UBT: Add a slightly better error message when a module that is invalid in the current configuration is referenced.

diff --git a/Engine/Source/Programs/UnrealBuildTool/Configuration/UEBuildModuleCPP.cs b/Engine/Source/Programs/UnrealBuildTool/Configuration/UEBuildModuleCPP.cs
--- a/Engine/Source/Programs/UnrealBuildTool/Configuration/UEBuildModuleCPP.cs
+++ b/Engine/Source/Programs/UnrealBuildTool/Configuration/UEBuildModuleCPP.cs
@@ -1022,7 +1022,7 @@ namespace UnrealBuildTool
 			{
 				// Adding this check here as otherwise the call to Binary.Config.IntermediateDirectory will give an 
 				// unhandled exception
-				throw new BuildException("UEBuildBinary not set up for module {0}", this.ToString());
+				throw new BuildException("Module {0} is required by this target, but is not compiled into any binary. Check any references to it are correct for this configuration, and whether it can be built.", this.ToString());
 			}
 
 			// Check if this is an engine module
