UBT: Fixing error about referencing game modules from engine modules (would have prevented current build failures in development streams due to HTTPChunkInstaller referencing a branched copy of OnlineSubsystem).

diff --git a/Engine/Source/Programs/UnrealBuildTool/Configuration/UEBuildModule.cs b/Engine/Source/Programs/UnrealBuildTool/Configuration/UEBuildModule.cs
--- a/Engine/Source/Programs/UnrealBuildTool/Configuration/UEBuildModule.cs
+++ b/Engine/Source/Programs/UnrealBuildTool/Configuration/UEBuildModule.cs
@@ -199,6 +199,28 @@ namespace UnrealBuildTool
 			WhitelistRestrictedFolders = new HashSet<DirectoryReference>(InRules.WhitelistRestrictedFolders.Select(x => DirectoryReference.Combine(ModuleDirectory, x)));
 		}
 
+		/// <summary>
+		/// Returns a list of this module's dependencies.
+		/// </summary>
+		/// <returns>An enumerable containing the dependencies of the module.</returns>
+		public HashSet<UEBuildModule> GetDependencies(bool bWithIncludePathModules, bool bWithDynamicallyLoadedModules)
+		{
+			HashSet<UEBuildModule> Modules = new HashSet<UEBuildModule>();
+			Modules.UnionWith(PublicDependencyModules);
+			Modules.UnionWith(PrivateDependencyModules);
+			if(bWithIncludePathModules)
+			{
+				Modules.UnionWith(PublicIncludePathModules);
+				Modules.UnionWith(PrivateIncludePathModules);
+			}
+			if(bWithDynamicallyLoadedModules)
+			{
+				Modules.UnionWith(DynamicallyLoadedModules);
+				Modules.UnionWith(PlatformSpecificDynamicallyLoadedModules);
+			}
+			return Modules;
+		}
+
 		/// <summary>
 		/// Returns a list of this module's immediate dependencies.
 		/// </summary>
diff --git a/Engine/Source/Programs/UnrealBuildTool/Configuration/UEBuildTarget.cs b/Engine/Source/Programs/UnrealBuildTool/Configuration/UEBuildTarget.cs
--- a/Engine/Source/Programs/UnrealBuildTool/Configuration/UEBuildTarget.cs
+++ b/Engine/Source/Programs/UnrealBuildTool/Configuration/UEBuildTarget.cs
@@ -2360,12 +2360,15 @@ namespace UnrealBuildTool
 				{
 					if(Module.Binary != null && Module.RulesFile.IsUnderDirectory(UnrealBuildTool.EngineDirectory))
 					{
-						foreach(UEBuildModule ReferencedModule in Module.GetDirectDependencyModules())
+						HashSet<UEBuildModule> ReferencedModules = Module.GetDependencies(bWithIncludePathModules: true, bWithDynamicallyLoadedModules: true);
+						foreach(UEBuildModule ReferencedModule in ReferencedModules)
 						{
 							// Hard-code specific exceptions until these are properly fixed up
-							if(ReferencedModule.RulesFile != null && !ReferencedModule.RulesFile.IsUnderDirectory(UnrealBuildTool.EngineDirectory) && !IsLegacyEngineToGameReference(Platform, Module.Name, ReferencedModule.Name))
+							if(ReferencedModule.RulesFile != null && !ReferencedModule.RulesFile.IsUnderDirectory(UnrealBuildTool.EngineDirectory))
 							{
-								throw new BuildException("Engine module '{0}' should not depend on game module '{1}'", Module.Name, ReferencedModule.Name);
+								string EngineModuleRelativePath = Module.RulesFile.MakeRelativeTo(UnrealBuildTool.EngineDirectory.ParentDirectory);
+								string ReferencedModuleRelativePath = ReferencedModule.RulesFile.IsUnderDirectory(ProjectFile.Directory)? ReferencedModule.RulesFile.MakeRelativeTo(ProjectFile.Directory.ParentDirectory) : ReferencedModule.RulesFile.FullName;
+								throw new BuildException("Engine module '{0}' should not depend on game module '{1}'", EngineModuleRelativePath, ReferencedModuleRelativePath);
 							}
 						}
 					}
@@ -2564,22 +2567,6 @@ namespace UnrealBuildTool
 			return ECompilationResult.Succeeded;
 		}
 
-		/// <summary>
-		/// Checks whether the reference from a game module
-		/// </summary>
-		/// <param name="Platform">The platform being built</param>
-		/// <param name="EngineModuleName">Name of the engine module</param>
-		/// <param name="GameModuleName">Name of the game module</param>
-		/// <returns>True if the reference is a known engine->game reference</returns>
-		static bool IsLegacyEngineToGameReference(UnrealTargetPlatform Platform, string EngineModuleName, string GameModuleName)
-		{
-			if(Platform == UnrealTargetPlatform.PS4 && EngineModuleName == "WebBrowser" && GameModuleName == "OnlineSubsystem")
-			{
-				return true;
-			}
-			return false;
-		}
-
 		/// <summary>
 		/// Export the definition of this target to a JSON file
 		/// </summary>
