HTML5 CIS - error hunt

#rnx

diff --git a/Engine/Source/Programs/UnrealBuildTool/HTML5/HTML5SDKInfo.cs b/Engine/Source/Programs/UnrealBuildTool/HTML5/HTML5SDKInfo.cs
--- a/Engine/Source/Programs/UnrealBuildTool/HTML5/HTML5SDKInfo.cs
+++ b/Engine/Source/Programs/UnrealBuildTool/HTML5/HTML5SDKInfo.cs
@@ -97,6 +97,7 @@ namespace UnrealBuildTool
 			string config = Environment.GetEnvironmentVariable("EM_CONFIG"); // This is either a string containing the config directly, or points to a file
 			if (config != null && File.Exists(config))
 			{
+				Log.TraceInformation( "NOTE[ReadEmscriptenConfigFile]: using EM_CONFIG" );
                 config = File.ReadAllText(config);
 			}
 			return config;
@@ -205,6 +206,7 @@ namespace UnrealBuildTool
 				// If user has configured a custom Emscripten toolchain, use that automatically.
                 if (Environment.GetEnvironmentVariable("EMSDK") != null && Environment.GetEnvironmentVariable("EM_CONFIG") != null && File.Exists(Environment.GetEnvironmentVariable("EM_CONFIG")))
                 {
+					Log.TraceInformation( "NOTE[DOT_EMSCRIPTEN]: using EM_CONFIG" );
                     return Environment.GetEnvironmentVariable("EM_CONFIG");
                 }
 
@@ -276,6 +278,7 @@ namespace UnrealBuildTool
 			// If user has configured a custom Emscripten toolchain, use that automatically.
 			if (Environment.GetEnvironmentVariable("EMSDK") != null && Environment.GetEnvironmentVariable("EM_CONFIG") != null && File.Exists(Environment.GetEnvironmentVariable("EM_CONFIG")))
 			{
+				Log.TraceInformation( "NOTE[SetUpEmscriptenConfigFile]: using EM_CONFIG" );
 				return Environment.GetEnvironmentVariable("EM_CONFIG");
 			}
 
@@ -340,7 +343,7 @@ namespace UnrealBuildTool
 							"# (EMSCRIPTEN_NATIVE_OPTIMIZER).*", "$1 = '" + optz + "'"),
 						"(BINARYEN_ROOT).*", "$1 = '" + byn + "'");
 				File.WriteAllText(DOT_EMSCRIPTEN, txt);
-				Log.TraceErrorOnce( txt );
+				Log.TraceInformation( txt );
 
 			// --------------------------------------------------
 			if (File.Exists(SaveDotEmscripten))
@@ -419,11 +422,11 @@ namespace UnrealBuildTool
 			// issue diagnostics about the relevant paths.
 			if (Environment.GetEnvironmentVariable("EMSDK") != null && !sdkInstalled)
 			{
-				Log.TraceError("EMSDK enviroment variable is set to point to " + Environment.GetEnvironmentVariable("EMSDK") + " but needed tools were not found:");
-				Log.TraceError("EMSCRIPTEN_ROOT: " + EMSCRIPTEN_ROOT + ", exists: " + Directory.Exists(EMSCRIPTEN_ROOT).ToString());
-				Log.TraceError("LLVM_ROOT: " + LLVM_ROOT + ", exists: " + Directory.Exists(LLVM_ROOT).ToString());
-				Log.TraceError("NODE_JS: " + NODE_JS + ", exists: " + File.Exists(NODE_JS).ToString());
-				Log.TraceError("PYTHON: " + PYTHON + ", exists: " + File.Exists(PYTHON).ToString());
+				Log.TraceInformation("EMSDK enviroment variable is set to point to " + Environment.GetEnvironmentVariable("EMSDK") + " but needed tools were not found:");
+				Log.TraceInformation("EMSCRIPTEN_ROOT: " + EMSCRIPTEN_ROOT + ", exists: " + Directory.Exists(EMSCRIPTEN_ROOT).ToString());
+				Log.TraceInformation("LLVM_ROOT: " + LLVM_ROOT + ", exists: " + Directory.Exists(LLVM_ROOT).ToString());
+				Log.TraceInformation("NODE_JS: " + NODE_JS + ", exists: " + File.Exists(NODE_JS).ToString());
+				Log.TraceInformation("PYTHON: " + PYTHON + ", exists: " + File.Exists(PYTHON).ToString());
 			}
 
 			return sdkInstalled;
