Merging CL 3408128 from Release-4.16

diff --git a/Engine/Source/Programs/UnrealBuildTool/IOS/UEDeployIOS.cs b/Engine/Source/Programs/UnrealBuildTool/IOS/UEDeployIOS.cs
--- a/Engine/Source/Programs/UnrealBuildTool/IOS/UEDeployIOS.cs
+++ b/Engine/Source/Programs/UnrealBuildTool/IOS/UEDeployIOS.cs
@@ -591,10 +591,22 @@ namespace UnrealBuildTool
 			List<string> ProjectArches = new List<string>();
 			ProjectArches.Add("None");
 
-			// get the receipt
-			string ReceiptFilename = TargetReceipt.GetDefaultPath(ProjectDirectory, ProjectName, UnrealTargetPlatform.IOS, Config, "");
-			string BundlePath = Path.Combine (ProjectDirectory, "Binaries", "IOS", "Payload", ProjectName + ".app");
-			string RelativeEnginePath = UnrealBuildTool.EngineDirectory.MakeRelativeTo(DirectoryReference.GetCurrentDirectory());
+            string ReceiptFilename;
+            string BundlePath;
+
+            // get the receipt
+            if (bIsUE4Game)
+            {
+                ReceiptFilename = TargetReceipt.GetDefaultPath(UnrealBuildTool.EngineDirectory.ToString(), "UE4Game", UnrealTargetPlatform.IOS, Config, "");
+                BundlePath = Path.Combine(UnrealBuildTool.EngineDirectory.ToString(), "Intermediate", "IOS-Deploy", "UE4Game", Config.ToString(), "Payload", "UE4Game.app");
+            }
+            else
+            {
+                ReceiptFilename = TargetReceipt.GetDefaultPath(ProjectDirectory, ProjectName, UnrealTargetPlatform.IOS, Config, "");
+                BundlePath = Path.Combine(ProjectDirectory, "Binaries", "IOS", "Payload", ProjectName + ".app");
+            }
+
+            string RelativeEnginePath = UnrealBuildTool.EngineDirectory.MakeRelativeTo(DirectoryReference.GetCurrentDirectory());
 
 			UPL = new UnrealPluginLanguage(ProjectFile, CollectPluginDataPaths(TargetReceipt.Read(ReceiptFilename)), ProjectArches, "", "", UnrealTargetPlatform.IOS);
 
