Copying latest UEBuildAndroid from Release-4.16 to Release-Staging-4.16 to resolve conflicts with //UE4/Main.

diff --git a/Engine/Source/Programs/UnrealBuildTool/Android/UEDeployAndroid.cs b/Engine/Source/Programs/UnrealBuildTool/Android/UEDeployAndroid.cs
--- a/Engine/Source/Programs/UnrealBuildTool/Android/UEDeployAndroid.cs
+++ b/Engine/Source/Programs/UnrealBuildTool/Android/UEDeployAndroid.cs
@@ -17,6 +17,15 @@ namespace UnrealBuildTool
 		// Minimum Android SDK that must be used for Java compiling
 		readonly int MinimumSDKLevel = 23;
 
+		// Reserved Java keywords not allowed in package names without modification
+		static private string[] JavaReservedKeywords = new string[] {
+			"abstract", "assert", "boolean", "break", "byte", "case", "catch", "char", "class", "const", "continue", "default", "do",
+			"double", "else", "enum", "extends", "final", "finally", "float", "for", "goto", "if", "implements", "import", "instanceof",
+			"int", "interface", "long", "native", "new", "package", "private", "protected", "public", "return", "short", "static",
+			"strictfp", "super", "switch", "sychronized", "this", "throw", "throws", "transient", "try", "void", "volatile", "while",
+			"false", "null", "true"
+		};
+
 		/// <summary>
 		/// Internal usage for GetApiLevel
 		/// </summary>
@@ -1363,15 +1372,106 @@ namespace UnrealBuildTool
 			}
 		}
 
+		private string CachedPackageName = null;
+
+		private bool IsLetter(char Input)
+		{
+			return (Input >= 'A' && Input <= 'Z') || (Input >= 'a' && Input <= 'z');
+		}
+
+		private bool IsDigit(char Input)
+		{
+			return (Input >= '0' && Input <= '9');
+		}
+
 		private string GetPackageName(string ProjectName)
 		{
-			ConfigHierarchy Ini = GetConfigCacheIni(ConfigHierarchyType.Engine);
-			string PackageName;
-			Ini.GetString("/Script/AndroidRuntimeSettings.AndroidRuntimeSettings", "PackageName", out PackageName);
-			// replace some variables
-			PackageName = PackageName.Replace("[PROJECT]", ProjectName);
-			PackageName = PackageName.Replace("-", "_");
-			return PackageName;
+			if (CachedPackageName == null)
+			{
+				ConfigHierarchy Ini = GetConfigCacheIni(ConfigHierarchyType.Engine);
+				string PackageName;
+				Ini.GetString("/Script/AndroidRuntimeSettings.AndroidRuntimeSettings", "PackageName", out PackageName);
+
+				if (PackageName.Contains("[PROJECT]"))
+				{
+					// project name must start with a letter
+					if (!IsLetter(ProjectName[0]))
+					{
+						throw new BuildException("Package name segments must all start with a letter. Please replace [PROJECT] with a valid name");
+					}
+
+					// hyphens not allowed so change them to underscores in project name
+					if (ProjectName.Contains("-"))
+					{
+						Trace.TraceWarning("Project name contained hyphens, converted to underscore");
+						ProjectName = ProjectName.Replace("-", "_");
+					}
+
+					// check for special characters
+					for (int Index = 0; Index < ProjectName.Length; Index++)
+					{
+						char c = ProjectName[Index];
+						if (c != '.' && c != '_' && !IsDigit(c) && !IsLetter(c))
+						{
+							throw new BuildException("Project name contains illegal characters (only letters, numbers, and underscore allowed); please replace [PROJECT] with a valid name");
+						}
+					}
+
+					PackageName = PackageName.Replace("[PROJECT]", ProjectName);
+				}
+
+				// verify minimum number of segments
+				string[] PackageParts = PackageName.Split('.');
+				int SectionCount = PackageParts.Length;
+				if (SectionCount < 2)
+				{
+					throw new BuildException("Package name must have at least 2 segments separated by periods (ex. com.projectname, not projectname); please change in Android Project Settings. Currently set to '" + PackageName + "'");
+				}
+
+				// hyphens not allowed
+				if (PackageName.Contains("-"))
+				{
+					throw new BuildException("Package names may not contain hyphens; please change in Android Project Settings. Currently set to '" + PackageName + "'");
+				}
+
+				// do not allow special characters
+				for (int Index = 0; Index < PackageName.Length; Index++)
+				{
+					char c = PackageName[Index];
+					if (c != '.' && c != '_' && !IsDigit(c) && !IsLetter(c))
+					{
+						throw new BuildException("Package name contains illegal characters (only letters, numbers, and underscore allowed); please change in Android Project Settings. Currently set to '" + PackageName + "'");
+					}
+				}
+
+				// validate each segment
+				for (int Index = 0; Index < SectionCount; Index++)
+				{
+					if (PackageParts[Index].Length < 1)
+					{
+						throw new BuildException("Package name segments must have at least one letter; please change in Android Project Settings. Currently set to '" + PackageName + "'");
+					}
+
+					if (!IsLetter(PackageParts[Index][0]))
+					{
+						throw new BuildException("Package name segments must start with a letter; please change in Android Project Settings. Currently set to '" + PackageName + "'");
+					}
+
+					// cannot use Java reserved keywords
+					foreach (string Keyword in JavaReservedKeywords)
+					{
+						if (PackageParts[Index] == Keyword)
+						{
+							throw new BuildException("Package name segments must not be a Java reserved keyword (" + Keyword + "); please change in Android Project Settings. Currently set to '" + PackageName + "'");
+						}
+					}
+				}
+
+				Console.WriteLine("Using package name: '{0}'", PackageName);
+				CachedPackageName = PackageName;
+			}
+
+			return CachedPackageName;
 		}
 
 		private string GetPublicKey()
@@ -1465,7 +1565,7 @@ namespace UnrealBuildTool
 
 			string InstallLocation;
 			Ini.GetString("/Script/AndroidRuntimeSettings.AndroidRuntimeSettings", "InstallLocation", out InstallLocation);
-			switch(InstallLocation.ToLower())
+			switch (InstallLocation.ToLower())
 			{
 				case "preferexternal":
 					InstallLocation = "preferExternal";
@@ -1529,7 +1629,7 @@ namespace UnrealBuildTool
 			}
 			else
 			{
-				switch(CookFlavor)
+				switch (CookFlavor)
 				{
 					case "_Multi":
 						//need to check ini to determine which are supported
@@ -1595,7 +1695,7 @@ namespace UnrealBuildTool
 					Text.AppendLine("\t             " + Line);
 				}
 			}
-            Text.AppendLine("\t             android:hardwareAccelerated=\"true\">");
+			Text.AppendLine("\t             android:hardwareAccelerated=\"true\">");
 			Text.AppendLine("\t             android:hasCode=\"true\">");
 			if (bShowLaunchImage)
 			{
@@ -1681,7 +1781,10 @@ namespace UnrealBuildTool
 			Text.AppendLine(string.Format("\t\t<meta-data android:name=\"com.epicgames.ue4.GameActivity.bHasOBBFiles\" android:value=\"{0}\"/>", bHasOBBFiles ? "true" : "false"));
 			Text.AppendLine(string.Format("\t\t<meta-data android:name=\"com.epicgames.ue4.GameActivity.BuildConfiguration\" android:value=\"{0}\"/>", Configuration));
 			Text.AppendLine(string.Format("\t\t<meta-data android:name=\"com.epicgames.ue4.GameActivity.bUseExternalFilesDir\" android:value=\"{0}\"/>", bUseExternalFilesDir ? "true" : "false"));
-			Text.AppendLine(string.Format("\t\t<meta-data android:name=\"com.epicgames.ue4.GameActivity.bDaydream\" android:value=\"{0}\"/>", bPackageForDaydream ? "true" : "false"));
+			if (bPackageForDaydream)
+			{
+				Text.AppendLine(string.Format("\t\t<meta-data android:name=\"com.epicgames.ue4.GameActivity.bDaydream\" android:value=\"true\"/>"));
+			}
 			Text.AppendLine("\t\t<meta-data android:name=\"com.google.android.gms.games.APP_ID\"");
 			Text.AppendLine("\t\t           android:value=\"@string/app_id\" />");
 			Text.AppendLine("\t\t<meta-data android:name=\"com.google.android.gms.version\"");
@@ -2022,51 +2125,51 @@ namespace UnrealBuildTool
 			string AntBatFilename = Path.Combine(AntBinPath, "ant.bat");
 			string AntOrigBatFilename = Path.Combine(AntBinPath, "ant.orig.bat");
 
-			if (!File.Exists(AntOrigBatFilename))
+			// check for an unused drive letter
+			string UnusedDriveLetter = "";
+			bool bFound = true;
+			DriveInfo[] AllDrives = DriveInfo.GetDrives();
+			for (char DriveLetter = 'Z'; DriveLetter >= 'A'; DriveLetter--)
 			{
-				// check for an unused drive letter
-				string UnusedDriveLetter = "";
-				bool bFound = true;
-				DriveInfo[] AllDrives = DriveInfo.GetDrives();
-				for (char DriveLetter = 'Z'; DriveLetter >= 'A'; DriveLetter--)
+				UnusedDriveLetter = Char.ToString(DriveLetter) + ":";
+				bFound = false;
+				for (int DriveIndex = AllDrives.Length-1; DriveIndex >= 0; DriveIndex--)
 				{
-					UnusedDriveLetter = Char.ToString(DriveLetter) + ":";
-					bFound = false;
-					foreach (DriveInfo drive in AllDrives)
-					{
-						if (drive.Name.ToUpper().StartsWith(UnusedDriveLetter))
-						{
-							bFound = true;
-							break;
-						}
-					}
-
-					if (!bFound)
+					if (AllDrives[DriveIndex].Name.ToUpper().StartsWith(UnusedDriveLetter))
 					{
+						bFound = true;
 						break;
 					}
 				}
-
-				if (bFound)
+				if (!bFound)
 				{
-					Log.TraceInformation("\nUnable to apply fixed ant.bat (all drive letters in use!)");
-					return;
+					break;
 				}
+			}
+
+			if (bFound)
+			{
+				Log.TraceInformation("\nUnable to apply fixed ant.bat (all drive letters in use!)");
+				return;
+			}
 
-				Log.TraceInformation("\nPatching ant.bat to work around commandline length limit (using unused drive letter {0})", UnusedDriveLetter);
+			Log.TraceInformation("\nPatching ant.bat to work around commandline length limit (using unused drive letter {0})", UnusedDriveLetter);
 
+			if (!File.Exists(AntOrigBatFilename))
+			{
 				// copy the existing ant.bat to ant.orig.bat
 				File.Copy(AntBatFilename, AntOrigBatFilename, true);
+			}
 
-				// make sure ant.bat isn't read-only
-				FileAttributes Attribs = File.GetAttributes(AntBatFilename);
-				if (Attribs.HasFlag(FileAttributes.ReadOnly))
-				{
-					File.SetAttributes(AntBatFilename, Attribs & ~FileAttributes.ReadOnly);
-				}
+			// make sure ant.bat isn't read-only
+			FileAttributes Attribs = File.GetAttributes(AntBatFilename);
+			if (Attribs.HasFlag(FileAttributes.ReadOnly))
+			{
+				File.SetAttributes(AntBatFilename, Attribs & ~FileAttributes.ReadOnly);
+			}
 
-				// generate new ant.bat with an unused drive letter for subst
-				string AntBatText =
+			// generate new ant.bat with an unused drive letter for subst
+			string AntBatText =
 					"@echo off\n" +
 					"set ANTPATH=%~dp0\n" +
 					"set ANT_CMD_LINE_ARGS =\n" +
@@ -2078,12 +2181,11 @@ namespace UnrealBuildTool
 					":doneStart\n" +
 					"subst " + UnusedDriveLetter + " \"%CD%\"\n" +
 					"pushd " + UnusedDriveLetter + "\n" +
-					"call %ANTPATH%\\ant.orig.bat %ANT_CMD_LINE_ARGS%\n" +
+					"call \"%ANTPATH%\\ant.orig.bat\" %ANT_CMD_LINE_ARGS%\n" +
 					"popd\n" +
 					"subst " + UnusedDriveLetter + " /d\n";
 
-				File.WriteAllText(AntBatFilename, AntBatText);
-			}
+			File.WriteAllText(AntBatFilename, AntBatText);
 		}
 
 		private void MakeApk(AndroidToolChain ToolChain, string ProjectName, string ProjectDirectory, string OutputPath, string EngineDirectory, bool bForDistribution, string CookFlavor, bool bMakeSeparateApks, bool bIncrementalPackage, bool bDisallowPackagingDataInApk, bool bDisallowExternalFilesDir)
