Patch ant.bat to handle long commandline issue on Windows
#ue4
#android

diff --git a/Engine/Source/Programs/UnrealBuildTool/Android/UEDeployAndroid.cs b/Engine/Source/Programs/UnrealBuildTool/Android/UEDeployAndroid.cs
--- a/Engine/Source/Programs/UnrealBuildTool/Android/UEDeployAndroid.cs
+++ b/Engine/Source/Programs/UnrealBuildTool/Android/UEDeployAndroid.cs
@@ -1954,10 +1954,88 @@ namespace UnrealBuildTool
             }
         }
 
+		private void PatchAntBatIfNeeded()
+		{
+			// only need to do this for Windows (other platforms are Mono so use that for check)
+			if (Utils.IsRunningOnMono)
+			{
+				return;
+			}
+
+			string AntBinPath = Environment.ExpandEnvironmentVariables("%ANT_HOME%/bin");
+			string AntBatFilename = Path.Combine(AntBinPath, "ant.bat");
+			string AntOrigBatFilename = Path.Combine(AntBinPath, "ant.orig.bat");
+
+			if (!File.Exists(AntOrigBatFilename))
+			{
+				// check for an unused drive letter
+				string UnusedDriveLetter = "";
+				bool bFound = true;
+				DriveInfo[] AllDrives = DriveInfo.GetDrives();
+				for (char DriveLetter = 'Z'; DriveLetter >= 'A'; DriveLetter--)
+				{
+					UnusedDriveLetter = Char.ToString(DriveLetter) + ":";
+					bFound = false;
+					foreach (DriveInfo drive in AllDrives)
+					{
+						if (drive.Name.ToUpper().StartsWith(UnusedDriveLetter))
+						{
+							bFound = true;
+							break;
+						}
+					}
+
+					if (!bFound)
+					{
+						break;
+					}
+				}
+
+				if (bFound)
+				{
+					Log.TraceInformation("\nUnable to apply fixed ant.bat (all drive letters in use!)");
+					return;
+				}
+
+				Log.TraceInformation("\nPatching ant.bat to work around commandline length limit (using unused drive letter {0})", UnusedDriveLetter);
+
+				// copy the existing ant.bat to ant.orig.bat
+				File.Copy(AntBatFilename, AntOrigBatFilename, true);
+
+				// make sure ant.bat isn't read-only
+				FileAttributes Attribs = File.GetAttributes(AntBatFilename);
+				if (Attribs.HasFlag(FileAttributes.ReadOnly))
+				{
+					File.SetAttributes(AntBatFilename, Attribs & ~FileAttributes.ReadOnly);
+				}
+
+				// generate new ant.bat with an unused drive letter for subst
+				string AntBatText =
+					"@echo off\n" +
+					"set ANTPATH=%~dp0\n" +
+					"set ANT_CMD_LINE_ARGS =\n" +
+					":setupArgs\n" +
+					"if \"\"%1\"\"==\"\"\"\" goto doneStart\n" +
+					"set ANT_CMD_LINE_ARGS=%ANT_CMD_LINE_ARGS% %1\n" +
+					"shift\n" +
+					"goto setupArgs\n\n" +
+					":doneStart\n" +
+					"subst " + UnusedDriveLetter + " \"%CD%\"\n" +
+					"pushd " + UnusedDriveLetter + "\n" +
+					"call %ANTPATH%\\ant.orig.bat %ANT_CMD_LINE_ARGS%\n" +
+					"popd\n" +
+					"subst " + UnusedDriveLetter + " /d\n";
+
+				File.WriteAllText(AntBatFilename, AntBatText);
+			}
+		}
+
 		private void MakeApk(AndroidToolChain ToolChain, string ProjectName, string ProjectDirectory, string OutputPath, string EngineDirectory, bool bForDistribution, string CookFlavor, bool bMakeSeparateApks, bool bIncrementalPackage, bool bDisallowPackagingDataInApk, bool bDisallowExternalFilesDir)
 		{
 			Log.TraceInformation("\n===={0}====PREPARING TO MAKE APK=================================================================", DateTime.Now.ToString());
 
+			PatchAntBatIfNeeded();
+			
 			// cache some tools paths
 			string NDKBuildPath = Environment.ExpandEnvironmentVariables("%NDKROOT%/ndk-build" + (Utils.IsRunningOnMono ? "" : ".cmd"));
 
