Fix UE4Game binding indirectly referenced precompiled modules to the incorrect binary, causing build errors.

diff --git a/Engine/Source/Programs/UnrealBuildTool/Configuration/UEBuildTarget.cs b/Engine/Source/Programs/UnrealBuildTool/Configuration/UEBuildTarget.cs
--- a/Engine/Source/Programs/UnrealBuildTool/Configuration/UEBuildTarget.cs
+++ b/Engine/Source/Programs/UnrealBuildTool/Configuration/UEBuildTarget.cs
@@ -1099,6 +1099,12 @@ namespace UnrealBuildTool
 		[NonSerialized]
 		public List<UEBuildPlugin> EnabledPlugins;
 
+		/// <summary>
+		/// All plugins which should be precompiled for this target
+		/// </summary>
+		[NonSerialized]
+		public List<UEBuildPlugin> PrecompilePlugins;
+
 		/// <summary>
 		/// Additional plugin filenames to include when building UnrealHeaderTool for the current target
 		/// </summary>
@@ -3403,9 +3409,11 @@ namespace UnrealBuildTool
 
 					// Also allow anything in the developer directory in non-shipping configurations (though we blacklist by default unless the PrecompileForTargets
 					// setting indicates that it's actually useful at runtime).
+					bool bAllowDeveloperModules = false;
 					if(Configuration != UnrealTargetConfiguration.Shipping)
 					{
 						Directories.Add(UnrealBuildTool.EngineSourceDeveloperDirectory);
+						bAllowDeveloperModules = true;
 					}
 
 					// Find all the modules that are not part of the standard set
@@ -3423,6 +3431,22 @@ namespace UnrealBuildTool
 						}
 					}
 
+					// Add all the plugins which aren't already being built
+					foreach(UEBuildPlugin PrecompilePlugin in PrecompilePlugins)
+					{
+						foreach (ModuleDescriptor ModuleDescriptor in PrecompilePlugin.Descriptor.Modules)
+						{
+							if (ModuleDescriptor.IsCompiledInConfiguration(Platform, TargetType, bAllowDeveloperModules && Rules.bBuildDeveloperTools, Rules.bBuildEditor, Rules.bBuildRequiresCookedData))
+							{
+								string RelativeFileName = RulesAssembly.GetModuleFileName(ModuleDescriptor.Name).MakeRelativeTo(UnrealBuildTool.EngineDirectory);
+								if (!ExcludeFolders.Any(x => RelativeFileName.Contains(x)) && !PrecompiledModules.Any(x => x.Name == ModuleDescriptor.Name))
+								{
+									FilteredModuleNames.Add(ModuleDescriptor.Name);
+								}
+							}
+						}
+					}
+
 					// Create rules for each remaining module, and check that it's set to be precompiled
 					foreach(string FilteredModuleName in FilteredModuleNames)
 					{
@@ -3753,7 +3777,7 @@ namespace UnrealBuildTool
 			EnabledPlugins = new List<UEBuildPlugin>(NameToInstance.Values);
 
 			// Try to add any the other plugins that are valid
-			if(Rules.bBuildAllPlugins || (bPrecompile && ProjectFile == null && TargetType != TargetType.Program))
+			if(Rules.bBuildAllPlugins)
 			{
 				foreach(PluginInfo Plugin in NameToInfo.Values)
 				{
@@ -3781,6 +3805,21 @@ namespace UnrealBuildTool
 
 			// Set the list of plugins that should be built
 			BuildPlugins = new List<UEBuildPlugin>(NameToInstance.Values);
+
+			// Create the list of plugins to precompile. This is separate to the list of binaries being built, because
+			// they and their dependencies won't be linked into the application binary by default on monolithic builds.
+			if (bPrecompile && ProjectFile == null && TargetType != TargetType.Program)
+			{
+				foreach(PluginInfo Plugin in NameToInfo.Values)
+				{
+					if (!NameToInstance.ContainsKey(Plugin.Name) && Plugin.Descriptor.Modules != null)
+					{
+						PluginReferenceDescriptor Reference = new PluginReferenceDescriptor(Plugin.Name, null, true);
+						AddPlugin(Reference, ExcludeFolders, NameToInstance, NameToInfo);
+					}
+				}
+				PrecompilePlugins = new List<UEBuildPlugin>(NameToInstance.Values.Except(BuildPlugins));
+			}
 		}
 
 		/// <summary>
