Copying //UE4/Dev-Enterprise to //UE4/Dev-Main (Source: //UE4/Dev-Enterprise @ 3438070)

==========================
MAJOR FEATURES + CHANGES
==========================

Change 3434228 on 2017/05/10 by [REDACTED_PERSON]

	Integrate UBT Task stream changes by JeanMichel.Dignard

	Unreal Build Tool
	- Added support for a new root folder /Enterprise with its own rules assembly, binaries and intermediate folders.
	- Supports enterprise specific modules, plugins, programs, samples and templates.
	- Enterprise plugins are automatically discovered by the PluginManager.

diff --git a/Engine/Source/Programs/UnrealBuildTool/Configuration/UEBuildBinaryCPP.cs b/Engine/Source/Programs/UnrealBuildTool/Configuration/UEBuildBinaryCPP.cs
--- a/Engine/Source/Programs/UnrealBuildTool/Configuration/UEBuildBinaryCPP.cs
+++ b/Engine/Source/Programs/UnrealBuildTool/Configuration/UEBuildBinaryCPP.cs
@@ -230,7 +230,7 @@ namespace UnrealBuildTool
 			List<UEBuildModule> GameModules = new List<UEBuildModule>();
 			foreach (UEBuildModule Module in Modules)
 			{
-				if (!Module.ModuleDirectory.IsUnderDirectory(UnrealBuildTool.EngineDirectory))
+				if (!UnrealBuildTool.IsUnderAnEngineDirectory(Module.ModuleDirectory))
 				{
 					GameModules.Add(Module);
 				}
diff --git a/Engine/Source/Programs/UnrealBuildTool/Configuration/UEBuildModuleCPP.cs b/Engine/Source/Programs/UnrealBuildTool/Configuration/UEBuildModuleCPP.cs
--- a/Engine/Source/Programs/UnrealBuildTool/Configuration/UEBuildModuleCPP.cs
+++ b/Engine/Source/Programs/UnrealBuildTool/Configuration/UEBuildModuleCPP.cs
@@ -376,7 +376,7 @@ namespace UnrealBuildTool
 
 			// Check to see if this is an Engine module (including program or plugin modules).  That is, the module is located under the "Engine" folder
 			bool IsPluginModule = (Target.ProjectFile != null && ModuleDirectory.IsUnderDirectory(DirectoryReference.Combine(Target.ProjectFile.Directory, "Plugins")));
-			bool IsGameModule = !IsPluginModule && !ModuleDirectory.IsUnderDirectory(UnrealBuildTool.EngineDirectory);
+			bool IsGameModule = !IsPluginModule && !UnrealBuildTool.IsUnderAnEngineDirectory(ModuleDirectory);
 
 			// Should we force a precompiled header to be generated for this module?  Usually, we only bother with a
 			// precompiled header if there are at least several source files in the module (after combining them for unity
@@ -1026,7 +1026,8 @@ namespace UnrealBuildTool
 		/// <returns>True if it is an engine module, false otherwise</returns>
 		public bool IsEngineModule()
 		{
-			return ModuleDirectory.IsUnderDirectory(UnrealBuildTool.EngineDirectory) && !ModuleDirectory.IsUnderDirectory(UnrealBuildTool.EngineSourceProgramsDirectory) && Name != "UE4Game";
+			return UnrealBuildTool.IsUnderAnEngineDirectory(ModuleDirectory) && !ModuleDirectory.IsUnderDirectory(UnrealBuildTool.EngineSourceProgramsDirectory) && Name != "UE4Game" &&
+				!ModuleDirectory.IsUnderDirectory(DirectoryReference.Combine(UnrealBuildTool.EnterpriseSourceDirectory, "Programs"));
 		}
 
 		/// <summary>
diff --git a/Engine/Source/Programs/UnrealBuildTool/Configuration/UEBuildTarget.cs b/Engine/Source/Programs/UnrealBuildTool/Configuration/UEBuildTarget.cs
--- a/Engine/Source/Programs/UnrealBuildTool/Configuration/UEBuildTarget.cs
+++ b/Engine/Source/Programs/UnrealBuildTool/Configuration/UEBuildTarget.cs
@@ -690,6 +690,17 @@ namespace UnrealBuildTool
 			else
 			{
 				RulesAssembly = RulesCompiler.CreateEngineRulesAssembly();
+
+				if (RulesAssembly.GetTargetFileName(Desc.TargetName) == null)
+				{
+					// Target isn't part of the engine assembly, try the enterprise assembly
+					RulesAssembly EnterpriseRulesAssembly = RulesCompiler.CreateEnterpriseRulesAssembly();
+
+					if (EnterpriseRulesAssembly != null)
+					{
+						RulesAssembly = EnterpriseRulesAssembly;
+					}
+				}
 			}
 			if (Desc.ForeignPlugins != null)
 			{
@@ -1366,7 +1377,14 @@ namespace UnrealBuildTool
 			}
 			else
 			{
-				ProjectDirectory = UnrealBuildTool.EngineDirectory;
+				if (TargetCsFilename.IsUnderDirectory(UnrealBuildTool.EnterpriseDirectory))
+				{
+					ProjectDirectory = UnrealBuildTool.EnterpriseDirectory;
+				}
+				else
+				{
+					ProjectDirectory = UnrealBuildTool.EngineDirectory;
+				}
 			}
 
 			// Build the project intermediate directory
@@ -4457,7 +4475,7 @@ namespace UnrealBuildTool
 					}
 				}
 
-				if (ModuleFileName.IsUnderDirectory(UnrealBuildTool.EngineDirectory))
+				if (UnrealBuildTool.IsUnderAnEngineDirectory(ModuleFileName.Directory))
 				{
 					if (RulesObject.Type == ModuleRules.ModuleType.External)
 					{
@@ -4472,7 +4490,14 @@ namespace UnrealBuildTool
 
 						if (!ModuleType.HasValue)
 						{
-							ModuleType = ExternalExecution.GetEngineModuleTypeBasedOnLocation(ModuleFileName);
+							if (ModuleFileName.IsUnderDirectory(UnrealBuildTool.EngineDirectory))
+							{
+								ModuleType = ExternalExecution.GetEngineModuleTypeBasedOnLocation(UnrealBuildTool.EngineSourceDirectory, ModuleFileName);
+							}
+							else if (ModuleFileName.IsUnderDirectory(UnrealBuildTool.EnterpriseSourceDirectory))
+							{
+								ModuleType = ExternalExecution.GetEngineModuleTypeBasedOnLocation(UnrealBuildTool.EnterpriseSourceDirectory, ModuleFileName);
+							}
 						}
 					}
 				}
@@ -4577,7 +4602,7 @@ namespace UnrealBuildTool
 					}
 
 					// So all we care about are the game module and/or plugins.
-					if (bDiscoverFiles && (!UnrealBuildTool.IsEngineInstalled() || !ModuleFileName.IsUnderDirectory(UnrealBuildTool.EngineDirectory)))
+					if (bDiscoverFiles && (!UnrealBuildTool.IsEngineInstalled() || !UnrealBuildTool.IsUnderAnEngineDirectory(ModuleFileName.Directory)))
 					{
 						List<FileReference> SourceFilePaths = new List<FileReference>();
 
@@ -4664,7 +4689,7 @@ namespace UnrealBuildTool
 		public void AddDefaultIncludePathsToModuleRules(FileReference ModuleFile, bool IsGameModule, PluginInfo Plugin, ModuleRules RulesObject)
 		{
 			// Get the base source directory for this module. This may be the project source directory, engine source directory, or plugin source directory.
-			if (!ModuleFile.IsUnderDirectory(UnrealBuildTool.EngineSourceDirectory))
+			if (!ModuleFile.IsUnderDirectory(UnrealBuildTool.EngineSourceDirectory) && !ModuleFile.IsUnderDirectory(UnrealBuildTool.EnterpriseSourceDirectory))
 			{
 				// Add the module source directory 
 				DirectoryReference BaseSourceDirectory;
diff --git a/Engine/Source/Programs/UnrealBuildTool/System/ExternalExecution.cs b/Engine/Source/Programs/UnrealBuildTool/System/ExternalExecution.cs
--- a/Engine/Source/Programs/UnrealBuildTool/System/ExternalExecution.cs
+++ b/Engine/Source/Programs/UnrealBuildTool/System/ExternalExecution.cs
@@ -389,29 +389,29 @@ namespace UnrealBuildTool
             return Type.GetValueOrDefault();
         }
 
-		public static UHTModuleType? GetEngineModuleTypeBasedOnLocation(FileReference ModuleFileName)
+		public static UHTModuleType? GetEngineModuleTypeBasedOnLocation(DirectoryReference SourceDirectory, FileReference ModuleFileName)
 		{
-			if (ModuleFileName.IsUnderDirectory(UnrealBuildTool.EngineSourceRuntimeDirectory))
+			if (ModuleFileName.IsUnderDirectory(DirectoryReference.Combine(SourceDirectory, "Runtime")))
 			{
 				return UHTModuleType.EngineRuntime;
 			}
 
-			if (ModuleFileName.IsUnderDirectory(UnrealBuildTool.EngineSourceDeveloperDirectory))
+			if (ModuleFileName.IsUnderDirectory(DirectoryReference.Combine(SourceDirectory, "Developer")))
 			{
 				return UHTModuleType.EngineDeveloper;
 			}
 
-			if (ModuleFileName.IsUnderDirectory(UnrealBuildTool.EngineSourceEditorDirectory))
+			if (ModuleFileName.IsUnderDirectory(DirectoryReference.Combine(SourceDirectory, "Editor")))
 			{
 				return UHTModuleType.EngineEditor;
 			}
-	
-			if (ModuleFileName.IsUnderDirectory(UnrealBuildTool.EngineSourceProgramsDirectory))
+
+			if (ModuleFileName.IsUnderDirectory(DirectoryReference.Combine(SourceDirectory, "Programs")))
 			{
 				return UHTModuleType.Program;
 			}
 
-			if (ModuleFileName.IsUnderDirectory(UnrealBuildTool.EngineSourceThirdPartyDirectory))
+			if (ModuleFileName.IsUnderDirectory(DirectoryReference.Combine(SourceDirectory, "ThirdParty")))
 			{
 				return UHTModuleType.EngineThirdParty;
 			}
@@ -727,7 +727,7 @@ namespace UnrealBuildTool
 			foreach (UHTModuleInfo Module in UObjectModules)
 			{
 				// If we're using a precompiled engine, skip skip checking timestamps for modules that are under the engine directory
-				if (bUsePrecompiled && Module.ModuleDirectory.IsUnderDirectory(UnrealBuildTool.EngineDirectory))
+				if (bUsePrecompiled && UnrealBuildTool.IsUnderAnEngineDirectory(Module.ModuleDirectory))
 				{
 					continue;
 				}
@@ -869,7 +869,7 @@ namespace UnrealBuildTool
 					if (GeneratedCodeDirectoryInfo.Exists)
 					{
 						// Don't write anything to the engine directory if we're running an installed build
-						if (bUsePrecompiled && Module.ModuleDirectory.IsUnderDirectory(UnrealBuildTool.EngineDirectory))
+						if (bUsePrecompiled && UnrealBuildTool.IsUnderAnEngineDirectory(Module.ModuleDirectory))
 						{
 							continue;
 						}
diff --git a/Engine/Source/Programs/UnrealBuildTool/System/LinkEnvironment.cs b/Engine/Source/Programs/UnrealBuildTool/System/LinkEnvironment.cs
--- a/Engine/Source/Programs/UnrealBuildTool/System/LinkEnvironment.cs
+++ b/Engine/Source/Programs/UnrealBuildTool/System/LinkEnvironment.cs
@@ -265,6 +265,12 @@ namespace UnrealBuildTool
 		/// </summary>
 		public List<FileItem> CommonResourceFiles = new List<FileItem>();
 
+		/// <summary>
+		/// Provides a Module Definition File (.def) to the linker to describe various attributes of a DLL.
+		/// Necessary when exporting functions by ordinal values instead of by name.
+		/// </summary>
+		public string ModuleDefinitionFile;
+
 		/// <summary>
 		/// Default constructor.
 		/// </summary>
@@ -322,6 +328,7 @@ namespace UnrealBuildTool
 			InputLibraries.AddRange(Other.InputLibraries);
 			DefaultResourceFiles.AddRange(Other.DefaultResourceFiles);
 			CommonResourceFiles.AddRange(Other.CommonResourceFiles);
-		}
+			ModuleDefinitionFile = Other.ModuleDefinitionFile;
+        }
 	}
 }
diff --git a/Engine/Source/Programs/UnrealBuildTool/System/ProjectFileGenerator.cs b/Engine/Source/Programs/UnrealBuildTool/System/ProjectFileGenerator.cs
--- a/Engine/Source/Programs/UnrealBuildTool/System/ProjectFileGenerator.cs
+++ b/Engine/Source/Programs/UnrealBuildTool/System/ProjectFileGenerator.cs
@@ -204,8 +204,11 @@ namespace UnrealBuildTool
 		/// Relative path to the directory where the master project file will be saved to
 		public static DirectoryReference MasterProjectPath = UnrealBuildTool.RootDirectory; // We'll save the master project to our "root" folder
 
-		/// Name of the UE4 engine project name that contains all of the engine code, config files and other files
-		static readonly string EngineProjectFileNameBase = "UE4";
+		/// Name of the UE4 engine project that contains all of the engine code, config files and other files
+		public static readonly string EngineProjectFileNameBase = "UE4";
+
+		/// Name of the UE4 enterprise project that contains all of the enterprise code, config files and other files
+		public static readonly string EnterpriseProjectFileNameBase = "Studio";
 
 		/// When ProjectsAreIntermediate is true, this is the directory to store generated project files
 		// @todo projectfiles: Ideally, projects for game modules/targets would be created in the game's Intermediate folder!
@@ -462,14 +465,15 @@ namespace UnrealBuildTool
 			List<FileReference> AllModuleFiles = DiscoverModules(AllGameProjects);
 
 			ProjectFile EngineProject = null;
+			ProjectFile EnterpriseProject = null;
 			Dictionary<DirectoryReference, ProjectFile> GameProjects = null;
 			Dictionary<DirectoryReference, ProjectFile> ModProjects = null;
-			Dictionary<string, ProjectFile> ProgramProjects = null;
-			HashSet<ProjectFile> TemplateGameProjects = null;
-			HashSet<ProjectFile> SampleGameProjects = null;
+			Dictionary<FileReference, ProjectFile> ProgramProjects = null;
+			Dictionary<DirectoryReference, ProjectFile> TemplateGameProjects = null;
+			Dictionary<DirectoryReference, ProjectFile> SampleGameProjects = null;
 			{
 				// Setup buildable projects for all targets
-				AddProjectsForAllTargets( AllGameProjects, out EngineProject, out GameProjects, out ModProjects, out ProgramProjects, out TemplateGameProjects, out SampleGameProjects );
+				AddProjectsForAllTargets( AllGameProjects, out EngineProject, out EnterpriseProject, out GameProjects, out ModProjects, out ProgramProjects, out TemplateGameProjects, out SampleGameProjects );
 
 				// Add all game projects and game config files
 				AddAllGameProjects(GameProjects, SupportedPlatformNames, RootFolder);
@@ -493,7 +497,7 @@ namespace UnrealBuildTool
 					// If we're still missing an engine project because we don't have any targets for it, make one up.
 					if( EngineProject == null )
 					{
-						FileReference ProjectFilePath = FileReference.Combine(IntermediateProjectFilesPath, "UE4" + ProjectFileExtension);
+						FileReference ProjectFilePath = FileReference.Combine(IntermediateProjectFilesPath, EngineProjectFileNameBase + ProjectFileExtension);
 
 						bool bAlreadyExisted;
 						EngineProject = FindOrAddProject(ProjectFilePath, true, out bAlreadyExisted);
@@ -561,6 +565,11 @@ namespace UnrealBuildTool
 						}
 					}
 
+					if (EnterpriseProject != null)
+					{
+						RootFolder.AddSubFolder(UnrealBuildTool.EnterpriseDirectory.GetDirectoryName()).ChildProjects.Add(EnterpriseProject);
+					}
+
 					foreach( ProjectFile CurModProject in ModProjects.Values )
 					{
 						RootFolder.AddSubFolder("Mods").ChildProjects.Add(CurModProject);
@@ -569,13 +578,30 @@ namespace UnrealBuildTool
 					foreach( ProjectFile CurGameProject in GameProjects.Values )
 					{
 						// Templates go under a different solution folder than games
-						if( TemplateGameProjects.Contains( CurGameProject ) )
+						DirectoryReference TemplateGameDirectory = TemplateGameProjects.FirstOrDefault( FolderAndProject => FolderAndProject.Value == CurGameProject ).Key;
+						DirectoryReference SampleGameDirectory = SampleGameProjects.FirstOrDefault( FolderAndProject => FolderAndProject.Value == CurGameProject ).Key;
+
+						if( TemplateGameDirectory != null )
 						{
-							RootFolder.AddSubFolder( "Templates" ).ChildProjects.Add( CurGameProject );
+							if( TemplateGameDirectory.IsUnderDirectory((UnrealBuildTool.EnterpriseDirectory)) )
+							{
+								RootFolder.AddSubFolder(UnrealBuildTool.EnterpriseDirectory.GetDirectoryName() + Path.DirectorySeparatorChar + "Templates").ChildProjects.Add(CurGameProject);
+							}
+							else
+							{
+								RootFolder.AddSubFolder("Templates").ChildProjects.Add(CurGameProject);
+							}
 						}
-						else if (SampleGameProjects.Contains( CurGameProject ) )
+						else if( SampleGameDirectory != null )
 						{
-							RootFolder.AddSubFolder("Samples").ChildProjects.Add(CurGameProject);
+							if( SampleGameDirectory.IsUnderDirectory((UnrealBuildTool.EnterpriseDirectory)) )
+                            {
+								RootFolder.AddSubFolder(UnrealBuildTool.EnterpriseDirectory.GetDirectoryName() + Path.DirectorySeparatorChar + "Samples").ChildProjects.Add(CurGameProject);
+							}
+							else
+							{
+								RootFolder.AddSubFolder("Samples").ChildProjects.Add(CurGameProject);
+							}
 						}
 						else
 						{
@@ -587,7 +613,7 @@ namespace UnrealBuildTool
 						if (NewProjectFiles != null)
 						{
 							DebugProjectFiles.AddRange(NewProjectFiles);
-					}
+						}
 
 					}
 
@@ -601,17 +627,24 @@ namespace UnrealBuildTool
 					}
 					
 
-					foreach( ProjectFile CurProgramProject in ProgramProjects.Values )
+					foreach( KeyValuePair<FileReference, ProjectFile> CurProgramProject in ProgramProjects )
 					{
-                        ProjectTarget Target = CurProgramProject.ProjectTargets.FirstOrDefault(t => !String.IsNullOrEmpty(t.TargetRules.SolutionDirectory));
+                        ProjectTarget Target = CurProgramProject.Value.ProjectTargets.FirstOrDefault(t => !String.IsNullOrEmpty(t.TargetRules.SolutionDirectory));
 
                         if (Target != null)
                         {
-                            RootFolder.AddSubFolder(Target.TargetRules.SolutionDirectory).ChildProjects.Add(CurProgramProject);
+                            RootFolder.AddSubFolder(Target.TargetRules.SolutionDirectory).ChildProjects.Add(CurProgramProject.Value);
                         }
                         else
                         {
-						    RootFolder.AddSubFolder( "Programs" ).ChildProjects.Add( CurProgramProject );
+							if (CurProgramProject.Key.IsUnderDirectory(UnrealBuildTool.EnterpriseDirectory))
+							{
+								RootFolder.AddSubFolder(UnrealBuildTool.EnterpriseDirectory.GetDirectoryName() + Path.DirectorySeparatorChar + "Programs").ChildProjects.Add(CurProgramProject.Value);
+							}
+							else
+							{
+								RootFolder.AddSubFolder( "Programs" ).ChildProjects.Add( CurProgramProject.Value );
+							}
                         }
 					}
 
@@ -668,6 +701,23 @@ namespace UnrealBuildTool
 						}
 					}
 
+					// Enterprise targets
+					if (EnterpriseProject != null)
+					{
+						foreach (ProjectTarget ProjectTarget in EnterpriseProject.ProjectTargets)
+						{
+							if (ProjectTarget.TargetFilePath != null)
+							{
+								// Only bother with the editor target.  We want to make sure that definitions are setup to be as inclusive as possible
+								// for good quality IntelliSense.  For example, we want WITH_EDITORONLY_DATA=1, so using the editor targets works well.
+								if (ProjectTarget.TargetRules.Type == TargetType.Editor)
+								{
+									IntelliSenseTargetFiles.Add(Tuple.Create(EnterpriseProject, ProjectTarget.TargetFilePath));
+								}
+							}
+						}
+					}
+
 					// Program targets
 					foreach( ProjectFile ProgramProject in ProgramProjects.Values )
 					{
@@ -1349,20 +1399,29 @@ namespace UnrealBuildTool
 		/// <summary>
 		/// Adds all of the config files for program targets to their project files
 		/// </summary>
-		private void AddEngineProgramConfigFiles( Dictionary<string, ProjectFile> ProgramProjects )
+		private void AddEngineProgramConfigFiles( Dictionary<FileReference, ProjectFile> ProgramProjects )
 		{
 			if( bIncludeConfigFiles )
 			{
-				foreach( KeyValuePair<string, ProjectFile> ProjectFolderAndFile in ProgramProjects )
+				foreach( KeyValuePair<FileReference, ProjectFile> FileAndProject in ProgramProjects )
 				{
-					string ProgramFolder = ProjectFolderAndFile.Key;
-					ProjectFile ProgramProjectFile = ProjectFolderAndFile.Value;
-
-					string ProgramName = ProgramFolder;
+					string ProgramName = FileAndProject.Key.GetFileNameWithoutAnyExtensions();
+					ProjectFile ProgramProjectFile = FileAndProject.Value;
 
 					// @todo projectfiles: The config folder for programs is kind of weird -- you end up going UP a few directories to get to it.  This stuff is not great.
 					// @todo projectfiles: Fragile assumption here about Programs always being under /Engine/Programs
-					DirectoryReference ProgramDirectory = DirectoryReference.Combine( UnrealBuildTool.EngineDirectory, "Programs", ProgramName );
+
+					DirectoryReference ProgramDirectory;
+
+                    if ( FileAndProject.Key.IsUnderDirectory(UnrealBuildTool.EnterpriseDirectory) )
+					{
+						ProgramDirectory = DirectoryReference.Combine( UnrealBuildTool.EnterpriseDirectory, "Programs", ProgramName );
+					}
+					else
+					{
+						ProgramDirectory = DirectoryReference.Combine( UnrealBuildTool.EngineDirectory, "Programs", ProgramName );
+					}
+
 					DirectoryReference ProgramConfigDirectory = DirectoryReference.Combine( ProgramDirectory, "Config" );
 					if( DirectoryReference.Exists(ProgramConfigDirectory) )
 					{
@@ -1535,10 +1594,8 @@ namespace UnrealBuildTool
 		/// <param name="ModProjects">All mod projects</param>
 		/// <param name="AllModuleFiles">List of *.Build.cs files for all engine programs and games</param>
 		/// <param name="bGatherThirdPartySource">True to gather source code from third party projects too</param>
-		protected void AddProjectsForAllModules( List<UProjectInfo> AllGames, Dictionary<string, ProjectFile> ProgramProjects, Dictionary<DirectoryReference, ProjectFile> ModProjects, List<FileReference> AllModuleFiles, bool bGatherThirdPartySource )
+		protected void AddProjectsForAllModules( List<UProjectInfo> AllGames, Dictionary<FileReference, ProjectFile> ProgramProjects, Dictionary<DirectoryReference, ProjectFile> ModProjects, List<FileReference> AllModuleFiles, bool bGatherThirdPartySource )
 		{
-			DirectoryReference EngineSourceThirdPartyDirectory = DirectoryReference.Combine(UnrealBuildTool.EngineSourceDirectory, "ThirdParty");
-
 			HashSet<ProjectFile> ProjectsWithPlugins = new HashSet<ProjectFile>();
 			foreach( FileReference CurModuleFile in AllModuleFiles )
 			{
@@ -1551,8 +1608,8 @@ namespace UnrealBuildTool
 				bool WantProjectFileForModule = true;
 
 				// We'll keep track of whether this is an "engine" or "external" module.  This is determined below while loading module rules.
-				bool IsEngineModule = CurModuleFile.IsUnderDirectory(UnrealBuildTool.EngineDirectory);
-				bool IsThirdPartyModule = CurModuleFile.IsUnderDirectory(EngineSourceThirdPartyDirectory);
+				bool IsEngineModule = UnrealBuildTool.IsUnderAnEngineDirectory(CurModuleFile.Directory);
+				bool IsThirdPartyModule = CurModuleFile.IsUnderDirectory(UnrealBuildTool.EngineSourceThirdPartyDirectory);
 
 				if( IsEngineModule && !bIncludeEngineSource )
 				{
@@ -1627,14 +1684,14 @@ namespace UnrealBuildTool
 			}
 		}
 
-		private ProjectFile FindProjectForModule(FileReference CurModuleFile, List<UProjectInfo> AllGames, Dictionary<string, ProjectFile> ProgramProjects, Dictionary<DirectoryReference, ProjectFile> ModProjects, out DirectoryReference BaseFolder)
+		private ProjectFile FindProjectForModule(FileReference CurModuleFile, List<UProjectInfo> AllGames, Dictionary<FileReference, ProjectFile> ProgramProjects, Dictionary<DirectoryReference, ProjectFile> ModProjects, out DirectoryReference BaseFolder)
 		{
 			string ProjectFileNameBase = null;
 
 			string PossibleProgramTargetName = CurModuleFile.GetFileNameWithoutAnyExtensions();
 
 			// @todo projectfiles: This works fine for now, but is pretty busted.  It assumes only one module per program and that it matches the program target file name. (see TTP 307091)
-			if( ProgramProjects != null && ProgramProjects.ContainsKey( PossibleProgramTargetName ) )	// @todo projectfiles: When building (in mem projects), ProgramProjects will be null so we are using the UE4 project instead
+			if( ProgramProjects != null && ProgramProjects.Any( ProgramProject => PossibleProgramTargetName.Equals( ProgramProject.Key.GetFileNameWithoutAnyExtensions() ) ) )	// @todo projectfiles: When building (in mem projects), ProgramProjects will be null so we are using the UE4 project instead
 			{
 				ProjectFileNameBase = PossibleProgramTargetName;
 				BaseFolder = CurModuleFile.Directory;
@@ -1644,6 +1701,12 @@ namespace UnrealBuildTool
 				ProjectFileNameBase = EngineProjectFileNameBase;
 				BaseFolder = UnrealBuildTool.EngineDirectory;
 			}
+			else if( CurModuleFile.IsUnderDirectory(UnrealBuildTool.EnterpriseSourceDirectory) ||
+				CurModuleFile.IsUnderDirectory(DirectoryReference.Combine(UnrealBuildTool.EnterpriseDirectory, "Plugins")) )
+			{
+				ProjectFileNameBase = EnterpriseProjectFileNameBase;
+				BaseFolder = UnrealBuildTool.EnterpriseDirectory;
+			}
 			else
 			{
 				// Check if it's a mod
@@ -1677,24 +1740,31 @@ namespace UnrealBuildTool
 		/// </summary>
 		/// <param name="AllGames">All game folders</param>
 		/// <param name="EngineProject">The engine project we created</param>
+		/// <param name="EnterpriseProject">The enterprise project we created</param>
 		/// <param name="GameProjects">Map of game folder name to all of the game projects we created</param>
 		/// <param name="ModProjects">Map of mod folder name to all the mod projects we created</param>
 		/// <param name="ProgramProjects">Map of program names to all of the program projects we created</param>
 		/// <param name="TemplateGameProjects">Set of template game projects we found.  These will also be in the GameProjects map</param>
 		/// <param name="SampleGameProjects">Set of sample game projects that were found</param>
-		private void AddProjectsForAllTargets( List<UProjectInfo> AllGames, out ProjectFile EngineProject, out Dictionary<DirectoryReference, ProjectFile> GameProjects, out Dictionary<DirectoryReference, ProjectFile> ModProjects, out Dictionary<string, ProjectFile> ProgramProjects, out HashSet<ProjectFile> TemplateGameProjects, out HashSet<ProjectFile> SampleGameProjects )
+		private void AddProjectsForAllTargets( List<UProjectInfo> AllGames, out ProjectFile EngineProject, out ProjectFile EnterpriseProject,
+			out Dictionary<DirectoryReference, ProjectFile> GameProjects, out Dictionary<DirectoryReference, ProjectFile> ModProjects,
+			out Dictionary<FileReference, ProjectFile> ProgramProjects, out Dictionary<DirectoryReference, ProjectFile> TemplateGameProjects, out Dictionary<DirectoryReference, ProjectFile> SampleGameProjects )
 		{
 			// As we're creating project files, we'll also keep track of whether we created an "engine" project and return that if we have one
 			EngineProject = null;
-			GameProjects = new Dictionary<DirectoryReference,ProjectFile>();
-			ProgramProjects = new Dictionary<string,ProjectFile>( StringComparer.InvariantCultureIgnoreCase );
-			TemplateGameProjects = new HashSet<ProjectFile>();
-			SampleGameProjects = new HashSet<ProjectFile>();
+			EnterpriseProject = null;
+			GameProjects = new Dictionary<DirectoryReference, ProjectFile>();
+			ProgramProjects = new Dictionary<FileReference, ProjectFile>();
+			TemplateGameProjects = new Dictionary<DirectoryReference, ProjectFile>();
+			SampleGameProjects = new Dictionary<DirectoryReference, ProjectFile>();
 
 			// Get some standard directories
 			DirectoryReference EngineSourceProgramsDirectory = DirectoryReference.Combine(UnrealBuildTool.EngineSourceDirectory, "Programs");
 			DirectoryReference TemplatesDirectory = DirectoryReference.Combine(UnrealBuildTool.EngineDirectory, "..", "Templates");
 			DirectoryReference SamplesDirectory = DirectoryReference.Combine(UnrealBuildTool.EngineDirectory, "..", "Samples");
+			DirectoryReference EnterpriseSourceProgramsDirectory = DirectoryReference.Combine(UnrealBuildTool.EnterpriseSourceDirectory, "Programs");
+			DirectoryReference EnterpriseTemplatesDirectory = DirectoryReference.Combine(UnrealBuildTool.EnterpriseDirectory, "Templates");
+			DirectoryReference EnterpriseSamplesDirectory = DirectoryReference.Combine(UnrealBuildTool.EnterpriseDirectory, "Samples");
 
 			// Find all of the target files.  This will filter out any modules or targets that don't
 			// belong to platforms we're generating project files for.
@@ -1705,6 +1775,7 @@ namespace UnrealBuildTool
 
 				// Check to see if this is an Engine target.  That is, the target is located under the "Engine" folder
 				bool IsEngineTarget = false;
+				bool IsEnterpriseTarget = false;
 				bool WantProjectFileForTarget = true;
 				if(TargetFilePath.IsUnderDirectory(UnrealBuildTool.EngineDirectory))
 				{
@@ -1720,6 +1791,20 @@ namespace UnrealBuildTool
 						WantProjectFileForTarget = bIncludeEngineSource;
 					}
 				}
+				else if (TargetFilePath.IsUnderDirectory(UnrealBuildTool.EnterpriseSourceDirectory))
+				{
+					// This is an enterprise target
+					IsEnterpriseTarget = true;
+
+					if(TargetFilePath.IsUnderDirectory(EnterpriseSourceProgramsDirectory))
+					{
+						WantProjectFileForTarget = IncludeEnginePrograms;
+					}
+					else if(TargetFilePath.IsUnderDirectory(UnrealBuildTool.EnterpriseSourceDirectory))
+					{
+						WantProjectFileForTarget = bIncludeEngineSource;
+					}
+				}
 
 				if (WantProjectFileForTarget)
 				{
@@ -1728,7 +1813,14 @@ namespace UnrealBuildTool
 					FileReference CheckProjectFile;
 					if(!UProjectInfo.TryGetProjectForTarget(TargetName, out CheckProjectFile))
 					{
-						RulesAssembly = RulesCompiler.CreateEngineRulesAssembly();
+						if(TargetFilePath.IsUnderDirectory(UnrealBuildTool.EnterpriseDirectory))
+						{
+							RulesAssembly = RulesCompiler.CreateEnterpriseRulesAssembly();
+						}
+						else
+						{
+							RulesAssembly = RulesCompiler.CreateEngineRulesAssembly();
+						}
 					}
 					else
 					{
@@ -1752,6 +1844,10 @@ namespace UnrealBuildTool
 					{
 						ProjectFileNameBase = EngineProjectFileNameBase;
 					}
+					else if (IsEnterpriseTarget)
+					{
+						ProjectFileNameBase = EnterpriseProjectFileNameBase;
+					}
 					else
 					{
 						// Figure out which game project this target belongs to
@@ -1791,13 +1887,13 @@ namespace UnrealBuildTool
 					}
 
 					// Check to see if this is a template target.  That is, the target is located under the "Templates" folder
-					bool IsTemplateTarget = TargetFilePath.IsUnderDirectory(TemplatesDirectory);
-					bool IsSampleTarget = TargetFilePath.IsUnderDirectory(SamplesDirectory);
+					bool IsTemplateTarget = TargetFilePath.IsUnderDirectory(TemplatesDirectory) || TargetFilePath.IsUnderDirectory(EnterpriseTemplatesDirectory);
+					bool IsSampleTarget = TargetFilePath.IsUnderDirectory(SamplesDirectory) || TargetFilePath.IsUnderDirectory(EnterpriseSamplesDirectory);
 
 					DirectoryReference BaseFolder = null;
 					if (IsProgramTarget)
 					{
-						ProgramProjects[TargetName] = ProjectFile;
+						ProgramProjects[TargetFilePath] = ProjectFile;
 						BaseFolder = TargetFilePath.Directory;
 					}
 					else if (IsEngineTarget)
@@ -1812,16 +1908,28 @@ namespace UnrealBuildTool
 							EngineProject.IsStubProject = true;
 						}
 					}
+					else if (IsEnterpriseTarget)
+					{
+						EnterpriseProject = ProjectFile;
+						BaseFolder = UnrealBuildTool.EnterpriseDirectory;
+						if (UnrealBuildTool.IsEngineInstalled())
+						{
+							// Allow enterprise projects to be created but not built for Installed Engine builds
+							EnterpriseProject.IsForeignProject = false;
+							EnterpriseProject.IsGeneratedProject = true;
+							EnterpriseProject.IsStubProject = true;
+						}
+					}
 					else
 					{
 						GameProjects[GameFolder] = ProjectFile;
 						if (IsTemplateTarget)
 						{
-							TemplateGameProjects.Add(ProjectFile);
+							TemplateGameProjects[GameFolder] = ProjectFile;
 						}
 						else if (IsSampleTarget)
 						{
-							SampleGameProjects.Add(ProjectFile);
+							SampleGameProjects[GameFolder] = ProjectFile;
 						}
 						BaseFolder = GameFolder;
 
diff --git a/Engine/Source/Programs/UnrealBuildTool/System/RulesCompiler.cs b/Engine/Source/Programs/UnrealBuildTool/System/RulesCompiler.cs
--- a/Engine/Source/Programs/UnrealBuildTool/System/RulesCompiler.cs
+++ b/Engine/Source/Programs/UnrealBuildTool/System/RulesCompiler.cs
@@ -67,6 +67,7 @@ namespace UnrealBuildTool
 			if (bIncludeEngine)
 			{
 				Folders.Add(UnrealBuildTool.EngineSourceDirectory);
+				Folders.Add(UnrealBuildTool.EnterpriseSourceDirectory);
 			}
 
 			// @todo plugin: Disallow modules from including plugin modules as dependency modules? (except when the module is part of that plugin)
@@ -76,6 +77,7 @@ namespace UnrealBuildTool
 			if (bIncludeEngine)
 			{
 				RootFolders.Add(UnrealBuildTool.EngineDirectory);
+				RootFolders.Add(UnrealBuildTool.EnterpriseDirectory);
 			}
 			if (GameFolders != null)
 			{
@@ -227,6 +229,11 @@ namespace UnrealBuildTool
 		/// </summary>
 		private static RulesAssembly EngineRulesAssembly;
 
+		/// <summary>
+		/// The cached rules assembly for enterprise modules and targets.
+		/// </summary>
+		private static RulesAssembly EnterpriseRulesAssembly;
+
 		/// Map of assembly names we've already compiled and loaded to their Assembly and list of game folders.  This is used to prevent
 		/// trying to recompile the same assembly when ping-ponging between different types of targets
 		private static Dictionary<FileReference, RulesAssembly> LoadedAssemblyMap = new Dictionary<FileReference, RulesAssembly>();
@@ -239,22 +246,49 @@ namespace UnrealBuildTool
 		{
 			if (EngineRulesAssembly == null)
 			{
-				// Find all the rules files
-				List<FileReference> ModuleFiles = new List<FileReference>(FindAllRulesFiles(UnrealBuildTool.EngineSourceDirectory, RulesFileType.Module));
-				List<FileReference> TargetFiles = new List<FileReference>(FindAllRulesFiles(UnrealBuildTool.EngineSourceDirectory, RulesFileType.Target));
-
-				// Add all the plugin modules too
-				IReadOnlyList<PluginInfo> EnginePlugins = Plugins.ReadEnginePlugins(UnrealBuildTool.EngineDirectory);
-				Dictionary<FileReference, PluginInfo> ModuleFileToPluginInfo = new Dictionary<FileReference, PluginInfo>();
-				FindModuleRulesForPlugins(EnginePlugins, ModuleFiles, ModuleFileToPluginInfo);
-
-				// Create a path to the assembly that we'll either load or compile
-				FileReference AssemblyFileName = FileReference.Combine(UnrealBuildTool.EngineDirectory, "Intermediate", "Build", "BuildRules", "UE4Rules.dll");
-				EngineRulesAssembly = new RulesAssembly(EnginePlugins, ModuleFiles, TargetFiles, ModuleFileToPluginInfo, AssemblyFileName, null);
+				EngineRulesAssembly = CreateRulesAssembly(UnrealBuildTool.EngineDirectory, ProjectFileGenerator.EngineProjectFileNameBase, null);
 			}
 			return EngineRulesAssembly;
 		}
 
+		/// <summary>
+		/// Creates the enterprise rules assembly
+		/// </summary>
+		/// <returns>New rules assembly. Returns null if the enterprise directory is unavailable.</returns>
+		public static RulesAssembly CreateEnterpriseRulesAssembly()
+		{
+			if (EnterpriseRulesAssembly == null && DirectoryReference.Exists(UnrealBuildTool.EnterpriseDirectory))
+			{
+				EnterpriseRulesAssembly = CreateRulesAssembly(UnrealBuildTool.EnterpriseDirectory, ProjectFileGenerator.EnterpriseProjectFileNameBase, CreateEngineRulesAssembly());
+			}
+			return EnterpriseRulesAssembly;
+		}
+
+		/// <summary>
+		/// Creates a rules assembly
+		/// </summary>
+		/// <param name="RootDirectory">The root directory to create rules for</param>
+		/// <param name="AssemblyPrefix">A prefix for the assembly file name</param>
+		/// <param name="Parent">The parent rules assembly</param>
+		/// <returns>New rules assembly</returns>
+		private static RulesAssembly CreateRulesAssembly(DirectoryReference RootDirectory, string AssemblyPrefix, RulesAssembly Parent)
+		{
+			DirectoryReference SourceDirectory = DirectoryReference.Combine(RootDirectory, "Source");
+
+			// Find all the rules files
+			List<FileReference> ModuleFiles = new List<FileReference>(FindAllRulesFiles(SourceDirectory, RulesFileType.Module));
+			List<FileReference> TargetFiles = new List<FileReference>(FindAllRulesFiles(SourceDirectory, RulesFileType.Target));
+
+			// Add all the plugin modules too
+			IReadOnlyList<PluginInfo> PluginInfos = Plugins.ReadEnginePlugins(RootDirectory);
+			Dictionary<FileReference, PluginInfo> ModuleFileToPluginInfo = new Dictionary<FileReference, PluginInfo>();
+			FindModuleRulesForPlugins(PluginInfos, ModuleFiles, ModuleFileToPluginInfo);
+
+			// Create a path to the assembly that we'll either load or compile
+			FileReference AssemblyFileName = FileReference.Combine(RootDirectory, "Intermediate", "Build", "BuildRules", AssemblyPrefix + "Rules.dll");
+			return new RulesAssembly(PluginInfos, ModuleFiles, TargetFiles, ModuleFileToPluginInfo, AssemblyFileName, Parent);
+		}
+
 		/// <summary>
 		/// Creates a rules assembly with the given parameters.
 		/// </summary>
@@ -267,7 +301,17 @@ namespace UnrealBuildTool
 			if (!LoadedAssemblyMap.TryGetValue(ProjectFileName, out ProjectRulesAssembly))
 			{
 				// Create the engine rules assembly
-				RulesAssembly Parent = CreateEngineRulesAssembly();
+				RulesAssembly Parent;
+
+				// When the enterprise folder is available, use it as our parent assembly. It's also a child of the engine assembly.
+				if (DirectoryReference.Exists(UnrealBuildTool.EnterpriseDirectory))
+				{
+					Parent = CreateEnterpriseRulesAssembly();
+				}
+				else
+				{
+					Parent = CreateEngineRulesAssembly();
+				}
 
 				// Find all the rules under the project source directory
 				DirectoryReference ProjectDirectory = ProjectFileName.Directory;
@@ -364,6 +408,11 @@ namespace UnrealBuildTool
 			{
 				return FileName.FullName;
 			}
+			else if (EnterpriseRulesAssembly != null && EnterpriseRulesAssembly.TryGetFileNameFromType(ExistingType, out FileName))
+			{
+				return FileName.FullName;
+			}
+
 			foreach (RulesAssembly RulesAssembly in LoadedAssemblyMap.Values)
 			{
 				if (RulesAssembly.TryGetFileNameFromType(ExistingType, out FileName))
diff --git a/Engine/Source/Programs/UnrealBuildTool/System/UnrealBuildTool.cs b/Engine/Source/Programs/UnrealBuildTool/System/UnrealBuildTool.cs
--- a/Engine/Source/Programs/UnrealBuildTool/System/UnrealBuildTool.cs
+++ b/Engine/Source/Programs/UnrealBuildTool/System/UnrealBuildTool.cs
@@ -76,6 +76,16 @@ namespace UnrealBuildTool
 		/// </summary>
 		public static readonly DirectoryReference EngineSourceThirdPartyDirectory = DirectoryReference.Combine(EngineSourceDirectory, "ThirdParty");
 
+		/// <summary>
+		/// The full name of the Enterprise directory
+		/// </summary>
+		public static readonly DirectoryReference EnterpriseDirectory = DirectoryReference.Combine(RootDirectory, "Enterprise");
+
+		/// <summary>
+		/// The full name of the Enterprise/Source directory
+		/// </summary>
+		public static readonly DirectoryReference EnterpriseSourceDirectory = DirectoryReference.Combine(EnterpriseDirectory, "Source");
+
 		/// <summary>
 		/// The Remote Ini directory.  This should always be valid when compiling using a remote server.
 		/// </summary>
@@ -240,6 +250,18 @@ namespace UnrealBuildTool
 			return InstalledPlatformInfo.Current.IsValidPlatform(InPlatform, EProjectType.Code);
 		}
 
+		/// <summary>
+		/// Determines whether a directory is part of the engine
+		/// </summary>
+		/// <param name="InDirectory"></param>
+		/// <returns>true if the directory is under of the engine directories, false if not</returns>
+		static public bool IsUnderAnEngineDirectory(DirectoryReference InDirectory)
+		{
+			// Enterprise modules are considered as engine modules
+			return InDirectory.IsUnderDirectory( UnrealBuildTool.EngineDirectory ) || InDirectory.IsUnderDirectory( UnrealBuildTool.EnterpriseSourceDirectory ) ||
+				InDirectory.IsUnderDirectory( DirectoryReference.Combine( UnrealBuildTool.EnterpriseDirectory, "Plugins" ) );
+		}
+
 		public static void RegisterAllUBTClasses(SDKOutputLevel OutputLevel, bool bValidatingPlatforms)
 		{
 			// Find and register all tool chains and build platforms that are present
@@ -465,7 +487,7 @@ namespace UnrealBuildTool
 			if (!bIsEngineInstalled.HasValue)
 			{
 				bIsEngineInstalled = FileReference.Exists(FileReference.Combine(RootDirectory, "Engine", "Build", "InstalledBuild.txt"));
-		}
+			}
 
 			DateTime StartTime = DateTime.UtcNow;
 
@@ -478,10 +500,10 @@ namespace UnrealBuildTool
 				bLogSourcesToConsole: false,
 				bColorConsoleOutput: true,
 				TraceListeners: new[] 
-                {
-                    new ConsoleTraceListener(),
-                    !string.IsNullOrEmpty(BuildConfiguration.LogFilename) ? new TextWriterTraceListener(new StreamWriter(new FileStream(BuildConfiguration.LogFilename, FileMode.Create, FileAccess.ReadWrite, FileShare.Read)) { AutoFlush = true }) : null,
-                });
+				{
+					new ConsoleTraceListener(),
+					!string.IsNullOrEmpty(BuildConfiguration.LogFilename) ? new TextWriterTraceListener(new StreamWriter(new FileStream(BuildConfiguration.LogFilename, FileMode.Create, FileAccess.ReadWrite, FileShare.Read)) { AutoFlush = true }) : null,
+				});
 
 			// Parse rocket-specific arguments.
 			FileReference ProjectFile = null;
@@ -492,7 +514,7 @@ namespace UnrealBuildTool
 				{
 					// This is to allow relative paths for the project file
 					Log.TraceVerbose("UBT Running for Rocket: " + ProjectFile);
-                    break;
+					break;
 				}
 			}
 
diff --git a/Engine/Source/Programs/UnrealBuildTool/Windows/UEBuildWindows.cs b/Engine/Source/Programs/UnrealBuildTool/Windows/UEBuildWindows.cs
--- a/Engine/Source/Programs/UnrealBuildTool/Windows/UEBuildWindows.cs
+++ b/Engine/Source/Programs/UnrealBuildTool/Windows/UEBuildWindows.cs
@@ -102,6 +102,12 @@ namespace UnrealBuildTool
 		[CommandLine("-StaticAnalyzer")]
 		public WindowsStaticAnalyzer StaticAnalyzer = WindowsStaticAnalyzer.None;
 
+		/// <summary>
+		/// Provides a Module Definition File (.def) to the linker to describe various attributes of a DLL.
+		/// Necessary when exporting functions by ordinal values instead of by name.
+		/// </summary>
+		public string ModuleDefinitionFile;
+
 		/// VS2015 updated some of the CRT definitions but not all of the Windows SDK has been updated to match.
 		/// Microsoft provides legacy_stdio_definitions library to enable building with VS2015 until they fix everything up.
 		public bool bNeedsLegacyStdioDefinitionsLib
@@ -185,6 +191,11 @@ namespace UnrealBuildTool
 			get { return Inner.StaticAnalyzer; }
 		}
 
+		public string ModuleDefinitionFile
+		{
+			get { return Inner.ModuleDefinitionFile; }
+		}
+
 		public bool bNeedsLegacyStdioDefinitionsLib
 		{
 			get { return Inner.bNeedsLegacyStdioDefinitionsLib; }
@@ -694,8 +705,8 @@ namespace UnrealBuildTool
 					Rules.Definitions.Add("D3D12_PROFILING_ENABLED=0");
 				}
 
-                // To enable platform specific D3D12 RHI Types
-                Rules.PrivateIncludePaths.Add("Runtime/Windows/D3D12RHI/Private/Windows");
+				// To enable platform specific D3D12 RHI Types
+				Rules.PrivateIncludePaths.Add("Runtime/Windows/D3D12RHI/Private/Windows");
 			}
 
 			// Delay-load D3D12 so we can use the latest features and still run on downlevel versions of the OS
@@ -878,6 +889,8 @@ namespace UnrealBuildTool
 			{
 				LinkEnvironment.DefaultStackSizeCommit = IniDefaultStackSizeCommit;
 			}
+
+			LinkEnvironment.ModuleDefinitionFile = Target.WindowsPlatform.ModuleDefinitionFile;
 		}
 
 		/// <summary>
diff --git a/Engine/Source/Programs/UnrealBuildTool/Windows/VCToolChain.cs b/Engine/Source/Programs/UnrealBuildTool/Windows/VCToolChain.cs
--- a/Engine/Source/Programs/UnrealBuildTool/Windows/VCToolChain.cs
+++ b/Engine/Source/Programs/UnrealBuildTool/Windows/VCToolChain.cs
@@ -1439,7 +1439,6 @@ namespace UnrealBuildTool
 			}
 
 
-			// Add delay loaded DLLs.
 			if (!bIsBuildingLibrary)
 			{
 				// Delay-load these DLLs.
@@ -1447,6 +1446,12 @@ namespace UnrealBuildTool
 				{
 					Arguments.Add(String.Format("/DELAYLOAD:\"{0}\"", DelayLoadDLL));
 				}
+
+				// Pass the module definition file to the linker if we have one
+				if (LinkEnvironment.ModuleDefinitionFile != null && LinkEnvironment.ModuleDefinitionFile.Length > 0)
+				{
+					Arguments.Add(String.Format("/DEF:\"{0}\"", LinkEnvironment.ModuleDefinitionFile));
+				}
 			}
 
 			// Set up the library paths for linking this binary
diff --git a/Engine/Source/Runtime/Core/Private/Misc/Paths.cpp b/Engine/Source/Runtime/Core/Private/Misc/Paths.cpp
--- a/Engine/Source/Runtime/Core/Private/Misc/Paths.cpp
+++ b/Engine/Source/Runtime/Core/Private/Misc/Paths.cpp
@@ -150,6 +150,11 @@ FString FPaths::EnginePluginsDir()
 	return FPaths::EngineDir() + TEXT("Plugins/");
 }
 
+FString FPaths::EnterprisePluginsDir()
+{
+	return FPaths::RootDir() + TEXT("Enterprise/Plugins/");
+}
+
 FString FPaths::RootDir()
 {
 	return FString(FPlatformMisc::RootDir());
diff --git a/Engine/Source/Runtime/Core/Public/Misc/Paths.h b/Engine/Source/Runtime/Core/Public/Misc/Paths.h
--- a/Engine/Source/Runtime/Core/Public/Misc/Paths.h
+++ b/Engine/Source/Runtime/Core/Public/Misc/Paths.h
@@ -82,6 +82,13 @@ public:
 	 */
 	static FString EnginePluginsDir();
 
+	/**
+	* Returns the enterprise plugins directory
+	*
+	* @return Plugins directory.
+	*/
+	static FString EnterprisePluginsDir();
+
 	/**
 	 * Returns the root directory of the engine directory tree
 	 *
diff --git a/Engine/Source/Runtime/Projects/Private/PluginManager.cpp b/Engine/Source/Runtime/Projects/Private/PluginManager.cpp
--- a/Engine/Source/Runtime/Projects/Private/PluginManager.cpp
+++ b/Engine/Source/Runtime/Projects/Private/PluginManager.cpp
@@ -214,6 +214,7 @@ void FPluginManager::ReadAllPlugins(TMap<FString, TSharedRef<FPlugin>>& Plugins,
 #if (WITH_ENGINE && !IS_PROGRAM) || WITH_PLUGIN_SUPPORT
 	// Find "built-in" plugins.  That is, plugins situated right within the Engine directory.
 	ReadPluginsInDirectory(FPaths::EnginePluginsDir(), EPluginLoadedFrom::Engine, Plugins);
+	ReadPluginsInDirectory(FPaths::EnterprisePluginsDir(), EPluginLoadedFrom::Engine, Plugins);
 
 	// Find plugins in the game project directory (<MyGameProject>/Plugins). If there are any engine plugins matching the name of a game plugin,
 	// assume that the game plugin version is preferred.
