Fix UE4Editor Linux cross-build failures.

diff --git a/Engine/Source/Programs/UnrealBuildTool/Linux/LinuxToolChain.cs b/Engine/Source/Programs/UnrealBuildTool/Linux/LinuxToolChain.cs
--- a/Engine/Source/Programs/UnrealBuildTool/Linux/LinuxToolChain.cs
+++ b/Engine/Source/Programs/UnrealBuildTool/Linux/LinuxToolChain.cs
@@ -1368,12 +1368,25 @@ namespace UnrealBuildTool
 					RelinkWriter.Write("@echo off\n");
 					RelinkWriter.Write("rem Automatically generated by UnrealBuildTool\n");
 					RelinkWriter.Write("rem *DO NOT EDIT*\n\n");
+					RelinkWriter.Write("set Retries=0\n");
+					RelinkWriter.Write(":relinkloop\n");
+					RelinkWriter.Write("if %Retries% GEQ 10 goto failedtorelink\n");
 					RelinkWriter.Write(RelinkInvocation + "\n");
-					RelinkWriter.Write("copy /B \"{0}\" \"{1}.temp\"\n", RelinkedFileForwardSlashes, OutputFile.AbsolutePath);
-					RelinkWriter.Write(":moveloop\n");
-					RelinkWriter.Write("move /Y \"{0}.temp\" \"{1}\"\n", OutputFile.AbsolutePath, OutputFile.AbsolutePath);
-					RelinkWriter.Write("if errorlevel 1 goto moveloop\n");
-					RelinkWriter.Write(string.Format("echo \"Dummy\" >> \"{0}\" && copy /b \"{0}\" +,,", RelinkActionDummyProductRef.FullName));
+					RelinkWriter.Write("if %errorlevel% neq 0 goto sleepandretry\n");
+					RelinkWriter.Write("copy /B \"{0}\" \"{1}.temp\" >NUL 2>NUL\n", RelinkedFileForwardSlashes, OutputFile.AbsolutePath);
+					RelinkWriter.Write("if %errorlevel% neq 0 goto sleepandretry\n");
+					RelinkWriter.Write("move /Y \"{0}.temp\" \"{1}\" >NUL 2>NUL\n", OutputFile.AbsolutePath, OutputFile.AbsolutePath);
+					RelinkWriter.Write("if %errorlevel% neq 0 goto sleepandretry\n");
+					RelinkWriter.Write(string.Format("echo \"Dummy\" >> \"{0}\" && copy /b \"{0}\" +,,\n", RelinkActionDummyProductRef.FullName));
+					RelinkWriter.Write("echo Relinked {0} successfully after %Retries% retries\n", OutputFile.AbsolutePath);
+					RelinkWriter.Write("exit 0\n");
+					RelinkWriter.Write(":sleepandretry\n");
+					RelinkWriter.Write("ping 127.0.0.1 -n 1 -w 5000 >NUL 2>NUL\n");		// timeout complains about lack of redirection
+					RelinkWriter.Write("set /a Retries+=1\n");
+					RelinkWriter.Write("goto relinkloop\n");
+					RelinkWriter.Write(":failedtorelink\n");
+					RelinkWriter.Write("echo Failed to relink {0} after %Retries% retries\n", OutputFile.AbsolutePath);
+					RelinkWriter.Write("exit 1\n");
 				}
 				else
 				{
