Prevent plugins with bEnabledByDefault being included as runtime dependencies in the receipt if they do not have any content or supported modules on the current platform.

diff --git a/Engine/Source/Programs/UnrealBuildTool/Configuration/UEBuildTarget.cs b/Engine/Source/Programs/UnrealBuildTool/Configuration/UEBuildTarget.cs
--- a/Engine/Source/Programs/UnrealBuildTool/Configuration/UEBuildTarget.cs
+++ b/Engine/Source/Programs/UnrealBuildTool/Configuration/UEBuildTarget.cs
@@ -3761,9 +3761,12 @@ namespace UnrealBuildTool
 				{
 					if(Plugin.Descriptor.bEnabledByDefault && !NameToReference.ContainsKey(Plugin.Name))
 					{
-						PluginReferenceDescriptor PluginReference = new PluginReferenceDescriptor(Plugin.Name, null, true);
-						PluginReference.bOptional = true;
-						NameToReference[Plugin.Name] = PluginReference;
+						if (Plugin.Descriptor.bCanContainContent || Plugin.Descriptor.Modules.Any(x => x.IsCompiledInConfiguration(Platform, TargetType, Rules.bBuildDeveloperTools, Rules.bBuildEditor, Rules.bBuildRequiresCookedData)))
+						{
+							PluginReferenceDescriptor PluginReference = new PluginReferenceDescriptor(Plugin.Name, null, true);
+							PluginReference.bOptional = true;
+							NameToReference[Plugin.Name] = PluginReference;
+						}
 					}
 				}
 			}
