Fix Linux editor cross-build.

- Switch cross-build to old way of dealing with circular dependencies.
  The new way apparently has a race hazard that is only exposed during a cross-build.

(Merging 3465361 from Release-Staging-4.16 to UE4 Main)

diff --git a/Engine/Source/Programs/UnrealBuildTool/Linux/LinuxToolChain.cs b/Engine/Source/Programs/UnrealBuildTool/Linux/LinuxToolChain.cs
--- a/Engine/Source/Programs/UnrealBuildTool/Linux/LinuxToolChain.cs
+++ b/Engine/Source/Programs/UnrealBuildTool/Linux/LinuxToolChain.cs
@@ -1,8 +1,4 @@
-﻿// Copyright 1998-2017 Epic Games, Inc. All Rights Reserved.
-
-// Define this to use older - slower, but tested and proven - way of dealing with cross-referenced libraries.
-// This path will be removed once new one is well tested (in 4.18 timeframe)
-//#define USE_FIXDEPS
+// Copyright 1998-2017 Epic Games, Inc. All Rights Reserved.
 
 using System;
 using System.Collections.Generic;
@@ -45,6 +41,10 @@ namespace UnrealBuildTool
 				{
 					GCCPath = null;
 				}
+
+				// When compiling on Linux, use a faster way to relink circularly dependent libraries.
+				// Race condition between actions linking to the .so and action overwriting it is avoided thanks to inodes
+				bUseFixdeps = false;
 			}
 			else
 			{
@@ -81,6 +81,9 @@ namespace UnrealBuildTool
 				ArPath = "ar.exe";
 				RanlibPath = "ranlib.exe";
 				StripPath = "strip.exe";
+
+				// When cross-compiling on Windows, use old FixDeps. It is slow, but it does not have timing issues
+				bUseFixdeps = (BuildHostPlatform.Current.Platform == UnrealTargetPlatform.Win64 || BuildHostPlatform.Current.Platform == UnrealTargetPlatform.Win32);
 			}
 
 			if (!DetermineCompilerVersion())
@@ -710,7 +713,12 @@ namespace UnrealBuildTool
 		/// </summary>
 		static int CompilerVersionPatch = -1;
 
-#if (USE_FIXDEPS)
+		/// <summary>
+		/// Whether to use old, slower way to relink circularly dependent libraries.
+		/// It makes sense to use it when cross-compiling on Windows due to race conditions between actions reading and modifying the libs.
+		/// </summary>
+		private bool bUseFixdeps = false;
+
 		/// <summary>
 		/// Track which scripts need to be deleted before appending to
 		/// </summary>
@@ -720,7 +728,6 @@ namespace UnrealBuildTool
 		/// Holds all the binaries for a particular target (except maybe the executable itself).
 		/// </summary>
 		private static List<FileItem> AllBinaries = new List<FileItem>();
-#endif
 
 		/// <summary>
 		/// Tracks that information about used C++ library is only printed once
@@ -737,6 +744,17 @@ namespace UnrealBuildTool
 				// inform the user which C++ library the engine is going to be compiled against - important for compatibility with third party code that uses STL
 				bool bUseLibCxx = ShouldUseLibcxx(CompileEnvironment.Architecture);
 				Console.WriteLine("Using {0} standard C++ library.", bUseLibCxx ? "bundled libc++" : "compiler default (most likely libstdc++)");
+
+				// Also print other once-per-build information
+				if (bUseFixdeps)
+				{
+					Console.WriteLine("Using old way to relink circularly dependent libraries (with a FixDeps step).");
+				}
+				else
+				{
+					Console.WriteLine("Using fast way to relink  circularly dependent libraries (no FixDeps).");
+				}
+
 				bHasPrintedLibcxxInformation = true;
 			}
 
@@ -962,62 +980,65 @@ namespace UnrealBuildTool
 
 		public FileItem FixDependencies(LinkEnvironment LinkEnvironment, FileItem Executable, ActionGraph ActionGraph)
 		{
-#if (USE_FIXDEPS)
-			if (!LinkEnvironment.bIsCrossReferenced)
+			if (bUseFixdeps)
 			{
-				return null;
-			}
+				if (!LinkEnvironment.bIsCrossReferenced)
+				{
+					return null;
+				}
 
-			Log.TraceVerbose("Adding postlink step");
+				Log.TraceVerbose("Adding postlink step");
 
-			bool bUseCmdExe = BuildHostPlatform.Current.Platform == UnrealTargetPlatform.Win64 || BuildHostPlatform.Current.Platform == UnrealTargetPlatform.Win32;
-			string ShellBinary = bUseCmdExe ? "cmd.exe" : "/bin/sh";
-			string ExecuteSwitch = bUseCmdExe ? " /C" : ""; // avoid -c so scripts don't need +x
-			string ScriptName = bUseCmdExe ? "FixDependencies.bat" : "FixDependencies.sh";
+				bool bUseCmdExe = BuildHostPlatform.Current.Platform == UnrealTargetPlatform.Win64 || BuildHostPlatform.Current.Platform == UnrealTargetPlatform.Win32;
+				string ShellBinary = bUseCmdExe ? "cmd.exe" : "/bin/sh";
+				string ExecuteSwitch = bUseCmdExe ? " /C" : ""; // avoid -c so scripts don't need +x
+				string ScriptName = bUseCmdExe ? "FixDependencies.bat" : "FixDependencies.sh";
 
-			FileItem FixDepsScript = FileItem.GetItemByFileReference(FileReference.Combine(LinkEnvironment.LocalShadowDirectory, ScriptName));
+				FileItem FixDepsScript = FileItem.GetItemByFileReference(FileReference.Combine(LinkEnvironment.LocalShadowDirectory, ScriptName));
 
-			Action PostLinkAction = ActionGraph.Add(ActionType.Link);
-			PostLinkAction.WorkingDirectory = UnrealBuildTool.EngineSourceDirectory.FullName;
-			PostLinkAction.CommandPath = ShellBinary;
-			PostLinkAction.StatusDescription = string.Format("{0}", Path.GetFileName(Executable.AbsolutePath));
-			PostLinkAction.CommandDescription = "FixDeps";
-			PostLinkAction.bCanExecuteRemotely = false;
-			PostLinkAction.CommandArguments = ExecuteSwitch;
+				Action PostLinkAction = ActionGraph.Add(ActionType.Link);
+				PostLinkAction.WorkingDirectory = UnrealBuildTool.EngineSourceDirectory.FullName;
+				PostLinkAction.CommandPath = ShellBinary;
+				PostLinkAction.StatusDescription = string.Format("{0}", Path.GetFileName(Executable.AbsolutePath));
+				PostLinkAction.CommandDescription = "FixDeps";
+				PostLinkAction.bCanExecuteRemotely = false;
+				PostLinkAction.CommandArguments = ExecuteSwitch;
 
-			PostLinkAction.CommandArguments += bUseCmdExe ? " \"" : " -c '";
+				PostLinkAction.CommandArguments += bUseCmdExe ? " \"" : " -c '";
 
-			FileItem OutputFile = FileItem.GetItemByFileReference(FileReference.Combine(LinkEnvironment.LocalShadowDirectory, Path.GetFileNameWithoutExtension(Executable.AbsolutePath) + ".link"));
+				FileItem OutputFile = FileItem.GetItemByFileReference(FileReference.Combine(LinkEnvironment.LocalShadowDirectory, Path.GetFileNameWithoutExtension(Executable.AbsolutePath) + ".link"));
 
-			// Make sure we don't run this script until the all executables and shared libraries
-			// have been built.
-			PostLinkAction.PrerequisiteItems.Add(Executable);
-			foreach (FileItem Dependency in AllBinaries)
-			{
-				PostLinkAction.PrerequisiteItems.Add(Dependency);
-			}
+				// Make sure we don't run this script until the all executables and shared libraries
+				// have been built.
+				PostLinkAction.PrerequisiteItems.Add(Executable);
+				foreach (FileItem Dependency in AllBinaries)
+				{
+					PostLinkAction.PrerequisiteItems.Add(Dependency);
+				}
 
-			PostLinkAction.CommandArguments += ShellBinary + ExecuteSwitch + " \"" + FixDepsScript.AbsolutePath + "\" && ";
+				PostLinkAction.CommandArguments += ShellBinary + ExecuteSwitch + " \"" + FixDepsScript.AbsolutePath + "\" && ";
 
-			// output file should not be empty or it will be rebuilt next time
-			string Touch = bUseCmdExe ? "echo \"Dummy\" >> \"{0}\" && copy /b \"{0}\" +,," : "echo \"Dummy\" >> \"{0}\"";
+				// output file should not be empty or it will be rebuilt next time
+				string Touch = bUseCmdExe ? "echo \"Dummy\" >> \"{0}\" && copy /b \"{0}\" +,," : "echo \"Dummy\" >> \"{0}\"";
 
-			PostLinkAction.CommandArguments += String.Format(Touch, OutputFile.AbsolutePath);
-			PostLinkAction.CommandArguments += bUseCmdExe ? "\"" : "'";
+				PostLinkAction.CommandArguments += String.Format(Touch, OutputFile.AbsolutePath);
+				PostLinkAction.CommandArguments += bUseCmdExe ? "\"" : "'";
 
-			System.Console.WriteLine("{0} {1}", PostLinkAction.CommandPath, PostLinkAction.CommandArguments);
+				System.Console.WriteLine("{0} {1}", PostLinkAction.CommandPath, PostLinkAction.CommandArguments);
 
-			// piping output through the handler during native builds is unnecessary and reportedly causes problems with tools like octobuild.
-			if (CrossCompiling())
+				// piping output through the handler during native builds is unnecessary and reportedly causes problems with tools like octobuild.
+				if (CrossCompiling())
+				{
+					PostLinkAction.OutputEventHandler = new DataReceivedEventHandler(CrossCompileOutputReceivedDataEventHandler);
+				}
+
+				PostLinkAction.ProducedItems.Add(OutputFile);
+				return OutputFile;
+			}
+			else
 			{
-				PostLinkAction.OutputEventHandler = new DataReceivedEventHandler(CrossCompileOutputReceivedDataEventHandler);
+				return null;
 			}
-
-			PostLinkAction.ProducedItems.Add(OutputFile);
-			return OutputFile;
-#else
-			return null;
-#endif
 		}
 
 
@@ -1257,161 +1278,164 @@ namespace UnrealBuildTool
 			// are created. This script will be called by action created in FixDependencies()
 			if (LinkEnvironment.bIsCrossReferenced && LinkEnvironment.bIsBuildingDLL)
 			{
-#if (USE_FIXDEPS)
-				bool bUseCmdExe = BuildHostPlatform.Current.Platform == UnrealTargetPlatform.Win64 || BuildHostPlatform.Current.Platform == UnrealTargetPlatform.Win32;
-				string ScriptName = bUseCmdExe ? "FixDependencies.bat" : "FixDependencies.sh";
-
-				string FixDepsScriptPath = Path.Combine(LinkEnvironment.LocalShadowDirectory.FullName, ScriptName);
-				if (!bHasWipedFixDepsScript)
+				if (bUseFixdeps)
 				{
-					bHasWipedFixDepsScript = true;
-					Log.TraceVerbose("Creating script: {0}", FixDepsScriptPath);
-					StreamWriter Writer = File.CreateText(FixDepsScriptPath);
+					bool bUseCmdExe = BuildHostPlatform.Current.Platform == UnrealTargetPlatform.Win64 || BuildHostPlatform.Current.Platform == UnrealTargetPlatform.Win32;
+					string ScriptName = bUseCmdExe ? "FixDependencies.bat" : "FixDependencies.sh";
+
+					string FixDepsScriptPath = Path.Combine(LinkEnvironment.LocalShadowDirectory.FullName, ScriptName);
+					if (!bHasWipedFixDepsScript)
+					{
+						bHasWipedFixDepsScript = true;
+						Log.TraceVerbose("Creating script: {0}", FixDepsScriptPath);
+						StreamWriter Writer = File.CreateText(FixDepsScriptPath);
+
+						if (bUseCmdExe)
+						{
+							Writer.Write("@echo off\n");
+							Writer.Write("rem Automatically generated by UnrealBuildTool\n");
+							Writer.Write("rem *DO NOT EDIT*\n\n");
+						}
+						else
+						{
+							Writer.Write("#!/bin/sh\n");
+							Writer.Write("# Automatically generated by UnrealBuildTool\n");
+							Writer.Write("# *DO NOT EDIT*\n\n");
+							Writer.Write("set -o errexit\n");
+						}
+						Writer.Close();
+					}
+
+					StreamWriter FixDepsScript = File.AppendText(FixDepsScriptPath);
+
+					string EngineAndGameLibrariesString = "";
+					foreach (string Library in EngineAndGameLibrariesLinkFlags)
+					{
+						EngineAndGameLibrariesString += Library;
+					}
 
+					FixDepsScript.Write(string.Format("echo Fixing {0}\n", Path.GetFileName(OutputFile.AbsolutePath)));
+					if (!bUseCmdExe)
+					{
+						FixDepsScript.Write(string.Format("TIMESTAMP=`stat --format %y \"{0}\"`\n", OutputFile.AbsolutePath));
+					}
+					string FixDepsLine = LinkAction.CommandPath + " " + LinkAction.CommandArguments;
+					string Replace = "-Wl,--allow-shlib-undefined";
+
+					FixDepsLine = FixDepsLine.Replace(Replace, EngineAndGameLibrariesString);
+					string OutputFileForwardSlashes = OutputFile.AbsolutePath.Replace("\\", "/");
+					FixDepsLine = FixDepsLine.Replace(OutputFileForwardSlashes, OutputFileForwardSlashes + ".fixed");
+					FixDepsLine = FixDepsLine.Replace("$", "\\$");
+					FixDepsScript.Write(FixDepsLine + "\n");
 					if (bUseCmdExe)
 					{
-						Writer.Write("@echo off\n");
-						Writer.Write("rem Automatically generated by UnrealBuildTool\n");
-						Writer.Write("rem *DO NOT EDIT*\n\n");
+						FixDepsScript.Write(string.Format("move /Y \"{0}.fixed\" \"{0}\"\n", OutputFile.AbsolutePath));
 					}
 					else
 					{
-						Writer.Write("#!/bin/sh\n");
-						Writer.Write("# Automatically generated by UnrealBuildTool\n");
-						Writer.Write("# *DO NOT EDIT*\n\n");
-						Writer.Write("set -o errexit\n");
+						FixDepsScript.Write(string.Format("mv \"{0}.fixed\" \"{0}\"\n", OutputFile.AbsolutePath));
+						FixDepsScript.Write(string.Format("touch -d \"$TIMESTAMP\" \"{0}\"\n\n", OutputFile.AbsolutePath));
 					}
-					Writer.Close();
-				}
-
-				StreamWriter FixDepsScript = File.AppendText(FixDepsScriptPath);
-
-				string EngineAndGameLibrariesString = "";
-				foreach (string Library in EngineAndGameLibrariesLinkFlags)
-				{
-					EngineAndGameLibrariesString += Library;
-				}
-
-				FixDepsScript.Write(string.Format("echo Fixing {0}\n", Path.GetFileName(OutputFile.AbsolutePath)));
-				if (!bUseCmdExe)
-				{
-					FixDepsScript.Write(string.Format("TIMESTAMP=`stat --format %y \"{0}\"`\n", OutputFile.AbsolutePath));
-				}
-				string FixDepsLine = LinkAction.CommandPath + " " + LinkAction.CommandArguments;
-				string Replace = "-Wl,--allow-shlib-undefined";
-
-				FixDepsLine = FixDepsLine.Replace(Replace, EngineAndGameLibrariesString);
-				string OutputFileForwardSlashes = OutputFile.AbsolutePath.Replace("\\", "/");
-				FixDepsLine = FixDepsLine.Replace(OutputFileForwardSlashes, OutputFileForwardSlashes + ".fixed");
-				FixDepsLine = FixDepsLine.Replace("$", "\\$");
-				FixDepsScript.Write(FixDepsLine + "\n");
-				if (bUseCmdExe)
-				{
-					FixDepsScript.Write(string.Format("move /Y \"{0}.fixed\" \"{0}\"\n", OutputFile.AbsolutePath));
+					FixDepsScript.Close();
 				}
 				else
 				{
-					FixDepsScript.Write(string.Format("mv \"{0}.fixed\" \"{0}\"\n", OutputFile.AbsolutePath));
-					FixDepsScript.Write(string.Format("touch -d \"$TIMESTAMP\" \"{0}\"\n\n", OutputFile.AbsolutePath));
-				}
-				FixDepsScript.Close();
-#else
-				// Create the action to relink the library. This actions does not overwrite the source file so it can be executed in parallel
-				Action RelinkAction = ActionGraph.Add(ActionType.Link);
-				RelinkAction.WorkingDirectory = LinkAction.WorkingDirectory;
-				RelinkAction.StatusDescription = LinkAction.StatusDescription;
-				RelinkAction.CommandDescription = "Relink";
-				RelinkAction.bCanExecuteRemotely = false;
-				RelinkAction.ProducedItems.Clear();
-				RelinkAction.PrerequisiteItems = new List<FileItem>( LinkAction.PrerequisiteItems );
-				foreach (FileItem Dependency in EngineAndGameLibrariesFiles)
-				{
-					RelinkAction.PrerequisiteItems.Add(Dependency);
-				}
-				RelinkAction.PrerequisiteItems.Add(OutputFile);	// also depend on the first link action's output
+					// Create the action to relink the library. This actions does not overwrite the source file so it can be executed in parallel
+					Action RelinkAction = ActionGraph.Add(ActionType.Link);
+					RelinkAction.WorkingDirectory = LinkAction.WorkingDirectory;
+					RelinkAction.StatusDescription = LinkAction.StatusDescription;
+					RelinkAction.CommandDescription = "Relink";
+					RelinkAction.bCanExecuteRemotely = false;
+					RelinkAction.ProducedItems.Clear();
+					RelinkAction.PrerequisiteItems = new List<FileItem>(LinkAction.PrerequisiteItems);
+					foreach (FileItem Dependency in EngineAndGameLibrariesFiles)
+					{
+						RelinkAction.PrerequisiteItems.Add(Dependency);
+					}
+					RelinkAction.PrerequisiteItems.Add(OutputFile); // also depend on the first link action's output
 
-				string LinkOutputFileForwardSlashes = OutputFile.AbsolutePath.Replace("\\", "/");
-				string RelinkedFileForwardSlashes = Path.Combine(LinkEnvironment.LocalShadowDirectory.FullName, OutputFile.Reference.GetFileName()) + ".relinked";
+					string LinkOutputFileForwardSlashes = OutputFile.AbsolutePath.Replace("\\", "/");
+					string RelinkedFileForwardSlashes = Path.Combine(LinkEnvironment.LocalShadowDirectory.FullName, OutputFile.Reference.GetFileName()) + ".relinked";
 
-				// cannot use the real product because we need to maintain the timestamp on it
-				FileReference RelinkActionDummyProductRef = FileReference.Combine(LinkEnvironment.LocalShadowDirectory, LinkEnvironment.OutputFilePath.GetFileNameWithoutExtension() + ".relinked_action_ran");
-				RelinkAction.ProducedItems.Add(FileItem.GetItemByFileReference(RelinkActionDummyProductRef));
+					// cannot use the real product because we need to maintain the timestamp on it
+					FileReference RelinkActionDummyProductRef = FileReference.Combine(LinkEnvironment.LocalShadowDirectory, LinkEnvironment.OutputFilePath.GetFileNameWithoutExtension() + ".relinked_action_ran");
+					RelinkAction.ProducedItems.Add(FileItem.GetItemByFileReference(RelinkActionDummyProductRef));
 
-				string EngineAndGameLibrariesString = "";
-				foreach (string Library in EngineAndGameLibrariesLinkFlags)
-				{
-					EngineAndGameLibrariesString += Library;
-				}
+					string EngineAndGameLibrariesString = "";
+					foreach (string Library in EngineAndGameLibrariesLinkFlags)
+					{
+						EngineAndGameLibrariesString += Library;
+					}
 
-				// create the action 
-				bool bUseCmdExe = BuildHostPlatform.Current.Platform == UnrealTargetPlatform.Win64 || BuildHostPlatform.Current.Platform == UnrealTargetPlatform.Win32;
-				string ShellBinary = bUseCmdExe ? "cmd.exe" : "/bin/sh";
-				string ExecuteSwitch = bUseCmdExe ? " /C" : ""; // avoid -c so scripts don't need +x
+					// create the action 
+					bool bUseCmdExe = BuildHostPlatform.Current.Platform == UnrealTargetPlatform.Win64 || BuildHostPlatform.Current.Platform == UnrealTargetPlatform.Win32;
+					string ShellBinary = bUseCmdExe ? "cmd.exe" : "/bin/sh";
+					string ExecuteSwitch = bUseCmdExe ? " /C" : ""; // avoid -c so scripts don't need +x
 
-				// create the relinking step
-				string RelinkScriptName = string.Format((bUseCmdExe ? "Relink-{0}.bat" : "Relink-{0}.sh"), OutputFile);
-				string RelinkScriptFullPath = Path.Combine(LinkEnvironment.LocalShadowDirectory.FullName, RelinkScriptName);
+					// create the relinking step
+					string RelinkScriptName = string.Format((bUseCmdExe ? "Relink-{0}.bat" : "Relink-{0}.sh"), OutputFile);
+					string RelinkScriptFullPath = Path.Combine(LinkEnvironment.LocalShadowDirectory.FullName, RelinkScriptName);
 
-				Log.TraceVerbose("Creating script: {0}", RelinkScriptFullPath);
-				StreamWriter RelinkWriter = File.CreateText(RelinkScriptFullPath);
+					Log.TraceVerbose("Creating script: {0}", RelinkScriptFullPath);
+					StreamWriter RelinkWriter = File.CreateText(RelinkScriptFullPath);
 
-				string RelinkInvocation = LinkAction.CommandPath + " " + LinkAction.CommandArguments;
-				string Replace = "-Wl,--allow-shlib-undefined";
-				RelinkInvocation = RelinkInvocation.Replace(Replace, EngineAndGameLibrariesString);
+					string RelinkInvocation = LinkAction.CommandPath + " " + LinkAction.CommandArguments;
+					string Replace = "-Wl,--allow-shlib-undefined";
+					RelinkInvocation = RelinkInvocation.Replace(Replace, EngineAndGameLibrariesString);
 
-				// should be the same as RelinkedFileRef
-				RelinkInvocation = RelinkInvocation.Replace(LinkOutputFileForwardSlashes, RelinkedFileForwardSlashes);
-				RelinkInvocation = RelinkInvocation.Replace("$", "\\$");
+					// should be the same as RelinkedFileRef
+					RelinkInvocation = RelinkInvocation.Replace(LinkOutputFileForwardSlashes, RelinkedFileForwardSlashes);
+					RelinkInvocation = RelinkInvocation.Replace("$", "\\$");
 
-				if (bUseCmdExe)
-				{
-					RelinkWriter.Write("@echo off\n");
-					RelinkWriter.Write("rem Automatically generated by UnrealBuildTool\n");
-					RelinkWriter.Write("rem *DO NOT EDIT*\n\n");
-					RelinkWriter.Write("set Retries=0\n");
-					RelinkWriter.Write(":relinkloop\n");
-					RelinkWriter.Write("if %Retries% GEQ 10 goto failedtorelink\n");
-					RelinkWriter.Write(RelinkInvocation + "\n");
-					RelinkWriter.Write("if %errorlevel% neq 0 goto sleepandretry\n");
-					RelinkWriter.Write("copy /B \"{0}\" \"{1}.temp\" >NUL 2>NUL\n", RelinkedFileForwardSlashes, OutputFile.AbsolutePath);
-					RelinkWriter.Write("if %errorlevel% neq 0 goto sleepandretry\n");
-					RelinkWriter.Write("move /Y \"{0}.temp\" \"{1}\" >NUL 2>NUL\n", OutputFile.AbsolutePath, OutputFile.AbsolutePath);
-					RelinkWriter.Write("if %errorlevel% neq 0 goto sleepandretry\n");
-					RelinkWriter.Write(string.Format("echo \"Dummy\" >> \"{0}\" && copy /b \"{0}\" +,,\n", RelinkActionDummyProductRef.FullName));
-					RelinkWriter.Write("echo Relinked {0} successfully after %Retries% retries\n", OutputFile.AbsolutePath);
-					RelinkWriter.Write("exit 0\n");
-					RelinkWriter.Write(":sleepandretry\n");
-					RelinkWriter.Write("ping 127.0.0.1 -n 1 -w 5000 >NUL 2>NUL\n");		// timeout complains about lack of redirection
-					RelinkWriter.Write("set /a Retries+=1\n");
-					RelinkWriter.Write("goto relinkloop\n");
-					RelinkWriter.Write(":failedtorelink\n");
-					RelinkWriter.Write("echo Failed to relink {0} after %Retries% retries\n", OutputFile.AbsolutePath);
-					RelinkWriter.Write("exit 1\n");
-				}
-				else
-				{
-					RelinkWriter.Write("#!/bin/sh\n");
-					RelinkWriter.Write("# Automatically generated by UnrealBuildTool\n");
-					RelinkWriter.Write("# *DO NOT EDIT*\n\n");
-					RelinkWriter.Write("set -o errexit\n");
-					RelinkWriter.Write(RelinkInvocation + "\n");
-					RelinkWriter.Write(string.Format("TIMESTAMP=`stat --format %y \"{0}\"`\n", OutputFile.AbsolutePath));
-					RelinkWriter.Write("cp \"{0}\" \"{1}.temp\"\n", RelinkedFileForwardSlashes, OutputFile.AbsolutePath);
-					RelinkWriter.Write("mv \"{0}.temp\" \"{1}\"\n", OutputFile.AbsolutePath, OutputFile.AbsolutePath);
-					RelinkWriter.Write(string.Format("touch -d \"$TIMESTAMP\" \"{0}\"\n\n", OutputFile.AbsolutePath));
-					RelinkWriter.Write(string.Format("echo \"Dummy\" >> \"{0}\"", RelinkActionDummyProductRef.FullName));
-				}
-				RelinkWriter.Close();
+					if (bUseCmdExe)
+					{
+						RelinkWriter.Write("@echo off\n");
+						RelinkWriter.Write("rem Automatically generated by UnrealBuildTool\n");
+						RelinkWriter.Write("rem *DO NOT EDIT*\n\n");
+						RelinkWriter.Write("set Retries=0\n");
+						RelinkWriter.Write(":relinkloop\n");
+						RelinkWriter.Write("if %Retries% GEQ 10 goto failedtorelink\n");
+						RelinkWriter.Write(RelinkInvocation + "\n");
+						RelinkWriter.Write("if %errorlevel% neq 0 goto sleepandretry\n");
+						RelinkWriter.Write("copy /B \"{0}\" \"{1}.temp\" >NUL 2>NUL\n", RelinkedFileForwardSlashes, OutputFile.AbsolutePath);
+						RelinkWriter.Write("if %errorlevel% neq 0 goto sleepandretry\n");
+						RelinkWriter.Write("move /Y \"{0}.temp\" \"{1}\" >NUL 2>NUL\n", OutputFile.AbsolutePath, OutputFile.AbsolutePath);
+						RelinkWriter.Write("if %errorlevel% neq 0 goto sleepandretry\n");
+						RelinkWriter.Write(string.Format("echo \"Dummy\" >> \"{0}\" && copy /b \"{0}\" +,,\n", RelinkActionDummyProductRef.FullName));
+						RelinkWriter.Write("echo Relinked {0} successfully after %Retries% retries\n", OutputFile.AbsolutePath);
+						RelinkWriter.Write("exit 0\n");
+						RelinkWriter.Write(":sleepandretry\n");
+						RelinkWriter.Write("ping 127.0.0.1 -n 1 -w 5000 >NUL 2>NUL\n");     // timeout complains about lack of redirection
+						RelinkWriter.Write("set /a Retries+=1\n");
+						RelinkWriter.Write("goto relinkloop\n");
+						RelinkWriter.Write(":failedtorelink\n");
+						RelinkWriter.Write("echo Failed to relink {0} after %Retries% retries\n", OutputFile.AbsolutePath);
+						RelinkWriter.Write("exit 1\n");
+					}
+					else
+					{
+						RelinkWriter.Write("#!/bin/sh\n");
+						RelinkWriter.Write("# Automatically generated by UnrealBuildTool\n");
+						RelinkWriter.Write("# *DO NOT EDIT*\n\n");
+						RelinkWriter.Write("set -o errexit\n");
+						RelinkWriter.Write(RelinkInvocation + "\n");
+						RelinkWriter.Write(string.Format("TIMESTAMP=`stat --format %y \"{0}\"`\n", OutputFile.AbsolutePath));
+						RelinkWriter.Write("cp \"{0}\" \"{1}.temp\"\n", RelinkedFileForwardSlashes, OutputFile.AbsolutePath);
+						RelinkWriter.Write("mv \"{0}.temp\" \"{1}\"\n", OutputFile.AbsolutePath, OutputFile.AbsolutePath);
+						RelinkWriter.Write(string.Format("touch -d \"$TIMESTAMP\" \"{0}\"\n\n", OutputFile.AbsolutePath));
+						RelinkWriter.Write(string.Format("echo \"Dummy\" >> \"{0}\"", RelinkActionDummyProductRef.FullName));
+					}
+					RelinkWriter.Close();
 
-				RelinkAction.CommandPath = ShellBinary;
-				RelinkAction.CommandArguments = ExecuteSwitch + " \"" + RelinkScriptFullPath + "\"";
+					RelinkAction.CommandPath = ShellBinary;
+					RelinkAction.CommandArguments = ExecuteSwitch + " \"" + RelinkScriptFullPath + "\"";
 
-				// piping output through the handler during native builds is unnecessary and reportedly causes problems with tools like octobuild.
-				if (CrossCompiling())
-				{
-					RelinkAction.OutputEventHandler = new DataReceivedEventHandler(CrossCompileOutputReceivedDataEventHandler);
+					// piping output through the handler during native builds is unnecessary and reportedly causes problems with tools like octobuild.
+					if (CrossCompiling())
+					{
+						RelinkAction.OutputEventHandler = new DataReceivedEventHandler(CrossCompileOutputReceivedDataEventHandler);
+					}
 				}
-#endif
 			}
 
 			if (CrossCompiling())
@@ -1426,15 +1450,16 @@ namespace UnrealBuildTool
 			throw new BuildException("Linux cannot compile C# files");
 		}
 
-#if (USE_FIXDEPS)
 		public override void SetupBundleDependencies(List<UEBuildBinary> Binaries, string GameName)
 		{
-			foreach (UEBuildBinary Binary in Binaries)
+			if (bUseFixdeps)
 			{
-				AllBinaries.Add(FileItem.GetItemByFileReference(Binary.Config.OutputFilePath));
+				foreach (UEBuildBinary Binary in Binaries)
+				{
+					AllBinaries.Add(FileItem.GetItemByFileReference(Binary.Config.OutputFilePath));
+				}
 			}
 		}
-#endif
 
 		/// <summary>
 		/// Converts the passed in path from UBT host to compiler native format.
@@ -1455,26 +1480,28 @@ namespace UnrealBuildTool
 		{
 			var OutputFiles = base.PostBuild(Executable, BinaryLinkEnvironment, ActionGraph);
 
-#if (USE_FIXDEPS)
-			if (BinaryLinkEnvironment.bIsBuildingDLL || BinaryLinkEnvironment.bIsBuildingLibrary)
+			if (bUseFixdeps)
 			{
-				return OutputFiles;
-			}
+				if (BinaryLinkEnvironment.bIsBuildingDLL || BinaryLinkEnvironment.bIsBuildingLibrary)
+				{
+					return OutputFiles;
+				}
 
-			FileItem FixDepsOutputFile = FixDependencies(BinaryLinkEnvironment, Executable, ActionGraph);
-			if (FixDepsOutputFile != null)
-			{
-				OutputFiles.Add(FixDepsOutputFile);
+				FileItem FixDepsOutputFile = FixDependencies(BinaryLinkEnvironment, Executable, ActionGraph);
+				if (FixDepsOutputFile != null)
+				{
+					OutputFiles.Add(FixDepsOutputFile);
+				}
 			}
-#else
-			// make build product of cross-referenced DSOs to be *..relinked_action_ran, so the relinking steps are executed
-			if (BinaryLinkEnvironment.bIsBuildingDLL && BinaryLinkEnvironment.bIsCrossReferenced)
+			else
 			{
-				FileReference RelinkedMapRef = FileReference.Combine(BinaryLinkEnvironment.LocalShadowDirectory, BinaryLinkEnvironment.OutputFilePath.GetFileNameWithoutExtension() + ".relinked_action_ran");
-				OutputFiles.Add(FileItem.GetItemByFileReference(RelinkedMapRef));
+				// make build product of cross-referenced DSOs to be *.relinked_action_ran, so the relinking steps are executed
+				if (BinaryLinkEnvironment.bIsBuildingDLL && BinaryLinkEnvironment.bIsCrossReferenced)
+				{
+					FileReference RelinkedMapRef = FileReference.Combine(BinaryLinkEnvironment.LocalShadowDirectory, BinaryLinkEnvironment.OutputFilePath.GetFileNameWithoutExtension() + ".relinked_action_ran");
+					OutputFiles.Add(FileItem.GetItemByFileReference(RelinkedMapRef));
+				}
 			}
-
-#endif // USE_FIXDEPS
 			return OutputFiles;
 		}
 
