Reverting XGE change that was accidentally copied up from Orion (stop on error was changed to false)

diff --git a/Engine/Source/Programs/UnrealBuildTool/System/XGE.cs b/Engine/Source/Programs/UnrealBuildTool/System/XGE.cs
--- a/Engine/Source/Programs/UnrealBuildTool/System/XGE.cs
+++ b/Engine/Source/Programs/UnrealBuildTool/System/XGE.cs
@@ -32,7 +32,7 @@ namespace UnrealBuildTool
 		/// When enabled, XGE will stop compiling targets after a compile error occurs.  Recommended, as it saves computing resources for others.
 		/// </summary>
 		[XmlConfigFile(Category = "BuildConfiguration")]
-		bool bStopXGECompilationAfterErrors = false;
+		bool bStopXGECompilationAfterErrors = true;
 
 		private const string ProgressMarkupPrefix = "@action";
 
