Only use android linker's --ICF switch with arm7 builds.

diff --git a/Engine/Source/Programs/UnrealBuildTool/Android/AndroidToolChain.cs b/Engine/Source/Programs/UnrealBuildTool/Android/AndroidToolChain.cs
--- a/Engine/Source/Programs/UnrealBuildTool/Android/AndroidToolChain.cs
+++ b/Engine/Source/Programs/UnrealBuildTool/Android/AndroidToolChain.cs
@@ -668,11 +668,6 @@ namespace UnrealBuildTool
 			Result += " -Wl,-shared,-Bsymbolic";
 			Result += " -Wl,--no-undefined";
 			Result += " -Wl,-gc-sections"; // Enable garbage collection of unused input sections. works best with -ffunction-sections, -fdata-sections
-			if (LinkEnvironment.Configuration == CppConfiguration.Shipping)
-			{
-				Result += " -Wl,--icf=all"; // Enables ICF (Identical Code Folding). [all, safe] safe == fold functions that can be proven not to have their address taken.
-				Result += " -Wl,--icf-iterations=3"; // Embiggen count of ICF iteration passes (default == 2)
-			}
 
 			if (Architecture == "-arm64")
 			{
@@ -693,7 +688,13 @@ namespace UnrealBuildTool
 			{
 				Result += ToolchainParamsArm;
 				Result += " -march=armv7-a";
-				Result += " -Wl,--fix-cortex-a8";		// required to route around a CPU bug in some Cortex-A8 implementations
+				Result += " -Wl,--fix-cortex-a8";       // required to route around a CPU bug in some Cortex-A8 implementations
+
+				if (LinkEnvironment.Configuration == CppConfiguration.Shipping)
+				{
+					Result += " -Wl,--icf=all"; // Enables ICF (Identical Code Folding). [all, safe] safe == fold functions that can be proven not to have their address taken.
+					Result += " -Wl,--icf-iterations=3";
+				}
 			}
 
 			if (bUseLdGold && CompilerVersionGreaterOrEqual(3, 6, 0) && CompilerVersionLessThan(3, 8, 0))
