Fix UAT crashing on startup during BuildCookRun. UAT currently makes an assumption that there is one editor/client/game/server target per project.

diff --git a/Engine/Source/Programs/AutomationTool/AutomationUtils/ProjectUtils.cs b/Engine/Source/Programs/AutomationTool/AutomationUtils/ProjectUtils.cs
--- a/Engine/Source/Programs/AutomationTool/AutomationUtils/ProjectUtils.cs
+++ b/Engine/Source/Programs/AutomationTool/AutomationUtils/ProjectUtils.cs
@@ -438,7 +438,7 @@ namespace AutomationTool
 				DirPushed = true;
 			}
 			var ExtraSearchDirectories = (ExtraSearchPaths == null)? null : ExtraSearchPaths.Select(x => new DirectoryReference(x)).ToList();
-			var TargetScripts = RulesCompiler.FindAllRulesSourceFiles(RulesCompiler.RulesFileType.Target, GameFolders: GameFolders, ForeignPlugins: null, AdditionalSearchPaths: ExtraSearchDirectories);
+			var TargetScripts = RulesCompiler.FindAllRulesSourceFiles(RulesCompiler.RulesFileType.Target, GameFolders: GameFolders, ForeignPlugins: null, AdditionalSearchPaths: ExtraSearchDirectories, bIncludeEnterprise: false);
 			if (DirPushed)
 			{
 				CommandUtils.PopDir();
diff --git a/Engine/Source/Programs/UnrealBuildTool/System/RulesCompiler.cs b/Engine/Source/Programs/UnrealBuildTool/System/RulesCompiler.cs
--- a/Engine/Source/Programs/UnrealBuildTool/System/RulesCompiler.cs
+++ b/Engine/Source/Programs/UnrealBuildTool/System/RulesCompiler.cs
@@ -58,8 +58,9 @@ namespace UnrealBuildTool
 		/// <param name="ForeignPlugins"></param>
 		/// <param name="AdditionalSearchPaths"></param>
 		/// <param name="bIncludeEngine"></param>
+		/// <param name="bIncludeEnterprise"></param>
 		/// <returns></returns>
-		public static List<FileReference> FindAllRulesSourceFiles(RulesFileType RulesFileType, List<DirectoryReference> GameFolders, List<FileReference> ForeignPlugins, List<DirectoryReference> AdditionalSearchPaths, bool bIncludeEngine = true)
+		public static List<FileReference> FindAllRulesSourceFiles(RulesFileType RulesFileType, List<DirectoryReference> GameFolders, List<FileReference> ForeignPlugins, List<DirectoryReference> AdditionalSearchPaths, bool bIncludeEngine = true, bool bIncludeEnterprise = true)
 		{
 			List<DirectoryReference> Folders = new List<DirectoryReference>();
 
@@ -67,6 +68,9 @@ namespace UnrealBuildTool
 			if (bIncludeEngine)
 			{
 				Folders.Add(UnrealBuildTool.EngineSourceDirectory);
+			}
+			if(bIncludeEnterprise)
+			{
 				Folders.Add(UnrealBuildTool.EnterpriseSourceDirectory);
 			}
 
@@ -77,6 +81,9 @@ namespace UnrealBuildTool
 			if (bIncludeEngine)
 			{
 				RootFolders.Add(UnrealBuildTool.EngineDirectory);
+			}
+			if(bIncludeEnterprise)
+			{
 				RootFolders.Add(UnrealBuildTool.EnterpriseDirectory);
 			}
 			if (GameFolders != null)
