Mirroring CL 3522159 from 4.17

Re-enabling the Oculus plugin by default, and explicitly excluding the OVRPlugin lib from all android builds but arm7 (the only one it is available for). This way we avoid a missing lib error when compiling for arm64, etc.

This fixes issues with pre-existing projects (like Robo Recall) that were relying on the plugin being automatically enabled.

diff --git a/Engine/Plugins/Runtime/Oculus/OculusVR/OculusVR.uplugin b/Engine/Plugins/Runtime/Oculus/OculusVR/OculusVR.uplugin
--- a/Engine/Plugins/Runtime/Oculus/OculusVR/OculusVR.uplugin
+++ b/Engine/Plugins/Runtime/Oculus/OculusVR/OculusVR.uplugin
@@ -10,7 +10,7 @@
 	"DocsURL" : "",
 	"MarketplaceURL" : "",
 	"SupportURL" : "",
-	"EnabledByDefault" : false,
+	"EnabledByDefault" : true,
 	"CanContainContent" : true,
 	"IsBetaVersion" : false,
 	"Installed" : false,
diff --git a/Engine/Source/Programs/UnrealBuildTool/Android/AndroidToolChain.cs b/Engine/Source/Programs/UnrealBuildTool/Android/AndroidToolChain.cs
--- a/Engine/Source/Programs/UnrealBuildTool/Android/AndroidToolChain.cs
+++ b/Engine/Source/Programs/UnrealBuildTool/Android/AndroidToolChain.cs
@@ -38,9 +38,9 @@ namespace UnrealBuildTool
 
 		static private Dictionary<string, string[]> LibrariesToSkip = new Dictionary<string, string[]> {
 			{ "-armv7", new string[] { } }, 
-			{ "-arm64", new string[] { "nvToolsExt", "nvToolsExtStub", "oculus", "vrapi", "ovrkernel", "systemutils", "openglloader", } },
-			{ "-x86",   new string[] { "nvToolsExt", "nvToolsExtStub", "oculus", "vrapi", "ovrkernel", "systemutils", "openglloader", "opus", "speex_resampler", } }, 
-			{ "-x64",   new string[] { "nvToolsExt", "nvToolsExtStub", "oculus", "vrapi", "ovrkernel", "systemutils", "openglloader", "gpg", } }, 
+			{ "-arm64", new string[] { "nvToolsExt", "nvToolsExtStub", "oculus", "OVRPlugin", "vrapi", "ovrkernel", "systemutils", "openglloader", } },
+			{ "-x86",   new string[] { "nvToolsExt", "nvToolsExtStub", "oculus", "OVRPlugin", "vrapi", "ovrkernel", "systemutils", "openglloader", "opus", "speex_resampler", } }, 
+			{ "-x64",   new string[] { "nvToolsExt", "nvToolsExtStub", "oculus", "OVRPlugin", "vrapi", "ovrkernel", "systemutils", "openglloader", "gpg", } }, 
 		};
 
 		static private Dictionary<string, string[]> ModulesToSkip = new Dictionary<string, string[]> {
