Clean up Gradle intermediates after packaging if on build machine
#ue4
#android

diff --git a/Engine/Source/Programs/UnrealBuildTool/Android/UEDeployAndroid.cs b/Engine/Source/Programs/UnrealBuildTool/Android/UEDeployAndroid.cs
--- a/Engine/Source/Programs/UnrealBuildTool/Android/UEDeployAndroid.cs
+++ b/Engine/Source/Programs/UnrealBuildTool/Android/UEDeployAndroid.cs
@@ -3325,6 +3325,13 @@ namespace UnrealBuildTool
 					string ShellParametersBegin = Utils.IsRunningOnMono ? "-c '" : "/c ";
 					string ShellParametersEnd = Utils.IsRunningOnMono ? "'" : "";
 					RunCommandLineProgramWithException(UE4BuildGradlePath, ShellExecutable, ShellParametersBegin + "\"" + GradleScriptPath + "\" " + GradleOptions + ShellParametersEnd, "Making .apk with Gradle...");
+					
+					// For build machine run a clean afterward to clean up intermediate files (does not remove final APK)
+					if (Environment.GetEnvironmentVariable("IsBuildMachine") == "1")
+					{
+						GradleOptions = "clean";
+						RunCommandLineProgramWithException(UE4BuildGradlePath, ShellExecutable, ShellParametersBegin + "\"" + GradleScriptPath + "\" " + GradleOptions + ShellParametersEnd, "Cleaning Gradle intermediates...");
+					}
 				}
 
 				bool bBuildWithHiddenSymbolVisibility = false;
