fix for tvOS CIS failure

diff --git a/Engine/Source/Programs/UnrealBuildTool/Platform/IOS/IOSToolChain.cs b/Engine/Source/Programs/UnrealBuildTool/Platform/IOS/IOSToolChain.cs
--- a/Engine/Source/Programs/UnrealBuildTool/Platform/IOS/IOSToolChain.cs
+++ b/Engine/Source/Programs/UnrealBuildTool/Platform/IOS/IOSToolChain.cs
@@ -949,7 +949,7 @@ namespace UnrealBuildTool
             string FullDestPathRoot = Path.Combine(Path.GetDirectoryName(LocalExecutable.AbsolutePath), "AssetCatalog", "Assets.car");
 
             FileItem OutputFile;
-            if (BuildHostPlatform.Current.Platform == UnrealTargetPlatform.Mac)
+            if (BuildHostPlatform.Current.Platform == UnrealTargetPlatform.Mac && Platform == CppPlatform.IOS)
             {
                 FullDestPathRoot = Path.Combine(Path.GetDirectoryName(LocalExecutable.AbsolutePath), "Payload", Path.GetFileNameWithoutExtension(LocalExecutable.AbsolutePath) + ".app", "Assets.car");
             }
