Updated min macOS requirement in Info.plist (generated by MacToolChain) to 10.12.6

diff --git a/Engine/Source/Programs/UnrealBuildTool/Platform/Mac/MacToolChain.cs b/Engine/Source/Programs/UnrealBuildTool/Platform/Mac/MacToolChain.cs
--- a/Engine/Source/Programs/UnrealBuildTool/Platform/Mac/MacToolChain.cs
+++ b/Engine/Source/Programs/UnrealBuildTool/Platform/Mac/MacToolChain.cs
@@ -36,7 +36,7 @@ namespace UnrealBuildTool
 		/// <summary>
 		/// Minimum version of Mac OS X to actually run on, running on earlier versions will display the system minimum version error dialog and exit.
 		/// </summary>
-		public string MinMacOSVersion = "10.11.6";
+		public string MinMacOSVersion = "10.12.6";
 
 		/// <summary>
 		/// Directory for the developer binaries
