Fix .version file not being included in list of build products, prevented it being included in zip of precompiled binaries for games.

diff --git a/Engine/Source/Programs/UnrealBuildTool/Configuration/UEBuildTarget.cs b/Engine/Source/Programs/UnrealBuildTool/Configuration/UEBuildTarget.cs
--- a/Engine/Source/Programs/UnrealBuildTool/Configuration/UEBuildTarget.cs
+++ b/Engine/Source/Programs/UnrealBuildTool/Configuration/UEBuildTarget.cs
@@ -1782,12 +1782,12 @@ namespace UnrealBuildTool
 						Receipt.PrecompiledBuildDependencies.Add(ExternalFile);
 					}
 				}
+			}
 
-				// Also add the version file if it's been specified
-				if (VersionFile != null)
-				{
-					Receipt.BuildProducts.Add(new BuildProduct(VersionFile, BuildProductType.BuildResource));
-				}
+			// Also add the version file if it's been specified
+			if (VersionFile != null)
+			{
+				Receipt.BuildProducts.Add(new BuildProduct(VersionFile, BuildProductType.BuildResource));
 			}
 
 			// Prepare all the version manifests
