Port to UE4/Main: Linux: fix OSSSteam cross-compilation.
- Broken by previous change that put compile-time arguments into response file and bulk-replaced \ with /, which affected things like Definitions.Add("STEAM_SDK_VER_PATH=TEXT(\"Steam" + SteamVersion + "\")");
#rnx

diff --git a/Engine/Source/Programs/UnrealBuildTool/Platform/Linux/LinuxToolChain.cs b/Engine/Source/Programs/UnrealBuildTool/Platform/Linux/LinuxToolChain.cs
--- a/Engine/Source/Programs/UnrealBuildTool/Platform/Linux/LinuxToolChain.cs
+++ b/Engine/Source/Programs/UnrealBuildTool/Platform/Linux/LinuxToolChain.cs
@@ -70,7 +70,7 @@ namespace UnrealBuildTool
 					throw new BuildException("LINUX_ROOT environment variable is not set; cannot instantiate Linux toolchain");
 				}
 
-				BaseLinuxPath = BaseLinuxPath.Replace("\"", "");
+				BaseLinuxPath = BaseLinuxPath.Replace("\"", "").Replace('\\', '/');
 
 				// set up the path to our toolchains
 				GCCPath = "";
@@ -816,22 +816,22 @@ namespace UnrealBuildTool
 				// Add the precompiled header file's path to the include path so Clang can find it.
 				// This needs to be before the other include paths to ensure Clang uses it instead of the source header file.
 				var PrecompiledFileExtension = UEBuildPlatform.GetBuildPlatform(UnrealTargetPlatform.Linux).GetBinaryExtension(UEBuildBinaryType.PrecompiledHeader);
-				PCHArguments += string.Format(" -include \"{0}\"", CompileEnvironment.PrecompiledHeaderFile.AbsolutePath.Replace(PrecompiledFileExtension, ""));
+				PCHArguments += string.Format(" -include \"{0}\"", CompileEnvironment.PrecompiledHeaderFile.AbsolutePath.Replace(PrecompiledFileExtension, "").Replace('\\', '/'));
 			}
 
 			foreach(FileReference ForceIncludeFile in CompileEnvironment.ForceIncludeFiles)
 			{
-				PCHArguments += String.Format(" -include \"{0}\"", ForceIncludeFile.FullName);
+				PCHArguments += String.Format(" -include \"{0}\"", ForceIncludeFile.FullName.Replace('\\', '/'));
 			}
 
 			// Add include paths to the argument list.
 			foreach (string IncludePath in CompileEnvironment.IncludePaths.UserIncludePaths)
 			{
-				Arguments += string.Format(" -I\"{0}\"", IncludePath);
+				Arguments += string.Format(" -I\"{0}\"", IncludePath.Replace('\\', '/'));
 			}
 			foreach (string IncludePath in CompileEnvironment.IncludePaths.SystemIncludePaths)
 			{
-				Arguments += string.Format(" -I\"{0}\"", IncludePath);
+				Arguments += string.Format(" -I\"{0}\"", IncludePath.Replace('\\', '/'));
 			}
 
 			// Add preprocessor definitions to the argument list.
@@ -900,7 +900,7 @@ namespace UnrealBuildTool
 					Result.PrecompiledHeaderFile = PrecompiledHeaderFile;
 
 					// Add the parameters needed to compile the precompiled header file to the command-line.
-					FileArguments += string.Format(" -o \"{0}\"", PrecompiledHeaderFile.AbsolutePath);
+					FileArguments += string.Format(" -o \"{0}\"", PrecompiledHeaderFile.AbsolutePath.Replace('\\', '/'));
 				}
 				else
 				{
@@ -921,11 +921,11 @@ namespace UnrealBuildTool
 					CompileAction.ProducedItems.Add(ObjectFile);
 					Result.ObjectFiles.Add(ObjectFile);
 
-					FileArguments += string.Format(" -o \"{0}\"", ObjectFile.AbsolutePath);
+					FileArguments += string.Format(" -o \"{0}\"", ObjectFile.AbsolutePath.Replace('\\', '/'));
 				}
 
 				// Add the source file path to the command-line.
-				FileArguments += string.Format(" \"{0}\"", SourceFile.AbsolutePath);
+				FileArguments += string.Format(" \"{0}\"", SourceFile.AbsolutePath.Replace('\\', '/'));
 
 				CompileAction.WorkingDirectory = UnrealBuildTool.EngineSourceDirectory.FullName;
 				if (!UsingClang())
@@ -937,8 +937,8 @@ namespace UnrealBuildTool
 					CompileAction.CommandPath = ClangPath;
 				}
 
-				string AllArguments = (Arguments + FileArguments + CompileEnvironment.AdditionalArguments).Replace('\\', '/');
-				// all response lines should have / instead of \
+				string AllArguments = (Arguments + FileArguments + CompileEnvironment.AdditionalArguments);
+				// all response lines should have / instead of \, but we cannot just bulk-replace it here since some \ are used to escape quotes, e.g. Definitions.Add("FOO=TEXT(\"Bar\")");
 
 				Debug.Assert(CompileAction.ProducedItems.Count > 0);
 
