Fix iOS/tvOS CIS crashes deu to not waiting for exit before accessing the exit code
#rnx

diff --git a/Engine/Source/Programs/UnrealBuildTool/System/Utils.cs b/Engine/Source/Programs/UnrealBuildTool/System/Utils.cs
--- a/Engine/Source/Programs/UnrealBuildTool/System/Utils.cs
+++ b/Engine/Source/Programs/UnrealBuildTool/System/Utils.cs
@@ -310,6 +310,7 @@ namespace UnrealBuildTool
 				StreamReader OutputReader = LocalProcess.StandardOutput;
 				// trim off any extraneous new lines, helpful for those one-line outputs
 				FullOutput = OutputReader.ReadToEnd().Trim();
+				LocalProcess.WaitForExit();
 				ExitCode = LocalProcess.ExitCode;
 			}
 
