Attempt to reduce command line length when compiling PCHs, currently causing CIS failure. Now uses relative paths for every include path under the root directory.

diff --git a/Engine/Source/Programs/UnrealBuildTool/Platform/Windows/VCToolChain.cs b/Engine/Source/Programs/UnrealBuildTool/Platform/Windows/VCToolChain.cs
--- a/Engine/Source/Programs/UnrealBuildTool/Platform/Windows/VCToolChain.cs
+++ b/Engine/Source/Programs/UnrealBuildTool/Platform/Windows/VCToolChain.cs
@@ -91,7 +91,7 @@ namespace UnrealBuildTool
 		{
 			// If the value has a space in it and isn't wrapped in quotes, do that now. Make sure it doesn't include a trailing slash, because that will escape the closing quote.
 			string IncludePathString;
-			if(IncludePath.IsUnderDirectory(UnrealBuildTool.EngineSourceDirectory))
+			if(IncludePath.IsUnderDirectory(UnrealBuildTool.RootDirectory))
 			{
 				IncludePathString = IncludePath.MakeRelativeTo(UnrealBuildTool.EngineSourceDirectory);
 			}
