Prevent generated code directories being added to a module's public include paths if the module does not contain any reflected classes. Attempt to reduce compiler command line length and fix CIS errors.

diff --git a/Engine/Source/Programs/UnrealBuildTool/Configuration/UEBuildModuleCPP.cs b/Engine/Source/Programs/UnrealBuildTool/Configuration/UEBuildModuleCPP.cs
--- a/Engine/Source/Programs/UnrealBuildTool/Configuration/UEBuildModuleCPP.cs
+++ b/Engine/Source/Programs/UnrealBuildTool/Configuration/UEBuildModuleCPP.cs
@@ -377,6 +377,18 @@ namespace UnrealBuildTool
 			}
 		}
 
+		/// <summary>
+		/// Remove the generated code directory from the list of public include paths. This is called for modules that do not contain any generated code, once the determination is
+		/// made that UHT will not generate anything, and is necessary to reduce the length of the command line to the compiler.
+		/// </summary>
+		public void RemoveGeneratedCodeDirectory()
+		{
+			if(GeneratedCodeDirectory != null)
+			{
+				PublicIncludePaths.Remove(GeneratedCodeDirectory);
+			}
+		}
+
 		/// <summary>
 		/// Path to the precompiled manifest location
 		/// </summary>
diff --git a/Engine/Source/Programs/UnrealBuildTool/System/ExternalExecution.cs b/Engine/Source/Programs/UnrealBuildTool/System/ExternalExecution.cs
--- a/Engine/Source/Programs/UnrealBuildTool/System/ExternalExecution.cs
+++ b/Engine/Source/Programs/UnrealBuildTool/System/ExternalExecution.cs
@@ -539,6 +539,9 @@ namespace UnrealBuildTool
 							Directory.Delete(Module.GeneratedCodeDirectory.FullName, true);
 						}
 					}
+
+					// Remove it from the list of public include paths
+					Module.RemoveGeneratedCodeDirectory();
 				}
 			}
 
