Check if all include paths exist before adding them to the compiler command line.

diff --git a/Engine/Source/Programs/UnrealBuildTool/Configuration/UEBuildModule.cs b/Engine/Source/Programs/UnrealBuildTool/Configuration/UEBuildModule.cs
--- a/Engine/Source/Programs/UnrealBuildTool/Configuration/UEBuildModule.cs
+++ b/Engine/Source/Programs/UnrealBuildTool/Configuration/UEBuildModule.cs
@@ -416,6 +416,18 @@ namespace UnrealBuildTool
 			}
 		}
 
+		private void AddIncludePaths(HashSet<DirectoryReference> IncludePaths, HashSet<DirectoryReference> IncludePathsToAdd)
+		{
+			// Need to check whether directories exist to avoid bloating compiler command line with generated code directories
+			foreach(DirectoryReference IncludePathToAdd in IncludePathsToAdd)
+			{
+				if(DirectoryReference.Exists(IncludePathToAdd))
+				{
+					IncludePaths.Add(IncludePathToAdd);
+				}
+			}
+		}
+
 		/// <summary>
 		/// Sets up the environment for compiling any module that includes the public interface of this module.
 		/// </summary>
@@ -432,10 +444,10 @@ namespace UnrealBuildTool
 			IncludePaths.Add(ModuleDirectory.ParentDirectory);
 
 			// Add this module's public include paths and definitions.
-			IncludePaths.UnionWith(PublicIncludePaths);
+			AddIncludePaths(IncludePaths, PublicIncludePaths);
 			if(bLegacyPublicIncludePaths)
 			{
-				IncludePaths.UnionWith(LegacyPublicIncludePaths);
+				AddIncludePaths(IncludePaths, LegacyPublicIncludePaths);
 			}
 			SystemIncludePaths.UnionWith(PublicSystemIncludePaths);
 			Definitions.AddRange(PublicDefinitions);
diff --git a/Engine/Source/Programs/UnrealBuildTool/Configuration/UEBuildModuleCPP.cs b/Engine/Source/Programs/UnrealBuildTool/Configuration/UEBuildModuleCPP.cs
--- a/Engine/Source/Programs/UnrealBuildTool/Configuration/UEBuildModuleCPP.cs
+++ b/Engine/Source/Programs/UnrealBuildTool/Configuration/UEBuildModuleCPP.cs
@@ -377,18 +377,6 @@ namespace UnrealBuildTool
 			}
 		}
 
-		/// <summary>
-		/// Remove the generated code directory from the list of public include paths. This is called for modules that do not contain any generated code, once the determination is
-		/// made that UHT will not generate anything, and is necessary to reduce the length of the command line to the compiler.
-		/// </summary>
-		public void RemoveGeneratedCodeDirectory()
-		{
-			if(GeneratedCodeDirectory != null)
-			{
-				PublicIncludePaths.Remove(GeneratedCodeDirectory);
-			}
-		}
-
 		/// <summary>
 		/// Path to the precompiled manifest location
 		/// </summary>
diff --git a/Engine/Source/Programs/UnrealBuildTool/System/ExternalExecution.cs b/Engine/Source/Programs/UnrealBuildTool/System/ExternalExecution.cs
--- a/Engine/Source/Programs/UnrealBuildTool/System/ExternalExecution.cs
+++ b/Engine/Source/Programs/UnrealBuildTool/System/ExternalExecution.cs
@@ -539,9 +539,6 @@ namespace UnrealBuildTool
 							Directory.Delete(Module.GeneratedCodeDirectory.FullName, true);
 						}
 					}
-
-					// Remove it from the list of public include paths
-					Module.RemoveGeneratedCodeDirectory();
 				}
 			}
 
