Fix missing IMPLEMENT_SIGNING_KEY_REGISTRATION() macro errors when compiling targets in adaptive non-unity or with PCHs disabled.

We allow any generic C++ macro to be defined through the "Definitions" setting in .build.cs files, but compiler command line switches are much stricter (not allowing function macros, and removing quote characters).

diff --git a/Engine/Source/Programs/UnrealBuildTool/Configuration/UEBuildModuleCPP.cs b/Engine/Source/Programs/UnrealBuildTool/Configuration/UEBuildModuleCPP.cs
--- a/Engine/Source/Programs/UnrealBuildTool/Configuration/UEBuildModuleCPP.cs
+++ b/Engine/Source/Programs/UnrealBuildTool/Configuration/UEBuildModuleCPP.cs
@@ -584,6 +584,9 @@ namespace UnrealBuildTool
 				}
 			}
 
+			// Write all the definitions to a separate file
+			CreateHeaderForDefinitions(CompileEnvironment, IntermediateDirectory, null);
+
 			// Compile CPP files
 			List<FileItem> CPPFilesToCompile = SourceFilesToBuild.CPPFiles;
 			if (bModuleUsesUnityBuild)
@@ -825,6 +828,9 @@ namespace UnrealBuildTool
 				}
 				AdaptiveUnityEnvironment.PrecompiledHeaderAction = PrecompiledHeaderAction.None;
 
+				// Write all the definitions out to a separate file
+				CreateHeaderForDefinitions(AdaptiveUnityEnvironment, IntermediateDirectory, "Adaptive");
+
 				// Compile the files
 				CPPOutput AdaptiveOutput = ToolChain.CompileCPPFiles(AdaptiveUnityEnvironment, AdaptiveFiles, IntermediateDirectory, Name, ActionGraph);
 
@@ -836,6 +842,37 @@ namespace UnrealBuildTool
 			return OutputFiles;
 		}
 
+		/// <summary>
+		/// Creates a header file containing all the preprocessor definitions for a compile environment, and force-include it. We allow a more flexible syntax for preprocessor definitions than
+		/// is typically allowed on the command line (allowing function macros or double-quote characters, for example). Ensuring all definitions are specified in a header files ensures consistent
+		/// behavior.
+		/// </summary>
+		/// <param name="CompileEnvironment">The compile environment</param>
+		/// <param name="IntermediateDirectory">Directory to create the intermediate file</param>
+		/// <param name="HeaderSuffix">Suffix for the included file</param>
+		static void CreateHeaderForDefinitions(CppCompileEnvironment CompileEnvironment, DirectoryReference IntermediateDirectory, string HeaderSuffix)
+		{
+			if(CompileEnvironment.Definitions.Count > 0)
+			{
+				StringBuilder PrivateDefinitionsName = new StringBuilder("Definitions");
+				if(!String.IsNullOrEmpty(HeaderSuffix))
+				{
+					PrivateDefinitionsName.Append('.');
+					PrivateDefinitionsName.Append(HeaderSuffix);
+				}
+				PrivateDefinitionsName.Append(".h");
+
+				FileReference PrivateDefinitionsFile = FileReference.Combine(IntermediateDirectory, PrivateDefinitionsName.ToString());
+				using (StringWriter Writer = new StringWriter())
+				{
+					WriteDefinitions(CompileEnvironment.Definitions, Writer);
+					CompileEnvironment.Definitions.Clear();
+					FileItem.CreateIntermediateTextFile(PrivateDefinitionsFile, Writer.ToString());
+				}
+				CompileEnvironment.ForceIncludeFiles.Add(PrivateDefinitionsFile);
+			}
+		}
+
 		/// <summary>
 		/// Create a header file containing the module definitions, which also includes the PCH itself. Including through another file is necessary on 
 		/// Clang, since we get warnings about #pragma once otherwise, but it also allows us to consistently define the preprocessor state on all 
