Add force-included headers as prerequisites when remote compiling on IOS.

diff --git a/Engine/Source/Programs/UnrealBuildTool/Platform/IOS/IOSToolChain.cs b/Engine/Source/Programs/UnrealBuildTool/Platform/IOS/IOSToolChain.cs
--- a/Engine/Source/Programs/UnrealBuildTool/Platform/IOS/IOSToolChain.cs
+++ b/Engine/Source/Programs/UnrealBuildTool/Platform/IOS/IOSToolChain.cs
@@ -643,7 +643,7 @@ namespace UnrealBuildTool
 
 			foreach(FileReference ForceIncludeFile in CompileEnvironment.ForceIncludeFiles)
 			{
-				PCHArguments += String.Format(" -include \"{0}\"", ConvertPath(ForceIncludeFile.FullName));
+				Arguments += String.Format(" -include \"{0}\"", ConvertPath(ForceIncludeFile.FullName));
 			}
 
 			// Add include paths to the argument list.
@@ -729,6 +729,12 @@ namespace UnrealBuildTool
 				// Add the C++ source file and its included files to the prerequisite item list.
 				AddPrerequisiteSourceFile(CompileEnvironment, SourceFile, CompileAction.PrerequisiteItems);
 
+				// Upload the force included files
+				foreach(FileReference ForceIncludeFile in CompileEnvironment.ForceIncludeFiles)
+				{
+					AddPrerequisiteSourceFile(CompileEnvironment, FileItem.GetItemByFileReference(ForceIncludeFile), CompileAction.PrerequisiteItems);
+				}
+
 				string OutputFilePath = null;
 				if (CompileEnvironment.PrecompiledHeaderAction == PrecompiledHeaderAction.Create)
 				{
