Fix android deploy when there is no project file (i.e. UE4)

diff --git a/Engine/Source/Programs/UnrealBuildTool/Platform/Android/UEDeployAndroid.cs b/Engine/Source/Programs/UnrealBuildTool/Platform/Android/UEDeployAndroid.cs
--- a/Engine/Source/Programs/UnrealBuildTool/Platform/Android/UEDeployAndroid.cs
+++ b/Engine/Source/Programs/UnrealBuildTool/Platform/Android/UEDeployAndroid.cs
@@ -3677,7 +3677,8 @@ namespace UnrealBuildTool
 
 			// make an apk at the end of compiling, so that we can run without packaging (debugger, cook on the fly, etc)
 			string RelativeEnginePath = UnrealBuildTool.EngineDirectory.MakeRelativeTo(DirectoryReference.GetCurrentDirectory());
-			MakeApk(ToolChain, InTarget.ProjectFile.GetFileNameWithoutAnyExtensions(), InTarget.TargetType, InTarget.ProjectDirectory.FullName, BaseSoName, RelativeEnginePath, bForDistribution: false, CookFlavor: "",
+			string TargetName = (InTarget.ProjectFile == null ? InTarget.TargetName : InTarget.ProjectFile.GetFileNameWithoutAnyExtensions());
+			MakeApk(ToolChain, TargetName, InTarget.TargetType, InTarget.ProjectDirectory.FullName, BaseSoName, RelativeEnginePath, bForDistribution: false, CookFlavor: "",
 				bMakeSeparateApks: ShouldMakeSeparateApks(), bIncrementalPackage: true, bDisallowPackagingDataInApk: false, bDisallowExternalFilesDir: true);
 
 			// if we made any non-standard .apk files, the generated debugger settings may be wrong
