Fix an incremental build issue caused by force-included files not being treated as dependencies, resulting in that definitions in the compile environment for adaptive non-unity files not triggering a build.

diff --git a/Engine/Source/Programs/UnrealBuildTool/Configuration/UEBuildModuleCPP.cs b/Engine/Source/Programs/UnrealBuildTool/Configuration/UEBuildModuleCPP.cs
--- a/Engine/Source/Programs/UnrealBuildTool/Configuration/UEBuildModuleCPP.cs
+++ b/Engine/Source/Programs/UnrealBuildTool/Configuration/UEBuildModuleCPP.cs
@@ -541,6 +541,8 @@ namespace UnrealBuildTool
 						PrecompiledHeaderInstance Instance = FindOrCreateSharedPCH(ToolChain, Template, ModuleCompileEnvironment.bOptimizeCode, ModuleCompileEnvironment.bUseRTTI, ActionGraph);
 
 						FileReference PrivateDefinitionsFile = FileReference.Combine(IntermediateDirectory, String.Format("Definitions.{0}.h", Name));
+
+						FileItem PrivateDefinitionsFileItem;
 						using (StringWriter Writer = new StringWriter())
 						{
 							// Remove the module _API definition for cases where there are circular dependencies between the shared PCH module and modules using it
@@ -555,12 +557,12 @@ namespace UnrealBuildTool
 							}
 
 							WriteDefinitions(CompileEnvironment.Definitions, Writer);
-							FileItem.CreateIntermediateTextFile(PrivateDefinitionsFile, Writer.ToString());
+							PrivateDefinitionsFileItem = FileItem.CreateIntermediateTextFile(PrivateDefinitionsFile, Writer.ToString());
 						}
 
 						CompileEnvironment = new CppCompileEnvironment(CompileEnvironment);
 						CompileEnvironment.Definitions.Clear();
-						CompileEnvironment.ForceIncludeFiles.Add(PrivateDefinitionsFile);
+						CompileEnvironment.ForceIncludeFiles.Add(PrivateDefinitionsFileItem);
 						CompileEnvironment.PrecompiledHeaderAction = PrecompiledHeaderAction.Include;
 						CompileEnvironment.PrecompiledHeaderIncludeFilename = Instance.HeaderFile.Location;
 						CompileEnvironment.PrecompiledHeaderFile = Instance.Output.PrecompiledHeaderFile;
@@ -867,9 +869,10 @@ namespace UnrealBuildTool
 				{
 					WriteDefinitions(CompileEnvironment.Definitions, Writer);
 					CompileEnvironment.Definitions.Clear();
-					FileItem.CreateIntermediateTextFile(PrivateDefinitionsFile, Writer.ToString());
+
+					FileItem PrivateDefinitionsFileItem = FileItem.CreateIntermediateTextFile(PrivateDefinitionsFile, Writer.ToString());
+					CompileEnvironment.ForceIncludeFiles.Add(PrivateDefinitionsFileItem);
 				}
-				CompileEnvironment.ForceIncludeFiles.Add(PrivateDefinitionsFile);
 			}
 		}
 
diff --git a/Engine/Source/Programs/UnrealBuildTool/Platform/Android/AndroidToolChain.cs b/Engine/Source/Programs/UnrealBuildTool/Platform/Android/AndroidToolChain.cs
--- a/Engine/Source/Programs/UnrealBuildTool/Platform/Android/AndroidToolChain.cs
+++ b/Engine/Source/Programs/UnrealBuildTool/Platform/Android/AndroidToolChain.cs
@@ -1248,9 +1248,9 @@ namespace UnrealBuildTool
 						PCHArguments += string.Format(" -include \"{0}\"", InlineArchName(BasePCHName, Arch, GPUArchitecture));
 					}
 
-					foreach(FileReference ForceIncludeFile in CompileEnvironment.ForceIncludeFiles)
+					foreach(FileItem ForceIncludeFile in CompileEnvironment.ForceIncludeFiles)
 					{
-						PCHArguments += string.Format(" -include \"{0}\"", ForceIncludeFile.FullName);
+						PCHArguments += string.Format(" -include \"{0}\"", ForceIncludeFile.Location);
 					}
 
 					// Add include paths to the argument list (filtered by architecture)
@@ -1272,6 +1272,8 @@ namespace UnrealBuildTool
 					foreach (FileItem SourceFile in InputFiles)
 					{
 						Action CompileAction = ActionGraph.Add(ActionType.Compile);
+						CompileAction.PrerequisiteItems.AddRange(CompileEnvironment.ForceIncludeFiles);
+
 						string FileArguments = "";
 						bool bIsPlainCFile = Path.GetExtension(SourceFile.AbsolutePath).ToUpperInvariant() == ".C";
 						bool bDisableShadowWarning = false;
diff --git a/Engine/Source/Programs/UnrealBuildTool/Platform/HTML5/HTML5ToolChain.cs b/Engine/Source/Programs/UnrealBuildTool/Platform/HTML5/HTML5ToolChain.cs
--- a/Engine/Source/Programs/UnrealBuildTool/Platform/HTML5/HTML5ToolChain.cs
+++ b/Engine/Source/Programs/UnrealBuildTool/Platform/HTML5/HTML5ToolChain.cs
@@ -443,15 +443,16 @@ namespace UnrealBuildTool
 			}
 
 			// Force include all the requested headers
-			foreach(FileReference ForceIncludeFile in CompileEnvironment.ForceIncludeFiles)
+			foreach(FileItem ForceIncludeFile in CompileEnvironment.ForceIncludeFiles)
 			{
-				Arguments += String.Format(" -include \"{0}\"", ForceIncludeFile.FullName);
+				Arguments += String.Format(" -include \"{0}\"", ForceIncludeFile.Location);
 			}
 
 			foreach (FileItem SourceFile in InputFiles)
 			{
 				Action CompileAction = ActionGraph.Add(ActionType.Compile);
 				CompileAction.CommandDescription = "Compile";
+				CompileAction.PrerequisiteItems.AddRange(CompileEnvironment.ForceIncludeFiles);
 //				CompileAction.bPrintDebugInfo = true;
 
 				bool bIsPlainCFile = Path.GetExtension(SourceFile.AbsolutePath).ToUpperInvariant() == ".C";
diff --git a/Engine/Source/Programs/UnrealBuildTool/Platform/IOS/IOSToolChain.cs b/Engine/Source/Programs/UnrealBuildTool/Platform/IOS/IOSToolChain.cs
--- a/Engine/Source/Programs/UnrealBuildTool/Platform/IOS/IOSToolChain.cs
+++ b/Engine/Source/Programs/UnrealBuildTool/Platform/IOS/IOSToolChain.cs
@@ -646,9 +646,9 @@ namespace UnrealBuildTool
 				PCHArguments += string.Format(" -include \"{0}\"", ConvertPath(CompileEnvironment.PrecompiledHeaderFile.AbsolutePath.Replace(PrecompiledFileExtension, "")));
 			}
 
-			foreach(FileReference ForceIncludeFile in CompileEnvironment.ForceIncludeFiles)
+			foreach(FileItem ForceIncludeFile in CompileEnvironment.ForceIncludeFiles)
 			{
-				Arguments += String.Format(" -include \"{0}\"", ConvertPath(ForceIncludeFile.FullName));
+				Arguments += String.Format(" -include \"{0}\"", ConvertPath(ForceIncludeFile.Location.FullName));
 			}
 
 			// Add include paths to the argument list.
@@ -735,9 +735,9 @@ namespace UnrealBuildTool
 				AddPrerequisiteSourceFile(CompileEnvironment, SourceFile, CompileAction.PrerequisiteItems);
 
 				// Upload the force included files
-				foreach(FileReference ForceIncludeFile in CompileEnvironment.ForceIncludeFiles)
+				foreach(FileItem ForceIncludeFile in CompileEnvironment.ForceIncludeFiles)
 				{
-					AddPrerequisiteSourceFile(CompileEnvironment, FileItem.GetItemByFileReference(ForceIncludeFile), CompileAction.PrerequisiteItems);
+					AddPrerequisiteSourceFile(CompileEnvironment, ForceIncludeFile, CompileAction.PrerequisiteItems);
 				}
 
 				string OutputFilePath = null;
diff --git a/Engine/Source/Programs/UnrealBuildTool/Platform/Linux/LinuxToolChain.cs b/Engine/Source/Programs/UnrealBuildTool/Platform/Linux/LinuxToolChain.cs
--- a/Engine/Source/Programs/UnrealBuildTool/Platform/Linux/LinuxToolChain.cs
+++ b/Engine/Source/Programs/UnrealBuildTool/Platform/Linux/LinuxToolChain.cs
@@ -863,9 +863,9 @@ namespace UnrealBuildTool
 				PCHArguments += string.Format(" -include \"{0}\"", CompileEnvironment.PrecompiledHeaderFile.AbsolutePath.Replace(PrecompiledFileExtension, "").Replace('\\', '/'));
 			}
 
-			foreach(FileReference ForceIncludeFile in CompileEnvironment.ForceIncludeFiles)
+			foreach(FileItem ForceIncludeFile in CompileEnvironment.ForceIncludeFiles)
 			{
-				PCHArguments += String.Format(" -include \"{0}\"", ForceIncludeFile.FullName.Replace('\\', '/'));
+				PCHArguments += String.Format(" -include \"{0}\"", ForceIncludeFile.Location.FullName.Replace('\\', '/'));
 			}
 
 			// Add include paths to the argument list.
@@ -889,6 +889,8 @@ namespace UnrealBuildTool
 			foreach (FileItem SourceFile in InputFiles)
 			{
 				Action CompileAction = ActionGraph.Add(ActionType.Compile);
+				CompileAction.PrerequisiteItems.AddRange(CompileEnvironment.ForceIncludeFiles);
+
 				string FileArguments = "";
 				string Extension = Path.GetExtension(SourceFile.AbsolutePath).ToUpperInvariant();
 
diff --git a/Engine/Source/Programs/UnrealBuildTool/Platform/Mac/MacToolChain.cs b/Engine/Source/Programs/UnrealBuildTool/Platform/Mac/MacToolChain.cs
--- a/Engine/Source/Programs/UnrealBuildTool/Platform/Mac/MacToolChain.cs
+++ b/Engine/Source/Programs/UnrealBuildTool/Platform/Mac/MacToolChain.cs
@@ -416,9 +416,9 @@ namespace UnrealBuildTool
 				PCHArguments.Append("\"");
 			}
 
-			foreach(FileReference ForceIncludeFile in CompileEnvironment.ForceIncludeFiles)
+			foreach(FileItem ForceIncludeFile in CompileEnvironment.ForceIncludeFiles)
 			{
-				PCHArguments.Append(String.Format(" -include \"{0}\"", ForceIncludeFile.FullName));
+				PCHArguments.Append(String.Format(" -include \"{0}\"", ForceIncludeFile.Location));
 			}
 
 			// Add include paths to the argument list.
@@ -471,6 +471,8 @@ namespace UnrealBuildTool
 			foreach (FileItem SourceFile in InputFiles)
 			{
 				Action CompileAction = ActionGraph.Add(ActionType.Compile);
+				CompileAction.PrerequisiteItems.AddRange(CompileEnvironment.ForceIncludeFiles);
+
 				string FileArguments = "";
 				string Extension = Path.GetExtension(SourceFile.AbsolutePath).ToUpperInvariant();
 
diff --git a/Engine/Source/Programs/UnrealBuildTool/Platform/Windows/PVSToolChain.cs b/Engine/Source/Programs/UnrealBuildTool/Platform/Windows/PVSToolChain.cs
--- a/Engine/Source/Programs/UnrealBuildTool/Platform/Windows/PVSToolChain.cs
+++ b/Engine/Source/Programs/UnrealBuildTool/Platform/Windows/PVSToolChain.cs
@@ -47,9 +47,9 @@ namespace UnrealBuildTool
 			{
 				SharedArguments.Add(String.Format("/D \"{0}\"", Definition));
 			}
-			foreach(FileReference ForceIncludeFile in CompileEnvironment.ForceIncludeFiles)
+			foreach(FileItem ForceIncludeFile in CompileEnvironment.ForceIncludeFiles)
 			{
-				SharedArguments.Add(String.Format("/FI\"{0}\"", ForceIncludeFile.FullName));
+				SharedArguments.Add(String.Format("/FI\"{0}\"", ForceIncludeFile.Location));
 			}
 
 			// Get the path to PVS studio
@@ -82,6 +82,7 @@ namespace UnrealBuildTool
 				PreprocessAction.CommandPath = EnvVars.CompilerPath.FullName;
 				PreprocessAction.WorkingDirectory = UnrealBuildTool.EngineSourceDirectory.FullName;
 				PreprocessAction.CommandArguments = " @\"" + ResponseFileItem.AbsolutePath + "\"";
+				PreprocessAction.PrerequisiteItems.AddRange(CompileEnvironment.ForceIncludeFiles);
 				PreprocessAction.PrerequisiteItems.Add(SourceFile);
 				PreprocessAction.PrerequisiteItems.Add(ResponseFileItem);
 				PreprocessAction.ProducedItems.Add(PreprocessedFileItem);
diff --git a/Engine/Source/Programs/UnrealBuildTool/Platform/Windows/VCToolChain.cs b/Engine/Source/Programs/UnrealBuildTool/Platform/Windows/VCToolChain.cs
--- a/Engine/Source/Programs/UnrealBuildTool/Platform/Windows/VCToolChain.cs
+++ b/Engine/Source/Programs/UnrealBuildTool/Platform/Windows/VCToolChain.cs
@@ -1201,9 +1201,9 @@ namespace UnrealBuildTool
 					CompileAction.StatusDescription = Path.GetFileName(SourceFile.AbsolutePath);
 				}
 
-				foreach(FileReference ForceIncludeFile in CompileEnvironment.ForceIncludeFiles)
+				foreach(FileItem ForceIncludeFile in CompileEnvironment.ForceIncludeFiles)
 				{
-					FileArguments.Add(String.Format("/FI\"{0}\"", ForceIncludeFile.FullName));
+					FileArguments.Add(String.Format("/FI\"{0}\"", ForceIncludeFile.Location));
 				}
 
 				if (bEmitsObjectFile)
@@ -1294,6 +1294,7 @@ namespace UnrealBuildTool
 
 				CompileAction.WorkingDirectory = UnrealBuildTool.EngineSourceDirectory.FullName;
 				CompileAction.CommandPath = EnvVars.CompilerPath.FullName;
+				CompileAction.PrerequisiteItems.AddRange(CompileEnvironment.ForceIncludeFiles);
 
 				string[] AdditionalArguments = String.IsNullOrEmpty(CompileEnvironment.AdditionalArguments)? new string[0] : new string[] { CompileEnvironment.AdditionalArguments };
 
@@ -1360,6 +1361,7 @@ namespace UnrealBuildTool
 				CompileAction.WorkingDirectory = UnrealBuildTool.EngineSourceDirectory.FullName;
 				CompileAction.CommandPath = EnvVars.ResourceCompilerPath.FullName;
 				CompileAction.StatusDescription = Path.GetFileName(RCFile.AbsolutePath);
+				CompileAction.PrerequisiteItems.AddRange(CompileEnvironment.ForceIncludeFiles);
 
 				// Resource tool can run remotely if possible
 				CompileAction.bCanExecuteRemotely = true;
diff --git a/Engine/Source/Programs/UnrealBuildTool/System/CppCompileEnvironment.cs b/Engine/Source/Programs/UnrealBuildTool/System/CppCompileEnvironment.cs
--- a/Engine/Source/Programs/UnrealBuildTool/System/CppCompileEnvironment.cs
+++ b/Engine/Source/Programs/UnrealBuildTool/System/CppCompileEnvironment.cs
@@ -277,7 +277,7 @@ namespace UnrealBuildTool
 		/// <summary>
 		/// List of header files to force include
 		/// </summary>
-		public List<FileReference> ForceIncludeFiles = new List<FileReference>();
+		public List<FileItem> ForceIncludeFiles = new List<FileItem>();
 
 		/// <summary>
 		/// The C++ preprocessor definitions to use.
