UBT: Fix errors during incremental builds if a build contains an exported symbol change, but the build fails after the source files have been compiled but before the import library is built.

ActionGraph was ignoring timestamps of LIB files when determining dependencies, so if one incremental build generates .obj files but stops before generating the import library, subsequent builds will not consider the import library to be out of date.

diff --git a/Engine/Source/Programs/UnrealBuildTool/Platform/Windows/VCToolChain.cs b/Engine/Source/Programs/UnrealBuildTool/Platform/Windows/VCToolChain.cs
--- a/Engine/Source/Programs/UnrealBuildTool/Platform/Windows/VCToolChain.cs
+++ b/Engine/Source/Programs/UnrealBuildTool/Platform/Windows/VCToolChain.cs
@@ -1657,12 +1657,19 @@ namespace UnrealBuildTool
 			LinkAction.PrerequisiteItems.AddRange(PrerequisiteItems);
 			LinkAction.StatusDescription = Path.GetFileName(OutputFile.AbsolutePath);
 			LinkAction.bUseIncrementalLinking = LinkEnvironment.bUseIncrementalLinking;
+
 			// ensure compiler timings are captured when we execute the action.
 			if (!WindowsPlatform.bCompileWithClang && LinkEnvironment.bPrintTimingInfo)
 			{
 				LinkAction.bPrintDebugInfo = true;
 			}
 
+			// VS 15.3+ does not touch lib files if they do not contain any modifications, but we need to ensure the timestamps are updated to avoid repeatedly building them.
+			if (bBuildImportLibraryOnly || (LinkEnvironment.bHasExports && !bIsBuildingLibrary))
+			{
+				LinkAction.bShouldDeleteProducedItems = true; 
+			}
+
 			if (WindowsPlatform.bCompileWithClang || WindowsPlatform.bCompileWithICL)
 			{
 				LinkAction.OutputEventHandler = new DataReceivedEventHandler(ClangCompilerOutputFormatter);
diff --git a/Engine/Source/Programs/UnrealBuildTool/System/ActionGraph.cs b/Engine/Source/Programs/UnrealBuildTool/System/ActionGraph.cs
--- a/Engine/Source/Programs/UnrealBuildTool/System/ActionGraph.cs
+++ b/Engine/Source/Programs/UnrealBuildTool/System/ActionGraph.cs
@@ -947,13 +947,6 @@ namespace UnrealBuildTool
 					// legitimate output to always be considered outdated.
 					if (ProducedItem.bExists && (ProducedItem.bIsRemoteFile || ProducedItem.Length > 0 || ProducedItem.IsDirectory))
 					{
-						// VS 15.3+ does not touch lib files if they do not contain any modifications, so ignore any timestamps on them.
-						String ProducedItemExtension = Path.GetExtension(ProducedItem.AbsolutePath).ToUpperInvariant();
-						if (ProducedItemExtension == ".LIB" || ProducedItemExtension == ".EXP" || ProducedItemExtension == ".PDB")
-						{
-							continue;
-						}
-
 						// Use the oldest produced item's time as the last execution time.
 						if (ProducedItem.LastWriteTime < LastExecutionTime)
 						{
@@ -1195,7 +1188,7 @@ namespace UnrealBuildTool
 					{
 						if (ProducedItem.bExists && (bShouldDeleteAllFiles || OutdatedAction.bShouldDeleteProducedItems))
 						{
-							Log.TraceVerbose("Deleting outdated item: {0}", ProducedItem.AbsolutePath);
+							Log.TraceLog("Deleting outdated item: {0}", ProducedItem.AbsolutePath);
 							ProducedItem.Delete();
 						}
 					}
