Fix definitions for UE4Game module being dependent on project being built (fixes CIS error in ProtoStarEditor).

diff --git a/Engine/Source/Programs/UnrealBuildTool/Configuration/UEBuildModuleCPP.cs b/Engine/Source/Programs/UnrealBuildTool/Configuration/UEBuildModuleCPP.cs
--- a/Engine/Source/Programs/UnrealBuildTool/Configuration/UEBuildModuleCPP.cs
+++ b/Engine/Source/Programs/UnrealBuildTool/Configuration/UEBuildModuleCPP.cs
@@ -1216,7 +1216,8 @@ namespace UnrealBuildTool
 			// For game modules, set the define for the project name. This will be used by the IMPLEMENT_PRIMARY_GAME_MODULE macro.
 			if (!Rules.bTreatAsEngineModule)
 			{
-				if (Target.ProjectFile != null)
+				// Make sure we don't set any define for a non-engine module that's under the engine directory (eg. UE4Game)
+				if (Target.ProjectFile != null && RulesFile.IsUnderDirectory(Target.ProjectFile.Directory))
 				{
 					string ProjectName = Target.ProjectFile.GetFileNameWithoutExtension();
 					Result.Definitions.Add(String.Format("UE_PROJECT_NAME={0}", ProjectName));
