Update restricted platforms list.

diff --git a/Engine/Source/Programs/UnrealBuildTool/System/RestrictedFolders.cs b/Engine/Source/Programs/UnrealBuildTool/System/RestrictedFolders.cs
--- a/Engine/Source/Programs/UnrealBuildTool/System/RestrictedFolders.cs
+++ b/Engine/Source/Programs/UnrealBuildTool/System/RestrictedFolders.cs
@@ -1,4 +1,4 @@
-﻿// Copyright 1998-2018 Epic Games, Inc. All Rights Reserved.
+// Copyright 1998-2018 Epic Games, Inc. All Rights Reserved.
 
 using System;
 using System.Collections.Generic;
@@ -44,6 +44,11 @@ namespace UnrealBuildTool
 		/// Switch source files
 		/// </summary>
 		Switch,
+
+		/// <summary>
+		/// Quail source files
+		/// </summary>
+		Quail,
 	}
 
 	/// <summary>
