Bringing over PVS-Studio changes in 3945532 to Main

diff --git a/Engine/Source/Programs/UnrealBuildTool/Platform/Windows/PVSToolChain.cs b/Engine/Source/Programs/UnrealBuildTool/Platform/Windows/PVSToolChain.cs
--- a/Engine/Source/Programs/UnrealBuildTool/Platform/Windows/PVSToolChain.cs
+++ b/Engine/Source/Programs/UnrealBuildTool/Platform/Windows/PVSToolChain.cs
@@ -54,7 +54,21 @@ namespace UnrealBuildTool
 
 			// Get the path to PVS studio
 			FileReference AnalyzerFile = new FileReference(Path.Combine(Environment.GetFolderPath(Environment.SpecialFolder.ProgramFilesX86), "PVS-Studio", "x64", "PVS-Studio.exe"));
-			if(!FileReference.Exists(AnalyzerFile))
+
+			// Check if AutoSDK is set up and PVS-Studio exists in it.
+			bool UseSavedLicense = false;
+			FileReference EnginePVS = FileReference.Combine(UnrealBuildTool.RootDirectory, "Engine", "Extras", "ThirdPartyNotUE", "NoRedist", "PVS-Studio", "x64", "PVS-Studio.exe");
+			FileReference EnginePVSLicense = FileReference.Combine(UnrealBuildTool.RootDirectory, "Engine", "Extras", "ThirdPartyNotUE", "NoRedist", "PVS-Studio", "x64", "PVS-Studio.lic");
+			if (FileReference.Exists(EnginePVS))
+			{
+				if(FileReference.Exists(EnginePVSLicense))
+				{
+					UseSavedLicense = true;
+				}
+				AnalyzerFile = EnginePVS;
+			}
+			
+			if (!FileReference.Exists(AnalyzerFile))
 			{
 				throw new BuildException("Unable to find PVS-Studio at {0}", AnalyzerFile);
 			}
@@ -120,7 +134,7 @@ namespace UnrealBuildTool
 				AnalyzeAction.StatusDescription = BaseFileName;
 				AnalyzeAction.WorkingDirectory = UnrealBuildTool.EngineSourceDirectory.FullName;
 				AnalyzeAction.CommandPath = AnalyzerFile.FullName;
-				AnalyzeAction.CommandArguments = String.Format("--cl-params \"{0}\" --source-file \"{1}\" --output-file \"{2}\" --cfg \"{3}\" --analysis-mode 4", PreprocessAction.CommandArguments, SourceFile.AbsolutePath, OutputFileLocation, ConfigFileItem.AbsolutePath);
+				AnalyzeAction.CommandArguments = String.Format("--cl-params \"{0}\" --source-file \"{1}\" --output-file \"{2}\" --cfg \"{3}\" --analysis-mode 4 {4}", PreprocessAction.CommandArguments, SourceFile.AbsolutePath, OutputFileLocation, ConfigFileItem.AbsolutePath, (UseSavedLicense ? string.Format("--lic-file \"{0}\"", EnginePVSLicense.FullName) : string.Empty));
 				AnalyzeAction.PrerequisiteItems.Add(ConfigFileItem);
 				AnalyzeAction.PrerequisiteItems.Add(PreprocessedFileItem);
 				AnalyzeAction.ProducedItems.Add(OutputFileItem);
