Fixing up erroneous paths to pvs exe

diff --git a/Engine/Source/Programs/UnrealBuildTool/Platform/Windows/PVSToolChain.cs b/Engine/Source/Programs/UnrealBuildTool/Platform/Windows/PVSToolChain.cs
--- a/Engine/Source/Programs/UnrealBuildTool/Platform/Windows/PVSToolChain.cs
+++ b/Engine/Source/Programs/UnrealBuildTool/Platform/Windows/PVSToolChain.cs
@@ -57,8 +57,8 @@ namespace UnrealBuildTool
 
 			// Check if AutoSDK is set up and PVS-Studio exists in it.
 			bool UseSavedLicense = false;
-			FileReference EnginePVS = FileReference.Combine(UnrealBuildTool.RootDirectory, "Engine", "Extras", "ThirdPartyNotUE", "NoRedist", "PVS-Studio", "x64", "PVS-Studio.exe");
-			FileReference EnginePVSLicense = FileReference.Combine(UnrealBuildTool.RootDirectory, "Engine", "Extras", "ThirdPartyNotUE", "NoRedist", "PVS-Studio", "x64", "PVS-Studio.lic");
+			FileReference EnginePVS = FileReference.Combine(UnrealBuildTool.RootDirectory, "Engine", "Extras", "ThirdPartyNotUE", "NoRedist", "PVS-Studio", "PVS-Studio.exe");
+			FileReference EnginePVSLicense = FileReference.Combine(UnrealBuildTool.RootDirectory, "Engine", "Extras", "ThirdPartyNotUE", "NoRedist", "PVS-Studio", "PVS-Studio.lic");
 			if (FileReference.Exists(EnginePVS))
 			{
 				if(FileReference.Exists(EnginePVSLicense))
