UBT: Fix incorrect paths for precompile-only binaries. Prevents QAGame running with binary distro.

diff --git a/Engine/Source/Programs/UnrealBuildTool/Configuration/UEBuildTarget.cs b/Engine/Source/Programs/UnrealBuildTool/Configuration/UEBuildTarget.cs
--- a/Engine/Source/Programs/UnrealBuildTool/Configuration/UEBuildTarget.cs
+++ b/Engine/Source/Programs/UnrealBuildTool/Configuration/UEBuildTarget.cs
@@ -2646,6 +2646,7 @@ namespace UnrealBuildTool
 			if (Platform == UnrealTargetPlatform.Mac && !Rules.bIsBuildingConsoleApplication)
 			{
 				TargetToolChain.FixBundleBinariesPaths(this, Binaries);
+				TargetToolChain.FixBundleBinariesPaths(this, PrecompileOnlyBinaries);
 			}
 		}
 
