Fix bad merge from Release-4.19 branch resulting in partial loss of changes in UEBuildTarget.cs#16, and preventing hot-reload from operating correctly when launched from the editor.

diff --git a/Engine/Source/Programs/UnrealBuildTool/Configuration/TargetDescriptor.cs b/Engine/Source/Programs/UnrealBuildTool/Configuration/TargetDescriptor.cs
--- a/Engine/Source/Programs/UnrealBuildTool/Configuration/TargetDescriptor.cs
+++ b/Engine/Source/Programs/UnrealBuildTool/Configuration/TargetDescriptor.cs
@@ -19,7 +19,6 @@ namespace UnrealBuildTool
 		public UnrealTargetConfiguration Configuration;
 		public string Architecture;
 		public bool bIsEditorRecompile;
-		public string RemoteRoot;
 		public List<OnlyModule> OnlyModules;
 		public FileReference ForeignPlugin;
 		public string ForceReceiptFileName;
@@ -30,7 +29,6 @@ namespace UnrealBuildTool
 			UnrealTargetConfiguration Configuration = UnrealTargetConfiguration.Unknown;
 			List<string> TargetNames = new List<string>();
 			string Architecture = null;
-			string RemoteRoot = null;
 			List<OnlyModule> OnlyModules = new List<OnlyModule>();
 			FileReference ForeignPlugin = null;
 			string ForceReceiptFileName = null;
@@ -71,90 +69,49 @@ namespace UnrealBuildTool
 				}
 				else
 				{
-					switch (Arguments[ArgumentIndex].ToUpperInvariant())
+					string Value;
+					if(ParseArgumentValue(Argument, "-Module=", out Value))
 					{
-						case "-MODULE":
-							// Specifies a module to recompile.  Can be specified more than once on the command-line to compile multiple specific modules.
-							{
-								if (ArgumentIndex + 1 >= Arguments.Length)
-								{
-									throw new BuildException("Expected module name after -Module argument, but found nothing.");
-								}
-								string OnlyModuleName = Arguments[++ArgumentIndex];
-
-								OnlyModules.Add(new OnlyModule(OnlyModuleName));
-							}
-							break;
-
-						case "-MODULEWITHSUFFIX":
-							{
-								// Specifies a module name to compile along with a suffix to append to the DLL file name.  Can be specified more than once on the command-line to compile multiple specific modules.
-								if (ArgumentIndex + 2 >= Arguments.Length)
-								{
-									throw new BuildException("Expected module name and module suffix -ModuleWithSuffix argument");
-								}
-
-								string OnlyModuleName = Arguments[++ArgumentIndex];
-								string OnlyModuleSuffix = Arguments[++ArgumentIndex];
-
-								OnlyModules.Add(new OnlyModule(OnlyModuleName, OnlyModuleSuffix));
-							}
-							break;
-
-						case "-PLUGIN":
-							{
-								if (ArgumentIndex + 1 >= Arguments.Length)
-								{
-									throw new BuildException("Expected plugin filename after -Plugin argument, but found nothing.");
-								}
-								if(ForeignPlugin != null)
-								{
-									throw new BuildException("Only one foreign plugin to compile may be specified per invocation");
-								}
-								ForeignPlugin = new FileReference(Arguments[++ArgumentIndex]);
-							}
-							break;
-
-						case "-RECEIPT":
-							{
-								if (ArgumentIndex + 1 >= Arguments.Length)
-								{
-									throw new BuildException("Expected path to the generated receipt after -Receipt argument, but found nothing.");
-								}
-
-								ForceReceiptFileName = Arguments[++ArgumentIndex];
-							}
-							break;
-
-						// -RemoteRoot <RemoteRoot> sets where the generated binaries are CookerSynced.
-						case "-REMOTEROOT":
-							if (ArgumentIndex + 1 >= Arguments.Length)
-							{
-								throw new BuildException("Expected path after -RemoteRoot argument, but found nothing.");
-							}
-							ArgumentIndex++;
-							if (Arguments[ArgumentIndex].StartsWith("xe:\\") == true)
-							{
-								RemoteRoot = Arguments[ArgumentIndex].Substring("xe:\\".Length);
-							}
-							else if (Arguments[ArgumentIndex].StartsWith("devkit:\\") == true)
-							{
-								RemoteRoot = Arguments[ArgumentIndex].Substring("devkit:\\".Length);
-							}
-							break;
-
-						case "-DEPLOY":
-							// Does nothing at the moment...
-							break;
-
-						case "-EDITORRECOMPILE":
-							{
-								bIsEditorRecompile = true;
-							}
-							break;
-
-						default:
-							break;
+						OnlyModules.Add(new OnlyModule(Value));
+					}
+					else if(ParseArgumentValue(Argument, "-ModuleWithSuffix=", out Value))
+					{
+						int SuffixIdx = Value.LastIndexOf(',');
+						if(SuffixIdx == -1)
+						{
+							throw new BuildException("Missing suffix argument from -ModuleWithSuffix=Name,Suffix");
+						}
+						OnlyModules.Add(new OnlyModule(Value.Substring(0, SuffixIdx), Value.Substring(SuffixIdx + 1)));
+					}
+					else if(ParseArgumentValue(Argument, "-Plugin=", out Value))
+					{
+						if(ForeignPlugin != null)
+						{
+							throw new BuildException("Only one foreign plugin to compile may be specified per invocation");
+						}
+						ForeignPlugin = new FileReference(Value);
+					}
+					else if(ParseArgumentValue(Argument, "-Receipt=", out Value))
+					{
+						ForceReceiptFileName = Value;
+					}
+					else if(Argument.Equals("-EditorRecompile", StringComparison.InvariantCultureIgnoreCase))
+					{
+						bIsEditorRecompile = true;
+					}
+					else
+					{
+						switch (Arguments[ArgumentIndex].ToUpperInvariant())
+						{
+							case "-MODULE":
+								throw new BuildException("'-Module <Name>' syntax is no longer supported on the command line. Use '-Module=<Name>' instead.");
+							case "-MODULEWITHSUFFIX":
+								throw new BuildException("'-ModuleWithSuffix <Name> <Suffix>' syntax is no longer supported on the command line. Use '-Module=<Name>,<Suffix>' instead.");
+							case "-PLUGIN":
+								throw new BuildException("'-Plugin <Path>' syntax is no longer supported on the command line. Use '-Plugin=<Path>' instead.");
+							case "-RECEIPT":
+								throw new BuildException("'-Receipt <Path>' syntax is no longer supported on the command line. Use '-Receipt=<Path>' instead.");
+						}
 					}
 				}
 			}
@@ -198,7 +155,6 @@ namespace UnrealBuildTool
 						Configuration = Configuration,
 						Architecture = Architecture,
 						bIsEditorRecompile = bIsEditorRecompile,
-						RemoteRoot = RemoteRoot,
 						OnlyModules = OnlyModules,
 						ForeignPlugin = ForeignPlugin,
 						ForceReceiptFileName = ForceReceiptFileName
@@ -210,5 +166,19 @@ namespace UnrealBuildTool
 			}
 			return Targets;
 		}
+
+		private static bool ParseArgumentValue(string Argument, string Prefix, out string Value)
+		{
+			if(Argument.StartsWith(Prefix, StringComparison.InvariantCultureIgnoreCase))
+			{
+				Value = Argument.Substring(Prefix.Length);
+				return true;
+			}
+			else
+			{
+				Value = null;
+				return false;
+			}
+		}
 	}
 }
