Fix AppIcon errors introduced merging from Release-4.19 -> Release-Staging-4.19
#rnx

diff --git a/Engine/Source/Programs/UnrealBuildTool/Platform/IOS/IOSToolChain.cs b/Engine/Source/Programs/UnrealBuildTool/Platform/IOS/IOSToolChain.cs
--- a/Engine/Source/Programs/UnrealBuildTool/Platform/IOS/IOSToolChain.cs
+++ b/Engine/Source/Programs/UnrealBuildTool/Platform/IOS/IOSToolChain.cs
@@ -136,13 +136,11 @@ namespace UnrealBuildTool
 
 		public override void ModifyBuildProducts(ReadOnlyTargetRules Target, UEBuildBinary Binary, List<string> Libraries, List<UEBuildBundleResource> BundleResources, Dictionary<FileReference, BuildProductType> BuildProducts)
 		{
-			if (Binary.Type != UEBuildBinaryType.StaticLibrary)
+			if (Target.bCreateStubIPA && Binary.Type != UEBuildBinaryType.StaticLibrary)
 			{
-				if (Target.bCreateStubIPA)
-				{
-					FileReference StubFile = FileReference.Combine(Binary.OutputFilePath.Directory, Binary.OutputFilePath.GetFileNameWithoutExtension() + ".stub");
+				FileReference StubFile = FileReference.Combine(Binary.OutputFilePath.Directory, Binary.OutputFilePath.GetFileNameWithoutExtension() + ".stub");
 				BuildProducts.Add(StubFile, BuildProductType.Package);
-				}
+
                 if (CppPlatform == CppPlatform.TVOS)
                 {
 					FileReference AssetFile = FileReference.Combine(Binary.OutputFilePath.Directory, "AssetCatalog", "Assets.car");
