Update restricted folder list.

diff --git a/Engine/Source/Programs/UnrealBuildTool/Configuration/UEBuildPlatform.cs b/Engine/Source/Programs/UnrealBuildTool/Configuration/UEBuildPlatform.cs
--- a/Engine/Source/Programs/UnrealBuildTool/Configuration/UEBuildPlatform.cs
+++ b/Engine/Source/Programs/UnrealBuildTool/Configuration/UEBuildPlatform.cs
@@ -43,7 +43,8 @@ namespace UnrealBuildTool
 			new FileSystemName("NoRedist"),
 			new FileSystemName("PS4"),
 			new FileSystemName("XboxOne"),
-			new FileSystemName("Switch")
+			new FileSystemName("Switch"),
+			new FileSystemName("Quail")
 		};
 
 		/// <summary>
