Fix compile errors caused by FindVSInstallDirs returning Visual C++ 2015 directories rather than Visual Studio 2015 directories.

diff --git a/Engine/Source/Programs/UnrealBuildTool/Platform/Windows/UEBuildWindows.cs b/Engine/Source/Programs/UnrealBuildTool/Platform/Windows/UEBuildWindows.cs
--- a/Engine/Source/Programs/UnrealBuildTool/Platform/Windows/UEBuildWindows.cs
+++ b/Engine/Source/Programs/UnrealBuildTool/Platform/Windows/UEBuildWindows.cs
@@ -597,7 +597,7 @@ namespace UnrealBuildTool
 				    {
 					    // VS2015 just installs one toolchain; use that.
 					    DirectoryReference InstallDir;
-					    if(TryReadInstallDirRegistryKey32("Microsoft\\VisualStudio\\SxS\\VC7", "14.0", out InstallDir))
+					    if(TryReadInstallDirRegistryKey32("Microsoft\\VisualStudio\\SxS\\VS7", "14.0", out InstallDir))
 					    {
 						    InstallDirs.Add(InstallDir);
 					    }
@@ -661,7 +661,7 @@ namespace UnrealBuildTool
 					    List<DirectoryReference> InstallDirs = FindVSInstallDirs(Compiler);
 					    foreach(DirectoryReference InstallDir in InstallDirs)
 					    {
-						    ToolChainVersionToDir[new VersionNumber(14, 0)] = InstallDir;
+						    ToolChainVersionToDir[new VersionNumber(14, 0)] = DirectoryReference.Combine(InstallDir, "VC");
 					    }
 				    }
 				    else if(Compiler == WindowsCompiler.VisualStudio2017)
