When generating project files for a single .uproject, only include the enterprise module if it's an enterprise project. Prevents exception later on due to Enterprise project files not having any valid configurations.

diff --git a/Engine/Source/Programs/UnrealBuildTool/ProjectFiles/ProjectFileGenerator.cs b/Engine/Source/Programs/UnrealBuildTool/ProjectFiles/ProjectFileGenerator.cs
--- a/Engine/Source/Programs/UnrealBuildTool/ProjectFiles/ProjectFileGenerator.cs
+++ b/Engine/Source/Programs/UnrealBuildTool/ProjectFiles/ProjectFileGenerator.cs
@@ -239,6 +239,11 @@ namespace UnrealBuildTool
 		/// </summary>
 		protected bool bIncludeEngineSource = true;
 
+		/// <summary>
+		/// Whether to include enterprise source in the generated solution
+		/// </summary>
+		protected bool bIncludeEnterpriseSource = true;
+
 		/// <summary>
 		/// True if shader source files should be included in the generated projects
 		/// </summary>
@@ -1117,6 +1122,16 @@ namespace UnrealBuildTool
 			{
 				// At least one extra argument was specified, but we weren't expected it.  Ignored.
 			}
+
+			// If we're generating a solution for only one project, only include the enterprise folder if it's an enterprise project
+			if(OnlyGameProject == null)
+			{
+				bIncludeEnterpriseSource = bIncludeEngineSource;
+			}
+			else
+			{
+				bIncludeEnterpriseSource = bIncludeEngineSource && ProjectDescriptor.FromFile(OnlyGameProject).IsEnterpriseProject;
+			}
 		}
 
 
@@ -1842,6 +1857,12 @@ namespace UnrealBuildTool
 					WantProjectFileForModule = false;
 				}
 
+				if( CurModuleFile.IsUnderDirectory(UnrealBuildTool.EnterpriseDirectory) && !bIncludeEnterpriseSource )
+				{
+					// We were asked to exclude enterprise modules from the generated projects
+					WantProjectFileForModule = false;
+				}
+
 				if( WantProjectFileForModule )
 				{
 					DirectoryReference BaseFolder;
@@ -2033,11 +2054,11 @@ namespace UnrealBuildTool
 
 					if(TargetFilePath.IsUnderDirectory(EnterpriseSourceProgramsDirectory))
 					{
-						WantProjectFileForTarget = IncludeEnginePrograms;
+						WantProjectFileForTarget = bIncludeEnterpriseSource && IncludeEnginePrograms;
 					}
-					else if(TargetFilePath.IsUnderDirectory(UnrealBuildTool.EnterpriseSourceDirectory))
+					else
 					{
-						WantProjectFileForTarget = bIncludeEngineSource;
+						WantProjectFileForTarget = bIncludeEnterpriseSource;
 					}
 				}
 
