Remove change to automatically enable fast PDB linking in newer versions of VS2017. Lots of reports of OOM errors.

diff --git a/Engine/Source/Programs/UnrealBuildTool/Platform/Windows/UEBuildWindows.cs b/Engine/Source/Programs/UnrealBuildTool/Platform/Windows/UEBuildWindows.cs
--- a/Engine/Source/Programs/UnrealBuildTool/Platform/Windows/UEBuildWindows.cs
+++ b/Engine/Source/Programs/UnrealBuildTool/Platform/Windows/UEBuildWindows.cs
@@ -457,16 +457,17 @@ namespace UnrealBuildTool
 				Target.bUsePDBFiles = true;
 			}
 
-			// Enable fast PDB linking if we're on VS2017 15.7 or later. Previous versions have OOM issues with large projects.
-			if(!Target.bFormalBuild && !Target.bUseFastPDBLinking.HasValue && Target.WindowsPlatform.Compiler >= WindowsCompiler.VisualStudio2017)
-			{
-				VersionNumber Version;
-				DirectoryReference ToolChainDir;
-				if(TryGetVCToolChainDir(Target.WindowsPlatform.Compiler, Target.WindowsPlatform.CompilerVersion, out Version, out ToolChainDir) && Version >= new VersionNumber(14, 14, 26316))
-				{
-					Target.bUseFastPDBLinking = true;
-				}
-			}
+//			@Todo: Still getting reports of frequent OOM issues with this enabled as of 15.7.
+//			// Enable fast PDB linking if we're on VS2017 15.7 or later. Previous versions have OOM issues with large projects.
+//			if(!Target.bFormalBuild && !Target.bUseFastPDBLinking.HasValue && Target.WindowsPlatform.Compiler >= WindowsCompiler.VisualStudio2017)
+//			{
+//				VersionNumber Version;
+//				DirectoryReference ToolChainDir;
+//				if(TryGetVCToolChainDir(Target.WindowsPlatform.Compiler, Target.WindowsPlatform.CompilerVersion, out Version, out ToolChainDir) && Version >= new VersionNumber(14, 14, 26316))
+//				{
+//					Target.bUseFastPDBLinking = true;
+//				}
+//			}
 		}
 
 		/// <summary>
