Disable automatic PIX behavior on dx12 temporarily to allow packaged Windows 7 builds to launch

diff --git a/Engine/Source/Programs/UnrealBuildTool/Platform/Windows/UEBuildWindows.cs b/Engine/Source/Programs/UnrealBuildTool/Platform/Windows/UEBuildWindows.cs
--- a/Engine/Source/Programs/UnrealBuildTool/Platform/Windows/UEBuildWindows.cs
+++ b/Engine/Source/Programs/UnrealBuildTool/Platform/Windows/UEBuildWindows.cs
@@ -106,7 +106,7 @@ namespace UnrealBuildTool
 		/// Enable PIX debugging (automatically disabled in Shipping and Test configs)
 		/// </summary>
 		[ConfigFile(ConfigHierarchyType.Engine, "/Script/WindowsTargetPlatform.WindowsTargetSettings", "bEnablePIXProfiling")]
-		public bool bPixProfilingEnabled = true;
+		public bool bPixProfilingEnabled = false;
 
 		/// <summary>
 		/// The name of the company (author, provider) that created the project.
