Fix Lumin compile errors. Path specified for compiler is a batch file; needs to be executed through a command prompt instance rather than passed to CreateProcess().

diff --git a/Engine/Source/Programs/UnrealBuildTool/Platform/Android/AndroidToolChain.cs b/Engine/Source/Programs/UnrealBuildTool/Platform/Android/AndroidToolChain.cs
--- a/Engine/Source/Programs/UnrealBuildTool/Platform/Android/AndroidToolChain.cs
+++ b/Engine/Source/Programs/UnrealBuildTool/Platform/Android/AndroidToolChain.cs
@@ -43,6 +43,7 @@ namespace UnrealBuildTool
 		private bool bUseLdGold;
 		private List<string> AdditionalArches;
 		private List<string> AdditionalGPUArches;
+		protected bool bExecuteCompilerThroughShell;
 
 		// the Clang version being used to compile
 		static int ClangVersionMajor = -1;
@@ -1436,8 +1437,16 @@ namespace UnrealBuildTool
 						string ResponseArgument = string.Format("@\"{0}\"", ResponseFileName);
 
 						CompileAction.WorkingDirectory = UnrealBuildTool.EngineSourceDirectory.FullName;
-						CompileAction.CommandPath = ClangPath;
-						CompileAction.CommandArguments = ResponseArgument;
+						if(bExecuteCompilerThroughShell)
+						{
+							CompileAction.CommandPath = "cmd.exe";
+							CompileAction.CommandArguments = String.Format("/c \"{0} {1}\"", ClangPath, ResponseArgument);
+						}
+						else
+						{
+							CompileAction.CommandPath = ClangPath;
+							CompileAction.CommandArguments = ResponseArgument;
+						}
 						CompileAction.PrerequisiteItems.Add(ResponseFileItem);
 						CompileAction.StatusDescription = string.Format("{0} [{1}-{2}]", Path.GetFileName(SourceFile.AbsolutePath), Arch.Replace("-", ""), GPUArchitecture.Replace("-", ""));
 
diff --git a/Engine/Source/Programs/UnrealBuildTool/Platform/Lumin/LuminToolChain.cs b/Engine/Source/Programs/UnrealBuildTool/Platform/Lumin/LuminToolChain.cs
--- a/Engine/Source/Programs/UnrealBuildTool/Platform/Lumin/LuminToolChain.cs
+++ b/Engine/Source/Programs/UnrealBuildTool/Platform/Lumin/LuminToolChain.cs
@@ -119,6 +119,9 @@ namespace UnrealBuildTool
 			StripPath = ToolsDict["STRIP"];
 			ObjCopyPath = ToolsDict["OBJCOPY"];
 
+			// force the compiler to be executed through a command prompt instance
+			bExecuteCompilerThroughShell = true;
+
 			// toolchain params
 			ToolchainParamsArm64 = " --sysroot=\"" + Path.Combine(MLSDKPath, "lumin") + "\"";
 
diff --git a/Engine/Source/Programs/UnrealBuildTool/System/FileItem.cs b/Engine/Source/Programs/UnrealBuildTool/System/FileItem.cs
--- a/Engine/Source/Programs/UnrealBuildTool/System/FileItem.cs
+++ b/Engine/Source/Programs/UnrealBuildTool/System/FileItem.cs
@@ -324,7 +324,7 @@ namespace UnrealBuildTool
 						FileReference PrevAbsolutePath = new FileReference(AbsolutePath.FullName + ".prev");
 						FileReference.Delete(PrevAbsolutePath);
 						FileReference.Move(AbsolutePath, PrevAbsolutePath);
-						Log.TraceLog("Updating {0} - contents have changed (previous version renamed to {1}).", PrevAbsolutePath);
+						Log.TraceLog("Updating {0} - contents have changed (previous version renamed to {1}).", AbsolutePath.FullName, PrevAbsolutePath);
 					}
 					catch
 					{
