Another fix for Lumin compile errors.

diff --git a/Engine/Source/Programs/UnrealBuildTool/Platform/Android/AndroidToolChain.cs b/Engine/Source/Programs/UnrealBuildTool/Platform/Android/AndroidToolChain.cs
--- a/Engine/Source/Programs/UnrealBuildTool/Platform/Android/AndroidToolChain.cs
+++ b/Engine/Source/Programs/UnrealBuildTool/Platform/Android/AndroidToolChain.cs
@@ -1659,6 +1659,12 @@ namespace UnrealBuildTool
 
 					// Only execute linking on the local PC.
 					LinkAction.bCanExecuteRemotely = false;
+
+					if(bExecuteCompilerThroughShell)
+					{
+						LinkAction.CommandArguments = String.Format("/c \"{0} {1}\"", LinkAction.CommandPath, LinkAction.CommandArguments);
+						LinkAction.CommandPath = "cmd.exe";
+					}
 				}
 			}
 
