UBT: fix for platforms overriding project generation (merge from Dev-Console)

diff --git a/Engine/Source/Programs/UnrealBuildTool/Configuration/UEBuildPlatform.cs b/Engine/Source/Programs/UnrealBuildTool/Configuration/UEBuildPlatform.cs
--- a/Engine/Source/Programs/UnrealBuildTool/Configuration/UEBuildPlatform.cs
+++ b/Engine/Source/Programs/UnrealBuildTool/Configuration/UEBuildPlatform.cs
@@ -121,6 +121,14 @@ namespace UnrealBuildTool
 		/// </summary>
 		public abstract SDKStatus HasRequiredSDKsInstalled();
 
+		/// <summary>
+		/// Whether this platform requires specific Visual Studio version.
+		/// </summary>
+		public virtual VCProjectFileFormat GetRequiredVisualStudioVersion()
+		{
+			return VCProjectFileFormat.Default;
+		}
+
 		/// <summary>
 		/// Gets all the registered platforms
 		/// </summary>
diff --git a/Engine/Source/Programs/UnrealBuildTool/Configuration/UEBuildPlatformSDK.cs b/Engine/Source/Programs/UnrealBuildTool/Configuration/UEBuildPlatformSDK.cs
--- a/Engine/Source/Programs/UnrealBuildTool/Configuration/UEBuildPlatformSDK.cs
+++ b/Engine/Source/Programs/UnrealBuildTool/Configuration/UEBuildPlatformSDK.cs
@@ -592,7 +592,7 @@ namespace UnrealBuildTool
 		/// Currently installed AutoSDK is written out to a text file in a known location.
 		/// This function just compares the file's contents with the current requirements.
 		/// </summary>
-		protected SDKStatus HasRequiredAutoSDKInstalled()
+		public SDKStatus HasRequiredAutoSDKInstalled()
 		{
 			if (PlatformSupportsAutoSDKs() && HasAutoSDKSystemEnabled())
 			{
@@ -645,7 +645,7 @@ namespace UnrealBuildTool
 			return bParentProcessSetupAutoSDK;
 		}
 
-		protected SDKStatus HasRequiredManualSDK()
+		public SDKStatus HasRequiredManualSDK()
 		{
 			if (HasSetupAutoSDK())
 			{
diff --git a/Engine/Source/Programs/UnrealBuildTool/ProjectFiles/VisualStudio/VCProjectFileGenerator.cs b/Engine/Source/Programs/UnrealBuildTool/ProjectFiles/VisualStudio/VCProjectFileGenerator.cs
--- a/Engine/Source/Programs/UnrealBuildTool/ProjectFiles/VisualStudio/VCProjectFileGenerator.cs
+++ b/Engine/Source/Programs/UnrealBuildTool/ProjectFiles/VisualStudio/VCProjectFileGenerator.cs
@@ -306,12 +306,15 @@ namespace UnrealBuildTool
 						// Don't worry about platforms that we're missing SDKs for
 						if (BuildPlatform.HasRequiredSDKsInstalled() == SDKStatus.Valid)
 						{
-							VCProjectFileFormat ProposedFormat = ProjectFileFormat;
+							VCProjectFileFormat ProposedFormat = BuildPlatform.GetRequiredVisualStudioVersion();
 
-							// Reduce the Visual Studio version to the max supported by each platform we plan to include.
-							if (ProjectFileFormat == VCProjectFileFormat.Default || ProposedFormat < ProjectFileFormat)
+							if (ProposedFormat != VCProjectFileFormat.Default)
 							{
-								ProjectFileFormat = ProposedFormat;
+								// Reduce the Visual Studio version to the max supported by each platform we plan to include.
+								if (ProjectFileFormat == VCProjectFileFormat.Default || ProposedFormat < ProjectFileFormat)
+								{
+									ProjectFileFormat = ProposedFormat;
+								}
 							}
 						}
 					}
