Prevent source folder being added as a private include path if it does not exist.

diff --git a/Engine/Source/Programs/UnrealBuildTool/Configuration/UEBuildTarget.cs b/Engine/Source/Programs/UnrealBuildTool/Configuration/UEBuildTarget.cs
--- a/Engine/Source/Programs/UnrealBuildTool/Configuration/UEBuildTarget.cs
+++ b/Engine/Source/Programs/UnrealBuildTool/Configuration/UEBuildTarget.cs
@@ -4177,7 +4177,10 @@ namespace UnrealBuildTool
 					// If it's a game module (plugin or otherwise), add the root source directory to the include paths.
 					if (ModuleFileName.IsUnderDirectory(TargetRulesFile.Directory) || (Plugin != null && Plugin.LoadedFrom == PluginLoadedFrom.Project))
 					{
-						RulesObject.PublicIncludePaths.Add(NormalizeIncludePath(BaseSourceDirectory));
+						if(DirectoryReference.Exists(BaseSourceDirectory))
+						{
+							RulesObject.PublicIncludePaths.Add(NormalizeIncludePath(BaseSourceDirectory));
+						}
 					}
 
 					// Resolve private include paths against the project source root
