Write the exception callstack to the UBT log when thrown reading a makefile.

#rnx

diff --git a/Engine/Source/Programs/UnrealBuildTool/System/UBTMakefile.cs b/Engine/Source/Programs/UnrealBuildTool/System/UBTMakefile.cs
--- a/Engine/Source/Programs/UnrealBuildTool/System/UBTMakefile.cs
+++ b/Engine/Source/Programs/UnrealBuildTool/System/UBTMakefile.cs
@@ -279,6 +279,7 @@ namespace UnrealBuildTool
 			catch (Exception Ex)
 			{
 				Log.TraceWarning("Failed to read makefile: {0}", Ex.Message);
+				Log.TraceLog("Exception: {0}", Ex.ToString());
 				ReasonNotLoaded = "couldn't read existing makefile";
 				return null;
 			}
