Ignore private include paths for any precompiled modules. Prevents warnings when they are stripped out from an installed engine build.

diff --git a/Engine/Source/Programs/UnrealBuildTool/Configuration/UEBuildModule.cs b/Engine/Source/Programs/UnrealBuildTool/Configuration/UEBuildModule.cs
--- a/Engine/Source/Programs/UnrealBuildTool/Configuration/UEBuildModule.cs
+++ b/Engine/Source/Programs/UnrealBuildTool/Configuration/UEBuildModule.cs
@@ -187,7 +187,14 @@ namespace UnrealBuildTool
 			PublicAdditionalShadowFiles = HashSetFromOptionalEnumerableStringParameter(InRules.PublicAdditionalShadowFiles);
 			PublicAdditionalBundleResources = InRules.AdditionalBundleResources == null ? new HashSet<UEBuildBundleResource>() : new HashSet<UEBuildBundleResource>(InRules.AdditionalBundleResources);
 			PublicDelayLoadDLLs = HashSetFromOptionalEnumerableStringParameter(InRules.PublicDelayLoadDLLs);
-			PrivateIncludePaths = CreateDirectoryHashSet(InRules.PrivateIncludePaths);
+			if(Rules.bUsePrecompiled)
+			{
+				PrivateIncludePaths = new HashSet<DirectoryReference>();
+			}
+			else
+			{
+				PrivateIncludePaths = CreateDirectoryHashSet(InRules.PrivateIncludePaths);
+			}
 			RuntimeDependencies = InRuntimeDependencies;
 			IsRedistributableOverride = InRules.IsRedistributableOverride;
 
