Fix missing DLL errors when compiling for Win32 with the VC++ 14.14 toolchain.

#rnx

diff --git a/Engine/Source/Programs/UnrealBuildTool/Platform/Windows/VCEnvironment.cs b/Engine/Source/Programs/UnrealBuildTool/Platform/Windows/VCEnvironment.cs
--- a/Engine/Source/Programs/UnrealBuildTool/Platform/Windows/VCEnvironment.cs
+++ b/Engine/Source/Programs/UnrealBuildTool/Platform/Windows/VCEnvironment.cs
@@ -109,16 +109,35 @@ namespace UnrealBuildTool
 			LibraryManagerPath = GetLibraryLinkerToolPath(Platform, LinkerDir);
 			ResourceCompilerPath = GetResourceCompilerToolPath(Platform, WindowsSdkDir, WindowsSdkVersion);
 
-            // Make sure the base 32-bit VS tool path is in the PATH, regardless of which configuration we're using. The toolchain may need to reference support DLLs from this directory (eg. mspdb120.dll).
+			// Add both toolchain paths to the PATH environment variable. There are some support DLLs which are only added to one of the paths, but which the toolchain in the other directory
+			// needs to run (eg. mspdbcore.dll).
+			if(Platform == CppPlatform.Win64)
+			{
+				AddDirectoryToPath(VCToolPath64);
+				AddDirectoryToPath(VCToolPath32);
+			}
+			if(Platform == CppPlatform.Win32)
+			{
+				AddDirectoryToPath(VCToolPath32);
+				AddDirectoryToPath(VCToolPath64);
+			}
+
+			// Get all the system include paths
+			SetupEnvironment();
+		}
+
+		/// <summary>
+		/// Add a directory to the PATH environment variable
+		/// </summary>
+		/// <param name="ToolPath">The path to add</param>
+		static void AddDirectoryToPath(DirectoryReference ToolPath)
+		{
             string PathEnvironmentVariable = Environment.GetEnvironmentVariable("PATH") ?? "";
-            if (!PathEnvironmentVariable.Split(';').Any(x => String.Compare(x, VCToolPath32.FullName, true) == 0))
+            if (!PathEnvironmentVariable.Split(';').Any(x => String.Compare(x, ToolPath.FullName, true) == 0))
             {
-                PathEnvironmentVariable = VCToolPath32.FullName + ";" + PathEnvironmentVariable;
+                PathEnvironmentVariable = ToolPath.FullName + ";" + PathEnvironmentVariable;
                 Environment.SetEnvironmentVariable("PATH", PathEnvironmentVariable);
             }
-
-			// Get all the system include paths
-			SetupEnvironment();
 		}
 
 		/// <summary>
