Include the BuildConfiguration.Schema.xsd file in the installed engine build, and don't write it out when running in an installed build.

diff --git a/Engine/Build/InstalledEngineBuild.xml b/Engine/Build/InstalledEngineBuild.xml
--- a/Engine/Build/InstalledEngineBuild.xml
+++ b/Engine/Build/InstalledEngineBuild.xml
@@ -460,6 +460,7 @@
 			<Compile Target="BuildPatchTool" Configuration="Shipping" Platform="Win64" Tag="#Build Tools Win64"/>
 
 			<Tag Files="Engine/Binaries/Win64/XGEControlWorker.exe" With="#Build Tools Win64"/>
+			<Tag Files="$(RootDir)/Engine/Saved/UnrealBuildTool/BuildConfiguration.Schema.xsd" With="#Build Tools Win64"/>
 
 			<Do If="$(EmbedSrcSrvInfo)">
 				<!-- Embed source info into the PDB files. Should be done from this machine to ensure that paths are correct. -->
@@ -514,6 +515,8 @@
 			<Compile Target="CrashReportClient" Platform="Linux" Configuration="Shipping"/>
 			<Compile Target="ShaderCompileWorker" Configuration="Development" Platform="Linux" Tag="#Build Tools Linux"/>
 			<Compile Target="UnrealLightmass" Configuration="Development" Platform="Linux" Tag="#Build Tools Linux"/>
+
+			<Tag Files="$(RootDir)/Engine/Saved/UnrealBuildTool/BuildConfiguration.Schema.xsd" With="#Build Tools Linux"/>
 		</Node>
 	</Agent>
 
@@ -529,6 +532,8 @@
 			<Compile Target="DsymExporter" Configuration="Development" Platform="Mac" Tag="#Build Tools Mac"/>
 			<Compile Target="UE4EditorServices" Configuration="Development" Platform="Mac" Tag="#Build Tools Mac"/>
 			<Compile Target="UnrealAtoS" Configuration="Development" Platform="Mac" Tag="#Build Tools Mac"/>
+
+			<Tag Files="$(RootDir)/Engine/Saved/UnrealBuildTool/BuildConfiguration.Schema.xsd" With="#Build Tools Mac"/>
 		</Node>
 	</Agent>
 
diff --git a/Engine/Source/Programs/UnrealBuildTool/System/XmlConfig.cs b/Engine/Source/Programs/UnrealBuildTool/System/XmlConfig.cs
--- a/Engine/Source/Programs/UnrealBuildTool/System/XmlConfig.cs
+++ b/Engine/Source/Programs/UnrealBuildTool/System/XmlConfig.cs
@@ -58,6 +58,16 @@ namespace UnrealBuildTool
 
 			// Get the path to the cache file
 			FileReference CacheFile = FileReference.Combine(UnrealBuildTool.EngineDirectory, "Intermediate", "Build", "XmlConfigCache.bin");
+			if(UnrealBuildTool.IsEngineInstalled())
+			{
+				DirectoryReference UserSettingsDir = Utils.GetUserSettingDirectory();
+				if(UserSettingsDir != null)
+				{
+					CacheFile = FileReference.Combine(UserSettingsDir, "UnrealEngine", String.Format("XmlConfigCache-{0}.bin", UnrealBuildTool.RootDirectory.FullName.Replace(":", "").Replace(Path.DirectorySeparatorChar, '+')));
+				}
+			}
+
+			// Get the path to the schema
 			FileReference SchemaFile = GetSchemaLocation();
 
 			// Try to read the existing cache from disk
@@ -79,7 +89,10 @@ namespace UnrealBuildTool
 
 				// Create a schema for the config files
 				XmlSchema Schema = CreateSchema(CategoryToFields);
-				WriteSchema(Schema, SchemaFile);
+				if(!UnrealBuildTool.IsEngineInstalled())
+				{
+					WriteSchema(Schema, SchemaFile);
+				}
 
 				// Read all the XML files and validate them against the schema
 				Dictionary<Type, Dictionary<FieldInfo, object>> TypeToValues = new Dictionary<Type, Dictionary<FieldInfo, object>>();
@@ -88,11 +101,13 @@ namespace UnrealBuildTool
 					if(!TryReadFile(InputFile, CategoryToFields, TypeToValues, Schema))
 					{
 						Log.TraceError("Failed to properly read XML file : {0}", InputFile.FullName);
-
 						return false;
 					}
 				}
 
+				// Make sure the cache directory exists
+				DirectoryReference.CreateDirectory(CacheFile.Directory);
+
 				// Create the new cache
 				Values = new XmlConfigData(InputFiles, TypeToValues.ToDictionary(x => x.Key, x => x.Value.ToArray()));
 				Values.Write(CacheFile);
@@ -167,20 +182,24 @@ namespace UnrealBuildTool
 			// Find all the input file locations
 			List<InputFile> InputFiles = new List<InputFile>();
 
-			// Check for the config file under /Engine/Programs/NotForLicensees/UnrealBuildTool
-			FileReference NotForLicenseesConfigLocation = FileReference.Combine(UnrealBuildTool.EngineDirectory, "Programs", "NotForLicensees", "UnrealBuildTool", "BuildConfiguration.xml");
-			if(FileReference.Exists(NotForLicenseesConfigLocation))
+			// Skip all the config files under the Engine folder if it's an installed build
+			if(!UnrealBuildTool.IsEngineInstalled())
 			{
-				InputFiles.Add(new InputFile { Location = NotForLicenseesConfigLocation, FolderName = "NotForLicensees" });
-			}
+				// Check for the config file under /Engine/Programs/NotForLicensees/UnrealBuildTool
+				FileReference NotForLicenseesConfigLocation = FileReference.Combine(UnrealBuildTool.EngineDirectory, "Programs", "NotForLicensees", "UnrealBuildTool", "BuildConfiguration.xml");
+				if(FileReference.Exists(NotForLicenseesConfigLocation))
+				{
+					InputFiles.Add(new InputFile { Location = NotForLicenseesConfigLocation, FolderName = "NotForLicensees" });
+				}
 
-			// Check for the user config file under /Engine/Programs/NotForLicensees/UnrealBuildTool
-			FileReference UserConfigLocation = FileReference.Combine(UnrealBuildTool.EngineDirectory, "Saved", "UnrealBuildTool", "BuildConfiguration.xml");
-			if(!FileReference.Exists(UserConfigLocation))
-			{
-				CreateDefaultConfigFile(UserConfigLocation);
+				// Check for the user config file under /Engine/Programs/NotForLicensees/UnrealBuildTool
+				FileReference UserConfigLocation = FileReference.Combine(UnrealBuildTool.EngineDirectory, "Saved", "UnrealBuildTool", "BuildConfiguration.xml");
+				if(!FileReference.Exists(UserConfigLocation))
+				{
+					CreateDefaultConfigFile(UserConfigLocation);
+				}
+				InputFiles.Add(new InputFile { Location = UserConfigLocation, FolderName = "User" });
 			}
-			InputFiles.Add(new InputFile { Location = UserConfigLocation, FolderName = "User" });
 
 			// Check for the global config file under AppData/Unreal Engine/UnrealBuildTool
 			string AppDataFolder = Environment.GetFolderPath(Environment.SpecialFolder.ApplicationData);
diff --git a/Engine/Source/Programs/UnrealBuildTool/UnrealBuildTool.cs b/Engine/Source/Programs/UnrealBuildTool/UnrealBuildTool.cs
--- a/Engine/Source/Programs/UnrealBuildTool/UnrealBuildTool.cs
+++ b/Engine/Source/Programs/UnrealBuildTool/UnrealBuildTool.cs
@@ -495,6 +495,24 @@ namespace UnrealBuildTool
 						Directory.SetCurrentDirectory(EngineSourceDirectory);
 					}
 
+					// Figure out if the engine is installed or not
+					foreach (string Argument in Arguments)
+					{
+						string LowercaseArg = Argument.ToLowerInvariant();
+						if (LowercaseArg == "-installed" || LowercaseArg == "-installedengine")
+						{
+							bIsEngineInstalled = true;
+						}
+						else if (LowercaseArg == "-notinstalledengine")
+						{
+							bIsEngineInstalled = false;
+						}
+					}
+					if (!bIsEngineInstalled.HasValue)
+					{
+						bIsEngineInstalled = FileReference.Exists(FileReference.Combine(RootDirectory, "Engine", "Build", "InstalledBuild.txt"));
+					}
+
 					// Read the XML configuration files
 					if (!XmlConfig.ReadConfigFiles())
 					{
@@ -538,14 +556,6 @@ namespace UnrealBuildTool
 							BuildConfiguration.bUseUBTMakefiles = false;
 							BuildConfiguration.SingleFileToCompile = LowercaseArg.Replace("-singlefile=", "");
 						}
-						else if (LowercaseArg == "-installed" || LowercaseArg == "-installedengine")
-						{
-							bIsEngineInstalled = true;
-						}
-						else if (LowercaseArg == "-notinstalledengine")
-						{
-							bIsEngineInstalled = false;
-						}
 						else if (LowercaseArg.StartsWith("-buildconfigurationdoc="))
 						{
 							XmlConfig.WriteDocumentation(new FileReference(Argument.Substring("-buildconfigurationdoc=".Length)));
@@ -563,12 +573,6 @@ namespace UnrealBuildTool
 						}
 					}
 
-					// If it wasn't set explicitly by a command line option, check for the installed build marker file
-					if (!bIsEngineInstalled.HasValue)
-					{
-						bIsEngineInstalled = FileReference.Exists(FileReference.Combine(RootDirectory, "Engine", "Build", "InstalledBuild.txt"));
-					}
-
 					// Create the log file, and flush the startup listener to it
 					if (!String.IsNullOrEmpty(BuildConfiguration.LogFileName))
 					{
