UBT: Fix issue where overlapping circularly header dependencies would not be parsed correctly, preventing files not being rebuilt on header changes (noted where Clang detected an out of date PCH).

diff --git a/Engine/Source/Programs/UnrealBuildTool/System/CPPHeaders.cs b/Engine/Source/Programs/UnrealBuildTool/System/CPPHeaders.cs
--- a/Engine/Source/Programs/UnrealBuildTool/System/CPPHeaders.cs
+++ b/Engine/Source/Programs/UnrealBuildTool/System/CPPHeaders.cs
@@ -101,22 +101,6 @@ namespace UnrealBuildTool
 		}
 	}
 
-	/// <summary>
-	/// List of all files included in a file and helper class for handling circular dependencies.
-	/// </summary>
-	class IncludedFilesSet : HashSet<FileItem>
-	{
-		/// <summary>
-		/// Whether this file list has been fully initialized or not.
-		/// </summary>
-		public bool bIsInitialized;
-
-		/// <summary>
-		/// List of files which include this file in one of its includes.
-		/// </summary>
-		public List<FileItem> CircularDependencies = new List<FileItem>();
-	}
-
 	class CPPHeaders
 	{
 		/// <summary>
@@ -152,12 +136,12 @@ namespace UnrealBuildTool
 		/// <summary>
 		/// A cache of the list of other files that are directly or indirectly included by a C++ file.
 		/// </summary>
-		Dictionary<FileItem, IncludedFilesSet> ExhaustiveIncludedFilesMap = new Dictionary<FileItem, IncludedFilesSet>();
+		Dictionary<FileItem, List<FileItem>> ExhaustiveIncludedFilesMap = new Dictionary<FileItem, List<FileItem>>();
 
 		/// <summary>
 		/// A cache of all files included by a C++ file, but only has files that we knew about from a previous session, loaded from a cache at startup
 		/// </summary>
-		Dictionary<FileItem, IncludedFilesSet> OnlyCachedIncludedFilesMap = new Dictionary<FileItem, IncludedFilesSet>();
+		Dictionary<FileItem, List<FileItem>> OnlyCachedIncludedFilesMap = new Dictionary<FileItem, List<FileItem>>();
 
 		/// <summary>
 		/// 
@@ -289,8 +273,6 @@ namespace UnrealBuildTool
 			}
 			else
 			{
-				// @todo ubtmake: HeaderParser.h is missing from the include set for Module.UnrealHeaderTool.cpp (failed to find include using:  FileItem DirectIncludeResolvedFile = CPPEnvironment.FindIncludedFile(DirectInclude.IncludeName, !BuildConfiguration.bCheckExternalHeadersForModification, IncludePathsToSearch, IncludeFileSearchDictionary );)
-
 				// If we're doing an exhaustive include scan, make sure that we have our include dependency cache loaded and ready
 				if (!bOnlyCachedDependencies)
 				{
@@ -300,14 +282,8 @@ namespace UnrealBuildTool
 					}
 				}
 
-				Result = new List<FileItem>();
-
-				IncludedFilesSet IncludedFileList = new IncludedFilesSet();
-				FindAndCacheAllIncludedFiles(SourceFile, IncludePaths, ref IncludedFileList, bOnlyCachedDependencies: bOnlyCachedDependencies);
-				foreach (FileItem IncludedFile in IncludedFileList)
-				{
-					Result.Add(IncludedFile);
-				}
+				// Get the headers
+				Result = FindAndCacheIncludedFiles(SourceFile, IncludePaths, bOnlyCachedDependencies);
 
 				// Update cache
 				if (bUseFlatCPPIncludeDependencyCache && !bOnlyCachedDependencies)
@@ -325,124 +301,124 @@ namespace UnrealBuildTool
 			return Result;
 		}
 
-
 		/// <summary>
-		/// Finds the files directly or indirectly included by the given C++ file.
+		/// Add all the files included by a source file to a set, using a cache.
 		/// </summary>
-		/// <param name="CPPFile">C++ file to get the dependencies for.</param>
-		/// <param name="IncludePaths"></param>
-		/// <param name="bOnlyCachedDependencies"></param>
-		/// <param name="Result">List of CPPFile dependencies.</param>
-		/// <returns>false if CPPFile is still being processed further down the callstack, true otherwise.</returns>
-		public bool FindAndCacheAllIncludedFiles(FileItem CPPFile, CppIncludePaths IncludePaths, ref IncludedFilesSet Result, bool bOnlyCachedDependencies)
+		/// <param name="SourceFile">The file to check.</param>
+		/// <param name="IncludePaths">Include paths to search.</param>
+		/// <param name="bOnlyCachedDependencies">Whether to just return cached dependencies, or update the cache with new results.</param>
+		private List<FileItem> FindAndCacheIncludedFiles(FileItem SourceFile, CppIncludePaths IncludePaths, bool bOnlyCachedDependencies)
 		{
-			IncludedFilesSet IncludedFileList;
-			Dictionary<FileItem, IncludedFilesSet> IncludedFilesMap = bOnlyCachedDependencies ? OnlyCachedIncludedFilesMap : ExhaustiveIncludedFilesMap;
-			if (!IncludedFilesMap.TryGetValue(CPPFile, out IncludedFileList))
+			// Get the map of files to their list of includes
+			Dictionary<FileItem, List<FileItem>> FileToIncludedFiles = bOnlyCachedDependencies ? OnlyCachedIncludedFilesMap : ExhaustiveIncludedFilesMap;
+
+			// Check if the included files is in the cache. If not, we'll create it.
+			List<FileItem> IncludedFiles;
+			if(!FileToIncludedFiles.TryGetValue(SourceFile, out IncludedFiles))
 			{
-				DateTime TimerStartTime = DateTime.UtcNow;
+				HashSet<FileItem> VisitedFiles = new HashSet<FileItem>();
+				VisitedFiles.Add(SourceFile);
 
-				IncludedFileList = new IncludedFilesSet();
+				HashSet<FileItem> IncludedFilesSet = new HashSet<FileItem>();
+				FindAndCacheIncludedFilesInner(SourceFile, IncludedFilesSet, IncludePaths, bOnlyCachedDependencies, VisitedFiles);
 
-				// Add an uninitialized entry for the include file to avoid infinitely recursing on include file loops.
-				IncludedFilesMap.Add(CPPFile, IncludedFileList);
+				IncludedFiles = IncludedFilesSet.ToList();
+				FileToIncludedFiles.Add(SourceFile, IncludedFiles);
+			}
 
-				// Gather a list of names of files directly included by this C++ file.
-				List<DependencyInclude> DirectIncludes = GetDirectIncludeDependencies(CPPFile, bOnlyCachedDependencies: bOnlyCachedDependencies);
+			return IncludedFiles;
+		}
 
-				// Build a list of the unique set of files that are included by this file.
-				HashSet<FileItem> DirectlyIncludedFiles = new HashSet<FileItem>();
-				// require a for loop here because we need to keep track of the index in the list.
-				for (int DirectlyIncludedFileNameIndex = 0; DirectlyIncludedFileNameIndex < DirectIncludes.Count; ++DirectlyIncludedFileNameIndex)
+		/// <summary>
+		/// Add all the files included by a source file to a set, using a cache.
+		/// </summary>
+		/// <param name="SourceFile">The file to check.</param>
+		/// <param name="IncludedFiles">Set of included files to add to</param>
+		/// <param name="IncludePaths">Include paths to search.</param>
+		/// <param name="bOnlyCachedDependencies">Whether to just return cached dependencies, or update the cache with new results.</param>
+		/// <param name="VisitedFiles">Set of files that have already been visited. Used to prevent infinite loops between circularly dependent headers.</param>
+		private void FindAndCacheIncludedFilesInner(FileItem SourceFile, HashSet<FileItem> IncludedFiles, CppIncludePaths IncludePaths, bool bOnlyCachedDependencies, HashSet<FileItem> VisitedFiles)
+		{
+			HashSet<FileItem> DirectlyIncludedFiles = GetDirectlyIncludedFiles(SourceFile, IncludePaths, bOnlyCachedDependencies);
+			foreach (FileItem DirectlyIncludedFile in DirectlyIncludedFiles)
+			{
+				if(IncludedFiles.Add(DirectlyIncludedFile))
 				{
-					// Resolve the included file name to an actual file.
-					DependencyInclude DirectInclude = DirectIncludes[DirectlyIncludedFileNameIndex];
-					if (!DirectInclude.HasAttemptedResolve ||
-						// ignore any preexisting resolve cache if we are not configured to use it.
-						!bUseIncludeDependencyResolveCache ||
-						// if we are testing the resolve cache, we force UBT to resolve every time to look for conflicts
-						bTestIncludeDependencyResolveCache
-						)
-					{
-						++TotalDirectIncludeResolveCacheMisses;
+					// Get the map of files to their list of includes
+					Dictionary<FileItem, List<FileItem>> FileToIncludedFiles = bOnlyCachedDependencies ? OnlyCachedIncludedFilesMap : ExhaustiveIncludedFilesMap;
 
-						// search the include paths to resolve the file
-						FileItem DirectIncludeResolvedFile = CPPHeaders.FindIncludedFile(CPPFile.Location, DirectInclude.IncludeName, IncludePaths);
-						if (DirectIncludeResolvedFile != null)
-						{
-							DirectlyIncludedFiles.Add(DirectIncludeResolvedFile);
-						}
-						IncludeDependencyCache.CacheResolvedIncludeFullPath(CPPFile, DirectlyIncludedFileNameIndex, DirectIncludeResolvedFile != null ? DirectIncludeResolvedFile.Location : null, bUseIncludeDependencyResolveCache, bTestIncludeDependencyResolveCache);
+					// Recursively add the files included by this file
+					List<FileItem> InnerFiles;
+					if(FileToIncludedFiles.TryGetValue(DirectlyIncludedFile, out InnerFiles))
+					{
+						// We already have the include paths cached; just add them directly.
+						IncludedFiles.UnionWith(InnerFiles);
+					}
+					else if(VisitedFiles.Add(DirectlyIncludedFile))
+					{
+						// We don't have include paths cached, and this isn't a recursive call. Create a new set and add it to the cache.
+						HashSet<FileItem> InnerFilesSet = new HashSet<FileItem>();
+						FindAndCacheIncludedFilesInner(DirectlyIncludedFile, InnerFilesSet, IncludePaths, bOnlyCachedDependencies, VisitedFiles);
+						FileToIncludedFiles.Add(DirectlyIncludedFile, InnerFilesSet.ToList());
+						IncludedFiles.UnionWith(InnerFilesSet);
 					}
 					else
 					{
-						// we might have cached an attempt to resolve the file, but couldn't actually find the file (system headers, etc).
-						if (DirectInclude.IncludeResolvedNameIfSuccessful != null)
-						{
-							DirectlyIncludedFiles.Add(FileItem.GetItemByFileReference(DirectInclude.IncludeResolvedNameIfSuccessful));
-						}
+						// We're already building a list of include paths for this file up the stack. Just recurse through it this time.
+						FindAndCacheIncludedFilesInner(DirectlyIncludedFile, IncludedFiles, IncludePaths, bOnlyCachedDependencies, VisitedFiles);
 					}
 				}
-				TotalDirectIncludeResolves += DirectIncludes.Count;
+			}
+		}
 
-				// Convert the dictionary of files included by this file into a list.
-				foreach (FileItem DirectlyIncludedFile in DirectlyIncludedFiles)
+		/// <summary>
+		/// Get a set of directly included files from the given source file, resolving their include paths to FileItem instances.
+		/// </summary>
+		/// <param name="SourceFile">The file to check.</param>
+		/// <param name="IncludePaths">Include paths to search.</param>
+		/// <param name="bOnlyCachedDependencies">Whether to just return cached dependencies, or update the cache with new results.</param>
+		/// <returns>Set of files that are included</returns>
+		private HashSet<FileItem> GetDirectlyIncludedFiles(FileItem SourceFile, CppIncludePaths IncludePaths, bool bOnlyCachedDependencies)
+		{
+			// Gather a list of names of files directly included by this C++ file.
+			List<DependencyInclude> DirectIncludes = GetDirectIncludeDependencies(SourceFile, bOnlyCachedDependencies: bOnlyCachedDependencies);
+
+			// Build a list of the unique set of files that are included by this file.
+			HashSet<FileItem> DirectlyIncludedFiles = new HashSet<FileItem>();
+			// require a for loop here because we need to keep track of the index in the list.
+			for (int DirectlyIncludedFileNameIndex = 0; DirectlyIncludedFileNameIndex < DirectIncludes.Count; ++DirectlyIncludedFileNameIndex)
+			{
+				// Resolve the included file name to an actual file.
+				DependencyInclude DirectInclude = DirectIncludes[DirectlyIncludedFileNameIndex];
+				if (!DirectInclude.HasAttemptedResolve ||
+					// ignore any preexisting resolve cache if we are not configured to use it.
+					!bUseIncludeDependencyResolveCache ||
+					// if we are testing the resolve cache, we force UBT to resolve every time to look for conflicts
+					bTestIncludeDependencyResolveCache
+					)
 				{
-					// Add the file we're directly including
-					IncludedFileList.Add(DirectlyIncludedFile);
+					++TotalDirectIncludeResolveCacheMisses;
 
-					// Also add all of the indirectly included files!
-					if (FindAndCacheAllIncludedFiles(DirectlyIncludedFile, IncludePaths, ref IncludedFileList, bOnlyCachedDependencies: bOnlyCachedDependencies) == false)
+					// search the include paths to resolve the file
+					FileItem DirectIncludeResolvedFile = CPPHeaders.FindIncludedFile(SourceFile.Location, DirectInclude.IncludeName, IncludePaths);
+					if (DirectIncludeResolvedFile != null)
 					{
-						// DirectlyIncludedFile is a circular dependency which is still being processed
-						// further down the callstack. Add this file to its circular dependencies list 
-						// so that it can update its dependencies later.
-						IncludedFilesSet DirectlyIncludedFileIncludedFileList;
-						if (IncludedFilesMap.TryGetValue(DirectlyIncludedFile, out DirectlyIncludedFileIncludedFileList))
-						{
-							DirectlyIncludedFileIncludedFileList.CircularDependencies.Add(CPPFile);
-						}
+						DirectlyIncludedFiles.Add(DirectIncludeResolvedFile);
 					}
+					IncludeDependencyCache.CacheResolvedIncludeFullPath(SourceFile, DirectlyIncludedFileNameIndex, DirectIncludeResolvedFile != null ? DirectIncludeResolvedFile.Location : null, bUseIncludeDependencyResolveCache, bTestIncludeDependencyResolveCache);
 				}
-
-				// All dependencies have been processed by now so update all circular dependencies
-				// with the full list.
-				foreach (FileItem CircularDependency in IncludedFileList.CircularDependencies)
+				else
 				{
-					IncludedFilesSet CircularDependencyIncludedFiles = IncludedFilesMap[CircularDependency];
-					foreach (FileItem IncludedFile in IncludedFileList)
+					// we might have cached an attempt to resolve the file, but couldn't actually find the file (system headers, etc).
+					if (DirectInclude.IncludeResolvedNameIfSuccessful != null)
 					{
-						CircularDependencyIncludedFiles.Add(IncludedFile);
+						DirectlyIncludedFiles.Add(FileItem.GetItemByFileReference(DirectInclude.IncludeResolvedNameIfSuccessful));
 					}
 				}
-				// No need to keep this around anymore.
-				IncludedFileList.CircularDependencies.Clear();
-
-				// Done collecting files.
-				IncludedFileList.bIsInitialized = true;
-
-				TimeSpan TimerDuration = DateTime.UtcNow - TimerStartTime;
-				TotalTimeSpentGettingIncludes += TimerDuration.TotalSeconds;
 			}
+			TotalDirectIncludeResolves += DirectIncludes.Count;
 
-			if (IncludedFileList.bIsInitialized)
-			{
-				// Copy the list of files included by this file into the result list.
-				foreach (FileItem IncludedFile in IncludedFileList)
-				{
-					// If the result list doesn't contain this file yet, add the file and the files it includes.
-					// NOTE: For some reason in .NET 4, Add() is over twice as fast as calling UnionWith() on the set
-					Result.Add(IncludedFile);
-				}
-
-				return true;
-			}
-			else
-			{
-				// The IncludedFileList.bIsInitialized was false because we added a dummy entry further down the call stack.  We're already processing
-				// the include list for this header elsewhere in the stack frame, so we don't need to add anything here.
-				return false;
-			}
+			return DirectlyIncludedFiles;
 		}
 
 		public FileItem CachePCHUsageForCPPFile(FileItem CPPFile, CppIncludePaths IncludePaths, CppPlatform Platform)
