PR #4700: This patch fixes a build failure that occurs on the very first build after adding new files with reflected types. (Contributed by [REDACTED_PERSON])

diff --git a/Engine/Source/Programs/UnrealBuildTool/System/ExternalExecution.cs b/Engine/Source/Programs/UnrealBuildTool/System/ExternalExecution.cs
--- a/Engine/Source/Programs/UnrealBuildTool/System/ExternalExecution.cs
+++ b/Engine/Source/Programs/UnrealBuildTool/System/ExternalExecution.cs
@@ -933,6 +933,9 @@ namespace UnrealBuildTool
 							AllUObjectFiles.AddRange(Module.PrivateUObjectHeaders.ConvertAll(Item => Item.AbsolutePath));
 							FileReference.WriteAllLines(TimestampFile, AllUObjectFiles);
 						}
+
+                        // Because new .cpp and .h files may have been generated by UHT, invalidate the DirectoryLookupCache
+                        DirectoryLookupCache.InvalidateCachedDirectory(new DirectoryReference(GeneratedCodeDirectoryInfo.FullName));
 					}
 				}
 				catch (Exception Exception)
