Fix detection of VS2017 toolchains that only have a 32-bit compiler for x64 (eg. VS2017 Express Desktop)

diff --git a/Engine/Source/Programs/UnrealBuildTool/Platform/Windows/UEBuildWindows.cs b/Engine/Source/Programs/UnrealBuildTool/Platform/Windows/UEBuildWindows.cs
--- a/Engine/Source/Programs/UnrealBuildTool/Platform/Windows/UEBuildWindows.cs
+++ b/Engine/Source/Programs/UnrealBuildTool/Platform/Windows/UEBuildWindows.cs
@@ -667,11 +667,8 @@ namespace UnrealBuildTool
 				    }
 				    else if(Compiler == WindowsCompiler.VisualStudio2017)
 				    {
-					    // Path of a file to look for to validate a directory contains any toolchain
-					    const string FileCheckPath = "bin\\Hostx64\\x64\\cl.exe";
-    
-					    // Enumerate all the manually installed toolchains
-					    List<DirectoryReference> InstallDirs = FindVSInstallDirs(Compiler);
+						// Enumerate all the manually installed toolchains
+						List<DirectoryReference> InstallDirs = FindVSInstallDirs(Compiler);
 					    foreach(DirectoryReference InstallDir in InstallDirs)
 					    {
 						    DirectoryReference ToolChainBaseDir = DirectoryReference.Combine(InstallDir, "VC", "Tools", "MSVC");
@@ -680,7 +677,7 @@ namespace UnrealBuildTool
 							    foreach(DirectoryReference ToolChainDir in DirectoryReference.EnumerateDirectories(ToolChainBaseDir))
 							    {
 								    VersionNumber Version;
-								    if(VersionNumber.TryParse(ToolChainDir.GetDirectoryName(), out Version) && FileReference.Exists(FileReference.Combine(ToolChainDir, FileCheckPath)))
+								    if(VersionNumber.TryParse(ToolChainDir.GetDirectoryName(), out Version) && IsValidToolChainDir(ToolChainDir))
 								    {
 									    ToolChainVersionToDir[Version] = ToolChainDir;
 								    }
@@ -698,7 +695,7 @@ namespace UnrealBuildTool
 							    foreach(DirectoryReference ToolChainDir in DirectoryReference.EnumerateDirectories(ToolChainBaseDir))
 							    {
 								    VersionNumber Version;
-								    if(VersionNumber.TryParse(ToolChainDir.GetDirectoryName(), out Version) && FileReference.Exists(FileReference.Combine(ToolChainDir, FileCheckPath)))
+								    if(VersionNumber.TryParse(ToolChainDir.GetDirectoryName(), out Version) && IsValidToolChainDir(ToolChainDir))
 								    {
 									    ToolChainVersionToDir[Version] = ToolChainDir;
 								    }
@@ -716,6 +713,16 @@ namespace UnrealBuildTool
 			return ToolChainVersionToDir;
 		}
 
+		/// <summary>
+		/// Checks if the given directory contains a valid toolchain
+		/// </summary>
+		/// <param name="ToolChainDir">Directory to check</param>
+		/// <returns>True if the given directory is valid</returns>
+		static bool IsValidToolChainDir(DirectoryReference ToolChainDir)
+		{
+			return FileReference.Exists(FileReference.Combine(ToolChainDir, "bin", "Hostx86", "x64", "cl.exe")) || FileReference.Exists(FileReference.Combine(ToolChainDir, "bin", "Hostx64", "x64", "cl.exe"));
+		}
+
 		/// <summary>
 		/// Determines if a given compiler is installed
 		/// </summary>
