Add a check that a 2015 compiler exists under VS2015 toolchain directories - it may not have been installed.

diff --git a/Engine/Source/Programs/UnrealBuildTool/Platform/Windows/UEBuildWindows.cs b/Engine/Source/Programs/UnrealBuildTool/Platform/Windows/UEBuildWindows.cs
--- a/Engine/Source/Programs/UnrealBuildTool/Platform/Windows/UEBuildWindows.cs
+++ b/Engine/Source/Programs/UnrealBuildTool/Platform/Windows/UEBuildWindows.cs
@@ -662,7 +662,11 @@ namespace UnrealBuildTool
 					    List<DirectoryReference> InstallDirs = FindVSInstallDirs(Compiler);
 					    foreach(DirectoryReference InstallDir in InstallDirs)
 					    {
-						    ToolChainVersionToDir[new VersionNumber(14, 0)] = DirectoryReference.Combine(InstallDir, "VC");
+							DirectoryReference ToolChainBaseDir = DirectoryReference.Combine(InstallDir, "VC");
+							if(IsValidToolChainDir2015(ToolChainBaseDir))
+							{
+								ToolChainVersionToDir[new VersionNumber(14, 0)] = ToolChainBaseDir;
+							}
 					    }
 				    }
 				    else if(Compiler == WindowsCompiler.VisualStudio2017)
@@ -677,7 +681,7 @@ namespace UnrealBuildTool
 							    foreach(DirectoryReference ToolChainDir in DirectoryReference.EnumerateDirectories(ToolChainBaseDir))
 							    {
 								    VersionNumber Version;
-								    if(VersionNumber.TryParse(ToolChainDir.GetDirectoryName(), out Version) && IsValidToolChainDir(ToolChainDir))
+								    if(VersionNumber.TryParse(ToolChainDir.GetDirectoryName(), out Version) && IsValidToolChainDir2017(ToolChainDir))
 								    {
 									    ToolChainVersionToDir[Version] = ToolChainDir;
 								    }
@@ -695,7 +699,7 @@ namespace UnrealBuildTool
 							    foreach(DirectoryReference ToolChainDir in DirectoryReference.EnumerateDirectories(ToolChainBaseDir))
 							    {
 								    VersionNumber Version;
-								    if(VersionNumber.TryParse(ToolChainDir.GetDirectoryName(), out Version) && IsValidToolChainDir(ToolChainDir))
+								    if(VersionNumber.TryParse(ToolChainDir.GetDirectoryName(), out Version) && IsValidToolChainDir2017(ToolChainDir))
 								    {
 									    ToolChainVersionToDir[Version] = ToolChainDir;
 								    }
@@ -714,11 +718,21 @@ namespace UnrealBuildTool
 		}
 
 		/// <summary>
-		/// Checks if the given directory contains a valid toolchain
+		/// Checks if the given directory contains a valid Visual Studio 2015 toolchain
+		/// </summary>
+		/// <param name="ToolChainDir">Directory to check</param>
+		/// <returns>True if the given directory is valid</returns>
+		static bool IsValidToolChainDir2015(DirectoryReference ToolChainDir)
+		{
+			return FileReference.Exists(FileReference.Combine(ToolChainDir, "bin", "amd64", "cl.exe")) || FileReference.Exists(FileReference.Combine(ToolChainDir, "bin", "x86_amd64", "cl.exe"));
+		}
+
+		/// <summary>
+		/// Checks if the given directory contains a valid Visual Studio 2017 toolchain
 		/// </summary>
 		/// <param name="ToolChainDir">Directory to check</param>
 		/// <returns>True if the given directory is valid</returns>
-		static bool IsValidToolChainDir(DirectoryReference ToolChainDir)
+		static bool IsValidToolChainDir2017(DirectoryReference ToolChainDir)
 		{
 			return FileReference.Exists(FileReference.Combine(ToolChainDir, "bin", "Hostx86", "x64", "cl.exe")) || FileReference.Exists(FileReference.Combine(ToolChainDir, "bin", "Hostx64", "x64", "cl.exe"));
 		}
