UBT: Fix intellisense search paths not being set correctly for modules which are not built by default.

diff --git a/Engine/Source/Programs/UnrealBuildTool/Configuration/UEBuildTarget.cs b/Engine/Source/Programs/UnrealBuildTool/Configuration/UEBuildTarget.cs
--- a/Engine/Source/Programs/UnrealBuildTool/Configuration/UEBuildTarget.cs
+++ b/Engine/Source/Programs/UnrealBuildTool/Configuration/UEBuildTarget.cs
@@ -2281,6 +2281,10 @@ namespace UnrealBuildTool
 				{
 					Binary.GatherDataForProjectFiles(Rules, GlobalCompileEnvironment);
 				}
+				foreach (UEBuildBinary Binary in PrecompileOnlyBinaries)
+				{
+					Binary.GatherDataForProjectFiles(Rules, GlobalCompileEnvironment);
+				}
 				return ECompilationResult.Succeeded;
 			}
 
