Allow the NativizedAssets plugin source to include monolithic engine header files without a compiler warning.

- Mirrored from //UE4/Dev-Framework (4081432).

#rnx

diff --git a/Engine/Source/Programs/UnrealBuildTool/Configuration/UEBuildTarget.cs b/Engine/Source/Programs/UnrealBuildTool/Configuration/UEBuildTarget.cs
--- a/Engine/Source/Programs/UnrealBuildTool/Configuration/UEBuildTarget.cs
+++ b/Engine/Source/Programs/UnrealBuildTool/Configuration/UEBuildTarget.cs
@@ -422,7 +422,7 @@ namespace UnrealBuildTool
 			}
 			if (Desc.ForeignPlugin != null)
 			{
-				RulesAssembly = RulesCompiler.CreatePluginRulesAssembly(Desc.ForeignPlugin, RulesAssembly);
+				RulesAssembly = RulesCompiler.CreatePluginRulesAssembly(Desc.ForeignPlugin, RulesAssembly, true);
 			}
 
 			FileReference TargetFileName;
@@ -599,7 +599,7 @@ namespace UnrealBuildTool
 				FileReference PluginFile = FileReference.Combine(RulesObject.ProjectFile.Directory, "Intermediate", "Plugins", "NativizedAssets", PlatformName, ProjectTargetType, "NativizedAssets.uplugin");
 				if (FileReference.Exists(PluginFile))
 				{
-					RulesAssembly = RulesCompiler.CreatePluginRulesAssembly(PluginFile, RulesAssembly);
+					RulesAssembly = RulesCompiler.CreatePluginRulesAssembly(PluginFile, RulesAssembly, false);
 				}
 				else
 				{
diff --git a/Engine/Source/Programs/UnrealBuildTool/System/RulesCompiler.cs b/Engine/Source/Programs/UnrealBuildTool/System/RulesCompiler.cs
--- a/Engine/Source/Programs/UnrealBuildTool/System/RulesCompiler.cs
+++ b/Engine/Source/Programs/UnrealBuildTool/System/RulesCompiler.cs
@@ -389,13 +389,14 @@ namespace UnrealBuildTool
 			return ProjectRulesAssembly;
 		}
 
-		/// <summary>
-		/// Creates a rules assembly with the given parameters.
-		/// </summary>
-		/// <param name="PluginFileName">The plugin file to create rules for</param>
-		/// <param name="Parent">The parent rules assembly</param>
-		/// <returns>The new rules assembly</returns>
-		public static RulesAssembly CreatePluginRulesAssembly(FileReference PluginFileName, RulesAssembly Parent)
+        /// <summary>
+        /// Creates a rules assembly with the given parameters.
+        /// </summary>
+        /// <param name="PluginFileName">The plugin file to create rules for</param>
+        /// <param name="Parent">The parent rules assembly</param>
+        /// <param name="bContainsEngineModules">Whether the plugin contains engine modules. Used to initialize the default value for ModuleRules.bTreatAsEngineModule.</param>
+        /// <returns>The new rules assembly</returns>
+        public static RulesAssembly CreatePluginRulesAssembly(FileReference PluginFileName, RulesAssembly Parent, bool bContainsEngineModules)
 		{
 			// Check if there's an existing assembly for this project
 			RulesAssembly PluginRulesAssembly;
@@ -418,7 +419,7 @@ namespace UnrealBuildTool
 
 				// Compile the assembly
 				FileReference AssemblyFileName = FileReference.Combine(PluginFileName.Directory, "Intermediate", "Build", "BuildRules", Path.GetFileNameWithoutExtension(PluginFileName.FullName) + "ModuleRules" + FrameworkAssemblyExtension);
-				PluginRulesAssembly = new RulesAssembly(ForeignPlugins, ModuleFiles, TargetFiles, ModuleFileToPluginInfo, AssemblyFileName, bContainsEngineModules: true, bUseBackwardsCompatibleDefaults: false, bInstalled: false, Parent: Parent);
+				PluginRulesAssembly = new RulesAssembly(ForeignPlugins, ModuleFiles, TargetFiles, ModuleFileToPluginInfo, AssemblyFileName, bContainsEngineModules, bUseBackwardsCompatibleDefaults: false, bInstalled: false, Parent: Parent);
 				LoadedAssemblyMap.Add(PluginFileName, PluginRulesAssembly);
 			}
 			return PluginRulesAssembly;
