Speculative fix for exception reading UBT makefiles on Mono. Doesn't seem to be able to deserialize dictionaries with custom comparison functions correctly.

diff --git a/Engine/Source/Programs/UnrealBuildTool/Configuration/UEBuildTarget.cs b/Engine/Source/Programs/UnrealBuildTool/Configuration/UEBuildTarget.cs
--- a/Engine/Source/Programs/UnrealBuildTool/Configuration/UEBuildTarget.cs
+++ b/Engine/Source/Programs/UnrealBuildTool/Configuration/UEBuildTarget.cs
@@ -281,6 +281,7 @@ namespace UnrealBuildTool
 	{
 		public FlatModuleCsDataType(SerializationInfo Info, StreamingContext Context)
 		{
+			ModuleName = Info.GetString("mn");
 			BuildCsFilename = Info.GetString("bf");
 			ModuleSourceFolder = (DirectoryReference)Info.GetValue("mf", typeof(DirectoryReference));
 			ExternalDependencies = (List<string>)Info.GetValue("ed", typeof(List<string>));
@@ -289,18 +290,21 @@ namespace UnrealBuildTool
 
 		public void GetObjectData(SerializationInfo Info, StreamingContext Context)
 		{
+			Info.AddValue("mn", ModuleName);
 			Info.AddValue("bf", BuildCsFilename);
 			Info.AddValue("mf", ModuleSourceFolder);
 			Info.AddValue("ed", ExternalDependencies);
 			Info.AddValue("hn", UHTHeaderNames);
 		}
 
-		public FlatModuleCsDataType(string InBuildCsFilename, IEnumerable<string> InExternalDependencies)
+		public FlatModuleCsDataType(string InModuleName, string InBuildCsFilename, IEnumerable<string> InExternalDependencies)
 		{
+			ModuleName = InModuleName;
 			BuildCsFilename = InBuildCsFilename;
 			ExternalDependencies = new List<string>(InExternalDependencies);
 		}
 
+		public string ModuleName;
 		public string BuildCsFilename;
 		public DirectoryReference ModuleSourceFolder;
 		public List<string> ExternalDependencies;
@@ -870,7 +874,7 @@ namespace UnrealBuildTool
 		/// <summary>
 		/// Used to map names of modules to their .Build.cs filename
 		/// </summary>
-		public Dictionary<string, FlatModuleCsDataType> FlatModuleCsData = new Dictionary<string, FlatModuleCsDataType>(StringComparer.InvariantCultureIgnoreCase);
+		public List<FlatModuleCsDataType> FlatModuleCsData = new List<FlatModuleCsDataType>();
 
 		/// <summary>
 		/// The receipt for this target, which contains a record of this build.
@@ -927,7 +931,7 @@ namespace UnrealBuildTool
 		/// <returns></returns>
 		public IEnumerable<string> GetAllModuleBuildCsFilenames()
 		{
-			return FlatModuleCsData.Values.Select(Data => Data.BuildCsFilename);
+			return FlatModuleCsData.Select(Data => Data.BuildCsFilename);
 		}
 
 		/// <summary>
@@ -936,7 +940,7 @@ namespace UnrealBuildTool
 		/// <returns></returns>
 		public IEnumerable<string> GetAllModuleFolders()
 		{
-			return FlatModuleCsData.Values.SelectMany(Data => Data.UHTHeaderNames);
+			return FlatModuleCsData.SelectMany(Data => Data.UHTHeaderNames);
 		}
 
 		/// <summary>
@@ -968,7 +972,7 @@ namespace UnrealBuildTool
 			bUsePrecompiled = Info.GetBoolean("up");
 			OnlyModules = (List<OnlyModule>)Info.GetValue("om", typeof(List<OnlyModule>));
 			bCompileMonolithic = Info.GetBoolean("cm");
-			FlatModuleCsData = (Dictionary<string, FlatModuleCsDataType>)Info.GetValue("fm", typeof(Dictionary<string, FlatModuleCsDataType>));
+			FlatModuleCsData = (List<FlatModuleCsDataType>)Info.GetValue("fm", typeof(List<FlatModuleCsDataType>));
 			Receipt = (TargetReceipt)Info.GetValue("re", typeof(TargetReceipt));
 			ReceiptFileName = (FileReference)Info.GetValue("rf", typeof(FileReference));
 			FileReferenceToModuleManifestPairs = (KeyValuePair<FileReference, ModuleManifest>[])Info.GetValue("vm", typeof(KeyValuePair<FileReference, ModuleManifest>[]));
@@ -2230,7 +2234,13 @@ namespace UnrealBuildTool
 					ModulesToGenerateHeadersFor = CorrectlyOrderedModules;
 				}
 
-				ExternalExecution.SetupUObjectModules(ModulesToGenerateHeadersFor, Rules, GlobalCompileEnvironment, UObjectModules, FlatModuleCsData, Rules.GeneratedCodeVersion, bIsAssemblingBuild);
+				Dictionary<string, FlatModuleCsDataType> NameToFlatModuleData = new Dictionary<string, FlatModuleCsDataType>(StringComparer.InvariantCultureIgnoreCase);
+				foreach(FlatModuleCsDataType FlatModuleData in FlatModuleCsData)
+				{
+					NameToFlatModuleData[FlatModuleData.ModuleName] = FlatModuleData;
+				}
+
+				ExternalExecution.SetupUObjectModules(ModulesToGenerateHeadersFor, Rules, GlobalCompileEnvironment, UObjectModules, NameToFlatModuleData, Rules.GeneratedCodeVersion, bIsAssemblingBuild);
 
 				// NOTE: Even in Gather mode, we need to run UHT to make sure the files exist for the static action graph to be setup correctly.  This is because UHT generates .cpp
 				// files that are injected as top level prerequisites.  If UHT only emitted included header files, we wouldn't need to run it during the Gather phase at all.
@@ -4277,7 +4287,7 @@ namespace UnrealBuildTool
 				// Now, go ahead and create the module builder instance
 				Module = InstantiateModule(RulesObject, ModuleName, ModuleType.Value, ModuleDirectory, GeneratedCodeDirectory, FoundSourceFiles, bBuildFiles, ModuleFileName, RuntimeDependencies, Plugin);
 				Modules.Add(Module.Name, Module);
-				FlatModuleCsData.Add(Module.Name, new FlatModuleCsDataType((Module.RulesFile == null) ? null : Module.RulesFile.FullName, RulesObject.ExternalDependencies));
+				FlatModuleCsData.Add(new FlatModuleCsDataType(Module.Name, (Module.RulesFile == null) ? null : Module.RulesFile.FullName, RulesObject.ExternalDependencies));
 			}
 			return Module;
 		}
diff --git a/Engine/Source/Programs/UnrealBuildTool/System/UBTMakefile.cs b/Engine/Source/Programs/UnrealBuildTool/System/UBTMakefile.cs
--- a/Engine/Source/Programs/UnrealBuildTool/System/UBTMakefile.cs
+++ b/Engine/Source/Programs/UnrealBuildTool/System/UBTMakefile.cs
@@ -323,7 +323,7 @@ namespace UnrealBuildTool
 					}
 				}
 
-				foreach (FlatModuleCsDataType FlatCsModuleData in Target.FlatModuleCsData.Values)
+				foreach (FlatModuleCsDataType FlatCsModuleData in Target.FlatModuleCsData)
 				{
 					if (FlatCsModuleData.BuildCsFilename != null && FlatCsModuleData.ExternalDependencies.Count > 0)
 					{
@@ -352,13 +352,13 @@ namespace UnrealBuildTool
 				HashSet<string> HFilesNewerThanMakefile =
 					new HashSet<string>(
 						Target.FlatModuleCsData
-						.SelectMany(x => x.Value.ModuleSourceFolder != null ? Directory.EnumerateFiles(x.Value.ModuleSourceFolder.FullName, "*.h", SearchOption.AllDirectories) : Enumerable.Empty<string>())
+						.SelectMany(x => x.ModuleSourceFolder != null ? Directory.EnumerateFiles(x.ModuleSourceFolder.FullName, "*.h", SearchOption.AllDirectories) : Enumerable.Empty<string>())
 						.Where(y => Directory.GetLastWriteTimeUtc(y) > UBTMakefileInfo.LastWriteTimeUtc)
 						.OrderBy(z => z).Distinct()
 					);
 
 				// Get all H files in all modules processed in the last makefile build
-				HashSet<string> AllUHTHeaders = new HashSet<string>(Target.FlatModuleCsData.Select(x => x.Value).SelectMany(x => x.UHTHeaderNames));
+				HashSet<string> AllUHTHeaders = new HashSet<string>(Target.FlatModuleCsData.SelectMany(x => x.UHTHeaderNames));
 
 				// Check whether any headers have been deleted. If they have, we need to regenerate the makefile since the module might now be empty. If we don't,
 				// and the file has been moved to a different module, we may include stale generated headers.
