Fix incorrect console colors after a warning or error on Mac. Cannot assume that the default foreground color is gray.

diff --git a/Engine/Source/Programs/UnrealBuildTool/System/Log.cs b/Engine/Source/Programs/UnrealBuildTool/System/Log.cs
--- a/Engine/Source/Programs/UnrealBuildTool/System/Log.cs
+++ b/Engine/Source/Programs/UnrealBuildTool/System/Log.cs
@@ -349,30 +349,22 @@ namespace UnrealBuildTool
 			if (Verbosity <= LogLevel)
 			{
 				// Do console color highlighting here.
-				ConsoleColor DefaultColor = ConsoleColor.Gray;
-				bool bIsWarning = false;
-				bool bIsError = false;
-				// don't try to touch the console unless we are told to color the output.
+				bool bResetConsoleColor = false;
 				if (bColorConsoleOutput)
 				{
-					DefaultColor = Console.ForegroundColor;
-					bIsWarning = Verbosity == LogEventType.Warning;
-					bIsError = Verbosity <= LogEventType.Error;
-					// @todo mono - mono doesn't seem to initialize the ForegroundColor properly, so we can't restore it properly.
-					// Avoid touching the console color unless we really need to.
-					if (bIsWarning || bIsError)
+					if(Verbosity == LogEventType.Warning)
 					{
-						Console.ForegroundColor = bIsWarning ? ConsoleColor.Yellow : ConsoleColor.Red;
+						Console.ForegroundColor = ConsoleColor.Yellow;
+						bResetConsoleColor = true;
+					}
+					if(Verbosity <= LogEventType.Error)
+					{
+						Console.ForegroundColor = ConsoleColor.Red;
+						bResetConsoleColor = true;
 					}
 				}
 				try
 				{
-					// @todo mono: mono has some kind of bug where calling mono recursively by spawning
-					// a new process causes Trace.WriteLine to stop functioning (it returns, but does nothing for some reason).
-					// work around this by simulating Trace.WriteLine on mono.
-					// We use UAT to spawn UBT instances recursively a lot, so this bug can effectively
-					// make all build output disappear outside of the top level UAT process.
-					//					#if MONO
 					lock (((System.Collections.ICollection)Trace.Listeners).SyncRoot)
 					{
 						foreach (TraceListener l in Trace.Listeners)
@@ -389,17 +381,13 @@ namespace UnrealBuildTool
 							l.Flush();
 						}
 					}
-					//					#else
-					// Call Trace directly here. Trace ensures that our logging is threadsafe using the GlobalLock.
-					//                    	Trace.WriteLine(FormatMessage(StackFramesToSkip + 1, CustomSource, Verbosity, Format, Args));
-					//					#endif
 				}
 				finally
 				{
 					// make sure we always put the console color back.
-					if (bColorConsoleOutput && (bIsWarning || bIsError))
+					if(bResetConsoleColor)
 					{
-						Console.ForegroundColor = DefaultColor;
+						Console.ResetColor();
 					}
 				}
 			}
