If we are using the native bundled toolchain set LC_ALL=C to avoid locale issues

diff --git a/Engine/Source/Programs/UnrealBuildTool/Platform/Linux/LinuxToolChain.cs b/Engine/Source/Programs/UnrealBuildTool/Platform/Linux/LinuxToolChain.cs
--- a/Engine/Source/Programs/UnrealBuildTool/Platform/Linux/LinuxToolChain.cs
+++ b/Engine/Source/Programs/UnrealBuildTool/Platform/Linux/LinuxToolChain.cs
@@ -76,6 +76,11 @@ namespace UnrealBuildTool
 				// When cross-compiling on Windows, use old FixDeps. It is slow, but it does not have timing issues
 				bUseFixdeps = (BuildHostPlatform.Current.Platform == UnrealTargetPlatform.Win64 || BuildHostPlatform.Current.Platform == UnrealTargetPlatform.Win32);
 
+				if (BuildHostPlatform.Current.Platform == UnrealTargetPlatform.Linux)
+				{
+					Environment.SetEnvironmentVariable("LC_ALL", "C");
+				}
+
 				bIsCrossCompiling = true;
 
 				bHasValidCompiler = DetermineCompilerVersion();
@@ -112,7 +117,7 @@ namespace UnrealBuildTool
 				bIsCrossCompiling = false;
 
 				bHasValidCompiler = DetermineCompilerVersion();
-			} 
+			}
 			else
 			{
 				ToolchainInfo = String.Format("toolchain located at '{0}'", BaseLinuxPath);
