Fix runtime dependencies referenced by precompiled modules for monolithic target platforms not being included in the receipt.

diff --git a/Engine/Source/Programs/UnrealBuildTool/Configuration/UEBuildTarget.cs b/Engine/Source/Programs/UnrealBuildTool/Configuration/UEBuildTarget.cs
--- a/Engine/Source/Programs/UnrealBuildTool/Configuration/UEBuildTarget.cs
+++ b/Engine/Source/Programs/UnrealBuildTool/Configuration/UEBuildTarget.cs
@@ -1610,19 +1610,19 @@ namespace UnrealBuildTool
 			// Find all the modules which are part of this target
 			HashSet<UEBuildModule> UniqueLinkedModules = new HashSet<UEBuildModule>();
 			foreach (UEBuildBinary Binary in Binaries)
+			{
+				foreach (UEBuildModule Module in Binary.Modules)
 				{
-					foreach (UEBuildModule Module in Binary.Modules)
+					if (UniqueLinkedModules.Add(Module))
 					{
-						if (UniqueLinkedModules.Add(Module))
+						foreach (RuntimeDependency RuntimeDependency in Module.RuntimeDependencies)
 						{
-							foreach (RuntimeDependency RuntimeDependency in Module.RuntimeDependencies)
-							{
-								Receipt.RuntimeDependencies.Add(RuntimeDependency.Path, RuntimeDependency.Type);
-							}
-							Receipt.AdditionalProperties.AddRange(Module.Rules.AdditionalPropertiesForReceipt.Inner);
+							Receipt.RuntimeDependencies.Add(RuntimeDependency.Path, RuntimeDependency.Type);
 						}
+						Receipt.AdditionalProperties.AddRange(Module.Rules.AdditionalPropertiesForReceipt.Inner);
 					}
 				}
+			}
 
 			// Add any dependencies of precompiled modules into the receipt
 			if(bPrecompile)
@@ -1649,20 +1649,28 @@ namespace UnrealBuildTool
 					}
 				}
 
+				// Find all the modules we need to add runtime dependencies for. This may include precompile-only modules, as well as third party binaries
+				HashSet<UEBuildModule> ReferencedModules = new HashSet<UEBuildModule>(PrecompileOnlyModules);
+				foreach(UEBuildModule Module in PrecompileOnlyModules)
+				{
+					ReferencedModules.UnionWith(Module.GetDependencies(false, false));
+				}
+				foreach(UEBuildBinary Binary in PrecompileOnlyBinaries)
+				{
+					ReferencedModules.UnionWith(Binary.Modules);
+				}
+
 				// Add the runtime dependencies of precompiled modules that are not directly part of this target
-				foreach (UEBuildBinary Binary in PrecompileOnlyBinaries)
+				foreach(UEBuildModule ReferencedModule in ReferencedModules)
+				{
+					if (UniqueLinkedModules.Add(ReferencedModule))
 					{
-						foreach (UEBuildModule Module in Binary.Modules)
+						foreach (RuntimeDependency RuntimeDependency in ReferencedModule.RuntimeDependencies)
 						{
-							if (UniqueLinkedModules.Add(Module))
-							{
-								foreach (RuntimeDependency RuntimeDependency in Module.RuntimeDependencies)
-								{
-									Receipt.PrecompiledRuntimeDependencies.Add(RuntimeDependency.Path);
-								}
-							}
+							Receipt.PrecompiledRuntimeDependencies.Add(RuntimeDependency.Path);
 						}
 					}
+				}
 
 				// Add all the files which are required to use the precompiled modules
 				HashSet<FileReference> ExternalFiles = new HashSet<FileReference>();
@@ -2956,10 +2964,6 @@ namespace UnrealBuildTool
 			{
 				if(Module.Binary != null && Module.Rules.bPrecompile)
 				{
-					if(Module.Name == "SpriterImporter")
-					{
-						Console.WriteLine();
-					}
 					PrecompiledModuleNames.Add(Module.Name);
 				}
 			}
