UE-59711 - TVOS/iOS fails to build from Windows with numerous errors 'no such file or directory' (reviewed)

diff --git a/Engine/Source/Programs/UnrealBuildTool/Configuration/UEBuildModuleCPP.cs b/Engine/Source/Programs/UnrealBuildTool/Configuration/UEBuildModuleCPP.cs
--- a/Engine/Source/Programs/UnrealBuildTool/Configuration/UEBuildModuleCPP.cs
+++ b/Engine/Source/Programs/UnrealBuildTool/Configuration/UEBuildModuleCPP.cs
@@ -1,4 +1,4 @@
-﻿// Copyright 1998-2018 Epic Games, Inc. All Rights Reserved.
+// Copyright 1998-2018 Epic Games, Inc. All Rights Reserved.
 
 using System;
 using System.Collections.Generic;
@@ -374,7 +374,7 @@ namespace UnrealBuildTool
 		/// <summary>
 		/// Path to the precompiled manifest location
 		/// </summary>
-		public FileReference PrecompiledManifestLocation
+		public virtual FileReference PrecompiledManifestLocation
 		{
 			get { return FileReference.Combine(IntermediateDirectory, String.Format("{0}.precompiled", Name)); }
 		}
@@ -431,6 +431,7 @@ namespace UnrealBuildTool
 				foreach(FileReference OutputFile in Manifest.OutputFiles)
 				{
 					FileItem ObjectFile = FileItem.GetExistingItemByFileReference(OutputFile);
+					ToolChain.DoLocalToRemoteFileItem(ObjectFile);
 					LinkInputFiles.Add(ObjectFile);
 				}
 				return LinkInputFiles;
diff --git a/Engine/Source/Programs/UnrealBuildTool/Platform/IOS/IOSToolChain.cs b/Engine/Source/Programs/UnrealBuildTool/Platform/IOS/IOSToolChain.cs
--- a/Engine/Source/Programs/UnrealBuildTool/Platform/IOS/IOSToolChain.cs
+++ b/Engine/Source/Programs/UnrealBuildTool/Platform/IOS/IOSToolChain.cs
@@ -954,7 +954,8 @@ namespace UnrealBuildTool
 			List<string> InputFileNames = new List<string>();
 			foreach (FileItem InputFile in LinkEnvironment.InputFiles)
 			{
-				InputFileNames.Add(string.Format("\"{0}\"", InputFile.AbsolutePath));
+				string FileNamePath = string.Format("\"{0}\"", ConvertPath(InputFile.AbsolutePath));
+				InputFileNames.Add(FileNamePath);
 				LinkAction.PrerequisiteItems.Add(InputFile);
 			}
 
diff --git a/Engine/Source/Programs/UnrealBuildTool/ToolChain/RemoteToolChain.cs b/Engine/Source/Programs/UnrealBuildTool/ToolChain/RemoteToolChain.cs
--- a/Engine/Source/Programs/UnrealBuildTool/ToolChain/RemoteToolChain.cs
+++ b/Engine/Source/Programs/UnrealBuildTool/ToolChain/RemoteToolChain.cs
@@ -723,7 +723,12 @@ namespace UnrealBuildTool
 			return RemoteFileItem;
 		}
 
-        public FileItem RemoteToLocalFileItem(FileItem RemoteFileItem)
+		public override void DoLocalToRemoteFileItem(FileItem LocalFileItem)
+		{
+			LocalToRemoteFileItem(LocalFileItem, true);
+		}
+
+		public FileItem RemoteToLocalFileItem(FileItem RemoteFileItem)
         {
             // Look to see if we've already made a remote FileItem for this local FileItem
             foreach (KeyValuePair<FileItem, FileItem> Item in CachedRemoteFileItems)
@@ -957,7 +962,8 @@ namespace UnrealBuildTool
 				// --exclude='*'  ??? why???
 				RsyncProcess.StartInfo.FileName = ResolvedRSyncExe;
                 RsyncProcess.StartInfo.Arguments = string.Format(
-                    "-vzrltgoDe \"{0}\" --rsync-path=\"mkdir -p {2} && rsync\" --chmod=ug=rwX,o=rxX --delete --files-from=\"{4}\" --include-from=\"{5}\" --include='*/' --exclude='*.o' --exclude='Timestamp' '{1}' \"{6}@{3}\":'{2}'",
+					//"-vzrltgoDe \"{0}\" --rsync-path=\"mkdir -p {2} && rsync\" --chmod=ug=rwX,o=rxX --delete --files-from=\"{4}\" --include-from=\"{5}\" --include='*/' --exclude='*.o' --exclude='Timestamp' '{1}' \"{6}@{3}\":'{2}'",
+					"-vzrltgoDe \"{0}\" --rsync-path=\"mkdir -p {2} && rsync\" --chmod=ug=rwX,o=rxX --delete --files-from=\"{4}\" --include-from=\"{5}\" --include='*/' --exclude='Timestamp' '{1}' \"{6}@{3}\":'{2}'",
 					ResolvedRsyncAuthentication,
 					CygRootPath,
 					RemotePath,
diff --git a/Engine/Source/Programs/UnrealBuildTool/ToolChain/UEToolChain.cs b/Engine/Source/Programs/UnrealBuildTool/ToolChain/UEToolChain.cs
--- a/Engine/Source/Programs/UnrealBuildTool/ToolChain/UEToolChain.cs
+++ b/Engine/Source/Programs/UnrealBuildTool/ToolChain/UEToolChain.cs
@@ -26,6 +26,10 @@ namespace UnrealBuildTool
 		{
 		}
 
+		public virtual void DoLocalToRemoteFileItem(FileItem LocalFileItem)
+		{
+		}
+
 		public abstract CPPOutput CompileCPPFiles(CppCompileEnvironment CompileEnvironment, List<FileItem> InputFiles, DirectoryReference OutputDir, string ModuleName, ActionGraph ActionGraph);
 
 		public virtual CPPOutput CompileRCFiles(CppCompileEnvironment Environment, List<FileItem> InputFiles, DirectoryReference OutputDir, ActionGraph ActionGraph)
