Handle issue with clang failing to use a response file to call ld by capturing the arguments and creating one manually instead
#ue4
#android

diff --git a/Engine/Source/Programs/UnrealBuildTool/Platform/Android/AndroidToolChain.cs b/Engine/Source/Programs/UnrealBuildTool/Platform/Android/AndroidToolChain.cs
--- a/Engine/Source/Programs/UnrealBuildTool/Platform/Android/AndroidToolChain.cs
+++ b/Engine/Source/Programs/UnrealBuildTool/Platform/Android/AndroidToolChain.cs
@@ -1665,12 +1665,69 @@ namespace UnrealBuildTool
 						LinkAction.CommandArguments = String.Format("/c \"{0} {1}\"", LinkAction.CommandPath, LinkAction.CommandArguments);
 						LinkAction.CommandPath = "cmd.exe";
 					}
+
+					// Windows can run into an issue with too long of a commandline when clang tries to call ld to link.
+					// To work around this we call clang to just get the command it would execute and generate a
+					// second response file to directly call ld with the right arguments instead of calling through clang.
+					if (!Utils.IsRunningOnMono)
+					{
+						// capture the actual link command without running it
+						ProcessStartInfo StartInfo = new ProcessStartInfo();
+						StartInfo.WorkingDirectory = LinkEnvironment.IntermediateDirectory.FullName;
+						StartInfo.FileName = LinkAction.CommandPath;
+						StartInfo.Arguments = "-### " + LinkAction.CommandArguments;
+						StartInfo.UseShellExecute = false;
+						StartInfo.CreateNoWindow = true;
+						StartInfo.RedirectStandardError = true;
+
+						LinkerCommandline = "";
+
+						Process Proc = new Process();
+						Proc.StartInfo = StartInfo;
+						Proc.ErrorDataReceived += new DataReceivedEventHandler(OutputReceivedForLinker);
+						Proc.Start();
+						Proc.BeginErrorReadLine();
+						Proc.WaitForExit(5000);
+
+						LinkerCommandline = LinkerCommandline.Trim();
+
+						// the command should be in quotes; if not we'll just use clang to link as usual
+						int FirstQuoteIndex = LinkerCommandline.IndexOf('"');
+						if (FirstQuoteIndex >= 0)
+						{
+							int SecondQuoteIndex = LinkerCommandline.Substring(FirstQuoteIndex + 1).IndexOf('"');
+							if (SecondQuoteIndex >= 0)
+							{
+								LinkAction.CommandPath = LinkerCommandline.Substring(FirstQuoteIndex + 1, SecondQuoteIndex - FirstQuoteIndex);
+								LinkAction.CommandArguments = LinkerCommandline.Substring(FirstQuoteIndex + SecondQuoteIndex + 3);
+
+								// replace double backslashes
+								LinkAction.CommandPath = LinkAction.CommandPath.Replace("\\\\", "/");
+
+								// now create a response file for the full command using ld directly
+								FileReference FinalResponseFileName = FileReference.Combine(LinkEnvironment.IntermediateDirectory, OutputFile.Location.GetFileName() + ".responseFinal");
+								FileItem FinalResponseFileItem = FileItem.CreateIntermediateTextFile(FinalResponseFileName, LinkAction.CommandArguments);
+								LinkAction.CommandArguments = string.Format("@\"{0}\"", FinalResponseFileName);
+								LinkAction.PrerequisiteItems.Add(FinalResponseFileItem);
+							}
+						}
+					}
 				}
 			}
 
 			return Outputs.ToArray();
 		}
 
+		// captures stderr from clang
+		private static string LinkerCommandline = "";
+		static public void OutputReceivedForLinker(Object Sender, DataReceivedEventArgs Line)
+		{
+			if ((Line != null) && (Line.Data != null) && (Line.Data.Contains("--sysroot")))
+			{
+				LinkerCommandline += Line.Data;
+			}
+		}
+
 		private void ExportObjectFilePaths(LinkEnvironment LinkEnvironment, string FileName)
 		{
 			// Write the list of object file directories
