Disable the alternative linker invocation for Android since it is causing other issues
#ue4
#android

diff --git a/Engine/Source/Programs/UnrealBuildTool/Platform/Android/AndroidToolChain.cs b/Engine/Source/Programs/UnrealBuildTool/Platform/Android/AndroidToolChain.cs
--- a/Engine/Source/Programs/UnrealBuildTool/Platform/Android/AndroidToolChain.cs
+++ b/Engine/Source/Programs/UnrealBuildTool/Platform/Android/AndroidToolChain.cs
@@ -1669,6 +1669,7 @@ namespace UnrealBuildTool
 					// Windows can run into an issue with too long of a commandline when clang tries to call ld to link.
 					// To work around this we call clang to just get the command it would execute and generate a
 					// second response file to directly call ld with the right arguments instead of calling through clang.
+/* disable while tracking down some linker errors this introduces
 					if (!Utils.IsRunningOnMono)
 					{
 						// capture the actual link command without running it
@@ -1712,6 +1713,7 @@ namespace UnrealBuildTool
 							}
 						}
 					}
+*/
 				}
 			}
 
