UBT: Fix issue causing incorrect compile environment for *.cc files. Treat them the same as .cpp files.

diff --git a/Engine/Source/Programs/UnrealBuildTool/Configuration/UEBuildModuleCPP.cs b/Engine/Source/Programs/UnrealBuildTool/Configuration/UEBuildModuleCPP.cs
--- a/Engine/Source/Programs/UnrealBuildTool/Configuration/UEBuildModuleCPP.cs
+++ b/Engine/Source/Programs/UnrealBuildTool/Configuration/UEBuildModuleCPP.cs
@@ -655,19 +655,31 @@ namespace UnrealBuildTool
 			}
 
 			// Compile C files directly. Do not use a PCH here, because a C++ PCH is not compatible with C source files.
-			LinkInputFiles.AddRange(ToolChain.CompileCPPFiles(ModuleCompileEnvironment, SourceFilesToBuild.CFiles, IntermediateDirectory, Name, ActionGraph).ObjectFiles);
+			if(SourceFilesToBuild.CFiles.Count > 0)
+			{
+				LinkInputFiles.AddRange(ToolChain.CompileCPPFiles(ModuleCompileEnvironment, SourceFilesToBuild.CFiles, IntermediateDirectory, Name, ActionGraph).ObjectFiles);
+			}
 
 			// Compile CC files directly.
-			LinkInputFiles.AddRange(ToolChain.CompileCPPFiles(CompileEnvironment, SourceFilesToBuild.CCFiles, IntermediateDirectory, Name, ActionGraph).ObjectFiles);
+			if(SourceFilesToBuild.CCFiles.Count > 0)
+			{
+				LinkInputFiles.AddRange(ToolChain.CompileCPPFiles(CompileEnvironment, SourceFilesToBuild.CCFiles, IntermediateDirectory, Name, ActionGraph).ObjectFiles);
+			}
 
 			// Compile MM files directly.
-			LinkInputFiles.AddRange(ToolChain.CompileCPPFiles(CompileEnvironment, SourceFilesToBuild.MMFiles, IntermediateDirectory, Name, ActionGraph).ObjectFiles);
+			if(SourceFilesToBuild.MMFiles.Count > 0)
+			{
+				LinkInputFiles.AddRange(ToolChain.CompileCPPFiles(CompileEnvironment, SourceFilesToBuild.MMFiles, IntermediateDirectory, Name, ActionGraph).ObjectFiles);
+			}
 
 			// Compile RC files. The resource compiler does not work with response files, and using the regular compile environment can easily result in the 
 			// command line length exceeding the OS limit. Use the binary compile environment to keep the size down, and require that all include paths
 			// must be specified relative to the resource file itself or Engine/Source.
-			CppCompileEnvironment ResourceCompileEnvironment = new CppCompileEnvironment(BinaryCompileEnvironment);
-			LinkInputFiles.AddRange(ToolChain.CompileRCFiles(ResourceCompileEnvironment, SourceFilesToBuild.RCFiles, IntermediateDirectory, ActionGraph).ObjectFiles);
+			if(SourceFilesToBuild.RCFiles.Count > 0)
+			{
+				CppCompileEnvironment ResourceCompileEnvironment = new CppCompileEnvironment(BinaryCompileEnvironment);
+				LinkInputFiles.AddRange(ToolChain.CompileRCFiles(ResourceCompileEnvironment, SourceFilesToBuild.RCFiles, IntermediateDirectory, ActionGraph).ObjectFiles);
+			}
 
 			// Write the compiled manifest
 			if(Rules.bPrecompile)
diff --git a/Engine/Source/Programs/UnrealBuildTool/Platform/IOS/IOSToolChain.cs b/Engine/Source/Programs/UnrealBuildTool/Platform/IOS/IOSToolChain.cs
--- a/Engine/Source/Programs/UnrealBuildTool/Platform/IOS/IOSToolChain.cs
+++ b/Engine/Source/Programs/UnrealBuildTool/Platform/IOS/IOSToolChain.cs
@@ -708,12 +708,6 @@ namespace UnrealBuildTool
 					// Compile the file as C code.
 					FileArguments += GetCompileArguments_C();
 				}
-				else if (Extension == ".CC")
-				{
-					// Compile the file as C++ code.
-					FileArguments += GetCompileArguments_CPP();
-					FileArguments += GetRTTIFlag(CompileEnvironment);
-				}
 				else if (Extension == ".MM")
 				{
 					// Compile the file as Objective-C++ code.
diff --git a/Engine/Source/Programs/UnrealBuildTool/Platform/Linux/LinuxToolChain.cs b/Engine/Source/Programs/UnrealBuildTool/Platform/Linux/LinuxToolChain.cs
--- a/Engine/Source/Programs/UnrealBuildTool/Platform/Linux/LinuxToolChain.cs
+++ b/Engine/Source/Programs/UnrealBuildTool/Platform/Linux/LinuxToolChain.cs
@@ -1060,11 +1060,6 @@ namespace UnrealBuildTool
 					// Compile the file as C code.
 					FileArguments += GetCompileArguments_C();
 				}
-				else if (Extension == ".CC")
-				{
-					// Compile the file as C++ code.
-					FileArguments += GetCompileArguments_CPP();
-				}
 				else if (Extension == ".MM")
 				{
 					// Compile the file as Objective-C++ code.
diff --git a/Engine/Source/Programs/UnrealBuildTool/Platform/Mac/MacToolChain.cs b/Engine/Source/Programs/UnrealBuildTool/Platform/Mac/MacToolChain.cs
--- a/Engine/Source/Programs/UnrealBuildTool/Platform/Mac/MacToolChain.cs
+++ b/Engine/Source/Programs/UnrealBuildTool/Platform/Mac/MacToolChain.cs
@@ -506,12 +506,6 @@ namespace UnrealBuildTool
 					// Compile the file as C code.
 					FileArguments += GetCompileArguments_C();
 				}
-				else if (Extension == ".CC")
-				{
-					// Compile the file as C++ code.
-					FileArguments += GetCompileArguments_CPP();
-					FileArguments += GetRTTIFlag(CompileEnvironment);
-				}
 				else if (Extension == ".MM")
 				{
 					// Compile the file as Objective-C++ code.
