UBT: Remove unused flag.

#rnx

diff --git a/Engine/Source/Programs/UnrealBuildTool/Configuration/UEBuildModuleCPP.cs b/Engine/Source/Programs/UnrealBuildTool/Configuration/UEBuildModuleCPP.cs
--- a/Engine/Source/Programs/UnrealBuildTool/Configuration/UEBuildModuleCPP.cs
+++ b/Engine/Source/Programs/UnrealBuildTool/Configuration/UEBuildModuleCPP.cs
@@ -155,11 +155,6 @@ namespace UnrealBuildTool
 		/// </summary>
 		public List<string> InvalidIncludeDirectiveMessages;
 
-		/// <summary>
-		/// Hack to skip adding definitions to compile environment. They will be baked into source code by external code.
-		/// </summary>
-		public bool bSkipDefinitionsForCompileEnvironment = false;
-
 		public IEnumerable<string> FindGeneratedCppFiles()
 		{
 			return ((null == GeneratedCodeDirectory) || !DirectoryReference.Exists(GeneratedCodeDirectory))
@@ -1264,13 +1259,6 @@ namespace UnrealBuildTool
 			// Setup the compile environment for the module.
 			SetupPrivateCompileEnvironment(Result.IncludePaths.UserIncludePaths, Result.IncludePaths.SystemIncludePaths, Result.Definitions, Result.AdditionalFrameworks, (Rules != null)? Rules.bLegacyPublicIncludePaths.Value : true);
 
-			// @hack to skip adding definitions to compile environment, they will be baked into source code files
-			if (bSkipDefinitionsForCompileEnvironment)
-			{
-				Result.Definitions.Clear();
-				Result.IncludePaths.UserIncludePaths = new HashSet<DirectoryReference>(BaseCompileEnvironment.IncludePaths.UserIncludePaths);
-			}
-
 			return Result;
 		}
 
