Fix missing third party libraries in installed build.

diff --git a/Engine/Source/Programs/UnrealBuildTool/Configuration/UEBuildTarget.cs b/Engine/Source/Programs/UnrealBuildTool/Configuration/UEBuildTarget.cs
--- a/Engine/Source/Programs/UnrealBuildTool/Configuration/UEBuildTarget.cs
+++ b/Engine/Source/Programs/UnrealBuildTool/Configuration/UEBuildTarget.cs
@@ -1311,22 +1311,12 @@ namespace UnrealBuildTool
 		/// <summary>
 		/// Create a list of all the externally referenced files
 		/// </summary>
+		/// <param name="Modules">All the modules to include files for</param>
 		/// <param name="Files">Set of referenced files</param>
-		void GetExternalFileList(HashSet<FileReference> Files)
+		void GetExternalFileList(HashSet<UEBuildModule> Modules, HashSet<FileReference> Files)
 		{
-			// Find all the modules we depend on
-			HashSet<UEBuildModule> Modules = new HashSet<UEBuildModule>(PrecompileOnlyModules);
-			foreach (UEBuildBinary Binary in Binaries.Concat(PrecompileOnlyBinaries))
-			{
-				foreach (UEBuildModule Module in Binary.GetAllDependencyModules(bIncludeDynamicallyLoaded: false, bForceCircular: false))
-				{
-					Modules.Add(Module);
-				}
-			}
-
 			// Get the platform we're building for
 			UEBuildPlatform BuildPlatform = UEBuildPlatform.GetBuildPlatform(Platform);
-
 			foreach (UEBuildModule Module in Modules)
 			{
 				// Skip artificial modules
@@ -1655,6 +1645,10 @@ namespace UnrealBuildTool
 				{
 					ReferencedModules.UnionWith(Module.GetDependencies(false, false));
 				}
+				foreach(UEBuildBinary Binary in Binaries)
+				{
+					ReferencedModules.UnionWith(Binary.Modules);
+				}
 				foreach(UEBuildBinary Binary in PrecompileOnlyBinaries)
 				{
 					ReferencedModules.UnionWith(Binary.Modules);
@@ -1674,7 +1668,7 @@ namespace UnrealBuildTool
 
 				// Add all the files which are required to use the precompiled modules
 				HashSet<FileReference> ExternalFiles = new HashSet<FileReference>();
-				GetExternalFileList(ExternalFiles);
+				GetExternalFileList(ReferencedModules, ExternalFiles);
 
 				// Convert them into relative to the target receipt
 				foreach(FileReference ExternalFile in ExternalFiles)
