Include .hpp files when enumerating headers that need to be included in the binary release.

diff --git a/Engine/Source/Programs/UnrealBuildTool/Configuration/UEBuildTarget.cs b/Engine/Source/Programs/UnrealBuildTool/Configuration/UEBuildTarget.cs
--- a/Engine/Source/Programs/UnrealBuildTool/Configuration/UEBuildTarget.cs
+++ b/Engine/Source/Programs/UnrealBuildTool/Configuration/UEBuildTarget.cs
@@ -1412,7 +1412,7 @@ namespace UnrealBuildTool
 						foreach (string IncludeFileName in Directory.EnumerateFiles(IncludePath, "*", SearchOption.AllDirectories))
 						{
 							string Extension = Path.GetExtension(IncludeFileName).ToLower();
-							if (Extension == ".h" || Extension == ".inl")
+							if (Extension == ".h" || Extension == ".inl" || Extension == ".hpp")
 							{
 								Files.Add(new FileReference(IncludeFileName));
 							}
