UnrealHeaderTool will now be compiled with Visual Studio 2015 if -2015 is specified on the command line to UBT.

#rnx

diff --git a/Engine/Source/Programs/UnrealBuildTool/Configuration/BuildConfiguration.cs b/Engine/Source/Programs/UnrealBuildTool/Configuration/BuildConfiguration.cs
--- a/Engine/Source/Programs/UnrealBuildTool/Configuration/BuildConfiguration.cs
+++ b/Engine/Source/Programs/UnrealBuildTool/Configuration/BuildConfiguration.cs
@@ -201,7 +201,14 @@ namespace UnrealBuildTool
 		/// If true, the Debug version of UnrealHeaderTool will be build and run instead of the Development version.
 		/// </summary>
 		[XmlConfigFile(Category = "UEBuildConfiguration")]
-		public static bool bForceDebugUnrealHeaderTool = false;
+		public bool bForceDebugUnrealHeaderTool = false;
+
+		/// <summary>
+		/// Compiler arguments to be forwarded to UnrealHeaderTool
+		/// </summary>
+		[CommandLine("-2015", Value = "-2015")]
+		[CommandLine("-2017", Value = "-2017")]
+		public string CompilerArgumentForUnrealHeaderTool = null;
 
 		/// <summary>
 		/// When true, the targets won't execute their link actions if there was nothing to compile
diff --git a/Engine/Source/Programs/UnrealBuildTool/System/ExternalExecution.cs b/Engine/Source/Programs/UnrealBuildTool/System/ExternalExecution.cs
--- a/Engine/Source/Programs/UnrealBuildTool/System/ExternalExecution.cs
+++ b/Engine/Source/Programs/UnrealBuildTool/System/ExternalExecution.cs
@@ -553,7 +553,7 @@ namespace UnrealBuildTool
 		/// Gets the path to the receipt for UHT
 		/// </summary>
 		/// <returns>Path to the UHT receipt</returns>
-		static FileReference GetHeaderToolReceiptFile()
+		static FileReference GetHeaderToolReceiptFile(BuildConfiguration BuildConfiguration)
 		{
 			UnrealTargetConfiguration Config = BuildConfiguration.bForceDebugUnrealHeaderTool ? UnrealTargetConfiguration.Debug : UnrealTargetConfiguration.Development;
 			return TargetReceipt.GetDefaultPath(UnrealBuildTool.EngineDirectory, "UnrealHeaderTool", BuildHostPlatform.Current.Platform, Config, "");
@@ -562,9 +562,9 @@ namespace UnrealBuildTool
 		/// <summary>
 		/// Gets UnrealHeaderTool.exe path. Does not care if UnrealheaderTool was build as a monolithic exe or not.
 		/// </summary>
-		static string GetHeaderToolPath()
+		static string GetHeaderToolPath(BuildConfiguration BuildConfiguration)
 		{
-			FileReference ReceiptFileName = GetHeaderToolReceiptFile();
+			FileReference ReceiptFileName = GetHeaderToolReceiptFile(BuildConfiguration);
 			TargetReceipt Receipt = TargetReceipt.Read(ReceiptFileName, UnrealBuildTool.EngineDirectory, null);
 
 			string HeaderToolPath = Receipt.BuildProducts[0].Path.FullName;
@@ -577,14 +577,14 @@ namespace UnrealBuildTool
 		/// <returns>
 		/// Latest timestamp of UHT binaries or DateTime.MaxValue if UnrealHeaderTool is out of date and needs to be rebuilt.
 		/// </returns>
-		static bool GetHeaderToolTimestamp(out DateTime Timestamp)
+		static bool GetHeaderToolTimestamp(BuildConfiguration BuildConfiguration, out DateTime Timestamp)
 		{
 			using (ScopedTimer TimestampTimer = new ScopedTimer("GetHeaderToolTimestamp"))
 			{
 				// Try to read the receipt for UHT.
 
 
-				FileReference ReceiptPath = GetHeaderToolReceiptFile();
+				FileReference ReceiptPath = GetHeaderToolReceiptFile(BuildConfiguration);
 				if (!FileReference.Exists(ReceiptPath))
 				{
 					Timestamp = DateTime.MaxValue;
@@ -1019,7 +1019,7 @@ namespace UnrealBuildTool
 
 				// check if UHT is out of date
 				DateTime HeaderToolTimestamp = DateTime.MaxValue;
-				bool bHaveHeaderTool = !bIsBuildingUHT && GetHeaderToolTimestamp(out HeaderToolTimestamp);
+				bool bHaveHeaderTool = !bIsBuildingUHT && GetHeaderToolTimestamp(BuildConfiguration, out HeaderToolTimestamp);
 
 				// ensure the headers are up to date
 				bool bUHTNeedsToRun = (BuildConfiguration.bForceHeaderGeneration || !bHaveHeaderTool || AreGeneratedCodeFilesOutOfDate(BuildConfiguration, UObjectModules, HeaderToolTimestamp, Target.bUsePrecompiled, HotReload, bIsGatheringBuild, bIsAssemblingBuild));
@@ -1107,7 +1107,13 @@ namespace UnrealBuildTool
 							{
 								UBTArguments.Append(" -PLUGIN=\"" + EnabledPlugin.Info.File + "\"");
 							}
-						}						
+						}
+
+						// Add any global override for the compiler
+						if(!String.IsNullOrEmpty(BuildConfiguration.CompilerArgumentForUnrealHeaderTool))
+						{
+							UBTArguments.AppendFormat(" {0}", BuildConfiguration.CompilerArgumentForUnrealHeaderTool);
+						}
 
 						// Output the log next to the current log
 						if(String.IsNullOrEmpty(BuildConfiguration.LogFileName))
@@ -1130,7 +1136,7 @@ namespace UnrealBuildTool
 					string ActualTargetName = String.IsNullOrEmpty(Target.GetTargetName()) ? "UE4" : Target.GetTargetName();
 					Log.TraceInformation("Parsing headers for {0}", ActualTargetName);
 
-					string HeaderToolPath = GetHeaderToolPath();
+					string HeaderToolPath = GetHeaderToolPath(BuildConfiguration);
 					if (!File.Exists(HeaderToolPath))
 					{
 						throw new BuildException("Unable to generate headers because UnrealHeaderTool binary was not found ({0}).", Path.GetFullPath(HeaderToolPath));
@@ -1161,7 +1167,7 @@ namespace UnrealBuildTool
 
 					Stopwatch s = new Stopwatch();
 					s.Start();
-					UHTResult = (ECompilationResult)RunExternalNativeExecutable(ExternalExecution.GetHeaderToolPath(), CmdLine);
+					UHTResult = (ECompilationResult)RunExternalNativeExecutable(ExternalExecution.GetHeaderToolPath(BuildConfiguration), CmdLine);
 					s.Stop();
 
 					if (UHTResult != ECompilationResult.Succeeded)
