Need to replace the correct module names with new suffix (script/response files) for hot reloading

diff --git a/Engine/Source/Programs/UnrealBuildTool/Platform/Linux/LinuxToolChain.cs b/Engine/Source/Programs/UnrealBuildTool/Platform/Linux/LinuxToolChain.cs
--- a/Engine/Source/Programs/UnrealBuildTool/Platform/Linux/LinuxToolChain.cs
+++ b/Engine/Source/Programs/UnrealBuildTool/Platform/Linux/LinuxToolChain.cs
@@ -1539,7 +1539,7 @@ namespace UnrealBuildTool
 				LinkCommandString = LinkCommandString.Replace("}", "}'");
 			}
 
-			string LinkScriptName = string.Format((bUseCmdExe ? "Link-{0}.bat" : "Link-{0}.sh"), OutputFile.Location.GetFileName());
+			string LinkScriptName = string.Format((bUseCmdExe ? "Link-{0}.link.bat" : "Link-{0}.link.sh"), OutputFile.Location.GetFileName());
 			string LinkScriptFullPath = Path.Combine(LinkEnvironment.LocalShadowDirectory.FullName, LinkScriptName);
 			Log.TraceVerbose("Creating link script: {0}", LinkScriptFullPath);
 			Directory.CreateDirectory(Path.GetDirectoryName(LinkScriptFullPath));
@@ -1577,6 +1577,8 @@ namespace UnrealBuildTool
 			};
 
 			LinkAction.CommandPath = ShellBinary;
+
+			// This must maintain the quotes around the LinkScriptFullPath
 			LinkAction.CommandArguments = ExecuteSwitch + " \"" + LinkScriptFullPath + "\"";
 
 			// prepare a linker script
diff --git a/Engine/Source/Programs/UnrealBuildTool/UnrealBuildTool.cs b/Engine/Source/Programs/UnrealBuildTool/UnrealBuildTool.cs
--- a/Engine/Source/Programs/UnrealBuildTool/UnrealBuildTool.cs
+++ b/Engine/Source/Programs/UnrealBuildTool/UnrealBuildTool.cs
@@ -2065,6 +2065,10 @@ namespace UnrealBuildTool
 			// names for all of the output files and import libraries
 			List<string> ResponseFilePaths = new List<string>();
 
+			// Same as Response files but for all of the link.sh files for link actions.
+			// Only used on BuildHostPlatform Linux
+			List<string> LinkScriptFilePaths = new List<string>();
+
 			// Keep a map of the original file names and their new file names, so we can fix up response files after
 			List<Tuple<string, string>> OriginalFileNameAndNewFileNameList_NoExtensions = new List<Tuple<string, string>>();
 
@@ -2105,6 +2109,70 @@ namespace UnrealBuildTool
 					ResponseFilePaths.Add(NewResponseFilePath);
 				}
 
+				// Find the *.link.sh file in the command line.  We'll need to make a copy of it with our new file name.
+				// Only currently used on Linux
+				if (UEBuildPlatform.IsPlatformInGroup(BuildHostPlatform.Current.Platform, UnrealPlatformGroup.Unix))
+				{
+					string LinkScriptFileExtension = ".link.sh";
+					int LinkScriptExtensionIndex = Action.CommandArguments.IndexOf(LinkScriptFileExtension, StringComparison.InvariantCultureIgnoreCase);
+					if (LinkScriptExtensionIndex != -1)
+					{
+						// We expect the script invocation to be quoted
+						int LinkScriptFilePathIndex = Action.CommandArguments.LastIndexOf("\"", LinkScriptExtensionIndex);
+						if (LinkScriptFilePathIndex == -1)
+						{
+							throw new BuildException("Couldn't find link script file path in action's command arguments when hot reloading. Is the path quoted?");
+						}
+
+						string OriginalLinkScriptFilePathWithoutExtension = Action.CommandArguments.Substring(LinkScriptFilePathIndex + 1, (LinkScriptExtensionIndex - LinkScriptFilePathIndex) - 1);
+						string OriginalLinkScriptFilePath = OriginalLinkScriptFilePathWithoutExtension + LinkScriptFileExtension;
+
+						string NewLinkScriptFilePath = OriginalLinkScriptFilePath.Replace(OriginalFileNameWithoutExtension, NewFileNameWithoutExtension);
+
+						// Copy the old response file to the new path
+						File.Copy(OriginalLinkScriptFilePath, NewLinkScriptFilePath, overwrite: true);
+
+						// Keep track of the new response file name.  We'll have to do some edits afterwards.
+						LinkScriptFilePaths.Add(NewLinkScriptFilePath);
+					}
+
+					// Update this action's list of prerequisite items too
+					for (int ItemIndex = 0; ItemIndex < Action.PrerequisiteItems.Count; ++ItemIndex)
+					{
+						FileItem OriginalPrerequisiteItem = Action.PrerequisiteItems[ItemIndex];
+						string NewPrerequisiteItemFilePath = OriginalPrerequisiteItem.AbsolutePath.Replace(OriginalFileNameWithoutExtension, NewFileNameWithoutExtension);
+
+						if (OriginalPrerequisiteItem.AbsolutePath != NewPrerequisiteItemFilePath)
+						{
+							// OK, the prerequisite item's file name changed so we'll update it to point to our new file
+							FileItem NewPrerequisiteItem = FileItem.GetItemByPath(NewPrerequisiteItemFilePath);
+							Action.PrerequisiteItems[ItemIndex] = NewPrerequisiteItem;
+
+							// Copy the other important settings from the original file item
+							NewPrerequisiteItem.bNeedsHotReloadNumbersDLLCleanUp = OriginalPrerequisiteItem.bNeedsHotReloadNumbersDLLCleanUp;
+							NewPrerequisiteItem.ProducingAction = OriginalPrerequisiteItem.ProducingAction;
+							NewPrerequisiteItem.bIsRemoteFile = OriginalPrerequisiteItem.bIsRemoteFile;
+
+							// Keep track of it so we can fix up dependencies in a second pass afterwards
+							AffectedOriginalFileItemAndNewFileItemMap.Add(OriginalPrerequisiteItem, NewPrerequisiteItem);
+
+							ResponseExtensionIndex = OriginalPrerequisiteItem.AbsolutePath.IndexOf(ResponseFileExtension, StringComparison.InvariantCultureIgnoreCase);
+							if (ResponseExtensionIndex != -1)
+							{
+								string OriginalResponseFilePathWithoutExtension = OriginalPrerequisiteItem.AbsolutePath.Substring(ResponseFilePathIndex + 1, (ResponseExtensionIndex - ResponseFilePathIndex) - 1);
+								string OriginalResponseFilePath = OriginalResponseFilePathWithoutExtension + ResponseFileExtension;
+
+								string NewResponseFilePath = OriginalResponseFilePath.Replace(OriginalFileNameWithoutExtension, NewFileNameWithoutExtension);
+
+								// Copy the old response file to the new path
+								File.Copy(OriginalResponseFilePath, NewResponseFilePath, overwrite: true);
+
+								// Keep track of the new response file name.  We'll have to do some edits afterwards.
+								ResponseFilePaths.Add(NewResponseFilePath);
+							}
+						}
+					}
+				}
 
 				// Go ahead and replace all occurrences of our file name in the command-line (ignoring extensions)
 				Action.CommandArguments = Action.CommandArguments.Replace(OriginalFileNameWithoutExtension, NewFileNameWithoutExtension);
@@ -2178,6 +2246,28 @@ namespace UnrealBuildTool
 					string FileContentsString = FileContents.ToString();
 					File.WriteAllText(ResponseFilePath, FileContentsString, new System.Text.UTF8Encoding(false));
 				}
+
+				if (UEBuildPlatform.IsPlatformInGroup(BuildHostPlatform.Current.Platform, UnrealPlatformGroup.Unix))
+				{
+					foreach (string LinkScriptFilePath in LinkScriptFilePaths)
+					{
+						// Load the file up
+						StringBuilder FileContents = new StringBuilder(Utils.ReadAllText(LinkScriptFilePath));
+
+						// Replace all of the old file names with new ones
+						foreach (Tuple<string, string> FileNameTuple in OriginalFileNameAndNewFileNameList_NoExtensions)
+						{
+							string OriginalFileNameWithoutExtension = FileNameTuple.Item1;
+							string NewFileNameWithoutExtension = FileNameTuple.Item2;
+
+							FileContents.Replace(OriginalFileNameWithoutExtension, NewFileNameWithoutExtension);
+						}
+
+						// Overwrite the original file
+						string FileContentsString = FileContents.ToString();
+						File.WriteAllText(LinkScriptFilePath, FileContentsString, new System.Text.UTF8Encoding(false));
+					}
+				}
 			}
 
 			foreach (UEBuildTarget Target in Targets)
