Fix removed index

diff --git a/Engine/Source/Programs/UnrealBuildTool/UnrealBuildTool.cs b/Engine/Source/Programs/UnrealBuildTool/UnrealBuildTool.cs
--- a/Engine/Source/Programs/UnrealBuildTool/UnrealBuildTool.cs
+++ b/Engine/Source/Programs/UnrealBuildTool/UnrealBuildTool.cs
@@ -2159,7 +2159,7 @@ namespace UnrealBuildTool
 							ResponseExtensionIndex = OriginalPrerequisiteItem.AbsolutePath.IndexOf(ResponseFileExtension, StringComparison.InvariantCultureIgnoreCase);
 							if (ResponseExtensionIndex != -1)
 							{
-								string OriginalResponseFilePathWithoutExtension = OriginalPrerequisiteItem.AbsolutePath.Substring(ResponseFilePathIndex + 1, (ResponseExtensionIndex - ResponseFilePathIndex) - 1);
+								string OriginalResponseFilePathWithoutExtension = OriginalPrerequisiteItem.AbsolutePath.Substring(0, ResponseExtensionIndex);
 								string OriginalResponseFilePath = OriginalResponseFilePathWithoutExtension + ResponseFileExtension;
 
 								string NewResponseFilePath = OriginalResponseFilePath.Replace(OriginalFileNameWithoutExtension, NewFileNameWithoutExtension);
