Integrated some VSCode fixes over from Dev-Core

diff --git a/Engine/Source/Programs/UnrealBuildTool/ProjectFiles/ProjectFileGenerator.cs b/Engine/Source/Programs/UnrealBuildTool/ProjectFiles/ProjectFileGenerator.cs
--- a/Engine/Source/Programs/UnrealBuildTool/ProjectFiles/ProjectFileGenerator.cs
+++ b/Engine/Source/Programs/UnrealBuildTool/ProjectFiles/ProjectFileGenerator.cs
@@ -553,7 +553,7 @@ namespace UnrealBuildTool
 					Format = ProjectFileFormat.XCode;
 					return true;
 				}
-				else if (PreferredAccessor == "visualstudiocodesourcecodeaccessor")
+				else if (PreferredAccessor == "visualstudiocode")
 				{
 					Format = ProjectFileFormat.VisualStudioCode;
 					return true;
diff --git a/Engine/Source/Programs/UnrealBuildTool/ProjectFiles/VisualStudioCode/VSCodeProjectFileGenerator.cs b/Engine/Source/Programs/UnrealBuildTool/ProjectFiles/VisualStudioCode/VSCodeProjectFileGenerator.cs
--- a/Engine/Source/Programs/UnrealBuildTool/ProjectFiles/VisualStudioCode/VSCodeProjectFileGenerator.cs
+++ b/Engine/Source/Programs/UnrealBuildTool/ProjectFiles/VisualStudioCode/VSCodeProjectFileGenerator.cs
@@ -689,7 +689,7 @@ namespace UnrealBuildTool
 								OutFile.AddUnnamedField(BuildProduct.Config.ToString());
 								if (bForeignProject)
 								{
-									OutFile.AddUnnamedField(MakeQuotedPathString(BuildProduct.UProjectFile, EPathType.Relative, null, EQuoteType.Single));
+									OutFile.AddUnnamedField(MakeUnquotedPathString(BuildProduct.UProjectFile, EPathType.Relative, null));
 								}
 								OutFile.AddUnnamedField("-waitmutex");
 
