UBT: Fix new plugin wizard not being able to compile engine plugins, due to engine plugins not being enabled by default.

diff --git a/Engine/Plugins/Editor/PluginBrowser/Source/PluginBrowser/Private/SNewPluginWizard.cpp b/Engine/Plugins/Editor/PluginBrowser/Source/PluginBrowser/Private/SNewPluginWizard.cpp
--- a/Engine/Plugins/Editor/PluginBrowser/Source/PluginBrowser/Private/SNewPluginWizard.cpp
+++ b/Engine/Plugins/Editor/PluginBrowser/Source/PluginBrowser/Private/SNewPluginWizard.cpp
@@ -743,7 +743,7 @@ FReply SNewPluginWizard::OnCreatePluginClicked()
 	if (bSucceeded && bHasModules)
 	{
 		FString ProjectFileName = FPaths::GetProjectFilePath();
-		FString Arguments = FString::Printf(TEXT("%s %s -TargetType=Editor -Module=%s -Project=\"%s\" -Progress -NoHotReloadFromIDE"), FModuleManager::Get().GetUBTConfiguration(), FPlatformMisc::GetUBTPlatform(), *PluginModuleName, *ProjectFileName);
+		FString Arguments = FString::Printf(TEXT("%s %s -TargetType=Editor -Module=%s -EnablePlugin=\"%s\" -Project=\"%s\" -Progress -NoHotReloadFromIDE"), FModuleManager::Get().GetUBTConfiguration(), FPlatformMisc::GetUBTPlatform(), *AutoPluginName, *PluginModuleName, *ProjectFileName);
 		if (!FDesktopPlatformModule::Get()->RunUnrealBuildTool(LOCTEXT("Compiling", "Compiling..."), FPaths::RootDir(), Arguments, GWarn))
 		{
 			PopErrorNotification(LOCTEXT("FailedToCompile", "Failed to compile source code."));
@@ -899,17 +899,8 @@ bool SNewPluginWizard::WritePluginDescriptor(const FString& PluginModuleName, co
 
 void SNewPluginWizard::PopErrorNotification(const FText& ErrorMessage)
 {
-	UE_LOG(LogPluginWizard, Log, TEXT("%s"), *ErrorMessage.ToString());
-
-	// Create and display a notification about the failure
-	FNotificationInfo Info(ErrorMessage);
-	Info.ExpireDuration = 2.0f;
-
-	TSharedPtr<SNotificationItem> Notification = FSlateNotificationManager::Get().AddNotification(Info);
-	if (Notification.IsValid())
-	{
-		Notification->SetCompletionState(SNotificationItem::CS_Fail);
-	}
+	FText Title = LOCTEXT("UnableToCreatePlugin", "Unable to create plugin");
+	FMessageDialog::Open(EAppMsgType::Ok, ErrorMessage, &Title);
 }
 
 void SNewPluginWizard::DeletePluginDirectory(const FString& InPath)
diff --git a/Engine/Source/Programs/UnrealBuildTool/Configuration/TargetRules.cs b/Engine/Source/Programs/UnrealBuildTool/Configuration/TargetRules.cs
--- a/Engine/Source/Programs/UnrealBuildTool/Configuration/TargetRules.cs
+++ b/Engine/Source/Programs/UnrealBuildTool/Configuration/TargetRules.cs
@@ -329,10 +329,26 @@ namespace UnrealBuildTool
 		public List<string> AdditionalPlugins = new List<string>();
 
 		/// <summary>
-		/// List of plugins to be excluded from this target. Note that the project file may still reference them, so they should be marked
+		/// Additional plugins that should be included for this target.
+		/// </summary>
+		[CommandLine("-EnablePlugin=", ListSeparator = '+')]
+		public List<string> EnablePlugins = new List<string>();
+
+		/// <summary>
+		/// List of plugins to be disabled for this target. Note that the project file may still reference them, so they should be marked
 		/// as optional to avoid failing to find them at runtime.
 		/// </summary>
-		public List<string> ExcludePlugins = new List<string>();
+		[CommandLine("-DisablePlugin=", ListSeparator = '+')]
+		public List<string> DisablePlugins = new List<string>();
+
+		/// <summary>
+		/// Accessor for
+		/// </summary>
+		[Obsolete("The ExcludePlugins setting has been renamed to DisablePlugins. Please update your code to avoid build failures in future versions of the engine.")]
+		public List<string> ExcludePlugins
+		{
+			get { return DisablePlugins; }
+		}
 
 		/// <summary>
 		/// Path to the set of pak signing keys to embed in the executable.
@@ -1498,9 +1514,14 @@ namespace UnrealBuildTool
 			get { return Inner.AdditionalPlugins; }
 		}
 
-		public IEnumerable<string> ExcludePlugins
+		public IEnumerable<string> EnablePlugins
+		{
+			get { return Inner.EnablePlugins; }
+		}
+
+		public IEnumerable<string> DisablePlugins
 		{
-			get { return Inner.ExcludePlugins; }
+			get { return Inner.DisablePlugins; }
 		}
 
 		public string PakSigningKeysFile
diff --git a/Engine/Source/Programs/UnrealBuildTool/Configuration/UEBuildTarget.cs b/Engine/Source/Programs/UnrealBuildTool/Configuration/UEBuildTarget.cs
--- a/Engine/Source/Programs/UnrealBuildTool/Configuration/UEBuildTarget.cs
+++ b/Engine/Source/Programs/UnrealBuildTool/Configuration/UEBuildTarget.cs
@@ -3359,20 +3359,33 @@ namespace UnrealBuildTool
 			// Map of plugin names to instances of that plugin
 			Dictionary<string, UEBuildPlugin> NameToInstance = new Dictionary<string, UEBuildPlugin>(StringComparer.InvariantCultureIgnoreCase);
 
+			// Configure plugins explicitly referenced via target settings
+			foreach(string PluginName in Rules.EnablePlugins)
+			{
+				if(ReferencedNames.Add(PluginName))
+				{
+					PluginReferenceDescriptor PluginReference = new PluginReferenceDescriptor(PluginName, null, true);
+					AddPlugin(PluginReference, "target settings", ExcludeFolders, NameToInstance, NameToInfo);
+				}
+			}
+
 			// Find a map of plugins which are explicitly referenced in the project file
 			if(ProjectDescriptor != null && ProjectDescriptor.Plugins != null)
 			{
 				string ProjectReferenceChain = ProjectFile.GetFileName();
 				foreach(PluginReferenceDescriptor PluginReference in ProjectDescriptor.Plugins)
 				{
-					// Make sure we don't have multiple references to the same plugin
-					if(!ReferencedNames.Add(PluginReference.Name))
+					if(!Rules.EnablePlugins.Contains(PluginReference.Name, StringComparer.InvariantCultureIgnoreCase))
 					{
-						Log.TraceWarning("Plugin '{0}' is listed multiple times in project file '{1}'.", PluginReference.Name, ProjectFile);
-					}
-					else
-					{
-						AddPlugin(PluginReference, ProjectReferenceChain, ExcludeFolders, NameToInstance, NameToInfo);
+						// Make sure we don't have multiple references to the same plugin
+						if(!ReferencedNames.Add(PluginReference.Name))
+						{
+							Log.TraceWarning("Plugin '{0}' is listed multiple times in project file '{1}'.", PluginReference.Name, ProjectFile);
+						}
+						else
+						{
+							AddPlugin(PluginReference, ProjectReferenceChain, ExcludeFolders, NameToInstance, NameToInfo);
+						}
 					}
 				}
 			}
@@ -3441,7 +3454,7 @@ namespace UnrealBuildTool
 			}
 
 			// If this plugin is listed to be excluded, do so here. The reference must be optional in this case.
-			if(Rules.ExcludePlugins.Contains(Reference.Name, StringComparer.InvariantCultureIgnoreCase))
+			if(Rules.DisablePlugins.Contains(Reference.Name, StringComparer.InvariantCultureIgnoreCase))
 			{
 				if(!Reference.bOptional)
 				{
