UE-60264 - iOS Launch-On rebuilds every time

diff --git a/Engine/Source/Programs/UnrealBuildTool/ToolChain/RemoteToolChain.cs b/Engine/Source/Programs/UnrealBuildTool/ToolChain/RemoteToolChain.cs
--- a/Engine/Source/Programs/UnrealBuildTool/ToolChain/RemoteToolChain.cs
+++ b/Engine/Source/Programs/UnrealBuildTool/ToolChain/RemoteToolChain.cs
@@ -959,11 +959,10 @@ namespace UnrealBuildTool
 					RsyncProcess.StartInfo.WorkingDirectory = ExeDir;
 				}
 
-				// --exclude='*'  ??? why???
+				// This will sync everything except .o files (because we want to keep old .o files on the target machine for iterative builds)
 				RsyncProcess.StartInfo.FileName = ResolvedRSyncExe;
                 RsyncProcess.StartInfo.Arguments = string.Format(
-					//"-vzrltgoDe \"{0}\" --rsync-path=\"mkdir -p {2} && rsync\" --chmod=ug=rwX,o=rxX --delete --files-from=\"{4}\" --include-from=\"{5}\" --include='*/' --exclude='*.o' --exclude='Timestamp' '{1}' \"{6}@{3}\":'{2}'",
-					"-vzrltgoDe \"{0}\" --rsync-path=\"mkdir -p {2} && rsync\" --chmod=ug=rwX,o=rxX --delete --files-from=\"{4}\" --include-from=\"{5}\" --include='*/' --exclude='Timestamp' '{1}' \"{6}@{3}\":'{2}'",
+					"-vzrltgoDe \"{0}\" --rsync-path=\"mkdir -p {2} && rsync\" --chmod=ug=rwX,o=rxX --delete --files-from=\"{4}\" --include-from=\"{5}\" --include='*/' --exclude='*.o' --exclude='Timestamp' '{1}' \"{6}@{3}\":'{2}'",
 					ResolvedRsyncAuthentication,
 					CygRootPath,
 					RemotePath,
@@ -979,6 +978,26 @@ namespace UnrealBuildTool
 				// run rsync
 				Utils.RunLocalProcess(RsyncProcess);
 
+				RsyncProcess = new Process();
+				if (ExeDir != "")
+				{
+					RsyncProcess.StartInfo.WorkingDirectory = ExeDir;
+				}
+				// this rsync, following the above rsync call will just sync newer .o files (this will usually do something only first time compiling a project, on a binary build)
+				RsyncProcess.StartInfo.FileName = ResolvedRSyncExe;
+				RsyncProcess.StartInfo.Arguments = string.Format(
+				   "-vzrltgoDe \"{0}\" --rsync-path=\"mkdir -p {2} && rsync\" --chmod=ug=rwX,o=rxX --files-from=\"{4}\" --include-from=\"{5}\" --include='*/' --exclude='Timestamp' '{1}' \"{6}@{3}\":'{2}'",
+				   ResolvedRsyncAuthentication,
+				   CygRootPath,
+				   RemotePath,
+				   RemoteServerName,
+				   ConvertPathToCygwin(RSyncPathsFile),
+				   ConvertPathToCygwin(IncludeFromFile),
+				   RSyncUsername);
+				Log.TraceInformation("Command: " + RsyncProcess.StartInfo.Arguments);
+				// run rsync
+				Utils.RunLocalProcess(RsyncProcess);
+
 				File.Delete(IncludeFromFile);
 				File.Delete(RSyncPathsFile);
 			}
