Linux: fix startup issues on some distros (UE-60398).

diff --git a/Engine/Source/Programs/UnrealBuildTool/Platform/Linux/LinuxToolChain.cs b/Engine/Source/Programs/UnrealBuildTool/Platform/Linux/LinuxToolChain.cs
--- a/Engine/Source/Programs/UnrealBuildTool/Platform/Linux/LinuxToolChain.cs
+++ b/Engine/Source/Programs/UnrealBuildTool/Platform/Linux/LinuxToolChain.cs
@@ -1537,6 +1537,7 @@ namespace UnrealBuildTool
 			{
 				LinkCommandString = LinkCommandString.Replace("{", "'{");
 				LinkCommandString = LinkCommandString.Replace("}", "}'");
+				LinkCommandString = LinkCommandString.Replace("$'{", "'${");	// fixing $'{ORIGIN}' to be '${ORIGIN}'
 			}
 
 			string LinkScriptName = string.Format((bUseCmdExe ? "Link-{0}.link.bat" : "Link-{0}.link.sh"), OutputFile.Location.GetFileName());
