Linux: fix suppression of -pie options in the compiler (UE-59928).

diff --git a/Engine/Source/Programs/UnrealBuildTool/Platform/Linux/LinuxToolChain.cs b/Engine/Source/Programs/UnrealBuildTool/Platform/Linux/LinuxToolChain.cs
--- a/Engine/Source/Programs/UnrealBuildTool/Platform/Linux/LinuxToolChain.cs
+++ b/Engine/Source/Programs/UnrealBuildTool/Platform/Linux/LinuxToolChain.cs
@@ -629,7 +629,7 @@ namespace UnrealBuildTool
 				Result += " -DPLATFORM_EXCEPTIONS_DISABLED=1";
 			}
 
-			if (bSuppressPIE)
+			if (bSuppressPIE && !CompileEnvironment.bIsBuildingDLL)
 			{
 				Result += " -fno-PIE";
 			}
@@ -799,9 +799,9 @@ namespace UnrealBuildTool
 
 			// This apparently can help LLDB speed up symbol lookups
 			Result += " -Wl,--build-id";
-			if (bSuppressPIE)
+			if (bSuppressPIE && !LinkEnvironment.bIsBuildingDLL)
 			{
-				Result += " -nopie";
+				Result += " -Wl,-nopie";
 			}
 
 			// whether we actually can do that is checked in CanUseLTO() earlier
