UE-60560 - iOS: Errors During Build Step with dSYM File Included Due to Space In Project Path

diff --git a/Engine/Source/Programs/UnrealBuildTool/Platform/IOS/IOSToolChain.cs b/Engine/Source/Programs/UnrealBuildTool/Platform/IOS/IOSToolChain.cs
--- a/Engine/Source/Programs/UnrealBuildTool/Platform/IOS/IOSToolChain.cs
+++ b/Engine/Source/Programs/UnrealBuildTool/Platform/IOS/IOSToolChain.cs
@@ -1178,7 +1178,7 @@ namespace UnrealBuildTool
             }
 
             GenDebugAction.CommandPath = "sh";
-            GenDebugAction.CommandArguments = string.Format("-c 'rm -rf \"{1}\"; dwarfdump --uuid {3} | cut -d\" \" -f2; chmod 777 ./DsymExporter; ./DsymExporter -UUID=$(dwarfdump --uuid {3} | cut -d\" \" -f2) \"{0}\" \"{2}\"'",
+            GenDebugAction.CommandArguments = string.Format("-c 'rm -rf \"{1}\"; dwarfdump --uuid \"{3}\" | cut -d\" \" -f2; chmod 777 ./DsymExporter; ./DsymExporter -UUID=$(dwarfdump --uuid \"{3}\" | cut -d\" \" -f2) \"{0}\" \"{2}\"'",
                     DWARFOutFile.AbsolutePath,
                     DestFile.AbsolutePath,
                     Path.GetDirectoryName(DestFile.AbsolutePath),
