CIS error fix

nuke "old" emscripten from build bots

#codeview ben.marsh
#rnx

diff --git a/Engine/Source/Programs/UnrealBuildTool/Platform/HTML5/UEBuildHTML5.cs b/Engine/Source/Programs/UnrealBuildTool/Platform/HTML5/UEBuildHTML5.cs
--- a/Engine/Source/Programs/UnrealBuildTool/Platform/HTML5/UEBuildHTML5.cs
+++ b/Engine/Source/Programs/UnrealBuildTool/Platform/HTML5/UEBuildHTML5.cs
@@ -261,44 +261,43 @@ namespace UnrealBuildTool
 	class HTML5PlatformSDK : UEBuildPlatformSDK
 	{
 		/// <summary>
-		/// Whether platform supports switching SDKs during runtime
+		/// platforms can choose if they prefer a correct the the AutoSDK install over the manual install.
 		/// </summary>
 		/// <returns>true if supports</returns>
-		protected override bool PlatformSupportsAutoSDKs()
-		{
-			return true;
-		}
-
-		// platforms can choose if they prefer a correct the the AutoSDK install over the manual install.
 		protected override bool PreferAutoSDK()
 		{
+			// HTML5 build tools are included in UE4 at: Engine/Extras/ThirdPartyNotUE/emsdk/emscripten/<version>
 			return false;
 		}
 
-		public override string GetSDKTargetPlatformName()
-		{
-			return "HTML5";
-		}
-
-		/// <summary>
-		/// Returns SDK string as required by the platform
-		/// </summary>
-		/// <returns>Valid SDK string</returns>
-		protected override string GetRequiredSDKString()
-		{
-			return HTML5SDKInfo.EmscriptenVersion();
-		}
-
-		protected override String GetRequiredScriptVersionString()
-		{
-			return "3.0";
-		}
-
 		/// <summary>
 		/// Whether the required external SDKs are installed for this platform
 		/// </summary>
 		protected override SDKStatus HasRequiredManualSDKInternal()
 		{
+			// ================================================================================
+			// HACK: CIS error fix -- UE-60768, UE-60769, UE-61032, etc...
+			// some build machines seems to be installed with emscripten v1.16.0
+			// that version is not even found in .../CarefullyRedist/HostWin64/HTML5/...
+			// nuke it if it exists
+			string TempPath = "C:\\\\Program\\ Files\\\\Emscripten\\\\emscripten\\\\1.16.0";
+			try
+			{
+				if (Directory.Exists(TempPath))
+				{
+					Directory.Delete(TempPath, true);
+				}
+			}
+			catch (Exception Ex)
+			{
+				Log.TraceErrorOnce("Unable to NUKE: " + TempPath + "\nPrehaps path is still in use?\n" + Ex.ToString());
+			}
+
+			// this nuke is done here before IsSDKInstalled() is called -- because that function
+			// will determine 'where' emscripten is located and will be used for HTML5 builds
+			// ================================================================================
+
+
 			if (!HTML5SDKInfo.IsSDKInstalled())
 			{
 				return SDKStatus.Invalid;
@@ -325,24 +324,19 @@ namespace UnrealBuildTool
 			// Make sure the SDK is installed
 			if ((ProjectFileGenerator.bGenerateProjectFiles == true) || (SDK.HasRequiredSDKsInstalled() == SDKStatus.Valid))
 			{
-				bool bRegisterBuildPlatform = true;
-
 				// make sure we have the HTML5 files; if not, then this user doesn't really have HTML5 access/files, no need to compile HTML5!
 				FileReference HTML5TargetPlatformFile = FileReference.Combine(UnrealBuildTool.EngineSourceDirectory, "Developer", "HTML5", "HTML5TargetPlatform", "HTML5TargetPlatform.Build.cs");
 				if (!FileReference.Exists(HTML5TargetPlatformFile))
 				{
-					bRegisterBuildPlatform = false;
 					Log.TraceWarning("Missing required components (.... HTML5TargetPlatformFile, others here...). Check source control filtering, or try resyncing.");
 				}
-
-				if (bRegisterBuildPlatform == true)
+				else
 				{
 					// Register this build platform for HTML5
 					Log.TraceVerbose("        Registering for {0}", UnrealTargetPlatform.HTML5.ToString());
 					UEBuildPlatform.RegisterBuildPlatform(new HTML5Platform(SDK));
 				}
 			}
-
 		}
 	}
 }
