CIS error fix

restore (removed) expected AutoSDK functions

#codeview ben.marsh
#rnx

diff --git a/Engine/Source/Programs/UnrealBuildTool/Platform/HTML5/UEBuildHTML5.cs b/Engine/Source/Programs/UnrealBuildTool/Platform/HTML5/UEBuildHTML5.cs
--- a/Engine/Source/Programs/UnrealBuildTool/Platform/HTML5/UEBuildHTML5.cs
+++ b/Engine/Source/Programs/UnrealBuildTool/Platform/HTML5/UEBuildHTML5.cs
@@ -261,15 +261,26 @@ namespace UnrealBuildTool
 	class HTML5PlatformSDK : UEBuildPlatformSDK
 	{
 		/// <summary>
-		/// platforms can choose if they prefer a correct the the AutoSDK install over the manual install.
+		/// Whether platform supports switching SDKs during runtime
 		/// </summary>
 		/// <returns>true if supports</returns>
+		protected override bool PlatformSupportsAutoSDKs()
+		{
+			return true;
+		}
+
+		// platforms can choose if they prefer a correct the the AutoSDK install over the manual install.
 		protected override bool PreferAutoSDK()
 		{
 			// HTML5 build tools are included in UE4 at: Engine/Extras/ThirdPartyNotUE/emsdk/emscripten/<version>
 			return false;
 		}
 
+		public override string GetSDKTargetPlatformName()
+		{
+			return "HTML5";
+		}
+
 		/// <summary>
 		/// Whether the required external SDKs are installed for this platform
 		/// </summary>
