CIS error fix

remove PlatformSupportsAutoSDK

#codeview ben.marsh
#rnx

diff --git a/Engine/Source/Programs/UnrealBuildTool/Platform/HTML5/UEBuildHTML5.cs b/Engine/Source/Programs/UnrealBuildTool/Platform/HTML5/UEBuildHTML5.cs
--- a/Engine/Source/Programs/UnrealBuildTool/Platform/HTML5/UEBuildHTML5.cs
+++ b/Engine/Source/Programs/UnrealBuildTool/Platform/HTML5/UEBuildHTML5.cs
@@ -260,15 +260,6 @@ namespace UnrealBuildTool
 
 	class HTML5PlatformSDK : UEBuildPlatformSDK
 	{
-		/// <summary>
-		/// Whether platform supports switching SDKs during runtime
-		/// </summary>
-		/// <returns>true if supports</returns>
-		protected override bool PlatformSupportsAutoSDKs()
-		{
-			return true;
-		}
-
 		// platforms can choose if they prefer a correct the the AutoSDK install over the manual install.
 		protected override bool PreferAutoSDK()
 		{
