CIS error fix

remove delete old emscripten HACK

#rnx

diff --git a/Engine/Source/Programs/UnrealBuildTool/Platform/HTML5/UEBuildHTML5.cs b/Engine/Source/Programs/UnrealBuildTool/Platform/HTML5/UEBuildHTML5.cs
--- a/Engine/Source/Programs/UnrealBuildTool/Platform/HTML5/UEBuildHTML5.cs
+++ b/Engine/Source/Programs/UnrealBuildTool/Platform/HTML5/UEBuildHTML5.cs
@@ -277,29 +277,6 @@ namespace UnrealBuildTool
 		/// </summary>
 		protected override SDKStatus HasRequiredManualSDKInternal()
 		{
-			// ================================================================================
-			// HACK: CIS error fix -- UE-60768, UE-60769, UE-61032, etc...
-			// some build machines seems to be installed with emscripten v1.16.0
-			// that version is not even found in .../CarefullyRedist/HostWin64/HTML5/...
-			// nuke it if it exists
-			string TempPath = "C:\\\\Program\\ Files\\\\Emscripten\\\\emscripten\\\\1.16.0";
-			try
-			{
-				if (Directory.Exists(TempPath))
-				{
-					Directory.Delete(TempPath, true);
-				}
-			}
-			catch (Exception Ex)
-			{
-				Log.TraceErrorOnce("Unable to NUKE: " + TempPath + "\nPrehaps path is still in use?\n" + Ex.ToString());
-			}
-
-			// this nuke is done here before IsSDKInstalled() is called -- because that function
-			// will determine 'where' emscripten is located and will be used for HTML5 builds
-			// ================================================================================
-
-
 			if (!HTML5SDKInfo.IsSDKInstalled())
 			{
 				return SDKStatus.Invalid;
