Prevent overwriting receipts in installed builds when building a target that's already build. Should fix issues with debug files being added back into the list of build products.

diff --git a/Engine/Source/Programs/UnrealBuildTool/Configuration/UEBuildTarget.cs b/Engine/Source/Programs/UnrealBuildTool/Configuration/UEBuildTarget.cs
--- a/Engine/Source/Programs/UnrealBuildTool/Configuration/UEBuildTarget.cs
+++ b/Engine/Source/Programs/UnrealBuildTool/Configuration/UEBuildTarget.cs
@@ -1900,25 +1900,35 @@ namespace UnrealBuildTool
 				UEBuildPlatform BuildPlatform = UEBuildPlatform.GetBuildPlatform(Platform);
 				if (OnlyModules == null || OnlyModules.Count == 0)
 				{
-					DirectoryReference.CreateDirectory(ReceiptFileName.Directory);
-					Receipt.Write(ReceiptFileName, UnrealBuildTool.EngineDirectory, ProjectDirectory);
+					if(!IsFileInstalled(ReceiptFileName))
+					{
+						DirectoryReference.CreateDirectory(ReceiptFileName.Directory);
+						Receipt.Write(ReceiptFileName, UnrealBuildTool.EngineDirectory, ProjectDirectory);
+					}
 				}
 				if (ForceReceiptFileName != null)
 				{
-					Directory.CreateDirectory(Path.GetDirectoryName(ForceReceiptFileName));
-					Receipt.Write(new FileReference(ForceReceiptFileName), UnrealBuildTool.EngineDirectory, ProjectDirectory);
+					FileReference ForceReceiptFile = new FileReference(ForceReceiptFileName);
+					if(!IsFileInstalled(ForceReceiptFile))
+					{
+						DirectoryReference.CreateDirectory(ForceReceiptFile.Directory);
+						Receipt.Write(ForceReceiptFile, UnrealBuildTool.EngineDirectory, ProjectDirectory);
+					}
 				}
 				if(VersionFile != null)
 				{
-					DirectoryReference.CreateDirectory(VersionFile.Directory);
+					if(!IsFileInstalled(VersionFile))
+					{
+						DirectoryReference.CreateDirectory(VersionFile.Directory);
 
-					StringWriter Writer = new StringWriter();
-					Receipt.Version.Write(Writer);
+						StringWriter Writer = new StringWriter();
+						Receipt.Version.Write(Writer);
 
-					string Text = Writer.ToString();
-					if(!FileReference.Exists(VersionFile) || File.ReadAllText(VersionFile.FullName) != Text)
-					{
-						File.WriteAllText(VersionFile.FullName, Text);
+						string Text = Writer.ToString();
+						if(!FileReference.Exists(VersionFile) || File.ReadAllText(VersionFile.FullName) != Text)
+						{
+							File.WriteAllText(VersionFile.FullName, Text);
+						}
 					}
 				}
 			}
@@ -1926,22 +1936,43 @@ namespace UnrealBuildTool
 			{
 				foreach (KeyValuePair<FileReference, ModuleManifest> FileNameToVersionManifest in FileReferenceToModuleManifestPairs)
 				{
-					// Write the manifest out to a string buffer, then only write it to disk if it's changed.
-					string OutputText;
-					using (StringWriter Writer = new StringWriter())
-					{
-						FileNameToVersionManifest.Value.Write(Writer);
-						OutputText = Writer.ToString();
-					}
-					if(!FileReference.Exists(FileNameToVersionManifest.Key) || File.ReadAllText(FileNameToVersionManifest.Key.FullName) != OutputText)
+					if(!IsFileInstalled(FileNameToVersionManifest.Key))
 					{
-						Directory.CreateDirectory(Path.GetDirectoryName(FileNameToVersionManifest.Key.FullName));
-						FileNameToVersionManifest.Value.Write(FileNameToVersionManifest.Key.FullName);
+						// Write the manifest out to a string buffer, then only write it to disk if it's changed.
+						string OutputText;
+						using (StringWriter Writer = new StringWriter())
+						{
+							FileNameToVersionManifest.Value.Write(Writer);
+							OutputText = Writer.ToString();
+						}
+						if(!FileReference.Exists(FileNameToVersionManifest.Key) || File.ReadAllText(FileNameToVersionManifest.Key.FullName) != OutputText)
+						{
+							Directory.CreateDirectory(Path.GetDirectoryName(FileNameToVersionManifest.Key.FullName));
+							FileNameToVersionManifest.Value.Write(FileNameToVersionManifest.Key.FullName);
+						}
 					}
 				}
 			}
 		}
 
+		/// <summary>
+		/// Checks whether the given file is under an installed directory, and should not be overridden
+		/// </summary>
+		/// <param name="File">File to test</param>
+		/// <returns>True if the file is part of the installed distribution, false otherwise</returns>
+		bool IsFileInstalled(FileReference File)
+		{
+			if(UnrealBuildTool.IsEngineInstalled() && File.IsUnderDirectory(UnrealBuildTool.EngineDirectory))
+			{
+				return true;
+			}
+			if(UnrealBuildTool.IsProjectInstalled() && ProjectFile != null && File.IsUnderDirectory(ProjectFile.Directory))
+			{
+				return true;
+			}
+			return false;
+		}
+
 		/// <summary>
 		/// Gathers dependency modules for given binaries list.
 		/// </summary>
diff --git a/Engine/Source/Programs/UnrealBuildTool/UnrealBuildTool.cs b/Engine/Source/Programs/UnrealBuildTool/UnrealBuildTool.cs
--- a/Engine/Source/Programs/UnrealBuildTool/UnrealBuildTool.cs
+++ b/Engine/Source/Programs/UnrealBuildTool/UnrealBuildTool.cs
@@ -1551,7 +1551,7 @@ namespace UnrealBuildTool
 						UBTMakefile.SourceFileWorkingSet = Unity.SourceFileWorkingSet;
 						UBTMakefile.CandidateSourceFilesForWorkingSet = Unity.CandidateSourceFilesForWorkingSet;
 
-						if (BuildConfiguration.bUseUBTMakefiles && !UBTMakefile.PrerequisiteActions.Any(x => x.ActionHandler != null))
+						if (BuildConfiguration.bUseUBTMakefiles && UBTMakefile.AllActions.Count > 0 && !UBTMakefile.PrerequisiteActions.Any(x => x.ActionHandler != null))
 						{
 							// We've been told to prepare to build, so let's go ahead and save out our action graph so that we can use in a later invocation 
 							// to assemble the build.  Even if we are configured to assemble the build in this same invocation, we want to save out the
