Fix packaging a plugin for the marketplace from installed builds.

diff --git a/Engine/Source/Programs/UnrealBuildTool/Configuration/UEBuildTarget.cs b/Engine/Source/Programs/UnrealBuildTool/Configuration/UEBuildTarget.cs
--- a/Engine/Source/Programs/UnrealBuildTool/Configuration/UEBuildTarget.cs
+++ b/Engine/Source/Programs/UnrealBuildTool/Configuration/UEBuildTarget.cs
@@ -4103,7 +4103,14 @@ namespace UnrealBuildTool
 
 				// Get the plugin for this module
 				PluginInfo Plugin;
-				RulesAssembly.TryGetPluginForModule(ModuleFileName, out Plugin);
+				if(RulesAssembly.TryGetPluginForModule(ModuleFileName, out Plugin))
+				{
+					// Clear the bUsePrecompiled flag if we're compiling a foreign plugin; since it's treated like an engine module, it will default to true in an installed build.
+					if(Plugin.File == ForeignPlugin)
+					{
+						RulesObject.bUsePrecompiled = false;
+					}
+				}
 
 				// Get the module descriptor for this module if it's a plugin
 				ModuleDescriptor PluginModuleDesc = null;
