PR #4781: Log out the amount of files that are to be compiled in the VS output . (Contributed by [REDACTED_PERSON])

diff --git a/Engine/Source/Programs/AutomationTool/AutomationUtils/ParallelExecutor.cs b/Engine/Source/Programs/AutomationTool/AutomationUtils/ParallelExecutor.cs
--- a/Engine/Source/Programs/AutomationTool/AutomationUtils/ParallelExecutor.cs
+++ b/Engine/Source/Programs/AutomationTool/AutomationUtils/ParallelExecutor.cs
@@ -85,7 +85,7 @@ namespace AutomationTool
 		{
 			List<BuildAction> Actions = ReadActions(ActionsFileName);
 
-			CommandUtils.Log("Building with {0} {1}...", MaxProcesses, (MaxProcesses == 1)? "process" : "processes");
+			CommandUtils.Log("Building {0} {1} with {2} {3}...", Actions.Count, (Actions.Count == 1) ? "action" : "actions", MaxProcesses, (MaxProcesses == 1)? "process" : "processes");
 
 			using (UnrealBuildTool.ScopedLogIndent Indent = new UnrealBuildTool.ScopedLogIndent("  "))
 			{
diff --git a/Engine/Source/Programs/UnrealBuildTool/System/ParallelExecutor.cs b/Engine/Source/Programs/UnrealBuildTool/System/ParallelExecutor.cs
--- a/Engine/Source/Programs/UnrealBuildTool/System/ParallelExecutor.cs
+++ b/Engine/Source/Programs/UnrealBuildTool/System/ParallelExecutor.cs
@@ -83,7 +83,7 @@ namespace UnrealBuildTool
 		{
 			// Figure out how many processors to use
 			int MaxProcesses = Math.Min((int)(Environment.ProcessorCount * ProcessorCountMultiplier), MaxProcessorCount);
-			Log.TraceInformation("Building with {0} {1}...", MaxProcesses, (MaxProcesses == 1)? "process" : "processes");
+			Log.TraceInformation("Building {0} {1} with {2} {3}...", InputActions.Count, (InputActions.Count == 1) ? "action" : "actions", MaxProcesses, (MaxProcesses == 1)? "process" : "processes");
 
 			// Create actions with all our internal metadata
 			List<BuildAction> Actions = new List<BuildAction>();
