UE-61188 - iOS Remote Build failure using blueprint project with icons set in a Binary Build

diff --git a/Engine/Source/Programs/UnrealBuildTool/Platform/IOS/IOSToolChain.cs b/Engine/Source/Programs/UnrealBuildTool/Platform/IOS/IOSToolChain.cs
--- a/Engine/Source/Programs/UnrealBuildTool/Platform/IOS/IOSToolChain.cs
+++ b/Engine/Source/Programs/UnrealBuildTool/Platform/IOS/IOSToolChain.cs
@@ -1310,6 +1310,11 @@ namespace UnrealBuildTool
 			if (BuildHostPlatform.Current.Platform != UnrealTargetPlatform.Mac)
 			{
 				BuiltBinaries = new List<FileReference>();
+				DirectoryReference BinaryIOS = DirectoryReference.MakeFromNormalizedFullPath(Path.GetFullPath(Path.Combine(BranchDirectory, "Engine/")) + "Binaries/IOS/");
+				if (DirectoryReference.Exists(BinaryIOS))
+				{
+					QueueDirectoryForBatchUpload(BinaryIOS);
+				}
 			}
 
 			RemoteToolChain.PreBuildSync();
