Linux: include more debug info in the build

(merged from Dev-Console)

diff --git a/Engine/Source/Programs/UnrealBuildTool/Platform/Linux/LinuxToolChain.cs b/Engine/Source/Programs/UnrealBuildTool/Platform/Linux/LinuxToolChain.cs
--- a/Engine/Source/Programs/UnrealBuildTool/Platform/Linux/LinuxToolChain.cs
+++ b/Engine/Source/Programs/UnrealBuildTool/Platform/Linux/LinuxToolChain.cs
@@ -593,6 +593,15 @@ namespace UnrealBuildTool
 			{
 				// libdwarf (from elftoolchain 0.6.1) doesn't support DWARF4. If we need to go back to depending on elftoolchain revert this back to dwarf-3
 				Result += " -gdwarf-4";
+				
+				// Include debug info
+				Result += " -g";
+				
+				// Include additional debug info to help LLDB
+				Result += " -glldb";
+				
+				// Makes debugging .so libraries better
+				Result += " -fstandalone-debug";
 			}
 
 			// optimization level
