Fix for custom icons not being created for IOS

diff --git a/Engine/Source/Programs/AutomationTool/AutomationUtils/ProjectUtils.cs b/Engine/Source/Programs/AutomationTool/AutomationUtils/ProjectUtils.cs
--- a/Engine/Source/Programs/AutomationTool/AutomationUtils/ProjectUtils.cs
+++ b/Engine/Source/Programs/AutomationTool/AutomationUtils/ProjectUtils.cs
@@ -217,7 +217,7 @@ namespace AutomationTool
 				bool RetVal = false;
 				foreach (UnrealTargetPlatform TargetPlatformType in TargetPlatforms)
 				{
-					if(!Automation.IsEngineInstalled() && !PlatformExports.HasDefaultBuildConfig(RawProjectPath, TargetPlatformType))
+					if((!Automation.IsEngineInstalled() && !PlatformExports.HasDefaultBuildConfig(RawProjectPath, TargetPlatformType)) || PlatformExports.RequiresBuild(RawProjectPath, TargetPlatformType))
 					{
 						RetVal = true;
 						break;
diff --git a/Engine/Source/Programs/UnrealBuildTool/Configuration/UEBuildPlatform.cs b/Engine/Source/Programs/UnrealBuildTool/Configuration/UEBuildPlatform.cs
--- a/Engine/Source/Programs/UnrealBuildTool/Configuration/UEBuildPlatform.cs
+++ b/Engine/Source/Programs/UnrealBuildTool/Configuration/UEBuildPlatform.cs
@@ -804,6 +804,15 @@ namespace UnrealBuildTool
 				BoolKeys, null, null);
 		}
 
+		/// <summary>
+		/// Check for whether we require a build for platform reasons
+		/// return true if the project requires a build
+		/// </summary>
+		public virtual bool RequiresBuild(UnrealTargetPlatform Platform, DirectoryReference ProjectDirectoryName)
+		{
+			return false;
+		}
+
 		/// <summary>
 		/// Get a list of extra modules the platform requires.
 		/// This is to allow undisclosed platforms to add modules they need without exposing information about the platform.
diff --git a/Engine/Source/Programs/UnrealBuildTool/Platform/IOS/UEBuildIOS.cs b/Engine/Source/Programs/UnrealBuildTool/Platform/IOS/UEBuildIOS.cs
--- a/Engine/Source/Programs/UnrealBuildTool/Platform/IOS/UEBuildIOS.cs
+++ b/Engine/Source/Programs/UnrealBuildTool/Platform/IOS/UEBuildIOS.cs
@@ -688,7 +688,20 @@ namespace UnrealBuildTool
 		{
 			IOSToolChain.PostCodeGeneration(Manifest);
 		}
-		
+
+		public bool HasCustomIcons(DirectoryReference ProjectDirectoryName)
+		{
+			string IconDir = Path.Combine(ProjectDirectoryName.FullName, "Build", "IOS", "Resources", "Graphics");
+			foreach (string f in Directory.EnumerateFiles(IconDir))
+			{
+				if (f.Contains("Icon") && Path.GetExtension(f).Contains(".png"))
+				{
+					return true;
+				}
+			}
+			return false;
+		}
+
 		/// <summary>
 		/// Check for the default configuration
 		/// return true if the project uses the default build config
@@ -707,6 +720,12 @@ namespace UnrealBuildTool
 				"AdditionalShippingLinkerFlags"
 			};
 
+			// check for custom icons
+			if (HasCustomIcons(ProjectDirectoryName))
+			{
+				return false;
+			}
+
 			// look up iOS specific settings
 			if (!DoProjectSettingsMatchDefault(Platform, ProjectDirectoryName, "/Script/IOSRuntimeSettings.IOSRuntimeSettings",
 					BoolKeys, null, StringKeys))
@@ -718,6 +737,16 @@ namespace UnrealBuildTool
 			return base.HasDefaultBuildConfig(Platform, ProjectDirectoryName);
 		}
 
+		/// <summary>
+		/// Check for the build requirement due to platform requirements
+		/// return true if the project requires a build
+		/// </summary>
+		public override bool RequiresBuild(UnrealTargetPlatform Platform, DirectoryReference ProjectDirectoryName)
+		{
+			// check for custom icons
+			return HasCustomIcons(ProjectDirectoryName);
+		}
+
 		/// <summary>
 		/// Whether the editor should be built for this platform or not
 		/// </summary>
diff --git a/Engine/Source/Programs/UnrealBuildTool/System/PlatformExports.cs b/Engine/Source/Programs/UnrealBuildTool/System/PlatformExports.cs
--- a/Engine/Source/Programs/UnrealBuildTool/System/PlatformExports.cs
+++ b/Engine/Source/Programs/UnrealBuildTool/System/PlatformExports.cs
@@ -56,6 +56,18 @@ namespace UnrealBuildTool
 			return (BuildPlat == null)? true : BuildPlat.HasDefaultBuildConfig(Platform, ProjectFile.Directory);
 		}
 
+		/// <summary>
+		/// Checks whether the given project requires a build because of platform needs
+		/// </summary>
+		/// <param name="ProjectFile">The project file</param>
+		/// <param name="Platform">Platform to check settings for</param>
+		/// <returns>True if the project requires a build for the platform</returns>
+		public static bool RequiresBuild(FileReference ProjectFile, UnrealTargetPlatform Platform)
+		{
+			UEBuildPlatform BuildPlat = UEBuildPlatform.GetBuildPlatform(Platform, true);
+			return (BuildPlat == null) ? false : BuildPlat.RequiresBuild(Platform, ProjectFile.Directory);
+		}
+
 		/// <summary>
 		/// Returns an array of all platform folder names
 		/// </summary>
