Fix exception building IOS projects without custom icons.

diff --git a/Engine/Source/Programs/UnrealBuildTool/Platform/IOS/UEBuildIOS.cs b/Engine/Source/Programs/UnrealBuildTool/Platform/IOS/UEBuildIOS.cs
--- a/Engine/Source/Programs/UnrealBuildTool/Platform/IOS/UEBuildIOS.cs
+++ b/Engine/Source/Programs/UnrealBuildTool/Platform/IOS/UEBuildIOS.cs
@@ -692,11 +692,14 @@ namespace UnrealBuildTool
 		public bool HasCustomIcons(DirectoryReference ProjectDirectoryName)
 		{
 			string IconDir = Path.Combine(ProjectDirectoryName.FullName, "Build", "IOS", "Resources", "Graphics");
-			foreach (string f in Directory.EnumerateFiles(IconDir))
+			if(Directory.Exists(IconDir))
 			{
-				if (f.Contains("Icon") && Path.GetExtension(f).Contains(".png"))
+				foreach (string f in Directory.EnumerateFiles(IconDir))
 				{
-					return true;
+					if (f.Contains("Icon") && Path.GetExtension(f).Contains(".png"))
+					{
+						return true;
+					}
 				}
 			}
 			return false;
