Merging from Dev-Rendering: Bump clang version to fix Android compiles.
#rnx

diff --git a/Engine/Source/Programs/UnrealBuildTool/Platform/Android/UEBuildAndroid.cs b/Engine/Source/Programs/UnrealBuildTool/Platform/Android/UEBuildAndroid.cs
--- a/Engine/Source/Programs/UnrealBuildTool/Platform/Android/UEBuildAndroid.cs
+++ b/Engine/Source/Programs/UnrealBuildTool/Platform/Android/UEBuildAndroid.cs
@@ -538,12 +538,12 @@ namespace UnrealBuildTool
 
 		protected override string GetRequiredSDKString()
 		{
-			return "-21";
+			return "-22";
 		}
 
 		protected override String GetRequiredScriptVersionString()
 		{
-			return "3.2";
+			return "3.3";
 		}
 
 		// prefer auto sdk on android as correct 'manual' sdk detection isn't great at the moment.
