Prevent post-build steps from running when linking is disabled. Fixes non-unity ShaderCompileWorker in CIS.

diff --git a/Engine/Source/Programs/UnrealBuildTool/UnrealBuildTool.cs b/Engine/Source/Programs/UnrealBuildTool/UnrealBuildTool.cs
--- a/Engine/Source/Programs/UnrealBuildTool/UnrealBuildTool.cs
+++ b/Engine/Source/Programs/UnrealBuildTool/UnrealBuildTool.cs
@@ -1694,11 +1694,11 @@ namespace UnrealBuildTool
 						}
 
 						// Execute all the post-build steps
-							if (BuildResult.Succeeded() && !ProjectFileGenerator.bGenerateProjectFiles && !BuildConfiguration.bXGEExport)
+						if (BuildResult.Succeeded() && !ProjectFileGenerator.bGenerateProjectFiles && !BuildConfiguration.bXGEExport)
 						{
-								foreach (UEBuildTarget Target in Targets)
+							foreach (UEBuildTarget Target in Targets)
 							{
-									if (!Target.ExecuteCustomPostBuildSteps())
+								if (!Target.Rules.bDisableLinking && !Target.ExecuteCustomPostBuildSteps())
 								{
 									BuildResult = ECompilationResult.OtherCompilationError;
 									break;
