Alternate fix for skipping post-build steps when linking is disabled. Rules instance doesn't exist when building using a makefile.

diff --git a/Engine/Source/Programs/UnrealBuildTool/Configuration/UEBuildTarget.cs b/Engine/Source/Programs/UnrealBuildTool/Configuration/UEBuildTarget.cs
--- a/Engine/Source/Programs/UnrealBuildTool/Configuration/UEBuildTarget.cs
+++ b/Engine/Source/Programs/UnrealBuildTool/Configuration/UEBuildTarget.cs
@@ -2740,17 +2740,20 @@ namespace UnrealBuildTool
 
 			// Find all the post-build steps
 			List<Tuple<string[], UEBuildPlugin>> PostBuildCommandBatches = new List<Tuple<string[], UEBuildPlugin>>();
-			if(ProjectDescriptor != null && ProjectDescriptor.PostBuildSteps != null)
+			if(!Rules.bDisableLinking)
 			{
-				AddCustomBuildSteps(ProjectDescriptor.PostBuildSteps, null, PostBuildCommandBatches);
-			}
-			if(Rules.PostBuildSteps.Count > 0)
-			{
-				PostBuildCommandBatches.Add(new Tuple<string[], UEBuildPlugin>(Rules.PostBuildSteps.ToArray(), null));
-			}
-			foreach(UEBuildPlugin BuildPlugin in BuildPlugins.Where(x => x.Descriptor.PostBuildSteps != null))
-			{
-				AddCustomBuildSteps(BuildPlugin.Descriptor.PostBuildSteps, BuildPlugin, PostBuildCommandBatches);
+				if(ProjectDescriptor != null && ProjectDescriptor.PostBuildSteps != null)
+				{
+					AddCustomBuildSteps(ProjectDescriptor.PostBuildSteps, null, PostBuildCommandBatches);
+				}
+				if(Rules.PostBuildSteps.Count > 0)
+				{
+					PostBuildCommandBatches.Add(new Tuple<string[], UEBuildPlugin>(Rules.PostBuildSteps.ToArray(), null));
+				}
+				foreach(UEBuildPlugin BuildPlugin in BuildPlugins.Where(x => x.Descriptor.PostBuildSteps != null))
+				{
+					AddCustomBuildSteps(BuildPlugin.Descriptor.PostBuildSteps, BuildPlugin, PostBuildCommandBatches);
+				}
 			}
 			PostBuildStepScripts = WriteCustomBuildStepScripts(BuildHostPlatform.Current.Platform, ScriptDirectory, "PostBuild", PostBuildCommandBatches);
 		}
diff --git a/Engine/Source/Programs/UnrealBuildTool/UnrealBuildTool.cs b/Engine/Source/Programs/UnrealBuildTool/UnrealBuildTool.cs
--- a/Engine/Source/Programs/UnrealBuildTool/UnrealBuildTool.cs
+++ b/Engine/Source/Programs/UnrealBuildTool/UnrealBuildTool.cs
@@ -1698,7 +1698,7 @@ namespace UnrealBuildTool
 						{
 							foreach (UEBuildTarget Target in Targets)
 							{
-								if (!Target.Rules.bDisableLinking && !Target.ExecuteCustomPostBuildSteps())
+								if (!Target.ExecuteCustomPostBuildSteps())
 								{
 									BuildResult = ECompilationResult.OtherCompilationError;
 									break;
