Fix for Mac stripping debug symbols not launching properly

diff --git a/Engine/Source/Programs/UnrealBuildTool/Platform/Android/AndroidToolChain.cs b/Engine/Source/Programs/UnrealBuildTool/Platform/Android/AndroidToolChain.cs
--- a/Engine/Source/Programs/UnrealBuildTool/Platform/Android/AndroidToolChain.cs
+++ b/Engine/Source/Programs/UnrealBuildTool/Platform/Android/AndroidToolChain.cs
@@ -1843,7 +1843,7 @@ namespace UnrealBuildTool
 			}
 
 			ProcessStartInfo StartInfo = new ProcessStartInfo();
-			StartInfo.FileName = GetStripPath(SourceFile);
+			StartInfo.FileName = GetStripPath(SourceFile).Trim('"');
 			StartInfo.Arguments = " --strip-debug \"" + TargetFile.FullName + "\"";
 			StartInfo.UseShellExecute = false;
 			StartInfo.CreateNoWindow = true;
diff --git a/Engine/Source/Programs/UnrealBuildTool/Platform/Android/UEDeployAndroid.cs b/Engine/Source/Programs/UnrealBuildTool/Platform/Android/UEDeployAndroid.cs
--- a/Engine/Source/Programs/UnrealBuildTool/Platform/Android/UEDeployAndroid.cs
+++ b/Engine/Source/Programs/UnrealBuildTool/Platform/Android/UEDeployAndroid.cs
@@ -1207,7 +1207,7 @@ namespace UnrealBuildTool
 			}
 
 			ProcessStartInfo StartInfo = new ProcessStartInfo();
-			StartInfo.FileName = AndroidToolChain.GetStripExecutablePath(UE4Arch);
+			StartInfo.FileName = AndroidToolChain.GetStripExecutablePath(UE4Arch).Trim('"');
 			if (bStripAll)
 			{
 				StartInfo.Arguments = "--strip-unneeded \"" + TargetFileName + "\"";
