UBT: Don't add the UHT timestamp file to the manifest, since it's updated whether the output files have been modified or not. Fixes CIS errors due to modified build products.

diff --git a/Engine/Source/Programs/UnrealBuildTool/Configuration/UEBuildBinary.cs b/Engine/Source/Programs/UnrealBuildTool/Configuration/UEBuildBinary.cs
--- a/Engine/Source/Programs/UnrealBuildTool/Configuration/UEBuildBinary.cs
+++ b/Engine/Source/Programs/UnrealBuildTool/Configuration/UEBuildBinary.cs
@@ -494,7 +494,11 @@ namespace UnrealBuildTool
 						{
 							foreach(FileReference GeneratedCodeFile in DirectoryReference.EnumerateFiles(Module.GeneratedCodeDirectory))
 							{
-								BuildProducts.Add(GeneratedCodeFile, BuildProductType.BuildResource);
+								// Exclude timestamp files, since they're always updated and cause collisions between builds
+								if(!GeneratedCodeFile.GetFileName().Equals("Timestamp", StringComparison.OrdinalIgnoreCase))
+								{
+									BuildProducts.Add(GeneratedCodeFile, BuildProductType.BuildResource);
+								}
 							}
 						}
 						if(Target.LinkType == TargetLinkType.Monolithic)
