UBT: Force UE4Game module to be precompiled as part of UE4Editor. Fixes CIS errors when explicitly compiled by ProtoStar.

diff --git a/Engine/Source/Programs/UnrealBuildTool/System/RulesAssembly.cs b/Engine/Source/Programs/UnrealBuildTool/System/RulesAssembly.cs
--- a/Engine/Source/Programs/UnrealBuildTool/System/RulesAssembly.cs
+++ b/Engine/Source/Programs/UnrealBuildTool/System/RulesAssembly.cs
@@ -374,7 +374,7 @@ namespace UnrealBuildTool
 				ModuleFileToPluginInfo.TryGetValue(RulesObject.File, out RulesObject.Plugin);
 				RulesObject.bTreatAsEngineModule = bContainsEngineModules;
 				RulesObject.bUseBackwardsCompatibleDefaults = bUseBackwardsCompatibleDefaults && Target.bUseBackwardsCompatibleDefaults;
-				RulesObject.bPrecompile = RulesObject.bTreatAsEngineModule && Target.bPrecompile;
+				RulesObject.bPrecompile = (RulesObject.bTreatAsEngineModule || ModuleName.Equals("UE4Game", StringComparison.OrdinalIgnoreCase)) && Target.bPrecompile;
 				RulesObject.bUsePrecompiled = bReadOnly;
 
 				// Call the constructor
