Automatically apply XML config to all platform rules objects.

diff --git a/Engine/Source/Programs/UnrealBuildTool/Configuration/TargetRules.cs b/Engine/Source/Programs/UnrealBuildTool/Configuration/TargetRules.cs
--- a/Engine/Source/Programs/UnrealBuildTool/Configuration/TargetRules.cs
+++ b/Engine/Source/Programs/UnrealBuildTool/Configuration/TargetRules.cs
@@ -1114,11 +1114,9 @@ namespace UnrealBuildTool
 			foreach(object ConfigurableObject in GetConfigurableObjects())
 			{
 				ConfigCache.ReadSettings(DirectoryReference.FromFile(ProjectFile), Platform, ConfigurableObject);
+				XmlConfig.ApplyTo(ConfigurableObject);
 			}
 
-			// Read settings from the XML config files
-			XmlConfig.ApplyTo(this);
-
 			// Allow the build platform to set defaults for this target
 			if(Platform != UnrealTargetPlatform.Unknown)
 			{
diff --git a/Engine/Source/Programs/UnrealBuildTool/Platform/Mac/UEBuildMac.cs b/Engine/Source/Programs/UnrealBuildTool/Platform/Mac/UEBuildMac.cs
--- a/Engine/Source/Programs/UnrealBuildTool/Platform/Mac/UEBuildMac.cs
+++ b/Engine/Source/Programs/UnrealBuildTool/Platform/Mac/UEBuildMac.cs
@@ -42,14 +42,6 @@ namespace UnrealBuildTool
 		[CommandLine("-EnableUBSan")]
 		[XmlConfigFile(Category = "BuildConfiguration", Name = "bEnableUndefinedBehaviorSanitizer")]
 		public bool bEnableUndefinedBehaviorSanitizer = false;
-
-		/// <summary>
-		/// Constructor
-		/// </summary>
-		public MacTargetRules()
-		{
-			XmlConfig.ApplyTo(this);
-		}
 	}
 
 	/// <summary>
diff --git a/Engine/Source/Programs/UnrealBuildTool/Platform/Windows/UEBuildWindows.cs b/Engine/Source/Programs/UnrealBuildTool/Platform/Windows/UEBuildWindows.cs
--- a/Engine/Source/Programs/UnrealBuildTool/Platform/Windows/UEBuildWindows.cs
+++ b/Engine/Source/Programs/UnrealBuildTool/Platform/Windows/UEBuildWindows.cs
@@ -76,14 +76,6 @@ namespace UnrealBuildTool
 	/// </summary>
 	public class WindowsTargetRules
 	{
-		/// <summary>
-		/// Constructor
-		/// </summary>
-		public WindowsTargetRules()
-		{
-			XmlConfig.ApplyTo(this);
-		}
-
 		/// <summary>
 		/// Version of the compiler toolchain to use on Windows platform. A value of "default" will be changed to a specific version at UBT startup.
 		/// </summary>
