UBT: Add an option to produce an error whenever a module manifest is changing, to make it easier to track down issues with files being modified.

diff --git a/Engine/Source/Programs/UnrealBuildTool/Configuration/UEBuildTarget.cs b/Engine/Source/Programs/UnrealBuildTool/Configuration/UEBuildTarget.cs
--- a/Engine/Source/Programs/UnrealBuildTool/Configuration/UEBuildTarget.cs
+++ b/Engine/Source/Programs/UnrealBuildTool/Configuration/UEBuildTarget.cs
@@ -1662,7 +1662,7 @@ namespace UnrealBuildTool
 		/// <summary>
 		/// Writes out the version manifest
 		/// </summary>
-		public void WriteReceipts()
+		public void WriteReceipts(bool bNoManifestChanges)
 		{
 			if (Receipt != null)
 			{
@@ -1707,18 +1707,36 @@ namespace UnrealBuildTool
 				{
 					if(!IsFileInstalled(FileNameToVersionManifest.Key))
 					{
-						// Write the manifest out to a string buffer, then only write it to disk if it's changed.
-						string OutputText;
-						using (StringWriter Writer = new StringWriter())
-						{
-							FileNameToVersionManifest.Value.Write(Writer);
-							OutputText = Writer.ToString();
-						}
-						if(!FileReference.Exists(FileNameToVersionManifest.Key) || File.ReadAllText(FileNameToVersionManifest.Key.FullName) != OutputText)
+						if(!FileReference.Exists(FileNameToVersionManifest.Key))
 						{
+							// If the file doesn't already exist, just write it out
 							Directory.CreateDirectory(Path.GetDirectoryName(FileNameToVersionManifest.Key.FullName));
 							FileNameToVersionManifest.Value.Write(FileNameToVersionManifest.Key.FullName);
 						}
+						else
+						{
+							// Otherwise write it to a buffer first
+							string OutputText;
+							using (StringWriter Writer = new StringWriter())
+							{
+								FileNameToVersionManifest.Value.Write(Writer);
+								OutputText = Writer.ToString();
+							}
+
+							// And only write it to disk if it's been modified
+							string CurrentText = FileReference.ReadAllText(FileNameToVersionManifest.Key);
+							if(CurrentText != OutputText)
+							{
+								if(bNoManifestChanges)
+								{
+									Log.TraceError("Build modifies {0}. This is not permitted. Before:\n    {1}\nAfter:\n    {2}", FileNameToVersionManifest.Key, CurrentText.Replace("\n", "\n    "), OutputText.Replace("\n", "\n    "));
+								}
+								else
+								{
+									FileReference.WriteAllText(FileNameToVersionManifest.Key, OutputText);
+								}
+							}
+						}
 					}
 				}
 			}
diff --git a/Engine/Source/Programs/UnrealBuildTool/UnrealBuildTool.cs b/Engine/Source/Programs/UnrealBuildTool/UnrealBuildTool.cs
--- a/Engine/Source/Programs/UnrealBuildTool/UnrealBuildTool.cs
+++ b/Engine/Source/Programs/UnrealBuildTool/UnrealBuildTool.cs
@@ -1665,9 +1665,10 @@ namespace UnrealBuildTool
 								// if the build succeeded, write the receipts and do any needed syncing
 								if (bSuccess)
 								{
+									bool bNoManifestChanges = Arguments.Any(x => x.Equals("-NoManifestChanges", StringComparison.OrdinalIgnoreCase));
 									foreach (UEBuildTarget Target in Targets)
 									{
-										Target.WriteReceipts();
+										Target.WriteReceipts(bNoManifestChanges);
 										UEBuildPlatform.GetBuildPlatform(Target.Platform).PostBuildSync(Target);
 									}
 									if (ActionsToExecute.Count == 0 && BuildConfiguration.bSkipLinkingWhenNothingToCompile)
@@ -2262,7 +2263,7 @@ namespace UnrealBuildTool
 			foreach (UEBuildTarget Target in Targets)
 			{
 				Target.PatchModuleManifestsForHotReloadAssembling(OnlyModulesLocal);
-				Target.WriteReceipts();
+				Target.WriteReceipts(false);
 			}
 		}
 	}
