UBT: Add a $(TargetOutputDir) variable that can be used in UBT paths and expands to the directory containing the executable. Can be used to ensure DLLs are copied to the appropriate output folder during a build.

Also rename the current $(OutputDir) variable (indicating the binary output directory) to $(BinaryOutputDir) to distinguish between them.

diff --git a/Engine/Source/Programs/UnrealBuildTool/Configuration/UEBuildBinary.cs b/Engine/Source/Programs/UnrealBuildTool/Configuration/UEBuildBinary.cs
--- a/Engine/Source/Programs/UnrealBuildTool/Configuration/UEBuildBinary.cs
+++ b/Engine/Source/Programs/UnrealBuildTool/Configuration/UEBuildBinary.cs
@@ -43,6 +43,11 @@ namespace UnrealBuildTool
 		/// </summary>
 		public UEBuildBinaryType Type;
 
+		/// <summary>
+		/// Output directory for this binary
+		/// </summary>
+		public DirectoryReference OutputDir;
+
 		/// <summary>
 		/// The output file path. This must be set before a binary can be built using it.
 		/// </summary>
@@ -143,6 +148,7 @@ namespace UnrealBuildTool
 			)
 		{
 			this.Type = Type;
+			this.OutputDir = OutputFilePaths.First().Directory;
 			this.OutputFilePaths = new List<FileReference>(OutputFilePaths);
 			this.IntermediateDirectory = IntermediateDirectory;
 			this.bAllowExports = bAllowExports;
@@ -172,9 +178,10 @@ namespace UnrealBuildTool
 		/// <param name="LinkEnvironment">The environment to link the binary in</param>
 		/// <param name="SharedPCHs">List of templates for shared PCHs</param>
 		/// <param name="WorkingSet">The working set of source files</param>
+		/// <param name="ExeDir">Directory containing the output executable</param>
 		/// <param name="ActionGraph">Graph to add build actions to</param>
 		/// <returns>Set of built products</returns>
-		public IEnumerable<FileItem> Build(ReadOnlyTargetRules Target, UEToolChain ToolChain, CppCompileEnvironment CompileEnvironment, LinkEnvironment LinkEnvironment, List<PrecompiledHeaderTemplate> SharedPCHs, ISourceFileWorkingSet WorkingSet, ActionGraph ActionGraph)
+		public IEnumerable<FileItem> Build(ReadOnlyTargetRules Target, UEToolChain ToolChain, CppCompileEnvironment CompileEnvironment, LinkEnvironment LinkEnvironment, List<PrecompiledHeaderTemplate> SharedPCHs, ISourceFileWorkingSet WorkingSet, DirectoryReference ExeDir, ActionGraph ActionGraph)
 		{
 			// Return nothing if we're using precompiled binaries
 			if(bUsePrecompiled)
@@ -183,7 +190,7 @@ namespace UnrealBuildTool
 			}
 
 			// Setup linking environment.
-			LinkEnvironment BinaryLinkEnvironment = SetupBinaryLinkEnvironment(Target, ToolChain, LinkEnvironment, CompileEnvironment, SharedPCHs, WorkingSet, ActionGraph);
+			LinkEnvironment BinaryLinkEnvironment = SetupBinaryLinkEnvironment(Target, ToolChain, LinkEnvironment, CompileEnvironment, SharedPCHs, WorkingSet, ExeDir, ActionGraph);
 
 			// If we're generating projects, we only need include paths and definitions, there is no need to run the linking logic.
 			if (ProjectFileGenerator.bGenerateProjectFiles)
@@ -241,8 +248,9 @@ namespace UnrealBuildTool
 		/// </summary>
 		/// <param name="RuntimeDependencies">The output list of runtime dependencies, mapping target file to type</param>
 		/// <param name="SourceFiles">Receives the list of source files that were copied</param>
+		/// <param name="ExeDir">Output directory for the executable</param>
 		/// <param name="ActionGraph">Actions to be executed</param>
-		public IEnumerable<FileItem> PrepareRuntimeDependencies(List<RuntimeDependency> RuntimeDependencies, List<FileReference> SourceFiles, ActionGraph ActionGraph)
+		public IEnumerable<FileItem> PrepareRuntimeDependencies(List<RuntimeDependency> RuntimeDependencies, List<FileReference> SourceFiles, DirectoryReference ExeDir, ActionGraph ActionGraph)
 		{
 			List<FileItem> CopiedFiles = new List<FileItem>();
 			foreach(UEBuildModule Module in Modules)
@@ -252,7 +260,7 @@ namespace UnrealBuildTool
 					if(Dependency.SourcePath == null)
 					{
 						// Expand the target path
-						string ExpandedPath = Module.ExpandPathVariables(Dependency.Path, this);
+						string ExpandedPath = Module.ExpandPathVariables(Dependency.Path, OutputDir, ExeDir);
 						if (FileFilter.FindWildcardIndex(ExpandedPath) == -1)
 						{
 							RuntimeDependencies.Add(new RuntimeDependency(new FileReference(ExpandedPath), Dependency.Type));
@@ -265,8 +273,8 @@ namespace UnrealBuildTool
 					else
 					{
 						// Parse the source and target patterns
-						FilePattern SourcePattern = new FilePattern(UnrealBuildTool.EngineSourceDirectory, Module.ExpandPathVariables(Dependency.SourcePath, this));
-						FilePattern TargetPattern = new FilePattern(UnrealBuildTool.EngineSourceDirectory, Module.ExpandPathVariables(Dependency.Path, this));
+						FilePattern SourcePattern = new FilePattern(UnrealBuildTool.EngineSourceDirectory, Module.ExpandPathVariables(Dependency.SourcePath, OutputDir, ExeDir));
+						FilePattern TargetPattern = new FilePattern(UnrealBuildTool.EngineSourceDirectory, Module.ExpandPathVariables(Dependency.Path, OutputDir, ExeDir));
 
 						// Resolve all the wildcards between the source and target paths
 						Dictionary<FileReference, FileReference> Mapping;
@@ -729,7 +737,7 @@ namespace UnrealBuildTool
 			return BinaryCompileEnvironment;
 		}
 
-		private LinkEnvironment SetupBinaryLinkEnvironment(ReadOnlyTargetRules Target, UEToolChain ToolChain, LinkEnvironment LinkEnvironment, CppCompileEnvironment CompileEnvironment, List<PrecompiledHeaderTemplate> SharedPCHs, ISourceFileWorkingSet WorkingSet, ActionGraph ActionGraph)
+		private LinkEnvironment SetupBinaryLinkEnvironment(ReadOnlyTargetRules Target, UEToolChain ToolChain, LinkEnvironment LinkEnvironment, CppCompileEnvironment CompileEnvironment, List<PrecompiledHeaderTemplate> SharedPCHs, ISourceFileWorkingSet WorkingSet, DirectoryReference ExeDir, ActionGraph ActionGraph)
 		{
 			LinkEnvironment BinaryLinkEnvironment = new LinkEnvironment(LinkEnvironment);
 			HashSet<UEBuildModule> LinkEnvironmentVisitedModules = new HashSet<UEBuildModule>();
@@ -763,7 +771,7 @@ namespace UnrealBuildTool
 				}
 
 				// Allow the module to modify the link environment for the binary.
-				Module.SetupPrivateLinkEnvironment(this, BinaryLinkEnvironment, BinaryDependencies, LinkEnvironmentVisitedModules);
+				Module.SetupPrivateLinkEnvironment(this, BinaryLinkEnvironment, BinaryDependencies, LinkEnvironmentVisitedModules, ExeDir);
 			}
 
 
diff --git a/Engine/Source/Programs/UnrealBuildTool/Configuration/UEBuildModule.cs b/Engine/Source/Programs/UnrealBuildTool/Configuration/UEBuildModule.cs
--- a/Engine/Source/Programs/UnrealBuildTool/Configuration/UEBuildModule.cs
+++ b/Engine/Source/Programs/UnrealBuildTool/Configuration/UEBuildModule.cs
@@ -239,7 +239,7 @@ namespace UnrealBuildTool
 			{
 				foreach(string InputString in InEnumerableStrings)
 				{
-					DirectoryReference Dir = new DirectoryReference(ExpandPathVariables(InputString, null));
+					DirectoryReference Dir = new DirectoryReference(ExpandPathVariables(InputString, null, null));
 					if(DirectoryReference.Exists(Dir))
 					{
 						Directories.Add(Dir);
@@ -539,9 +539,10 @@ namespace UnrealBuildTool
 		/// Expand path variables within the context of this module
 		/// </summary>
 		/// <param name="Path">Path to expand variables within</param>
-		/// <param name="Binary">The source binary containing the module. May be null.</param>
+		/// <param name="BinaryOutputDir">Directory containing the binary that links this module. May be mull.</param>
+		/// <param name="TargetOutputDir">Directory containing the output executable. May be null.</param>
 		/// <returns>The path with variables expanded</returns>
-		public string ExpandPathVariables(string Path, UEBuildBinary Binary)
+		public string ExpandPathVariables(string Path, DirectoryReference BinaryOutputDir, DirectoryReference TargetOutputDir)
 		{
 			if(Path.StartsWith("$(", StringComparison.Ordinal))
 			{
@@ -573,9 +574,13 @@ namespace UnrealBuildTool
 						{
 							Path = Rules.PluginDirectory + Path.Substring(EndIdx + 1);
 						}
-						else if(Binary != null && MatchVariableName(Path, StartIdx, EndIdx, "OutputDir"))
+						else if(BinaryOutputDir != null && MatchVariableName(Path, StartIdx, EndIdx, "BinaryOutputDir"))
 						{
-							Path = Binary.OutputFilePaths[0].Directory.FullName + Path.Substring(EndIdx + 1);
+							Path = BinaryOutputDir.FullName + Path.Substring(EndIdx + 1);
+						}
+						else if(TargetOutputDir != null && MatchVariableName(Path, StartIdx, EndIdx, "TargetOutputDir"))
+						{
+							Path = TargetOutputDir.FullName + Path.Substring(EndIdx + 1);
 						}
 						else
 						{
@@ -611,13 +616,14 @@ namespace UnrealBuildTool
 		/// Expand path variables within the context of this module
 		/// </summary>
 		/// <param name="Paths">Path to expand variables within</param>
-		/// <param name="Binary">The source binary containing the module. May be null.</param>
+		/// <param name="BinaryDir">Directory containing the binary that links this module. May be mull.</param>
+		/// <param name="ExeDir">Directory containing the output executable. May be null.</param>
 		/// <returns>The path with variables expanded</returns>
-		private IEnumerable<string> ExpandPathVariables(IEnumerable<string> Paths, UEBuildBinary Binary)
+		private IEnumerable<string> ExpandPathVariables(IEnumerable<string> Paths, DirectoryReference BinaryDir, DirectoryReference ExeDir)
 		{
 			foreach(string Path in Paths)
 			{
-				yield return ExpandPathVariables(Path, Binary);
+				yield return ExpandPathVariables(Path, BinaryDir, ExeDir);
 			}
 		}
 
@@ -636,7 +642,8 @@ namespace UnrealBuildTool
 			List<UEBuildBundleResource> AdditionalBundleResources,
 			List<string> DelayLoadDLLs,
 			List<UEBuildBinary> BinaryDependencies,
-			HashSet<UEBuildModule> VisitedModules
+			HashSet<UEBuildModule> VisitedModules,
+			DirectoryReference ExeDir
 			)
 		{
 			// There may be circular dependencies in compile dependencies, so we need to avoid reentrance.
@@ -669,7 +676,7 @@ namespace UnrealBuildTool
 						if (bIsExternalModule || bIsInStaticLibrary)
 						{
 							DependencyModule.SetupPublicLinkEnvironment(SourceBinary, LibraryPaths, AdditionalLibraries, RuntimeLibraryPaths, Frameworks, WeakFrameworks,
-								AdditionalFrameworks, AdditionalShadowFiles, AdditionalBundleResources, DelayLoadDLLs, BinaryDependencies, VisitedModules);
+								AdditionalFrameworks, AdditionalShadowFiles, AdditionalBundleResources, DelayLoadDLLs, BinaryDependencies, VisitedModules, ExeDir);
 						}
 					}
 				}
@@ -677,7 +684,7 @@ namespace UnrealBuildTool
 				// Add this module's public include library paths and additional libraries.
 				LibraryPaths.AddRange(PublicLibraryPaths);
 				AdditionalLibraries.AddRange(PublicAdditionalLibraries);
-				RuntimeLibraryPaths.AddRange(ExpandPathVariables(Rules.PublicRuntimeLibraryPaths, SourceBinary));
+				RuntimeLibraryPaths.AddRange(ExpandPathVariables(Rules.PublicRuntimeLibraryPaths, SourceBinary.OutputDir, ExeDir));
 				Frameworks.AddRange(PublicFrameworks);
 				WeakFrameworks.AddRange(PublicWeakFrameworks);
 				AdditionalBundleResources.AddRange(PublicAdditionalBundleResources);
@@ -699,15 +706,16 @@ namespace UnrealBuildTool
 			UEBuildBinary SourceBinary,
 			LinkEnvironment LinkEnvironment,
 			List<UEBuildBinary> BinaryDependencies,
-			HashSet<UEBuildModule> VisitedModules
+			HashSet<UEBuildModule> VisitedModules,
+			DirectoryReference ExeDir
 			)
 		{
 			// Add the private rpaths
-			LinkEnvironment.RuntimeLibraryPaths.AddRange(ExpandPathVariables(Rules.PrivateRuntimeLibraryPaths, SourceBinary));
+			LinkEnvironment.RuntimeLibraryPaths.AddRange(ExpandPathVariables(Rules.PrivateRuntimeLibraryPaths, SourceBinary.OutputDir, ExeDir));
 
 			// Allow the module's public dependencies to add library paths and additional libraries to the link environment.
 			SetupPublicLinkEnvironment(SourceBinary, LinkEnvironment.LibraryPaths, LinkEnvironment.AdditionalLibraries, LinkEnvironment.RuntimeLibraryPaths, LinkEnvironment.Frameworks, LinkEnvironment.WeakFrameworks,
-				LinkEnvironment.AdditionalFrameworks, LinkEnvironment.AdditionalShadowFiles, LinkEnvironment.AdditionalBundleResources, LinkEnvironment.DelayLoadDLLs, BinaryDependencies, VisitedModules);
+				LinkEnvironment.AdditionalFrameworks, LinkEnvironment.AdditionalShadowFiles, LinkEnvironment.AdditionalBundleResources, LinkEnvironment.DelayLoadDLLs, BinaryDependencies, VisitedModules, ExeDir);
 
 			// Also allow the module's public and private dependencies to modify the link environment.
 			List<UEBuildModule> AllDependencyModules = new List<UEBuildModule>();
@@ -717,7 +725,7 @@ namespace UnrealBuildTool
 			foreach (UEBuildModule DependencyModule in AllDependencyModules)
 			{
 				DependencyModule.SetupPublicLinkEnvironment(SourceBinary, LinkEnvironment.LibraryPaths, LinkEnvironment.AdditionalLibraries, LinkEnvironment.RuntimeLibraryPaths, LinkEnvironment.Frameworks, LinkEnvironment.WeakFrameworks,
-					LinkEnvironment.AdditionalFrameworks, LinkEnvironment.AdditionalShadowFiles, LinkEnvironment.AdditionalBundleResources, LinkEnvironment.DelayLoadDLLs, BinaryDependencies, VisitedModules);
+					LinkEnvironment.AdditionalFrameworks, LinkEnvironment.AdditionalShadowFiles, LinkEnvironment.AdditionalBundleResources, LinkEnvironment.DelayLoadDLLs, BinaryDependencies, VisitedModules, ExeDir);
 			}
 		}
 
diff --git a/Engine/Source/Programs/UnrealBuildTool/Configuration/UEBuildTarget.cs b/Engine/Source/Programs/UnrealBuildTool/Configuration/UEBuildTarget.cs
--- a/Engine/Source/Programs/UnrealBuildTool/Configuration/UEBuildTarget.cs
+++ b/Engine/Source/Programs/UnrealBuildTool/Configuration/UEBuildTarget.cs
@@ -2116,9 +2116,10 @@ namespace UnrealBuildTool
 			}
 
 			// Build the target's binaries.
+			DirectoryReference ExeDir = OutputPaths[0].Directory;
 			foreach (UEBuildBinary Binary in Binaries)
 			{
-				OutputItems.AddRange(Binary.Build(Rules, TargetToolChain, GlobalCompileEnvironment, GlobalLinkEnvironment, SharedPCHs, WorkingSet, ActionGraph));
+				OutputItems.AddRange(Binary.Build(Rules, TargetToolChain, GlobalCompileEnvironment, GlobalLinkEnvironment, SharedPCHs, WorkingSet, ExeDir, ActionGraph));
 			}
 
 			// Prepare all the runtime dependencies, copying them from their source folders if necessary
@@ -2126,7 +2127,7 @@ namespace UnrealBuildTool
 			List<RuntimeDependency> RuntimeDependencies = new List<RuntimeDependency>();
 			foreach(UEBuildBinary Binary in Binaries)
 			{
-				OutputItems.AddRange(Binary.PrepareRuntimeDependencies(RuntimeDependencies, RuntimeDependencySourceFiles, ActionGraph));
+				OutputItems.AddRange(Binary.PrepareRuntimeDependencies(RuntimeDependencies, RuntimeDependencySourceFiles, ExeDir, ActionGraph));
 			}
 
 			// If we're just precompiling a plugin, only include output items which are under that directory
