UBT: Add a field to the target receipt which specifies the executable to launch. Allows procedurally determining the location of the editor for a particular target.

diff --git a/Engine/Source/Programs/UnrealBuildTool/Configuration/UEBuildTarget.cs b/Engine/Source/Programs/UnrealBuildTool/Configuration/UEBuildTarget.cs
--- a/Engine/Source/Programs/UnrealBuildTool/Configuration/UEBuildTarget.cs
+++ b/Engine/Source/Programs/UnrealBuildTool/Configuration/UEBuildTarget.cs
@@ -1418,8 +1418,16 @@ namespace UnrealBuildTool
 				}
 			}
 
-			// Find all the build products and modules from this binary
+			// Create the receipt
 			Receipt = new TargetReceipt(TargetName, Platform, Configuration, Version);
+
+			// Set the launch executable if there is one
+			if(!Rules.bShouldCompileAsDLL)
+			{
+				Receipt.Launch = OutputPaths[0];
+			}
+
+			// Find all the build products and modules from this binary
 			foreach (KeyValuePair<FileReference, BuildProductType> BuildProductPair in BuildProducts)
 			{
 				if(BuildProductPair.Value != BuildProductType.BuildResource)
diff --git a/Engine/Source/Programs/UnrealBuildTool/System/TargetReceipt.cs b/Engine/Source/Programs/UnrealBuildTool/System/TargetReceipt.cs
--- a/Engine/Source/Programs/UnrealBuildTool/System/TargetReceipt.cs
+++ b/Engine/Source/Programs/UnrealBuildTool/System/TargetReceipt.cs
@@ -275,6 +275,11 @@ namespace UnrealBuildTool
 		/// </summary>
 		public BuildVersion Version;
 
+		/// <summary>
+		/// The exectuable to launch for this target
+		/// </summary>
+		public FileReference Launch;
+
 		/// <summary>
 		/// The build products which are part of this target
 		/// </summary>
@@ -318,6 +323,7 @@ namespace UnrealBuildTool
 		/// <param name="Other">Receipt to copy from</param>
 		public TargetReceipt(TargetReceipt Other)
 		{
+			Launch = Other.Launch;
 			foreach (BuildProduct OtherBuildProduct in Other.BuildProducts)
 			{
 				BuildProducts.Add(new BuildProduct(OtherBuildProduct));
@@ -475,6 +481,13 @@ namespace UnrealBuildTool
 			// Create the receipt
 			TargetReceipt Receipt = new TargetReceipt(TargetName, Platform, Configuration, Version);
 
+			// Read the launch executable
+			string Launch;
+			if(RawObject.TryGetStringField("Launch", out Launch))
+			{
+				Receipt.Launch = ExpandPathVariables(Launch, EngineDir, ProjectDir);
+			}
+
 			// Read the build products
 			JsonObject[] BuildProductObjects;
 			if (RawObject.TryGetObjectArrayField("BuildProducts", out BuildProductObjects))
@@ -591,6 +604,11 @@ namespace UnrealBuildTool
 				Version.WriteProperties(Writer);
 				Writer.WriteObjectEnd();
 
+				if(Launch != null)
+				{
+					Writer.WriteValue("Launch", InsertPathVariables(Launch, EngineDir, ProjectDir));
+				}
+
 				Writer.WriteArrayStart("BuildProducts");
 				foreach (BuildProduct BuildProduct in BuildProducts)
 				{
