UBT: Fix project files not containing include search paths for dependent modules.

diff --git a/Engine/Source/Programs/UnrealBuildTool/Configuration/UEBuildModuleCPP.cs b/Engine/Source/Programs/UnrealBuildTool/Configuration/UEBuildModuleCPP.cs
--- a/Engine/Source/Programs/UnrealBuildTool/Configuration/UEBuildModuleCPP.cs
+++ b/Engine/Source/Programs/UnrealBuildTool/Configuration/UEBuildModuleCPP.cs
@@ -299,13 +299,6 @@ namespace UnrealBuildTool
 			ProjectFile.AddIntelliSensePreprocessorDefinitions(ModuleCompileEnvironment.Definitions);
 			ProjectFile.AddIntelliSenseIncludePaths(ModuleCompileEnvironment.IncludePaths.SystemIncludePaths, true);
 			ProjectFile.AddIntelliSenseIncludePaths(ModuleCompileEnvironment.IncludePaths.UserIncludePaths, false);
-
-			// This directory may not exist for this module (or ever exist, if it doesn't contain any generated headers), but we want the project files
-			// to search it so we can pick up generated code definitions after UHT is run for the first time.
-			if(GeneratedCodeDirectory != null)
-			{
-				ProjectFile.AddIntelliSenseIncludePaths(new HashSet<DirectoryReference>{ GeneratedCodeDirectory }, false);
-			}
 		}
 
 		/// <summary>
@@ -320,7 +313,9 @@ namespace UnrealBuildTool
 			bool bLegacyPublicIncludePaths
 			)
 		{
-			if(bAddGeneratedCodeIncludePath)
+			// This directory may not exist for this module (or ever exist, if it doesn't contain any generated headers), but we want the project files
+			// to search it so we can pick up generated code definitions after UHT is run for the first time.
+			if(bAddGeneratedCodeIncludePath || (ProjectFileGenerator.bGenerateProjectFiles && GeneratedCodeDirectory != null))
 			{
 				IncludePaths.Add(GeneratedCodeDirectory);
 			}
