UBT: Catch the case where manifests are recycled if the -NoManifestChanges argument is specified.

#rnx

diff --git a/Engine/Source/Programs/UnrealBuildTool/Configuration/UEBuildTarget.cs b/Engine/Source/Programs/UnrealBuildTool/Configuration/UEBuildTarget.cs
--- a/Engine/Source/Programs/UnrealBuildTool/Configuration/UEBuildTarget.cs
+++ b/Engine/Source/Programs/UnrealBuildTool/Configuration/UEBuildTarget.cs
@@ -1541,8 +1541,9 @@ namespace UnrealBuildTool
 		/// additional engine modules to be built, so we don't prohibit files being added or removed.
 		/// </summary>
 		/// <param name="OutputFiles">List of files being modified by this build</param>
+		/// <param name="bNoManifestChanges">Whether manifest changes are allowed. If a manifest has to be changed, an error will be output.</param>
 		/// <returns>True if the existing version manifests will remain valid during this build, false if they are invalidated</returns>
-		public bool TryRecycleVersionManifests(HashSet<FileReference> OutputFiles)
+		public bool TryRecycleVersionManifests(HashSet<FileReference> OutputFiles, bool bNoManifestChanges)
 		{
 			// Make sure we've got a list of version manifests to check against
 			if(FileReferenceToModuleManifestPairs == null)
@@ -1587,7 +1588,14 @@ namespace UnrealBuildTool
 					{
 						if(OutputFiles.Contains(ExistingFile))
 						{
-							Log.TraceLog("Unable to recycle manifests - modifying {0} invalidates {1}. Using build id {2}.", ExistingFile, ExistingPair.Key, Receipt.Version.BuildId);
+							if(bNoManifestChanges)
+							{
+								Log.TraceError("Unable to recycle manifests - modifying {0} invalidates {1}.", ExistingFile, ExistingPair.Key);
+							}
+							else
+							{
+								Log.TraceLog("Unable to recycle manifests - modifying {0} invalidates {1}. Using build id {2}.", ExistingFile, ExistingPair.Key, Receipt.Version.BuildId);
+							}
 							return false;
 						}
 					}
diff --git a/Engine/Source/Programs/UnrealBuildTool/UnrealBuildTool.cs b/Engine/Source/Programs/UnrealBuildTool/UnrealBuildTool.cs
--- a/Engine/Source/Programs/UnrealBuildTool/UnrealBuildTool.cs
+++ b/Engine/Source/Programs/UnrealBuildTool/UnrealBuildTool.cs
@@ -1638,11 +1638,14 @@ namespace UnrealBuildTool
 									CppIncludeThread = new CppIncludeBackgroundThread(TargetToOutdatedPrerequisitesMap, TargetToHeaders);
 								}
 
+								// Check if we're allowed to modify manifests
+								bool bNoManifestChanges = Arguments.Any(x => x.Equals("-NoManifestChanges", StringComparison.OrdinalIgnoreCase));
+
 								// If we're not touching any shared files (ie. anything under Engine), allow the build ids to be recycled between applications.
 								HashSet<FileReference> OutputFiles = new HashSet<FileReference>(ActionsToExecute.SelectMany(x => x.ProducedItems.Select(y => y.Location)));
 								foreach (UEBuildTarget Target in Targets)
 								{
-									if (!Target.TryRecycleVersionManifests(OutputFiles))
+									if (!Target.TryRecycleVersionManifests(OutputFiles, bNoManifestChanges) && !bNoManifestChanges)
 									{
 										Target.InvalidateVersionManifests();
 									}
@@ -1665,7 +1668,6 @@ namespace UnrealBuildTool
 								// if the build succeeded, write the receipts and do any needed syncing
 								if (bSuccess)
 								{
-									bool bNoManifestChanges = Arguments.Any(x => x.Equals("-NoManifestChanges", StringComparison.OrdinalIgnoreCase));
 									foreach (UEBuildTarget Target in Targets)
 									{
 										Target.WriteReceipts(bNoManifestChanges);
