UBT: Prevent continuing past manifest errors.

#rnx

diff --git a/Engine/Source/Programs/UnrealBuildTool/Configuration/UEBuildTarget.cs b/Engine/Source/Programs/UnrealBuildTool/Configuration/UEBuildTarget.cs
--- a/Engine/Source/Programs/UnrealBuildTool/Configuration/UEBuildTarget.cs
+++ b/Engine/Source/Programs/UnrealBuildTool/Configuration/UEBuildTarget.cs
@@ -1590,7 +1590,7 @@ namespace UnrealBuildTool
 						{
 							if(bNoManifestChanges)
 							{
-								Log.TraceError("Unable to recycle manifests - modifying {0} invalidates {1}.", ExistingFile, ExistingPair.Key);
+								Log.TraceError("Previous build product would be overwritten: {0}.", ExistingFile);
 							}
 							else
 							{
diff --git a/Engine/Source/Programs/UnrealBuildTool/UnrealBuildTool.cs b/Engine/Source/Programs/UnrealBuildTool/UnrealBuildTool.cs
--- a/Engine/Source/Programs/UnrealBuildTool/UnrealBuildTool.cs
+++ b/Engine/Source/Programs/UnrealBuildTool/UnrealBuildTool.cs
@@ -1645,9 +1645,16 @@ namespace UnrealBuildTool
 								HashSet<FileReference> OutputFiles = new HashSet<FileReference>(ActionsToExecute.SelectMany(x => x.ProducedItems.Select(y => y.Location)));
 								foreach (UEBuildTarget Target in Targets)
 								{
-									if (!Target.TryRecycleVersionManifests(OutputFiles, bNoManifestChanges) && !bNoManifestChanges)
+									if (!Target.TryRecycleVersionManifests(OutputFiles, bNoManifestChanges))
 									{
-										Target.InvalidateVersionManifests();
+										if(bNoManifestChanges)
+										{
+											throw new BuildException("Stopping build due to manifest changes.");
+										}
+										else
+										{
+											Target.InvalidateVersionManifests();
+										}
 									}
 								}
 
