UBT: Suppress message about change to action history if the previous build product did not exist.

#rnx

diff --git a/Engine/Source/Programs/UnrealBuildTool/System/ActionGraph.cs b/Engine/Source/Programs/UnrealBuildTool/System/ActionGraph.cs
--- a/Engine/Source/Programs/UnrealBuildTool/System/ActionGraph.cs
+++ b/Engine/Source/Programs/UnrealBuildTool/System/ActionGraph.cs
@@ -872,13 +872,16 @@ namespace UnrealBuildTool
 						if (!ActionHistory.GetProducingCommandLine(ProducedItem, out OldProducingCommandLine)
 						|| !String.Equals(OldProducingCommandLine, NewProducingCommandLine, StringComparison.InvariantCultureIgnoreCase))
 						{
-							Log.TraceLog(
-								"{0}: Produced item \"{1}\" was produced by outdated command-line.\n  Old command-line: {2}\n  New command-line: {3}",
-								RootAction.StatusDescription,
-								Path.GetFileName(ProducedItem.AbsolutePath),
-								OldProducingCommandLine,
-								NewProducingCommandLine
-								);
+							if(ProducedItem.bExists)
+							{
+								Log.TraceLog(
+									"{0}: Produced item \"{1}\" was produced by outdated command-line.\n  Old command-line: {2}\n  New command-line: {3}",
+									RootAction.StatusDescription,
+									Path.GetFileName(ProducedItem.AbsolutePath),
+									OldProducingCommandLine,
+									NewProducingCommandLine
+									);
+							}
 
 							bIsOutdated = true;
 
