UBT: Fix cook problems caused by copied runtime dependencies not being added to the receipt or manifest.

#rnx

diff --git a/Engine/Source/Programs/UnrealBuildTool/Configuration/UEBuildBinary.cs b/Engine/Source/Programs/UnrealBuildTool/Configuration/UEBuildBinary.cs
--- a/Engine/Source/Programs/UnrealBuildTool/Configuration/UEBuildBinary.cs
+++ b/Engine/Source/Programs/UnrealBuildTool/Configuration/UEBuildBinary.cs
@@ -247,12 +247,10 @@ namespace UnrealBuildTool
 		/// Gets all the runtime dependencies Copies all the runtime dependencies from any modules in 
 		/// </summary>
 		/// <param name="RuntimeDependencies">The output list of runtime dependencies, mapping target file to type</param>
-		/// <param name="SourceFiles">Receives the list of source files that were copied</param>
+		/// <param name="TargetFileToSourceFile">Map of target files to source files that need to be copied</param>
 		/// <param name="ExeDir">Output directory for the executable</param>
-		/// <param name="ActionGraph">Actions to be executed</param>
-		public IEnumerable<FileItem> PrepareRuntimeDependencies(List<RuntimeDependency> RuntimeDependencies, List<FileReference> SourceFiles, DirectoryReference ExeDir, ActionGraph ActionGraph)
+		public void PrepareRuntimeDependencies(List<RuntimeDependency> RuntimeDependencies, Dictionary<FileReference, FileReference> TargetFileToSourceFile, DirectoryReference ExeDir)
 		{
-			List<FileItem> CopiedFiles = new List<FileItem>();
 			foreach(UEBuildModule Module in Modules)
 			{
 				foreach (ModuleRules.RuntimeDependency Dependency in Module.Rules.RuntimeDependencies.Inner)
@@ -290,48 +288,20 @@ namespace UnrealBuildTool
 						// Add actions to copy everything
 						foreach(KeyValuePair<FileReference, FileReference> Pair in Mapping)
 						{
-							CopiedFiles.Add(CreateCopyAction(Pair.Value, Pair.Key, ActionGraph));
-							SourceFiles.Add(Pair.Value);
-							RuntimeDependencies.Add(new RuntimeDependency(Pair.Key, Dependency.Type));
+							FileReference ExistingSourceFile;
+							if(!TargetFileToSourceFile.TryGetValue(Pair.Key, out ExistingSourceFile))
+							{
+								TargetFileToSourceFile[Pair.Key] = Pair.Value;
+								RuntimeDependencies.Add(new RuntimeDependency(Pair.Key, Dependency.Type));
+							}
+							else if(ExistingSourceFile != Pair.Value)
+							{
+								throw new BuildException("Runtime dependency '{0}' is configured to be staged from '{1}' and '{2}'", Pair.Key, Pair.Value, ExistingSourceFile);
+							}
 						}
 					}
 				}
 			}
-			return CopiedFiles;
-		}
-
-		/// <summary>
-		/// Creates an action which copies a file from one location to another
-		/// </summary>
-		/// <param name="SourceFile">The source file location</param>
-		/// <param name="TargetFile">The target file location</param>
-		/// <param name="ActionGraph">The action graph</param>
-		/// <returns>File item for the output file</returns>
-		static FileItem CreateCopyAction(FileReference SourceFile, FileReference TargetFile, ActionGraph ActionGraph)
-		{
-			FileItem SourceFileItem = FileItem.GetItemByFileReference(SourceFile);
-			FileItem TargetFileItem = FileItem.GetItemByFileReference(TargetFile);
-
-			Action CopyAction = ActionGraph.Add(ActionType.BuildProject);
-			CopyAction.CommandDescription = "Copy";
-			if(BuildHostPlatform.Current.Platform == UnrealTargetPlatform.Win64)
-			{
-				CopyAction.CommandPath = "cmd.exe";
-				CopyAction.CommandArguments = String.Format("/C \"copy /Y \"{0}\" \"{1}\" 1>nul\"", SourceFile, TargetFile);
-			}
-			else
-			{
-				CopyAction.CommandPath = "/bin/sh";
-				CopyAction.CommandArguments = String.Format("cp -f {0} {1}", Utils.EscapeShellArgument(SourceFile.FullName), Utils.EscapeShellArgument(TargetFile.FullName));
-			}
-			CopyAction.WorkingDirectory = Environment.CurrentDirectory;
-			CopyAction.PrerequisiteItems.Add(SourceFileItem);
-			CopyAction.ProducedItems.Add(TargetFileItem);
-			CopyAction.DeleteItems.Add(TargetFileItem);
-			CopyAction.StatusDescription = TargetFileItem.Location.GetFileName();
-			CopyAction.bCanExecuteRemotely = false;
-
-			return TargetFileItem;
 		}
 
 		/// <summary>
diff --git a/Engine/Source/Programs/UnrealBuildTool/Configuration/UEBuildTarget.cs b/Engine/Source/Programs/UnrealBuildTool/Configuration/UEBuildTarget.cs
--- a/Engine/Source/Programs/UnrealBuildTool/Configuration/UEBuildTarget.cs
+++ b/Engine/Source/Programs/UnrealBuildTool/Configuration/UEBuildTarget.cs
@@ -2139,11 +2139,15 @@ namespace UnrealBuildTool
 			}
 
 			// Prepare all the runtime dependencies, copying them from their source folders if necessary
-			List<FileReference> RuntimeDependencySourceFiles = new List<FileReference>();
 			List<RuntimeDependency> RuntimeDependencies = new List<RuntimeDependency>();
+			Dictionary<FileReference, FileReference> RuntimeDependencyTargetFileToSourceFile = new Dictionary<FileReference, FileReference>();
 			foreach(UEBuildBinary Binary in Binaries)
 			{
-				OutputItems.AddRange(Binary.PrepareRuntimeDependencies(RuntimeDependencies, RuntimeDependencySourceFiles, ExeDir, ActionGraph));
+				Binary.PrepareRuntimeDependencies(RuntimeDependencies, RuntimeDependencyTargetFileToSourceFile, ExeDir);
+			}
+			foreach(KeyValuePair<FileReference, FileReference> Pair in RuntimeDependencyTargetFileToSourceFile)
+			{
+				OutputItems.Add(CreateCopyAction(Pair.Value, Pair.Key, ActionGraph));
 			}
 
 			// If we're just precompiling a plugin, only include output items which are under that directory
@@ -2163,6 +2167,7 @@ namespace UnrealBuildTool
 				Binary.GetBuildProducts(Rules, TargetToolChain, BinaryBuildProducts, GlobalLinkEnvironment.bCreateDebugInfo);
 				BuildProducts.AddRange(BinaryBuildProducts);
 			}
+			BuildProducts.AddRange(RuntimeDependencyTargetFileToSourceFile.Select(x => new KeyValuePair<FileReference, BuildProductType>(x.Key, BuildProductType.RequiredResource)));
 
 			// Create a receipt for the target
 			if (!ProjectFileGenerator.bGenerateProjectFiles)
@@ -2185,7 +2190,7 @@ namespace UnrealBuildTool
 			// Build a list of all the externally files
 			foreach(FileReference DependencyListFileName in Rules.DependencyListFileNames)
 			{
-				WriteDependencyList(DependencyListFileName, RuntimeDependencySourceFiles);
+				WriteDependencyList(DependencyListFileName, RuntimeDependencyTargetFileToSourceFile.Values);
 			}
 
 			// If we're only generating the manifest, return now
@@ -2203,6 +2208,40 @@ namespace UnrealBuildTool
 			return ECompilationResult.Succeeded;
 		}
 
+		/// <summary>
+		/// Creates an action which copies a file from one location to another
+		/// </summary>
+		/// <param name="SourceFile">The source file location</param>
+		/// <param name="TargetFile">The target file location</param>
+		/// <param name="ActionGraph">The action graph</param>
+		/// <returns>File item for the output file</returns>
+		static FileItem CreateCopyAction(FileReference SourceFile, FileReference TargetFile, ActionGraph ActionGraph)
+		{
+			FileItem SourceFileItem = FileItem.GetItemByFileReference(SourceFile);
+			FileItem TargetFileItem = FileItem.GetItemByFileReference(TargetFile);
+
+			Action CopyAction = ActionGraph.Add(ActionType.BuildProject);
+			CopyAction.CommandDescription = "Copy";
+			if(BuildHostPlatform.Current.Platform == UnrealTargetPlatform.Win64)
+			{
+				CopyAction.CommandPath = "cmd.exe";
+				CopyAction.CommandArguments = String.Format("/C \"copy /Y \"{0}\" \"{1}\" 1>nul\"", SourceFile, TargetFile);
+			}
+			else
+			{
+				CopyAction.CommandPath = "/bin/sh";
+				CopyAction.CommandArguments = String.Format("cp -f {0} {1}", Utils.EscapeShellArgument(SourceFile.FullName), Utils.EscapeShellArgument(TargetFile.FullName));
+			}
+			CopyAction.WorkingDirectory = Environment.CurrentDirectory;
+			CopyAction.PrerequisiteItems.Add(SourceFileItem);
+			CopyAction.ProducedItems.Add(TargetFileItem);
+			CopyAction.DeleteItems.Add(TargetFileItem);
+			CopyAction.StatusDescription = TargetFileItem.Location.GetFileName();
+			CopyAction.bCanExecuteRemotely = false;
+
+			return TargetFileItem;
+		}
+
 		/// <summary>
 		/// Creates a toolchain for the current target. May be overridden by the target rules.
 		/// </summary>
