Prevent runtime dependencies being copied into installed locations.

#rnx

diff --git a/Engine/Source/Programs/UnrealBuildTool/Configuration/UEBuildTarget.cs b/Engine/Source/Programs/UnrealBuildTool/Configuration/UEBuildTarget.cs
--- a/Engine/Source/Programs/UnrealBuildTool/Configuration/UEBuildTarget.cs
+++ b/Engine/Source/Programs/UnrealBuildTool/Configuration/UEBuildTarget.cs
@@ -2147,7 +2147,10 @@ namespace UnrealBuildTool
 			}
 			foreach(KeyValuePair<FileReference, FileReference> Pair in RuntimeDependencyTargetFileToSourceFile)
 			{
-				OutputItems.Add(CreateCopyAction(Pair.Value, Pair.Key, ActionGraph));
+				if(!IsFileInstalled(Pair.Key))
+				{
+					OutputItems.Add(CreateCopyAction(Pair.Value, Pair.Key, ActionGraph));
+				}
 			}
 
 			// If we're just precompiling a plugin, only include output items which are under that directory
