UBT: Output a message when there is nothing to build for a target (eg. because it's an installed build)

#rnx

diff --git a/Engine/Source/Programs/UnrealBuildTool/UnrealBuildTool.cs b/Engine/Source/Programs/UnrealBuildTool/UnrealBuildTool.cs
--- a/Engine/Source/Programs/UnrealBuildTool/UnrealBuildTool.cs
+++ b/Engine/Source/Programs/UnrealBuildTool/UnrealBuildTool.cs
@@ -1542,11 +1542,19 @@ namespace UnrealBuildTool
 
 							if (BuildConfiguration.bUseUBTMakefiles && !UBTMakefile.PrerequisiteActions.Any(x => x.ActionHandler != null))
 							{
-								// We've been told to prepare to build, so let's go ahead and save out our action graph so that we can use in a later invocation 
-								// to assemble the build.  Even if we are configured to assemble the build in this same invocation, we want to save out the
-								// Makefile so that it can be used on subsequent 'assemble only' runs, for the fastest possible iteration times
-								// @todo ubtmake: Optimization: We could make 'gather + assemble' mode slightly faster by saving this while busy compiling (on our worker thread)
-								UBTMakefile.SaveUBTMakefile(TargetDescs, HotReload, UBTMakefile);
+								if(UBTMakefile.PrerequisiteActions.Length == 0)
+								{
+									// Nothing to build for this target
+									Log.TraceInformation("No buildable actions for this target. Makefile will not be saved.");
+								}
+								else
+								{
+									// We've been told to prepare to build, so let's go ahead and save out our action graph so that we can use in a later invocation 
+									// to assemble the build.  Even if we are configured to assemble the build in this same invocation, we want to save out the
+									// Makefile so that it can be used on subsequent 'assemble only' runs, for the fastest possible iteration times
+									// @todo ubtmake: Optimization: We could make 'gather + assemble' mode slightly faster by saving this while busy compiling (on our worker thread)
+									UBTMakefile.SaveUBTMakefile(TargetDescs, HotReload, UBTMakefile);
+								}
 							}
 						}
 
