UBT: Fixed issue where MaxProcessorCount and ProcessorCountMultiplier settings would only affect the first class found. These values can now be configured for each executor separately, via the LocalExecutor, ParallelExecutor and SNDBS sections in BuildConfiguration.xml.

diff --git a/Engine/Source/Programs/UnrealBuildTool/Executors/LocalExecutor.cs b/Engine/Source/Programs/UnrealBuildTool/Executors/LocalExecutor.cs
--- a/Engine/Source/Programs/UnrealBuildTool/Executors/LocalExecutor.cs
+++ b/Engine/Source/Programs/UnrealBuildTool/Executors/LocalExecutor.cs
@@ -265,13 +265,13 @@ namespace UnrealBuildTool
 		/// When using the local executor (not XGE), run a single action on each CPU core.  Note that you can set this to a larger value
 		/// to get slightly faster build times in many cases, but your computer's responsiveness during compiling may be much worse.
 		/// </summary>
-		[XmlConfigFile(Category = "BuildConfiguration")]
+		[XmlConfigFile]
 		double ProcessorCountMultiplier = 1.0;
 
 		/// <summary>
 		/// Maximum processor count for local execution. 
 		/// </summary>
-		[XmlConfigFile(Category = "BuildConfiguration")]
+		[XmlConfigFile]
 		int MaxProcessorCount = int.MaxValue;
 
 		public LocalExecutor()
diff --git a/Engine/Source/Programs/UnrealBuildTool/Executors/ParallelExecutor.cs b/Engine/Source/Programs/UnrealBuildTool/Executors/ParallelExecutor.cs
--- a/Engine/Source/Programs/UnrealBuildTool/Executors/ParallelExecutor.cs
+++ b/Engine/Source/Programs/UnrealBuildTool/Executors/ParallelExecutor.cs
@@ -35,19 +35,19 @@ namespace UnrealBuildTool
 		/// <summary>
 		/// Processor count multiplier for local execution. Can be below 1 to reserve CPU for other tasks.
 		/// </summary>
-		[XmlConfigFile(Category = "BuildConfiguration")]
+		[XmlConfigFile]
 		double ProcessorCountMultiplier = 1.0;
 
 		/// <summary>
 		/// Maximum processor count for local execution. 
 		/// </summary>
-		[XmlConfigFile(Category = "BuildConfiguration")]
+		[XmlConfigFile]
 		int MaxProcessorCount = int.MaxValue;
 
 		/// <summary>
 		/// When enabled, will stop compiling targets after a compile error occurs.
 		/// </summary>
-		[XmlConfigFile(Category = "BuildConfiguration")]
+		[XmlConfigFile]
 		bool bStopCompilationAfterErrors = false;
 
 		/// <summary>
