Add back -funwind-tables for arm

This was removed an only tested on x86 which worked just fine. Arm reqiures this for backtrace

diff --git a/Engine/Source/Programs/UnrealBuildTool/Platform/Linux/LinuxToolChain.cs b/Engine/Source/Programs/UnrealBuildTool/Platform/Linux/LinuxToolChain.cs
--- a/Engine/Source/Programs/UnrealBuildTool/Platform/Linux/LinuxToolChain.cs
+++ b/Engine/Source/Programs/UnrealBuildTool/Platform/Linux/LinuxToolChain.cs
@@ -472,8 +472,12 @@ namespace UnrealBuildTool
 			}
 
 			Result += " -Wall -Werror";
-			// Comment out for now, if we run into issues with our stack tracing double check this flag
-			//Result += " -funwind-tables";               // generate unwind tables as they seem to be needed for stack tracing (why??)
+
+			if (Architecture.StartsWith("arm") || Architecture.StartsWith("aarch64"))
+			{
+				Result += " -funwind-tables";           // generate unwind tables as they are needed for arm
+			}
+
 			Result += " -Wsequence-point";              // additional warning not normally included in Wall: warns if order of operations is ambigious
 			//Result += " -Wunreachable-code";            // additional warning not normally included in Wall: warns if there is code that will never be executed - not helpful due to bIsGCC and similar
 			//Result += " -Wshadow";                      // additional warning not normally included in Wall: warns if there variable/typedef shadows some other variable - not helpful because we have gobs of code that shadows variables
