Linux: enable XGE during cross-builds to see whether the build issues persist.

- Licensees are asking for this and XGE folks are eager to help investigating the crashes, if any.

diff --git a/Engine/Source/Programs/UnrealBuildTool/Platform/Linux/UEBuildLinux.cs b/Engine/Source/Programs/UnrealBuildTool/Platform/Linux/UEBuildLinux.cs
--- a/Engine/Source/Programs/UnrealBuildTool/Platform/Linux/UEBuildLinux.cs
+++ b/Engine/Source/Programs/UnrealBuildTool/Platform/Linux/UEBuildLinux.cs
@@ -182,9 +182,9 @@ namespace UnrealBuildTool
 
 		public override bool CanUseXGE()
 		{
-			// XGE crashes with very high probability when v8_clang-3.9.0-centos cross-toolchain is used on Windows. Please make sure this is resolved before re-enabling it.
 			// [RCL] 2018-05-02: disabling XGE even during a native build because the support is not ready and you can have mysterious build failures when ib_console is installed.
-			return false;//BuildHostPlatform.Current.Platform == UnrealTargetPlatform.Linux;
+			// [RCL] 2018-07-10: enabling XGE for Windows to see if the crash from 2016 still persists. Please disable if you see spurious build errors that don't repro without XGE
+			return BuildHostPlatform.Current.Platform == UnrealTargetPlatform.Win64;
 		}
 
 		public override bool CanUseParallelExecutor()
