Linux: Check we are not x86 otherwise add unwind tables

Copying the change that went over into 4.20.1 here

diff --git a/Engine/Source/Programs/UnrealBuildTool/Platform/Linux/LinuxToolChain.cs b/Engine/Source/Programs/UnrealBuildTool/Platform/Linux/LinuxToolChain.cs
--- a/Engine/Source/Programs/UnrealBuildTool/Platform/Linux/LinuxToolChain.cs
+++ b/Engine/Source/Programs/UnrealBuildTool/Platform/Linux/LinuxToolChain.cs
@@ -478,9 +478,9 @@ namespace UnrealBuildTool
 
 			Result += " -Wall -Werror";
 
-			if (Architecture.StartsWith("arm") || Architecture.StartsWith("aarch64"))
+			if (!CompileEnvironment.Architecture.StartsWith("x86_64") && !CompileEnvironment.Architecture.StartsWith("i686"))
 			{
-				Result += " -funwind-tables";           // generate unwind tables as they are needed for arm
+				Result += " -funwind-tables";               // generate unwind tables as they are needed for backtrace (on x86(64) they are generated implicitly)
 			}
 
 			Result += " -Wsequence-point";              // additional warning not normally included in Wall: warns if order of operations is ambigious
