Unshelved from pending changelist '4238012':

Xcode project generator improvements

- Per-project precompiled header that wraps UnrealEd.h with #ifdef __cplusplus to allow Xcode to compile the pch in ObjC mode. Later we could replace UnrealEd.h with some other header file for non-editor targets
- Moved commands that disable compile warnings from MacToolChain's GetCompileArguments_Global to ApplePlatformCompilerPreSetup.h. Thanks to this we can have all the Xcode recommended warnings enabled in the project, but still allow Clang to index our code without reporting warnings
- Few more minor changes to fix Xcode's project validation and indexing warnings

Also, unify compile warning flags across all Apple platforms.

diff --git a/Engine/Source/Programs/UnrealBuildTool/Platform/IOS/IOSToolChain.cs b/Engine/Source/Programs/UnrealBuildTool/Platform/IOS/IOSToolChain.cs
--- a/Engine/Source/Programs/UnrealBuildTool/Platform/IOS/IOSToolChain.cs
+++ b/Engine/Source/Programs/UnrealBuildTool/Platform/IOS/IOSToolChain.cs
@@ -274,39 +274,6 @@ namespace UnrealBuildTool
 				Result += " -Wundef" + (CompileEnvironment.bUndefinedIdentifierWarningsAsErrors ? "" : " -Wno-error=undef");
 			}
 
-			Result += " -Wno-unused-variable";
-			Result += " -Wno-unused-value";
-			// this will hide the warnings about static functions in headers that aren't used in every single .cpp file
-			Result += " -Wno-unused-function";
-			// this hides the "enumeration value 'XXXXX' not handled in switch [-Wswitch]" warnings - we should maybe remove this at some point and add UE_LOG(, Fatal, ) to default cases
-			Result += " -Wno-switch";
-			// this hides the "warning : comparison of unsigned expression < 0 is always false" type warnings due to constant comparisons, which are possible with template arguments
-			Result += " -Wno-tautological-compare";
-			//This will prevent the issue of warnings for unused private variables.
-			Result += " -Wno-unused-private-field";
-			// needed to suppress warnings about using offsetof on non-POD types.
-			Result += " -Wno-invalid-offsetof";
-			// we use this feature to allow static FNames.
-			Result += " -Wno-gnu-string-literal-operator-template";
-
-			if (Settings.Value.IOSSDKVersionFloat >= 9.0)
-			{
-				Result += " -Wno-inconsistent-missing-override"; // too many missing overrides...
-				Result += " -Wno-unused-local-typedef"; // PhysX has some, hard to remove
-			}
-			
-			// fix for Xcode 8.3 enabling nonportable include checks, but p4 has some invalid cases in it
-			if (Settings.Value.IOSSDKVersionFloat >= 10.3)
-			{
-				Result += " -Wno-nonportable-include-path";
-			}
-
-			// fix for Xcode 8.3 enabling nonportable include checks, but p4 has some invalid cases in it
-			if (Settings.Value.IOSSDKVersionFloat >= 10.3)
-			{
-				Result += " -Wno-nonportable-include-path";
-			}
-
 			if (IsBitcodeCompilingEnabled(CompileEnvironment.Configuration))
 			{
 				Result += " -fembed-bitcode";
diff --git a/Engine/Source/Programs/UnrealBuildTool/Platform/Mac/MacToolChain.cs b/Engine/Source/Programs/UnrealBuildTool/Platform/Mac/MacToolChain.cs
--- a/Engine/Source/Programs/UnrealBuildTool/Platform/Mac/MacToolChain.cs
+++ b/Engine/Source/Programs/UnrealBuildTool/Platform/Mac/MacToolChain.cs
@@ -193,29 +193,6 @@ namespace UnrealBuildTool
 			Result += " -Wdelete-non-virtual-dtor";
 			//Result += " -Wsign-compare"; // fed up of not seeing the signed/unsigned warnings we get on Windows - lets enable them here too.
 
-			Result += " -Wno-unused-variable";
-			Result += " -Wno-unused-value";
-			// This will hide the warnings about static functions in headers that aren't used in every single .cpp file
-			Result += " -Wno-unused-function";
-			// This hides the "enumeration value 'XXXXX' not handled in switch [-Wswitch]" warnings - we should maybe remove this at some point and add UE_LOG(, Fatal, ) to default cases
-			Result += " -Wno-switch";
-			// This hides the "warning : comparison of unsigned expression < 0 is always false" type warnings due to constant comparisons, which are possible with template arguments
-			Result += " -Wno-tautological-compare";
-			// This will prevent the issue of warnings for unused private variables.
-			Result += " -Wno-unused-private-field";
-			// needed to suppress warnings about using offsetof on non-POD types.
-			Result += " -Wno-invalid-offsetof";
-			// we use this feature to allow static FNames.
-			Result += " -Wno-gnu-string-literal-operator-template";
-			// Needed for Alembic third party lib
-			Result += " -Wno-deprecated-register";
-
-			if (Settings.MacOSSDKVersionFloat < 10.9f && Settings.MacOSSDKVersionFloat >= 10.11f)
-			{
-				Result += " -Wno-inconsistent-missing-override"; // too many missing overrides...
-				Result += " -Wno-unused-local-typedef"; // PhysX has some, hard to remove
-			}
-
 			if (CompileEnvironment.bEnableShadowVariableWarnings)
 			{
 				Result += " -Wshadow" + (CompileEnvironment.bShadowVariableWarningsAsErrors ? "" : " -Wno-error=shadow");
@@ -226,11 +203,6 @@ namespace UnrealBuildTool
 				Result += " -Wundef" + (CompileEnvironment.bUndefinedIdentifierWarningsAsErrors ? "" : " -Wno-error=undef");
 			}
 
-			// @todo: Remove these three when the code is fixed and they're no longer needed
-			Result += " -Wno-logical-op-parentheses";
-			Result += " -Wno-unknown-pragmas";
-			Result += " -Wno-undefined-var-template";
-
 			Result += " -c";
 
 			Result += " -arch x86_64";
diff --git a/Engine/Source/Programs/UnrealBuildTool/ProjectFiles/Xcode/XcodeProject.cs b/Engine/Source/Programs/UnrealBuildTool/ProjectFiles/Xcode/XcodeProject.cs
--- a/Engine/Source/Programs/UnrealBuildTool/ProjectFiles/Xcode/XcodeProject.cs
+++ b/Engine/Source/Programs/UnrealBuildTool/ProjectFiles/Xcode/XcodeProject.cs
@@ -158,8 +158,10 @@ namespace UnrealBuildTool
 			{
 				foreach (string PlatformName in Enum.GetNames(typeof(UnrealTargetPlatform)))
 				{
-					string AltName = PlatformName == "Win32" || PlatformName == "Win64" ? "windows" : PlatformName.ToLower();
-					if ((SourceFile.Reference.FullName.ToLower().Contains("/" + PlatformName.ToLower() + "/") || SourceFile.Reference.FullName.ToLower().Contains("/" + AltName + "/"))
+					string PlatformNameLowerCase = PlatformName.ToLower();
+					string AltName = PlatformNameLowerCase == "win32" || PlatformNameLowerCase== "win64" ? "windows" : PlatformNameLowerCase;
+					string FullNameLowerCase = SourceFile.Reference.FullName.ToLower();
+					if ((FullNameLowerCase.Contains("/" + PlatformNameLowerCase + "/") || FullNameLowerCase.ToLower().Contains("/" + AltName + "/") || FullNameLowerCase.Contains("/engine/build/"))
 						&& PlatformName != "Mac" && PlatformName != "IOS" && PlatformName != "TVOS")
 					{
 						// Build phase is used for indexing only and indexing currently works only with files that can be compiled for Mac, so skip files for other platforms
@@ -513,7 +515,7 @@ namespace UnrealBuildTool
 			Content.Append("\t\t" + ProjectGuid + " /* Project object */ = {" + ProjectFileGenerator.NewLine);
 			Content.Append("\t\t\tisa = PBXProject;" + ProjectFileGenerator.NewLine);
 			Content.Append("\t\t\tattributes = {" + ProjectFileGenerator.NewLine);
-			Content.Append("\t\t\t\tLastUpgradeCheck = 0900;" + ProjectFileGenerator.NewLine);
+			Content.Append("\t\t\t\tLastUpgradeCheck = 0940;" + ProjectFileGenerator.NewLine);
 			Content.Append("\t\t\t\tORGANIZATIONNAME = \"Epic Games, Inc.\";" + ProjectFileGenerator.NewLine);
             Content.Append("\t\t\t\tTargetAttributes = {" + ProjectFileGenerator.NewLine);
             Content.Append("\t\t\t\t\t" + TargetGuid + " = {" + ProjectFileGenerator.NewLine);
@@ -599,10 +601,11 @@ namespace UnrealBuildTool
 			Content.Append("\t\t\t\tGCC_PREPROCESSOR_DEFINITIONS = (" + ProjectFileGenerator.NewLine);
 			foreach (string Definition in IntelliSensePreprocessorDefinitions)
 			{
-				Content.Append("\t\t\t\t\t\"" + Definition.Replace("\"", "").Replace("\\", "") + "\"," + ProjectFileGenerator.NewLine);
+				Content.Append("\t\t\t\t\t\"" + Definition.Replace("\"", "").Replace("\\", "").Replace("()", "") + "\"," + ProjectFileGenerator.NewLine);
 			}
-			Content.Append("\t\t\t\t\t\"__INTELLISENSE__\"," + ProjectFileGenerator.NewLine);
+			Content.Append("\t\t\t\t\t__INTELLISENSE__," + ProjectFileGenerator.NewLine);
 			Content.Append("\t\t\t\t\t\"MONOLITHIC_BUILD=1\"," + ProjectFileGenerator.NewLine);
+			Content.Append("\t\t\t\t\tSUPPRESS_MONOLITHIC_HEADER_WARNINGS," + ProjectFileGenerator.NewLine);
 			Content.Append("\t\t\t\t);" + ProjectFileGenerator.NewLine);
 
 			Content.Append("\t\t\t\tHEADER_SEARCH_PATHS = (" + ProjectFileGenerator.NewLine);
@@ -621,13 +624,35 @@ namespace UnrealBuildTool
 			}
 			Content.Append("\t\t\t\t);" + ProjectFileGenerator.NewLine);
 
-			if (ConfigName == "Debug")
-			{
-				Content.Append("\t\t\t\tONLY_ACTIVE_ARCH = YES;" + ProjectFileGenerator.NewLine);
-				Content.Append("\t\t\t\tENABLE_TESTABILITY = YES;" + ProjectFileGenerator.NewLine);
-			}
+			Content.Append("\t\t\t\tONLY_ACTIVE_ARCH = YES;" + ProjectFileGenerator.NewLine);
+			Content.Append("\t\t\t\tENABLE_TESTABILITY = YES;" + ProjectFileGenerator.NewLine);
 			Content.Append("\t\t\t\tALWAYS_SEARCH_USER_PATHS = NO;" + ProjectFileGenerator.NewLine);
 			Content.Append("\t\t\t\tCLANG_CXX_LANGUAGE_STANDARD = \"c++14\";" + ProjectFileGenerator.NewLine);
+			Content.Append("\t\t\t\tCLANG_WARN_BLOCK_CAPTURE_AUTORELEASING = YES;" + ProjectFileGenerator.NewLine);
+			Content.Append("\t\t\t\tCLANG_WARN_BOOL_CONVERSION = YES;" + ProjectFileGenerator.NewLine);
+			Content.Append("\t\t\t\tCLANG_WARN_COMMA = YES;" + ProjectFileGenerator.NewLine);
+			Content.Append("\t\t\t\tCLANG_WARN_CONSTANT_CONVERSION = YES;" + ProjectFileGenerator.NewLine);
+			Content.Append("\t\t\t\tCLANG_WARN_DEPRECATED_OBJC_IMPLEMENTATIONS = YES;" + ProjectFileGenerator.NewLine);
+			Content.Append("\t\t\t\tCLANG_WARN_EMPTY_BODY = YES;" + ProjectFileGenerator.NewLine);
+			Content.Append("\t\t\t\tCLANG_WARN_ENUM_CONVERSION = YES;" + ProjectFileGenerator.NewLine);
+			Content.Append("\t\t\t\tCLANG_WARN_INFINITE_RECURSION = YES;" + ProjectFileGenerator.NewLine);
+			Content.Append("\t\t\t\tCLANG_WARN_INT_CONVERSION = YES;" + ProjectFileGenerator.NewLine);
+			Content.Append("\t\t\t\tCLANG_WARN_NON_LITERAL_NULL_CONVERSION = YES;" + ProjectFileGenerator.NewLine);
+			Content.Append("\t\t\t\tCLANG_WARN_OBJC_IMPLICIT_RETAIN_SELF = YES;" + ProjectFileGenerator.NewLine);
+			Content.Append("\t\t\t\tCLANG_WARN_OBJC_LITERAL_CONVERSION = YES;" + ProjectFileGenerator.NewLine);
+			Content.Append("\t\t\t\tCLANG_WARN_RANGE_LOOP_ANALYSIS = YES;" + ProjectFileGenerator.NewLine);
+			Content.Append("\t\t\t\tCLANG_WARN_STRICT_PROTOTYPES = YES;" + ProjectFileGenerator.NewLine);
+			Content.Append("\t\t\t\tCLANG_WARN_SUSPICIOUS_MOVE = YES;" + ProjectFileGenerator.NewLine);
+			Content.Append("\t\t\t\tCLANG_WARN_UNREACHABLE_CODE = YES;" + ProjectFileGenerator.NewLine);
+			Content.Append("\t\t\t\tCLANG_WARN__DUPLICATE_METHOD_MATCH = YES;" + ProjectFileGenerator.NewLine);
+			Content.Append("\t\t\t\tENABLE_STRICT_OBJC_MSGSEND = YES;" + ProjectFileGenerator.NewLine);
+			Content.Append("\t\t\t\tGCC_NO_COMMON_BLOCKS = YES;" + ProjectFileGenerator.NewLine);
+			Content.Append("\t\t\t\tGCC_WARN_64_TO_32_BIT_CONVERSION = YES;" + ProjectFileGenerator.NewLine);
+			Content.Append("\t\t\t\tGCC_WARN_ABOUT_RETURN_TYPE = YES;" + ProjectFileGenerator.NewLine);
+			Content.Append("\t\t\t\tGCC_WARN_UNDECLARED_SELECTOR = YES;" + ProjectFileGenerator.NewLine);
+			Content.Append("\t\t\t\tGCC_WARN_UNINITIALIZED_AUTOS = YES;" + ProjectFileGenerator.NewLine);
+			Content.Append("\t\t\t\tGCC_WARN_UNUSED_FUNCTION = YES;" + ProjectFileGenerator.NewLine);
+			Content.Append("\t\t\t\tGCC_WARN_UNUSED_VARIABLE = YES;" + ProjectFileGenerator.NewLine);
 			Content.Append("\t\t\t\tGCC_ENABLE_CPP_RTTI = NO;" + ProjectFileGenerator.NewLine);
 			Content.Append("\t\t\t\tGCC_WARN_CHECK_SWITCH_STATEMENTS = NO;" + ProjectFileGenerator.NewLine);
 			Content.Append("\t\t\t\tUSE_HEADERMAP = NO;" + ProjectFileGenerator.NewLine);
@@ -650,6 +675,7 @@ namespace UnrealBuildTool
 			Content.Append("\t\t" + ConfigGuid + " /* \"" + Config.DisplayName + "\" */ = {" + ProjectFileGenerator.NewLine);
 			Content.Append("\t\t\tisa = XCBuildConfiguration;" + ProjectFileGenerator.NewLine);
 			Content.Append("\t\t\tbuildSettings = {" + ProjectFileGenerator.NewLine);
+			Content.Append("\t\t\t\tCLANG_ENABLE_OBJC_WEAK = YES;" + ProjectFileGenerator.NewLine);
 
             string UE4Dir = ConvertPath(Path.GetFullPath(Directory.GetCurrentDirectory() + "../../.."));
 			string MacExecutableDir = ConvertPath(Config.MacExecutablePath.Directory.FullName);
@@ -679,6 +705,7 @@ namespace UnrealBuildTool
                 string TEAM_TVOS = "";
                 string IOS_CERT = "iPhone Developer";
                 string TVOS_CERT = "iPhone Developer";
+				string BundleIdentifier = "";
                 if (InstalledPlatformInfo.IsValidPlatform(UnrealTargetPlatform.IOS, EProjectType.Code))
                 {
 					IOSPlatform IOSPlatform = ((IOSPlatform)UEBuildPlatform.GetBuildPlatform(UnrealTargetPlatform.IOS));
@@ -686,6 +713,7 @@ namespace UnrealBuildTool
 					IOSProvisioningData ProvisioningData = IOSPlatform.ReadProvisioningData(ProjectSettings);
 					IOSRunTimeVersion = ProjectSettings.RuntimeVersion;
 					IOSRunTimeDevices = ProjectSettings.RuntimeDevices;
+					BundleIdentifier = ProjectSettings.BundleIdentifier;
 					ValidArchs += " arm64 armv7 armv7s";
 					SupportedPlatforms += " iphoneos";
 					bAutomaticSigning = ProjectSettings.bAutomaticSigning;
@@ -704,6 +732,7 @@ namespace UnrealBuildTool
 					TVOSProvisioningData ProvisioningData = TVOSPlatform.ReadProvisioningData(ProjectSettings);
 					TVOSRunTimeVersion = ProjectSettings.RuntimeVersion;
 					TVOSRunTimeDevices = ProjectSettings.RuntimeDevices;
+					BundleIdentifier = ProjectSettings.BundleIdentifier;
 					if (ValidArchs == "x86_64")
 					{
 						ValidArchs += " arm64 armv7 armv7s";
@@ -723,6 +752,7 @@ namespace UnrealBuildTool
 				if (IOSRunTimeVersion != null)
 				{
 					Content.Append("\t\t\t\tIPHONEOS_DEPLOYMENT_TARGET = " + IOSRunTimeVersion + ";" + ProjectFileGenerator.NewLine);
+					Content.Append("\t\t\t\tPRODUCT_BUNDLE_IDENTIFIER = " + BundleIdentifier + ";" + ProjectFileGenerator.NewLine);
 					Content.Append("\t\t\t\t\"PRODUCT_NAME[sdk=iphoneos*]\" = \"" + Config.BuildTarget + "\";" + ProjectFileGenerator.NewLine); // @todo: change to Path.GetFileName(Config.IOSExecutablePath) when we stop using payload
 					Content.Append("\t\t\t\t\"TARGETED_DEVICE_FAMILY[sdk=iphoneos*]\" = \"" + IOSRunTimeDevices + "\";" + ProjectFileGenerator.NewLine);
                     Content.Append("\t\t\t\t\"SDKROOT[sdk=iphoneos]\" = iphoneos;" + ProjectFileGenerator.NewLine);
@@ -739,6 +769,7 @@ namespace UnrealBuildTool
                 if (TVOSRunTimeVersion != null)
 				{
 					Content.Append("\t\t\t\tTVOS_DEPLOYMENT_TARGET = " + TVOSRunTimeVersion + ";" + ProjectFileGenerator.NewLine);
+					Content.Append("\t\t\t\tPRODUCT_BUNDLE_IDENTIFIER = " + BundleIdentifier + ";" + ProjectFileGenerator.NewLine);
 					Content.Append("\t\t\t\t\"PRODUCT_NAME[sdk=appletvos*]\" = \"" + Config.BuildTarget + "\";" + ProjectFileGenerator.NewLine); // @todo: change to Path.GetFileName(Config.TVOSExecutablePath) when we stop using payload
 					Content.Append("\t\t\t\t\"TARGETED_DEVICE_FAMILY[sdk=appletvos*]\" = \"" + TVOSRunTimeDevices + "\";" + ProjectFileGenerator.NewLine);
                     Content.Append("\t\t\t\t\"SDKROOT[sdk=appletvos]\" = appletvos;" + ProjectFileGenerator.NewLine);
@@ -926,7 +957,7 @@ namespace UnrealBuildTool
                 Content.Append("\t\t\t\tSDKROOT = macosx;" + ProjectFileGenerator.NewLine);
             }
             Content.Append("\t\t\t\tGCC_PRECOMPILE_PREFIX_HEADER = YES;" + ProjectFileGenerator.NewLine);
-			Content.Append("\t\t\t\tGCC_PREFIX_HEADER = \"" + UE4Dir + "/Engine/Source/Editor/UnrealEd/Public/UnrealEd.h\";" + ProjectFileGenerator.NewLine);
+			Content.Append("\t\t\t\tGCC_PREFIX_HEADER = \"" + ProjectFilePath.ChangeExtension("h").FullName + "\";" + ProjectFileGenerator.NewLine);
 			Content.Append("\t\t\t};" + ProjectFileGenerator.NewLine);
 			Content.Append("\t\t\tname = \"" + Config.DisplayName + "\";" + ProjectFileGenerator.NewLine);
 			Content.Append("\t\t};" + ProjectFileGenerator.NewLine);
@@ -946,6 +977,7 @@ namespace UnrealBuildTool
 			Content.Append("\t\t" + ConfigGuid + " /* \"" + Config.DisplayName + "\" */ = {" + ProjectFileGenerator.NewLine);
 			Content.Append("\t\t\tisa = XCBuildConfiguration;" + ProjectFileGenerator.NewLine);
 			Content.Append("\t\t\tbuildSettings = {" + ProjectFileGenerator.NewLine);
+			Content.Append("\t\t\t\tCLANG_ENABLE_OBJC_WEAK = YES;" + ProjectFileGenerator.NewLine);
 			if (bMacOnly)
 			{
 				Content.Append("\t\t\t\tVALID_ARCHS = \"x86_64\";" + ProjectFileGenerator.NewLine);
@@ -1442,7 +1474,7 @@ namespace UnrealBuildTool
 			ProjectFileContent.Append("\tarchiveVersion = 1;" + ProjectFileGenerator.NewLine);
 			ProjectFileContent.Append("\tclasses = {" + ProjectFileGenerator.NewLine);
 			ProjectFileContent.Append("\t};" + ProjectFileGenerator.NewLine);
-			ProjectFileContent.Append("\tobjectVersion = 46;" + ProjectFileGenerator.NewLine);
+			ProjectFileContent.Append("\tobjectVersion = 48;" + ProjectFileGenerator.NewLine);
 			ProjectFileContent.Append("\tobjects = {" + ProjectFileGenerator.NewLine + ProjectFileGenerator.NewLine);
 
 			AppendBuildFileSection(ProjectFileContent, PBXBuildFileSection);
@@ -1470,6 +1502,15 @@ namespace UnrealBuildTool
 			{
 				FileReference PBXProjFilePath = ProjectFilePath + "/project.pbxproj";
 				bSuccess = ProjectFileGenerator.WriteFileIfChanged(PBXProjFilePath.FullName, ProjectFileContent.ToString(), new UTF8Encoding());
+				if (bSuccess)
+				{
+					string UE4Dir = ConvertPath(Path.GetFullPath(Directory.GetCurrentDirectory() + "../../.."));
+					StringBuilder PCHFileContent = new StringBuilder();
+					PCHFileContent.Append("#ifdef __cplusplus" + ProjectFileGenerator.NewLine);
+					PCHFileContent.Append("#include \"" + UE4Dir + "/Engine/Source/Editor/UnrealEd/Public/UnrealEd.h\"" + ProjectFileGenerator.NewLine);
+					PCHFileContent.Append("#endif // __cplusplus" + ProjectFileGenerator.NewLine);
+					bSuccess = ProjectFileGenerator.WriteFileIfChanged(ProjectFilePath.ChangeExtension("h").FullName, PCHFileContent.ToString(), new UTF8Encoding());
+				}
 			}
 
 			if (bSuccess)
diff --git a/Engine/Source/Runtime/Core/Public/Apple/ApplePlatformCompilerPreSetup.h b/Engine/Source/Runtime/Core/Public/Apple/ApplePlatformCompilerPreSetup.h
--- a/Engine/Source/Runtime/Core/Public/Apple/ApplePlatformCompilerPreSetup.h
+++ b/Engine/Source/Runtime/Core/Public/Apple/ApplePlatformCompilerPreSetup.h
@@ -33,6 +33,20 @@
 #pragma clang diagnostic ignored "-Wreserved-user-defined-literal"
 #pragma clang diagnostic ignored "-Wdelete-non-virtual-dtor"
 #pragma clang diagnostic ignored "-Wlogical-op-parentheses"
+#pragma clang diagnostic ignored "-Wunused-variable"
+#pragma clang diagnostic ignored "-Wunused-value"
+#pragma clang diagnostic ignored "-Wunused-function" // This will hide the warnings about static functions in headers that aren't used in every single .cpp file
+#pragma clang diagnostic ignored "-Wswitch" // This hides the "enumeration value 'XXXXX' not handled in switch [-Wswitch]" warnings - we should maybe remove this at some point and add UE_LOG(, Fatal, ) to default cases
+#pragma clang diagnostic ignored "-Wtautological-compare" // This hides the "warning : comparison of unsigned expression < 0 is always false" type warnings due to constant comparisons, which are possible with template arguments
+#pragma clang diagnostic ignored "-Wunused-private-field" // This will prevent the issue of warnings for unused private variables.
+#pragma clang diagnostic ignored "-Winvalid-offsetof" // needed to suppress warnings about using offsetof on non-POD types.
+#pragma clang diagnostic ignored "-Wgnu-string-literal-operator-template" // we use this feature to allow static FNames.
+#pragma clang diagnostic ignored "-Wdeprecated-register" // Needed for Alembic third party lib
+#pragma clang diagnostic ignored "-Winconsistent-missing-override" // too many missing overrides...
+#pragma clang diagnostic ignored "-Wunused-local-typedef" // PhysX has some, hard to remove
+#pragma clang diagnostic ignored "-Wlogical-op-parentheses"
+#pragma clang diagnostic ignored "-Wunknown-pragmas"
+#pragma clang diagnostic ignored "-Wundefined-var-template"
 
 // Apple LLVM 9.1.0 (Xcode 9.3)
 #if (__clang_major__ > 9) || (__clang_major__ == 9 && __clang_minor__ >= 1)
