Linux: Dont disable ICU by default on Servers

diff --git a/Engine/Source/Programs/UnrealBuildTool/Platform/Linux/UEBuildLinux.cs b/Engine/Source/Programs/UnrealBuildTool/Platform/Linux/UEBuildLinux.cs
--- a/Engine/Source/Programs/UnrealBuildTool/Platform/Linux/UEBuildLinux.cs
+++ b/Engine/Source/Programs/UnrealBuildTool/Platform/Linux/UEBuildLinux.cs
@@ -156,8 +156,8 @@ namespace UnrealBuildTool
 			// check if OS update invalidated our build
 			Target.bCheckSystemHeadersForModification = (BuildHostPlatform.Current.Platform == UnrealTargetPlatform.Linux);
 
-			// At the moment ICU has not been compiled for AArch64 and i686. Also, localization isn't needed on servers by default, and ICU is pretty heavy
-			if (Target.Architecture.StartsWith("aarch64") || Target.Architecture.StartsWith("i686") || Target.Type == TargetType.Server)
+			// At the moment ICU has not been compiled for AArch64 and i686.
+			if (Target.Architecture.StartsWith("aarch64") || Target.Architecture.StartsWith("i686"))
 			{
 				Target.bCompileICU = false;
 			}
