Linux: Disable XGE builds as it appears to be lower casing folders when the build platform is Windows

diff --git a/Engine/Source/Programs/UnrealBuildTool/Platform/Linux/UEBuildLinux.cs b/Engine/Source/Programs/UnrealBuildTool/Platform/Linux/UEBuildLinux.cs
--- a/Engine/Source/Programs/UnrealBuildTool/Platform/Linux/UEBuildLinux.cs
+++ b/Engine/Source/Programs/UnrealBuildTool/Platform/Linux/UEBuildLinux.cs
@@ -184,7 +184,8 @@ namespace UnrealBuildTool
 		{
 			// [RCL] 2018-05-02: disabling XGE even during a native build because the support is not ready and you can have mysterious build failures when ib_console is installed.
 			// [RCL] 2018-07-10: enabling XGE for Windows to see if the crash from 2016 still persists. Please disable if you see spurious build errors that don't repro without XGE
-			return BuildHostPlatform.Current.Platform == UnrealTargetPlatform.Win64;
+			// [bschaefer] 2018-08-24: disabling XGE due to a bug where XGE seems to be lower casing folders names that are headers ie. misc/Header.h vs Misc/Header.h
+			return false;//BuildHostPlatform.Current.Platform == UnrealTargetPlatform.Win64;
 		}
 
 		public override bool CanUseParallelExecutor()
