Speculative fix for UE-63365

diff --git a/Engine/Source/Programs/UnrealBuildTool/ProjectFiles/Xcode/XcodeProject.cs b/Engine/Source/Programs/UnrealBuildTool/ProjectFiles/Xcode/XcodeProject.cs
--- a/Engine/Source/Programs/UnrealBuildTool/ProjectFiles/Xcode/XcodeProject.cs
+++ b/Engine/Source/Programs/UnrealBuildTool/ProjectFiles/Xcode/XcodeProject.cs
@@ -731,7 +731,10 @@ namespace UnrealBuildTool
 				if (IOSRunTimeVersion != null)
 				{
 					Content.Append("\t\t\t\tIPHONEOS_DEPLOYMENT_TARGET = " + IOSRunTimeVersion + ";" + ProjectFileGenerator.NewLine);
-					Content.Append("\t\t\t\tPRODUCT_BUNDLE_IDENTIFIER = " + BundleIdentifier + ";" + ProjectFileGenerator.NewLine);
+					if (!XcodeProjectFileGenerator.bGeneratingRunIOSProject && !XcodeProjectFileGenerator.bGeneratingRunTVOSProject)
+					{
+						Content.Append("\t\t\t\tPRODUCT_BUNDLE_IDENTIFIER = " + BundleIdentifier + ";" + ProjectFileGenerator.NewLine);
+					}
 					Content.Append("\t\t\t\t\"PRODUCT_NAME[sdk=iphoneos*]\" = \"" + Config.BuildTarget + "\";" + ProjectFileGenerator.NewLine); // @todo: change to Path.GetFileName(Config.IOSExecutablePath) when we stop using payload
 					Content.Append("\t\t\t\t\"TARGETED_DEVICE_FAMILY[sdk=iphoneos*]\" = \"" + IOSRunTimeDevices + "\";" + ProjectFileGenerator.NewLine);
                     Content.Append("\t\t\t\t\"SDKROOT[sdk=iphoneos]\" = iphoneos;" + ProjectFileGenerator.NewLine);
@@ -748,7 +751,10 @@ namespace UnrealBuildTool
                 if (TVOSRunTimeVersion != null)
 				{
 					Content.Append("\t\t\t\tTVOS_DEPLOYMENT_TARGET = " + TVOSRunTimeVersion + ";" + ProjectFileGenerator.NewLine);
-					Content.Append("\t\t\t\tPRODUCT_BUNDLE_IDENTIFIER = " + BundleIdentifier + ";" + ProjectFileGenerator.NewLine);
+					if (!XcodeProjectFileGenerator.bGeneratingRunIOSProject && !XcodeProjectFileGenerator.bGeneratingRunTVOSProject)
+					{
+						Content.Append("\t\t\t\tPRODUCT_BUNDLE_IDENTIFIER = " + BundleIdentifier + ";" + ProjectFileGenerator.NewLine);
+					}
 					Content.Append("\t\t\t\t\"PRODUCT_NAME[sdk=appletvos*]\" = \"" + Config.BuildTarget + "\";" + ProjectFileGenerator.NewLine); // @todo: change to Path.GetFileName(Config.TVOSExecutablePath) when we stop using payload
 					Content.Append("\t\t\t\t\"TARGETED_DEVICE_FAMILY[sdk=appletvos*]\" = \"" + TVOSRunTimeDevices + "\";" + ProjectFileGenerator.NewLine);
                     Content.Append("\t\t\t\t\"SDKROOT[sdk=appletvos]\" = appletvos;" + ProjectFileGenerator.NewLine);
