Disable nonportable-include-path warning in iOS toolchain to allow incorrect case in paths to headers passed using -include

diff --git a/Engine/Source/Programs/UnrealBuildTool/Platform/IOS/IOSToolChain.cs b/Engine/Source/Programs/UnrealBuildTool/Platform/IOS/IOSToolChain.cs
--- a/Engine/Source/Programs/UnrealBuildTool/Platform/IOS/IOSToolChain.cs
+++ b/Engine/Source/Programs/UnrealBuildTool/Platform/IOS/IOSToolChain.cs
@@ -275,6 +275,12 @@ namespace UnrealBuildTool
 				Result += " -Wundef" + (CompileEnvironment.bUndefinedIdentifierWarningsAsErrors ? "" : " -Wno-error=undef");
 			}
 
+			// fix for Xcode 8.3 enabling nonportable include checks, but p4 has some invalid cases in it
+			if (Settings.Value.IOSSDKVersionFloat >= 10.3)
+			{
+				Result += " -Wno-nonportable-include-path";
+			}
+
 			if (IsBitcodeCompilingEnabled(CompileEnvironment.Configuration))
 			{
 				Result += " -fembed-bitcode";
