UBT: Check the registry for the installation location of XGE, rather than just searching the system path.

diff --git a/Engine/Source/Programs/AutomationTool/AutomationUtils/UE4Build.cs b/Engine/Source/Programs/AutomationTool/AutomationUtils/UE4Build.cs
--- a/Engine/Source/Programs/AutomationTool/AutomationUtils/UE4Build.cs
+++ b/Engine/Source/Programs/AutomationTool/AutomationUtils/UE4Build.cs
@@ -22,8 +22,6 @@ namespace AutomationTool
 	{
 		private BuildCommand OwnerCommand;
 
-		private String XGEConsoleExePath = "";
-
 		public bool HasBuildProduct(string InFile)
 		{
 			string File = CommandUtils.CombinePaths(InFile);
@@ -587,30 +585,6 @@ namespace AutomationTool
 			return Result;
 		}
 
-		public string XGEConsoleExe()
-		{
-			if (string.IsNullOrEmpty(XGEConsoleExePath))
-			{
-				string[] PathDirs = Environment.GetEnvironmentVariable("PATH").Split(Path.PathSeparator);
-				foreach (string PathDir in PathDirs)
-				{
-					try
-					{
-						string FullPath = Path.Combine(PathDir, "xgConsole" + Platform.GetExeExtension(BuildHostPlatform.Current.Platform));
-						if (CommandUtils.FileExists(FullPath))
-						{
-							XGEConsoleExePath = FullPath;
-							break;
-						}
-					}
-					catch
-					{
-					}
-				}
-			}
-			return XGEConsoleExePath;
-		}
-
 		public bool ProcessXGEItems(List<XGEItem> Actions, string XGETool, string Args, string TaskFilePath, bool ShowProgress)
 		{
 			TelemetryStopwatch CombineXGEStopwatch = new TelemetryStopwatch("CombineXGEItemFiles.{0}", Path.GetFileNameWithoutExtension(XGETool));
@@ -1144,9 +1118,9 @@ namespace AutomationTool
 			bool bForceNonUnity = ParseParam("ForceNonUnity") || InForceNonUnity;
 			bool bForceUnity = ParseParam("ForceUnity") || InForceUnity;
 			bool bForceDebugInfo = ParseParam("ForceDebugInfo");
-			string XGEConsole = XGEConsoleExe();
+			string XGEConsole = null;
 			bool bDisableXGE = ParseParam("NoXGE") || InForceNoXGE;
-			bool bCanUseXGE = !bDisableXGE && !String.IsNullOrEmpty(XGEConsole);
+			bool bCanUseXGE = !bDisableXGE && PlatformExports.TryGetXgConsoleExecutable(out XGEConsole);
 
 			// only run ParallelExecutor if not running XGE (and we've requested ParallelExecutor and it exists)
 			bool bCanUseParallelExecutor = !bCanUseXGE && InUseParallelExecutor && (HostPlatform.Current.HostEditorPlatform == UnrealTargetPlatform.Win64);
@@ -1238,7 +1212,7 @@ namespace AutomationTool
 				string Args = null;
 				if (bCanUseXGE) 
 				{
-					XGETool = XGEConsoleExePath;
+					XGETool = XGEConsole;
 					Args = "\"" + TaskFilePath + "\" /Rebuild /NoLogo /ShowAgent /ShowTime";
 					if (ParseParam("StopOnErrors"))
 					{
diff --git a/Engine/Source/Programs/UnrealBuildTool/Executors/XGE.cs b/Engine/Source/Programs/UnrealBuildTool/Executors/XGE.cs
--- a/Engine/Source/Programs/UnrealBuildTool/Executors/XGE.cs
+++ b/Engine/Source/Programs/UnrealBuildTool/Executors/XGE.cs
@@ -47,8 +47,24 @@ namespace UnrealBuildTool
 			get { return "XGE"; }
 		}
 
-		public static bool IsAvailable()
+		public static bool TryGetXgConsoleExecutable(out string OutXgConsoleExe)
 		{
+			// Try to get the path from the registry
+			if(BuildHostPlatform.Current.Platform == UnrealTargetPlatform.Win64)
+			{
+				string XgConsoleExe;
+				if(TryGetXgConsoleExecutableFromRegistry(RegistryView.Registry32, out XgConsoleExe))
+				{
+					OutXgConsoleExe = XgConsoleExe;
+					return true;
+				}
+				if(TryGetXgConsoleExecutableFromRegistry(RegistryView.Registry64, out XgConsoleExe))
+				{
+					OutXgConsoleExe = XgConsoleExe;
+					return true;
+				}
+			}
+
 			// Get the name of the XgConsole executable.
 			string XgConsole = "xgConsole";
 			if (BuildHostPlatform.Current.Platform == UnrealTargetPlatform.Win64)
@@ -69,6 +85,7 @@ namespace UnrealBuildTool
 					string PotentialPath = Path.Combine(SearchPath, XgConsole);
 					if(File.Exists(PotentialPath))
 					{
+						OutXgConsoleExe = PotentialPath;
 						return true;
 					}
 				}
@@ -77,9 +94,50 @@ namespace UnrealBuildTool
 					// PATH variable may contain illegal characters; just ignore them.
 				}
 			}
+
+			OutXgConsoleExe = null;
 			return false;
 		}
 
+		private static bool TryGetXgConsoleExecutableFromRegistry(RegistryView View, out string OutXgConsoleExe)
+		{
+			try
+			{
+				using(RegistryKey BaseKey = RegistryKey.OpenBaseKey(RegistryHive.LocalMachine, View))
+				{
+					using (RegistryKey Key = BaseKey.OpenSubKey("SOFTWARE\\Xoreax\\IncrediBuild\\Builder", false))
+					{
+						if(Key != null)
+						{
+							string Folder = Key.GetValue("Folder", null) as string;
+							if(!String.IsNullOrEmpty(Folder))
+							{
+								string FileName = Path.Combine(Folder, "xgConsole.exe");
+								if(File.Exists(FileName))
+								{
+									OutXgConsoleExe = FileName;
+									return true;
+								}
+							}
+						}
+					}
+				}
+			}
+			catch(Exception Ex)
+			{
+				Log.WriteException(Ex, null);
+			}
+
+			OutXgConsoleExe = null;
+			return false;
+		}
+
+		public static bool IsAvailable()
+		{
+			string XgConsoleExe;
+			return TryGetXgConsoleExecutable(out XgConsoleExe);
+		}
+
 		// precompile the Regex needed to parse the XGE output (the ones we want are of the form "File (Duration at +time)"
 		private static Regex XGEDurationRegex = new Regex(@"(?<Filename>.*) *\((?<Duration>[0-9:\.]+) at [0-9\+:\.]+\)", RegexOptions.ExplicitCapture);
 
diff --git a/Engine/Source/Programs/UnrealBuildTool/System/PlatformExports.cs b/Engine/Source/Programs/UnrealBuildTool/System/PlatformExports.cs
--- a/Engine/Source/Programs/UnrealBuildTool/System/PlatformExports.cs
+++ b/Engine/Source/Programs/UnrealBuildTool/System/PlatformExports.cs
@@ -119,6 +119,16 @@ namespace UnrealBuildTool
 			return UEBuildPlatform.IsPlatformAvailable(Platform) && UEBuildPlatform.GetBuildPlatform(Platform).CanUseParallelExecutor();
 		}
 
+		/// <summary>
+		/// Gets the path for the XGE console executable
+		/// </summary>
+		/// <param name="OutXgConsoleExe">On success, receives the path to the XGE console executable</param>
+		/// <returns>True if the path was found, false otherwise</returns>
+		public static bool TryGetXgConsoleExecutable(out string OutXgConsoleExe)
+		{
+			return XGE.TryGetXgConsoleExecutable(out OutXgConsoleExe);
+		}
+
 		/// <summary>
 		///
 		/// </summary>
