Fix manifests being overriden inside Enterprise folders.

diff --git a/Engine/Source/Programs/UnrealBuildTool/Configuration/UEBuildTarget.cs b/Engine/Source/Programs/UnrealBuildTool/Configuration/UEBuildTarget.cs
--- a/Engine/Source/Programs/UnrealBuildTool/Configuration/UEBuildTarget.cs
+++ b/Engine/Source/Programs/UnrealBuildTool/Configuration/UEBuildTarget.cs
@@ -1567,14 +1567,20 @@ namespace UnrealBuildTool
 				return false;
 			}
 
+			// Get the project directory. We will ignore any manifests under this directory (ie. anything not under engine/enterprise folders).
+			DirectoryReference ProjectDir = DirectoryReference.FromFile(ProjectFile);
+
 			// Read any the existing module manifests under the engine directory
 			Dictionary<FileReference, ModuleManifest> ExistingFileToManifest = new Dictionary<FileReference, ModuleManifest>();
 			foreach(FileReference ExistingFile in FileReferenceToModuleManifestPairs.Select(x => x.Key))
 			{
-				ModuleManifest ExistingManifest;
-				if(UnrealBuildTool.IsUnderAnEngineDirectory(ExistingFile.Directory) && ModuleManifest.TryRead(ExistingFile, out ExistingManifest))
+				if(ProjectDir == null || !ExistingFile.IsUnderDirectory(ProjectDir))
 				{
-					ExistingFileToManifest.Add(ExistingFile, ExistingManifest);
+					ModuleManifest ExistingManifest;
+					if(ModuleManifest.TryRead(ExistingFile, out ExistingManifest))
+					{
+						ExistingFileToManifest.Add(ExistingFile, ExistingManifest);
+					}
 				}
 			}
 
