HTML5 : re-enabling [ -s BINARYEN_TRAP_MODE='clamp'; ]

warning: this is technically slower performing -- but, this fixes a lot of 64-bit issues (HTML5 is "only" 32-bits)

#rnx

diff --git a/Engine/Source/Programs/UnrealBuildTool/Platform/HTML5/HTML5ToolChain.cs b/Engine/Source/Programs/UnrealBuildTool/Platform/HTML5/HTML5ToolChain.cs
--- a/Engine/Source/Programs/UnrealBuildTool/Platform/HTML5/HTML5ToolChain.cs
+++ b/Engine/Source/Programs/UnrealBuildTool/Platform/HTML5/HTML5ToolChain.cs
@@ -274,7 +274,7 @@ namespace UnrealBuildTool
 			Result += " -s BINARYEN=1 -s ALLOW_MEMORY_GROWTH=1";
 //			Result += " -s BINARYEN_METHOD=\\'native-wasm\\'";
 //			Result += " -s BINARYEN_MEM_MAX=-1";
-//			Result += " -s BINARYEN_TRAP_MODE=\\'clamp\\'";
+			Result += " -s BINARYEN_TRAP_MODE=\\'clamp\\'";
 
 			// no need for exceptions
 			Result += " -s DISABLE_EXCEPTION_CATCHING=1";
