UE-59381 Ios application default orientation

#iOS
#4.21

- added Initial interface orientation in the generated Plist file (UIInterfaceOrientation key)
- implemented preferredInterfaceOrientationForPresentation in the IOSViewController
- re-arranged the orientation list in the Plist file

diff --git a/Engine/Source/Developer/IOS/IOSPlatformEditor/Private/IOSTargetSettingsCustomization.cpp b/Engine/Source/Developer/IOS/IOSPlatformEditor/Private/IOSTargetSettingsCustomization.cpp
--- a/Engine/Source/Developer/IOS/IOSPlatformEditor/Private/IOSTargetSettingsCustomization.cpp
+++ b/Engine/Source/Developer/IOS/IOSPlatformEditor/Private/IOSTargetSettingsCustomization.cpp
@@ -805,6 +805,7 @@ void FIOSTargetSettingsCustomization::BuildPListSection(IDetailLayoutBuilder& De
 	SETUP_PLIST_PROP(bSupportsUpsideDownOrientation, OrientationCategory);
 	SETUP_PLIST_PROP(bSupportsLandscapeLeftOrientation, OrientationCategory);
 	SETUP_PLIST_PROP(bSupportsLandscapeRightOrientation, OrientationCategory);
+	SETUP_PLIST_PROP(PreferredLandscapeOrientation, OrientationCategory);
 	
 	SETUP_PLIST_PROP(bSupportsMetal, RenderCategory);
 	
diff --git a/Engine/Source/Programs/UnrealBuildTool/Platform/IOS/UEDeployIOS.cs b/Engine/Source/Programs/UnrealBuildTool/Platform/IOS/UEDeployIOS.cs
--- a/Engine/Source/Programs/UnrealBuildTool/Platform/IOS/UEDeployIOS.cs
+++ b/Engine/Source/Programs/UnrealBuildTool/Platform/IOS/UEDeployIOS.cs
@@ -222,20 +222,48 @@ namespace UnrealBuildTool
 			ConfigHierarchy Ini = ConfigCache.ReadHierarchy(ConfigHierarchyType.Engine, DirRef, UnrealTargetPlatform.IOS);
 
 			// orientations
+			string InterfaceOrientation = "";
+			string PreferredLandscapeOrientation = "";
+			Ini.GetString("/Script/IOSRuntimeSettings.IOSRuntimeSettings", "PreferredLandscapeOrientation", out PreferredLandscapeOrientation);
+
 			string SupportedOrientations = "";
 			bool bSupported = true;
 			Ini.GetBool("/Script/IOSRuntimeSettings.IOSRuntimeSettings", "bSupportsPortraitOrientation", out bSupported);
 			SupportedOrientations += bSupported ? "\t\t<string>UIInterfaceOrientationPortrait</string>\n" : "";
 			bSupportsPortrait = bSupported;
+
 			Ini.GetBool("/Script/IOSRuntimeSettings.IOSRuntimeSettings", "bSupportsUpsideDownOrientation", out bSupported);
 			SupportedOrientations += bSupported ? "\t\t<string>UIInterfaceOrientationPortraitUpsideDown</string>\n" : "";
 			bSupportsPortrait |= bSupported;
-			Ini.GetBool("/Script/IOSRuntimeSettings.IOSRuntimeSettings", "bSupportsLandscapeLeftOrientation", out bSupported);
-			SupportedOrientations += bSupported ? "\t\t<string>UIInterfaceOrientationLandscapeLeft</string>\n" : "";
-			bSupportsLandscape = bSupported;
-			Ini.GetBool("/Script/IOSRuntimeSettings.IOSRuntimeSettings", "bSupportsLandscapeRightOrientation", out bSupported);
-			SupportedOrientations += bSupported ? "\t\t<string>UIInterfaceOrientationLandscapeRight</string>\n" : "";
-			bSupportsLandscape |= bSupported;
+
+			bool bSupportsLandscapeLeft = false;
+			Ini.GetBool("/Script/IOSRuntimeSettings.IOSRuntimeSettings", "bSupportsLandscapeLeftOrientation", out bSupportsLandscapeLeft);
+			bool bSupportsLandscapeRight = false;
+			Ini.GetBool("/Script/IOSRuntimeSettings.IOSRuntimeSettings", "bSupportsLandscapeRightOrientation", out bSupportsLandscapeRight);
+			bSupportsLandscape = bSupportsLandscapeLeft || bSupportsLandscapeRight;
+
+			if (bSupportsLandscapeLeft && bSupportsLandscapeRight)
+			{
+				// if both landscape orientations are present, set the UIInterfaceOrientation key
+				// in the orientation list, the preferred orientation should be first
+				if (PreferredLandscapeOrientation == "LandscapeLeft")
+				{
+					InterfaceOrientation = "\t<key>UIInterfaceOrientation</key>\n\t<string>UIInterfaceOrientationLandscapeLeft</string>\n";
+					SupportedOrientations += "\t\t<string>UIInterfaceOrientationLandscapeLeft</string>\n\t\t<string>UIInterfaceOrientationLandscapeRight</string>\n";
+				}
+				else
+				{
+					// by default, landscape right is the preferred orientation - Apple's UI guidlines
+					InterfaceOrientation = "\t<key>UIInterfaceOrientation</key>\n\t<string>UIInterfaceOrientationLandscapeRight</string>\n";
+					SupportedOrientations += "\t\t<string>UIInterfaceOrientationLandscapeRight</string>\n\t\t<string>UIInterfaceOrientationLandscapeLeft</string>\n";
+				}
+			}
+			else
+			{
+				// max one landscape orientation is supported
+				SupportedOrientations += bSupportsLandscapeRight? "\t\t<string>UIInterfaceOrientationLandscapeRight</string>\n": "";
+				SupportedOrientations += bSupportsLandscapeLeft? "\t\t<string>UIInterfaceOrientationLandscapeLeft</string>\n" : "";
+			}
 
 			// bundle display name
 			string BundleDisplayName;
@@ -405,6 +433,10 @@ namespace UnrealBuildTool
 			Text.AppendLine("\t<true/>");
 			Text.AppendLine("\t<key>UIViewControllerBasedStatusBarAppearance</key>");
 			Text.AppendLine("\t<false/>");
+			if (InterfaceOrientation != "")
+			{ 
+				Text.AppendLine(InterfaceOrientation);
+			}
 			Text.AppendLine("\t<key>UISupportedInterfaceOrientations</key>");
 			Text.AppendLine("\t<array>");
 			foreach(string Line in SupportedOrientations.Split("\r\n".ToCharArray()))
diff --git a/Engine/Source/Runtime/ApplicationCore/Private/IOS/IOSView.cpp b/Engine/Source/Runtime/ApplicationCore/Private/IOS/IOSView.cpp
--- a/Engine/Source/Runtime/ApplicationCore/Private/IOS/IOSView.cpp
+++ b/Engine/Source/Runtime/ApplicationCore/Private/IOS/IOSView.cpp
@@ -1070,6 +1070,28 @@ id<MTLDevice> GMetalDevice = nil;
 	return YES;
 }
 
+/*
+* Set the preferred landscape orientation 
+*/
+- (UIInterfaceOrientation)preferredInterfaceOrientationForPresentation
+{
+	FString PreferredLandscapeOrientation = "";
+
+	GConfig->GetString(TEXT("/Script/IOSRuntimeSettings.IOSRuntimeSettings"), TEXT("PreferredLandscapeOrientation"), PreferredLandscapeOrientation, GEngineIni);
+	if (PreferredLandscapeOrientation.Equals("LandscapeLeft"))
+	{
+		return UIInterfaceOrientationLandscapeLeft;
+	}
+	else if (PreferredLandscapeOrientation.Equals("LandscapeRight"))
+	{
+		return UIInterfaceOrientationLandscapeRight;
+	}
+	else
+	{
+		return UIInterfaceOrientationPortrait;
+	}
+}
+
 - (UIRectEdge)preferredScreenEdgesDeferringSystemGestures
 {
 	return UIRectEdgeBottom;
diff --git a/Engine/Source/Runtime/IOS/IOSRuntimeSettings/Classes/IOSRuntimeSettings.h b/Engine/Source/Runtime/IOS/IOSRuntimeSettings/Classes/IOSRuntimeSettings.h
--- a/Engine/Source/Runtime/IOS/IOSRuntimeSettings/Classes/IOSRuntimeSettings.h
+++ b/Engine/Source/Runtime/IOS/IOSRuntimeSettings/Classes/IOSRuntimeSettings.h
@@ -66,6 +66,16 @@ enum class EIOSMetalShaderStandard : uint8
 	IOSMetalSLStandard_2_0 = 3 UMETA(DisplayName="Metal v2.0 (iOS 11.0/tvOS 11.0)"),
 };
 
+UENUM()
+enum class EIOSLandscapeOrientation : uint8
+{
+	/** Landscape Left */
+	LandscapeLeft = 0 UMETA(DisplayName = "Landscape (left home button)"),
+
+	/** Landscape Right */
+	LandscapeRight = 1 UMETA(DisplayName = "Landscape (right home button)"),
+};
+
 /**
  *	IOS Build resource file struct, used to serialize filepaths to the configs for use in the build system,
  */
@@ -329,6 +339,10 @@ public:
 	UPROPERTY(GlobalConfig, EditAnywhere, Category = DeviceOrientations)
 	uint32 bSupportsLandscapeRightOrientation : 1;
 
+	// The Preferred Orientation will be used as the initial orientation at launch when both Landscape Left and Landscape Right orientations are to be supported.
+	UPROPERTY(GlobalConfig, EditAnywhere, Category = DeviceOrientations)
+	EIOSLandscapeOrientation PreferredLandscapeOrientation;
+
 	// Specifies the the display name for the application. This will be displayed under the icon on the device.
 	UPROPERTY(GlobalConfig, EditAnywhere, Category = BundleInformation)
 	FString BundleDisplayName;
