Fixed -vsmac argument not working with GenerateProjectFiles

diff --git a/Engine/Source/Programs/UnrealBuildTool/UnrealBuildTool.cs b/Engine/Source/Programs/UnrealBuildTool/UnrealBuildTool.cs
--- a/Engine/Source/Programs/UnrealBuildTool/UnrealBuildTool.cs
+++ b/Engine/Source/Programs/UnrealBuildTool/UnrealBuildTool.cs
@@ -875,6 +875,9 @@ namespace UnrealBuildTool
 								case ProjectFileFormat.CLion:
 									Generator = new CLionGenerator(ProjectFile);
 									break;
+								case ProjectFileFormat.VisualStudioMac:
+									Generator = new VCMacProjectFileGenerator(ProjectFile, OverrideWindowsCompiler);
+									break;
 								default:
 									throw new BuildException("Unhandled project file type '{0}", ProjectFileFormat);
 							}
