Prevent Xcode upgrade warnings for >9.0 versions

diff --git a/Engine/Source/Programs/UnrealBuildTool/ProjectFiles/Xcode/XcodeProject.cs b/Engine/Source/Programs/UnrealBuildTool/ProjectFiles/Xcode/XcodeProject.cs
--- a/Engine/Source/Programs/UnrealBuildTool/ProjectFiles/Xcode/XcodeProject.cs
+++ b/Engine/Source/Programs/UnrealBuildTool/ProjectFiles/Xcode/XcodeProject.cs
@@ -492,7 +492,7 @@ namespace UnrealBuildTool
 			Content.Append("\t\t" + ProjectGuid + " /* Project object */ = {" + ProjectFileGenerator.NewLine);
 			Content.Append("\t\t\tisa = PBXProject;" + ProjectFileGenerator.NewLine);
 			Content.Append("\t\t\tattributes = {" + ProjectFileGenerator.NewLine);
-			Content.Append("\t\t\t\tLastUpgradeCheck = 0900;" + ProjectFileGenerator.NewLine);
+			Content.Append("\t\t\t\tLastUpgradeCheck = 2000;" + ProjectFileGenerator.NewLine);
 			Content.Append("\t\t\t\tORGANIZATIONNAME = \"Epic Games, Inc.\";" + ProjectFileGenerator.NewLine);
             Content.Append("\t\t\t\tTargetAttributes = {" + ProjectFileGenerator.NewLine);
             Content.Append("\t\t\t\t\t" + TargetGuid + " = {" + ProjectFileGenerator.NewLine);
@@ -1180,7 +1180,7 @@ namespace UnrealBuildTool
 
 			Content.Append("<?xml version=\"1.0\" encoding=\"UTF-8\"?>" + ProjectFileGenerator.NewLine);
 			Content.Append("<Scheme" + ProjectFileGenerator.NewLine);
-			Content.Append("   LastUpgradeVersion = \"0710\"" + ProjectFileGenerator.NewLine);
+			Content.Append("   LastUpgradeVersion = \"2000\"" + ProjectFileGenerator.NewLine);
 			Content.Append("   version = \"1.3\">" + ProjectFileGenerator.NewLine);
 			Content.Append("   <BuildAction" + ProjectFileGenerator.NewLine);
 			Content.Append("      parallelizeBuildables = \"YES\"" + ProjectFileGenerator.NewLine);
