Fix issue with Android PCH include files

#ue4
#android

diff --git a/Engine/Source/Programs/UnrealBuildTool/Platform/Android/AndroidToolChain.cs b/Engine/Source/Programs/UnrealBuildTool/Platform/Android/AndroidToolChain.cs
--- a/Engine/Source/Programs/UnrealBuildTool/Platform/Android/AndroidToolChain.cs
+++ b/Engine/Source/Programs/UnrealBuildTool/Platform/Android/AndroidToolChain.cs
@@ -1290,9 +1290,12 @@ namespace UnrealBuildTool
 					string PCHArguments = "";
 					if (CompileEnvironment.PrecompiledHeaderAction == PrecompiledHeaderAction.Include)
 					{
+						// add the platform-specific PCH reference
+						PCHArguments += string.Format(" -include-pch \"{0}\"", InlineArchName(BasePCHName, Arch, GPUArchitecture) + PCHExtension);
+
 						// Add the precompiled header file's path to the include path so Clang can find it.
 						// This needs to be before the other include paths to ensure Clang uses it instead of the source header file.
-						PCHArguments += string.Format(" -include \"{0}\"", InlineArchName(BasePCHName, Arch, GPUArchitecture));
+						PCHArguments += string.Format(" -include \"{0}\"", BasePCHName);
 					}
 
 					foreach (FileItem ForceIncludeFile in CompileEnvironment.ForceIncludeFiles)
