UBT: Move HTML5 environment setup into the toolchain constructor.

diff --git a/Engine/Source/Programs/UnrealBuildTool/Platform/HTML5/HTML5ToolChain.cs b/Engine/Source/Programs/UnrealBuildTool/Platform/HTML5/HTML5ToolChain.cs
--- a/Engine/Source/Programs/UnrealBuildTool/Platform/HTML5/HTML5ToolChain.cs
+++ b/Engine/Source/Programs/UnrealBuildTool/Platform/HTML5/HTML5ToolChain.cs
@@ -66,10 +66,7 @@ namespace UnrealBuildTool
 			Log.TraceInformation("HTML5ToolChain: EnableTracing = "      + bEnableTracing       );
 
 			PrintOnce = new VerbosePrint(PrintOnceOn); // reset
-		}
 
-		public static void PreBuildSync()
-		{
 			Log.TraceInformation("Setting Emscripten SDK: located in " + HTML5SDKInfo.EMSCRIPTEN_ROOT);
 			HTML5SDKInfo.SetupEmscriptenTemp();
 			HTML5SDKInfo.SetUpEmscriptenConfigFile();
diff --git a/Engine/Source/Programs/UnrealBuildTool/Platform/HTML5/UEBuildHTML5.cs b/Engine/Source/Programs/UnrealBuildTool/Platform/HTML5/UEBuildHTML5.cs
--- a/Engine/Source/Programs/UnrealBuildTool/Platform/HTML5/UEBuildHTML5.cs
+++ b/Engine/Source/Programs/UnrealBuildTool/Platform/HTML5/UEBuildHTML5.cs
@@ -57,11 +57,6 @@ namespace UnrealBuildTool
 			Target.bDeployAfterCompile = true;
 		}
 
-		public override void PreBuildSync()
-		{
-			HTML5ToolChain.PreBuildSync();
-		}
-
 		public override bool CanUseXGE()
 		{
 			return false; // NOTE: setting to true may break CIS builds...
