UBT: Remove some unused stuff from UEBuildPlatform.

#rnx

diff --git a/Engine/Source/Programs/UnrealBuildTool/Configuration/UEBuildPlatform.cs b/Engine/Source/Programs/UnrealBuildTool/Configuration/UEBuildPlatform.cs
--- a/Engine/Source/Programs/UnrealBuildTool/Configuration/UEBuildPlatform.cs
+++ b/Engine/Source/Programs/UnrealBuildTool/Configuration/UEBuildPlatform.cs
@@ -656,30 +656,6 @@ namespace UnrealBuildTool
 			return TempList;
 		}
 
-		/// <summary>
-		/// Return whether this platform has uniquely named binaries across multiple games
-		/// </summary>
-		public virtual bool HasUniqueBinaries()
-		{
-			return true;
-		}
-
-		/// <summary>
-		/// Return whether we wish to have this platform's binaries in our builds
-		/// </summary>
-		public virtual bool IsBuildRequired()
-		{
-			return true;
-		}
-
-		/// <summary>
-		/// Return whether we wish to have this platform's binaries in our CIS tests
-		/// </summary>
-		public virtual bool IsCISRequired()
-		{
-			return true;
-		}
-
 		/// <summary>
 		/// Whether the build platform requires deployment prep
 		/// </summary>
@@ -888,16 +864,6 @@ namespace UnrealBuildTool
 		/// <returns>New toolchain instance.</returns>
 		public abstract UEToolChain CreateToolChain(CppPlatform CppPlatform, ReadOnlyTargetRules Target);
 
-		/// <summary>
-		/// Creates a temp toolchain instance for the given project - will not be used to compile with, and is only needed on some platforms
-		/// </summary>
-		/// <param name="ProjectFile">The project to make the toolchain for</param>
-		/// <returns>New toolchain instance.</returns>
-		public virtual UEToolChain CreateTempToolChainForProject(FileReference ProjectFile)
-		{
-			return null;
-		}
-
 		/// <summary>
 		/// Deploys the given target
 		/// </summary>
diff --git a/Engine/Source/Programs/UnrealBuildTool/Platform/Android/UEBuildAndroid.cs b/Engine/Source/Programs/UnrealBuildTool/Platform/Android/UEBuildAndroid.cs
--- a/Engine/Source/Programs/UnrealBuildTool/Platform/Android/UEBuildAndroid.cs
+++ b/Engine/Source/Programs/UnrealBuildTool/Platform/Android/UEBuildAndroid.cs
@@ -462,7 +462,7 @@ namespace UnrealBuildTool
 			bool bUseLdGold = Target.bUseUnityBuild;
 			return new AndroidToolChain(Target.ProjectFile, bUseLdGold, Target.AndroidPlatform.Architectures, Target.AndroidPlatform.GPUArchitectures);
 		}
-		public override UEToolChain CreateTempToolChainForProject(FileReference ProjectFile)
+		public virtual UEToolChain CreateTempToolChainForProject(FileReference ProjectFile)
 		{
 			return new AndroidToolChain(ProjectFile, true, null, null);
 		}
diff --git a/Engine/Source/Programs/UnrealBuildTool/Platform/Android/UEDeployAndroid.cs b/Engine/Source/Programs/UnrealBuildTool/Platform/Android/UEDeployAndroid.cs
--- a/Engine/Source/Programs/UnrealBuildTool/Platform/Android/UEDeployAndroid.cs
+++ b/Engine/Source/Programs/UnrealBuildTool/Platform/Android/UEDeployAndroid.cs
@@ -3896,7 +3896,7 @@ namespace UnrealBuildTool
 
 		public override bool PrepTargetForDeployment(UEBuildDeployTarget InTarget)
 		{
-			AndroidToolChain ToolChain = UEBuildPlatform.GetBuildPlatform(InTarget.Platform).CreateTempToolChainForProject(InTarget.ProjectFile) as AndroidToolChain;
+			AndroidToolChain ToolChain = ((AndroidPlatform)UEBuildPlatform.GetBuildPlatform(InTarget.Platform)).CreateTempToolChainForProject(InTarget.ProjectFile) as AndroidToolChain;
 
 			// we need to strip architecture from any of the output paths
 			string BaseSoName = ToolChain.RemoveArchName(InTarget.OutputPaths[0].FullName);
diff --git a/Engine/Source/Programs/UnrealBuildTool/Platform/Mac/UEBuildMac.cs b/Engine/Source/Programs/UnrealBuildTool/Platform/Mac/UEBuildMac.cs
--- a/Engine/Source/Programs/UnrealBuildTool/Platform/Mac/UEBuildMac.cs
+++ b/Engine/Source/Programs/UnrealBuildTool/Platform/Mac/UEBuildMac.cs
@@ -232,22 +232,6 @@ namespace UnrealBuildTool
 		{
 		}
 
-		/// <summary>
-		/// Return whether we wish to have this platform's binaries in our builds
-		/// </summary>
-		public override bool IsBuildRequired()
-		{
-			return false;
-		}
-
-		/// <summary>
-		/// Return whether we wish to have this platform's binaries in our CIS tests
-		/// </summary>
-		public override bool IsCISRequired()
-		{
-			return false;
-		}
-
 		public override DirectoryReference GetBundleDirectory(ReadOnlyTargetRules Rules, List<FileReference> OutputFiles)
 		{
 			if(Rules.bIsBuildingConsoleApplication)
diff --git a/Engine/Source/Programs/UnrealBuildTool/Platform/Windows/UEBuildWindows.cs b/Engine/Source/Programs/UnrealBuildTool/Platform/Windows/UEBuildWindows.cs
--- a/Engine/Source/Programs/UnrealBuildTool/Platform/Windows/UEBuildWindows.cs
+++ b/Engine/Source/Programs/UnrealBuildTool/Platform/Windows/UEBuildWindows.cs
@@ -1318,15 +1318,6 @@ namespace UnrealBuildTool
 		{
 		}
 
-		/// <summary>
-		/// Return whether this platform has uniquely named binaries across multiple games
-		/// </summary>
-		public override bool HasUniqueBinaries()
-		{
-			// Windows applications have many shared binaries between games
-			return false;
-		}
-
 		/// <summary>
 		/// Modify the rules for a newly created module, in a target that's being built for this platform.
 		/// This is not required - but allows for hiding details of a particular platform.
