UBT: Remove ConvertPath() method, which is no longer needed since remote toolchain has been removed.

#rnx

diff --git a/Engine/Source/Programs/UnrealBuildTool/Configuration/UEBuildPlatform.cs b/Engine/Source/Programs/UnrealBuildTool/Configuration/UEBuildPlatform.cs
--- a/Engine/Source/Programs/UnrealBuildTool/Configuration/UEBuildPlatform.cs
+++ b/Engine/Source/Programs/UnrealBuildTool/Configuration/UEBuildPlatform.cs
@@ -303,16 +303,6 @@ namespace UnrealBuildTool
 		{
 		}
 
-		/// <summary>
-		/// Converts the passed in path from UBT host to compiler native format.
-		/// </summary>
-		/// <param name="OriginalPath">The path to convert</param>
-		/// <returns>The path in native format for the toolchain</returns>
-		public virtual string ConvertPath(string OriginalPath)
-		{
-			return OriginalPath;
-		}
-
 		/// <summary>
 		/// Get the bundle directory for the shared link environment
 		/// </summary>
diff --git a/Engine/Source/Programs/UnrealBuildTool/Platform/Linux/LinuxToolChain.cs b/Engine/Source/Programs/UnrealBuildTool/Platform/Linux/LinuxToolChain.cs
--- a/Engine/Source/Programs/UnrealBuildTool/Platform/Linux/LinuxToolChain.cs
+++ b/Engine/Source/Programs/UnrealBuildTool/Platform/Linux/LinuxToolChain.cs
@@ -1283,7 +1283,7 @@ namespace UnrealBuildTool
 				FileItem ResponseFileItem = FileItem.CreateIntermediateTextFile(ResponsePath, InputFileNames);
 				ArchiveAction.PrerequisiteItems.Add(ResponseFileItem);
 			}
-			ArchiveAction.CommandArguments += string.Format(" @\"{0}\"", ConvertPath(ResponsePath.FullName));
+			ArchiveAction.CommandArguments += string.Format(" @\"{0}\"", ResponsePath.FullName);
 
 			// add ranlib if not using llvm-ar
 			if (String.IsNullOrEmpty(LlvmArPath))
@@ -1840,21 +1840,6 @@ namespace UnrealBuildTool
 			}
 		}
 
-		/// <summary>
-		/// Converts the passed in path from UBT host to compiler native format.
-		/// </summary>
-		public static string ConvertPath(string OriginalPath)
-		{
-			if (BuildHostPlatform.Current.Platform == UnrealTargetPlatform.Linux)
-			{
-				return OriginalPath.Replace("\\", "/");
-			}
-			else
-			{
-				return OriginalPath;
-			}
-		}
-
 		public override ICollection<FileItem> PostBuild(FileItem Executable, LinkEnvironment BinaryLinkEnvironment, ActionGraph ActionGraph)
 		{
 			ICollection<FileItem> OutputFiles = base.PostBuild(Executable, BinaryLinkEnvironment, ActionGraph);
diff --git a/Engine/Source/Programs/UnrealBuildTool/Platform/Linux/UEBuildLinux.cs b/Engine/Source/Programs/UnrealBuildTool/Platform/Linux/UEBuildLinux.cs
--- a/Engine/Source/Programs/UnrealBuildTool/Platform/Linux/UEBuildLinux.cs
+++ b/Engine/Source/Programs/UnrealBuildTool/Platform/Linux/UEBuildLinux.cs
@@ -378,14 +378,6 @@ namespace UnrealBuildTool
 			}
 		}
 
-		/// <summary>
-		/// Converts the passed in path from UBT host to compiler native format.
-		/// </summary>
-		public override string ConvertPath(string OriginalPath)
-		{
-			return LinuxToolChain.ConvertPath(OriginalPath);
-		}
-
 		/// <summary>
 		/// Modify the rules for a newly created module, in a target that's being built for this platform.
 		/// This is not required - but allows for hiding details of a particular platform.
diff --git a/Engine/Source/Programs/UnrealBuildTool/System/ExternalExecution.cs b/Engine/Source/Programs/UnrealBuildTool/System/ExternalExecution.cs
--- a/Engine/Source/Programs/UnrealBuildTool/System/ExternalExecution.cs
+++ b/Engine/Source/Programs/UnrealBuildTool/System/ExternalExecution.cs
@@ -309,16 +309,14 @@ namespace UnrealBuildTool
 
 		public bool IsGameTarget;     // True if the current target is a game target
 		public string RootLocalPath;    // The engine path on the local machine
-		public string RootBuildPath;    // The engine path on the build machine, if different (e.g. Mac/iOS builds)
 		public string TargetName;       // Name of the target currently being compiled
 		public string ExternalDependenciesFile; // File to contain additional dependencies that the generated code depends on
 		public List<Module> Modules;
 
-		public UHTManifest(UEBuildTarget Target, string InRootLocalPath, string InRootBuildPath, string InExternalDependenciesFile, List<Module> InModules)
+		public UHTManifest(UEBuildTarget Target, string InRootLocalPath, string InExternalDependenciesFile, List<Module> InModules)
 		{
 			IsGameTarget = (Target.TargetType != TargetType.Program);
 			RootLocalPath = InRootLocalPath;
-			RootBuildPath = InRootBuildPath;
 			TargetName = Target.GetTargetName();
 			ExternalDependenciesFile = InExternalDependenciesFile;
 			Modules = InModules;
@@ -1126,7 +1124,7 @@ namespace UnrealBuildTool
 				{
 					Modules.Add(new UHTManifest.Module(UObjectModule));
 				}
-				UHTManifest Manifest = new UHTManifest(Target, RootLocalPath, UEBuildPlatform.GetBuildPlatform(Target.Platform).ConvertPath(RootLocalPath + '\\'), ExternalDependenciesFile.FullName, Modules);
+				UHTManifest Manifest = new UHTManifest(Target, RootLocalPath, ExternalDependenciesFile.FullName, Modules);
 
 				if (!bIsBuildingUHT && bUHTNeedsToRun)
 				{
