UBT: Remove unused DoLocalToRemoteFileItem function.

#rnx

diff --git a/Engine/Source/Programs/UnrealBuildTool/Configuration/UEBuildModuleCPP.cs b/Engine/Source/Programs/UnrealBuildTool/Configuration/UEBuildModuleCPP.cs
--- a/Engine/Source/Programs/UnrealBuildTool/Configuration/UEBuildModuleCPP.cs
+++ b/Engine/Source/Programs/UnrealBuildTool/Configuration/UEBuildModuleCPP.cs
@@ -339,7 +339,6 @@ namespace UnrealBuildTool
 				foreach(FileReference OutputFile in Manifest.OutputFiles)
 				{
 					FileItem ObjectFile = FileItem.GetExistingItemByFileReference(OutputFile);
-					ToolChain.DoLocalToRemoteFileItem(ObjectFile);
 					LinkInputFiles.Add(ObjectFile);
 				}
 				return LinkInputFiles;
diff --git a/Engine/Source/Programs/UnrealBuildTool/ToolChain/UEToolChain.cs b/Engine/Source/Programs/UnrealBuildTool/ToolChain/UEToolChain.cs
--- a/Engine/Source/Programs/UnrealBuildTool/ToolChain/UEToolChain.cs
+++ b/Engine/Source/Programs/UnrealBuildTool/ToolChain/UEToolChain.cs
@@ -26,10 +26,6 @@ namespace UnrealBuildTool
 		{
 		}
 
-		public virtual void DoLocalToRemoteFileItem(FileItem LocalFileItem)
-		{
-		}
-
 		public abstract CPPOutput CompileCPPFiles(CppCompileEnvironment CompileEnvironment, List<FileItem> InputFiles, DirectoryReference OutputDir, string ModuleName, ActionGraph ActionGraph);
 
 		public virtual CPPOutput CompileRCFiles(CppCompileEnvironment Environment, List<FileItem> InputFiles, DirectoryReference OutputDir, ActionGraph ActionGraph)
