UBT: Remove another unused function (UEToolChain.ParseProjectSettings)

#rnx

diff --git a/Engine/Source/Programs/UnrealBuildTool/ToolChain/UEToolChain.cs b/Engine/Source/Programs/UnrealBuildTool/ToolChain/UEToolChain.cs
--- a/Engine/Source/Programs/UnrealBuildTool/ToolChain/UEToolChain.cs
+++ b/Engine/Source/Programs/UnrealBuildTool/ToolChain/UEToolChain.cs
@@ -59,11 +59,6 @@ namespace UnrealBuildTool
 		}
 
 		public virtual void SetUpGlobalEnvironment(ReadOnlyTargetRules Target)
-		{
-			ParseProjectSettings();
-		}
-
-		public virtual void ParseProjectSettings()
 		{
 		}
 
