UBT: The UEBuildFramework class instantiated from ModuleRules classes has been renamed to Framework, to decouple it from the internal representation used inside UBT. A deprecated ModuleRules.UEBuildFramework class has been added to help migration.

diff --git a/Engine/Source/Programs/UnrealBuildTool/Configuration/ModuleRules.cs b/Engine/Source/Programs/UnrealBuildTool/Configuration/ModuleRules.cs
--- a/Engine/Source/Programs/UnrealBuildTool/Configuration/ModuleRules.cs
+++ b/Engine/Source/Programs/UnrealBuildTool/Configuration/ModuleRules.cs
@@ -263,6 +263,58 @@ namespace UnrealBuildTool
 			}
 		}
 
+		/// <summary>
+		/// Stores information about a framework on IOS or MacOS
+		/// </summary>
+		public class Framework
+		{
+			/// <summary>
+			/// Name of the framework
+			/// </summary>
+			internal string Name;
+
+			/// <summary>
+			/// For non-system frameworks, specifies the path to a zip file that contains it.
+			/// </summary>
+			internal string ZipPath;
+
+			/// <summary>
+			/// 
+			/// </summary>
+			internal string CopyBundledAssets = null;
+
+			/// <summary>
+			/// Constructor
+			/// </summary>
+			/// <param name="Name">Name of the framework</param>
+			/// <param name="ZipPath">Path to a zip file containing the framework. May be null.</param>
+			/// <param name="CopyBundledAssets"></param>
+			public Framework(string Name, string ZipPath = null, string CopyBundledAssets = null)
+			{
+				this.Name = Name;
+				this.ZipPath = ZipPath;
+				this.CopyBundledAssets = CopyBundledAssets;
+			}
+		}
+
+		/// <summary>
+		/// Deprecated; wrapper for Framework.
+		/// </summary>
+		[Obsolete("The UEBuildFramework class has been deprecated in UE 4.22. Please use the Framework class instead.")]
+		public class UEBuildFramework : Framework
+		{
+			/// <summary>
+			/// Constructor
+			/// </summary>
+			/// <param name="Name">Name of the framework</param>
+			/// <param name="ZipPath">Path to a zip file containing the framework. May be null.</param>
+			/// <param name="CopyBundledAssets"></param>
+			public UEBuildFramework(string Name, string ZipPath = null, string CopyBundledAssets = null)
+				: base(Name, ZipPath, CopyBundledAssets)
+			{
+			}
+		}
+
 		/// <summary>
 		/// Name of this module
 		/// </summary>
@@ -510,7 +562,7 @@ namespace UnrealBuildTool
 		/// <summary>
 		/// List of addition frameworks - typically used for External (third party) modules on Mac and iOS
 		/// </summary>
-		public List<UEBuildFramework> PublicAdditionalFrameworks = new List<UEBuildFramework>();
+		public List<Framework> PublicAdditionalFrameworks = new List<Framework>();
 
 		/// <summary>
 		/// List of addition resources that should be copied to the app bundle for Mac or iOS
diff --git a/Engine/Source/Programs/UnrealBuildTool/Configuration/UEBuildFramework.cs b/Engine/Source/Programs/UnrealBuildTool/Configuration/UEBuildFramework.cs
--- a/Engine/Source/Programs/UnrealBuildTool/Configuration/UEBuildFramework.cs
+++ b/Engine/Source/Programs/UnrealBuildTool/Configuration/UEBuildFramework.cs
@@ -1,4 +1,4 @@
-﻿// Copyright 1998-2018 Epic Games, Inc. All Rights Reserved.
+// Copyright 1998-2018 Epic Games, Inc. All Rights Reserved.
 
 using System;
 using System.Collections.Generic;
@@ -11,39 +11,30 @@ namespace UnrealBuildTool
 	/// <summary>
 	/// 
 	/// </summary>
-	public class UEBuildFramework
+	class UEBuildFramework
 	{
 		/// <summary>
 		/// 
 		/// </summary>
-		/// <param name="InFrameworkName"></param>
-		public UEBuildFramework(string InFrameworkName)
+		/// <param name="Name"></param>
+		/// <param name="ZipPath"></param>
+		/// <param name="CopyBundledAssets"></param>
+		public UEBuildFramework(string Name, string ZipPath = null, string CopyBundledAssets = null)
 		{
-			FrameworkName = InFrameworkName;
+			this.FrameworkName = Name;
+			this.FrameworkZipPath = ZipPath;
+			this.CopyBundledAssets = CopyBundledAssets;
 		}
 
 		/// <summary>
 		/// 
 		/// </summary>
-		/// <param name="InFrameworkName"></param>
-		/// <param name="InFrameworkZipPath"></param>
-		public UEBuildFramework(string InFrameworkName, string InFrameworkZipPath)
+		/// <param name="Framework"></param>
+		public UEBuildFramework(ModuleRules.Framework Framework)
 		{
-			FrameworkName = InFrameworkName;
-			FrameworkZipPath = InFrameworkZipPath;
-		}
-
-		/// <summary>
-		/// 
-		/// </summary>
-		/// <param name="InFrameworkName"></param>
-		/// <param name="InFrameworkZipPath"></param>
-		/// <param name="InCopyBundledAssets"></param>
-		public UEBuildFramework(string InFrameworkName, string InFrameworkZipPath, string InCopyBundledAssets)
-		{
-			FrameworkName = InFrameworkName;
-			FrameworkZipPath = InFrameworkZipPath;
-			CopyBundledAssets = InCopyBundledAssets;
+			FrameworkName = Framework.Name;
+			FrameworkZipPath = Framework.ZipPath;
+			CopyBundledAssets = Framework.CopyBundledAssets;
 		}
 
 		internal UEBuildModule OwningModule = null;
diff --git a/Engine/Source/Programs/UnrealBuildTool/Configuration/UEBuildModule.cs b/Engine/Source/Programs/UnrealBuildTool/Configuration/UEBuildModule.cs
--- a/Engine/Source/Programs/UnrealBuildTool/Configuration/UEBuildModule.cs
+++ b/Engine/Source/Programs/UnrealBuildTool/Configuration/UEBuildModule.cs
@@ -173,7 +173,7 @@ namespace UnrealBuildTool
 			PublicAdditionalLibraries = HashSetFromOptionalEnumerableStringParameter(Rules.PublicAdditionalLibraries);
 			PublicFrameworks = HashSetFromOptionalEnumerableStringParameter(Rules.PublicFrameworks);
 			PublicWeakFrameworks = HashSetFromOptionalEnumerableStringParameter(Rules.PublicWeakFrameworks);
-			PublicAdditionalFrameworks = Rules.PublicAdditionalFrameworks == null ? new HashSet<UEBuildFramework>() : new HashSet<UEBuildFramework>(Rules.PublicAdditionalFrameworks);
+			PublicAdditionalFrameworks = Rules.PublicAdditionalFrameworks == null ? new HashSet<UEBuildFramework>() : new HashSet<UEBuildFramework>(Rules.PublicAdditionalFrameworks.Select(x => new UEBuildFramework(x)));
 			PublicAdditionalShadowFiles = HashSetFromOptionalEnumerableStringParameter(Rules.PublicAdditionalShadowFiles);
 			PublicAdditionalBundleResources = Rules.AdditionalBundleResources == null ? new HashSet<UEBuildBundleResource>() : new HashSet<UEBuildBundleResource>(Rules.AdditionalBundleResources);
 			PublicDelayLoadDLLs = HashSetFromOptionalEnumerableStringParameter(Rules.PublicDelayLoadDLLs);
diff --git a/Engine/Source/Programs/UnrealBuildTool/Configuration/UEBuildTarget.cs b/Engine/Source/Programs/UnrealBuildTool/Configuration/UEBuildTarget.cs
--- a/Engine/Source/Programs/UnrealBuildTool/Configuration/UEBuildTarget.cs
+++ b/Engine/Source/Programs/UnrealBuildTool/Configuration/UEBuildTarget.cs
@@ -1207,11 +1207,11 @@ namespace UnrealBuildTool
 				}
 
 				// Add any zip files from Additional Frameworks
-				foreach (UEBuildFramework Framework in Rules.PublicAdditionalFrameworks)
+				foreach (ModuleRules.Framework Framework in Rules.PublicAdditionalFrameworks)
 				{
-					if (!String.IsNullOrEmpty(Framework.FrameworkZipPath))
+					if (!String.IsNullOrEmpty(Framework.ZipPath))
 					{
-						Files.Add(FileReference.Combine(Module.ModuleDirectory, Framework.FrameworkZipPath));
+						Files.Add(FileReference.Combine(Module.ModuleDirectory, Framework.ZipPath));
 					}
 				}
 
