UBT: Allow marking attribute-driven command line arguments as required.

diff --git a/Engine/Source/Programs/UnrealBuildTool/System/CommandLine.cs b/Engine/Source/Programs/UnrealBuildTool/System/CommandLine.cs
--- a/Engine/Source/Programs/UnrealBuildTool/System/CommandLine.cs
+++ b/Engine/Source/Programs/UnrealBuildTool/System/CommandLine.cs
@@ -134,6 +134,27 @@ namespace UnrealBuildTool
 					}
 				}
 			}
+
+			// Make sure there are no required parameters that are missing
+			Dictionary<FieldInfo, Parameter> MissingFieldToParameter = new Dictionary<FieldInfo, Parameter>();
+			foreach(Parameter Parameter in PrefixToParameter.Values)
+			{
+				if(Parameter.Attribute.Required && !AssignedFieldToParameter.ContainsKey(Parameter.FieldInfo) && !MissingFieldToParameter.ContainsKey(Parameter.FieldInfo))
+				{
+					MissingFieldToParameter.Add(Parameter.FieldInfo, Parameter);
+				}
+			}
+			if(MissingFieldToParameter.Count > 0)
+			{
+				if(MissingFieldToParameter.Count == 1)
+				{
+					throw new BuildException("Missing {0} argument", MissingFieldToParameter.First().Value.Prefix.Replace("=", "=..."));
+				}
+				else
+				{
+					throw new BuildException("Missing {0} arguments", Utils.FormatList(MissingFieldToParameter.Values.Select(x => x.Prefix.Replace("=", "=..."))));
+				}
+			}
 		}
 
 		/// <summary>
diff --git a/Engine/Source/Programs/UnrealBuildTool/System/CommandLineAttribute.cs b/Engine/Source/Programs/UnrealBuildTool/System/CommandLineAttribute.cs
--- a/Engine/Source/Programs/UnrealBuildTool/System/CommandLineAttribute.cs
+++ b/Engine/Source/Programs/UnrealBuildTool/System/CommandLineAttribute.cs
@@ -25,6 +25,11 @@ namespace UnrealBuildTool
 		/// </summary>
 		public string Value = null;
 
+		/// <summary>
+		/// Whether this argument is required
+		/// </summary>
+		public bool Required;
+
 		/// <summary>
 		/// For collection types, specifies the separator character between multiple arguments
 		/// </summary>
diff --git a/Engine/Source/Programs/UnrealBuildTool/System/Utils.cs b/Engine/Source/Programs/UnrealBuildTool/System/Utils.cs
--- a/Engine/Source/Programs/UnrealBuildTool/System/Utils.cs
+++ b/Engine/Source/Programs/UnrealBuildTool/System/Utils.cs
@@ -1054,5 +1054,42 @@ namespace UnrealBuildTool
 			}
 			return true;
 		}
+
+		/// <summary>
+		/// Formats a list of strings in the style "1, 2, 3 and 4"
+		/// </summary>
+		/// <param name="Arguments">List of strings to format</param>
+		/// <returns>Formatted list of strings</returns>
+		public static string FormatList(string[] Arguments)
+		{
+			StringBuilder Result = new StringBuilder();
+			if(Arguments.Length > 0)
+			{
+				Result.Append(Arguments[0]);
+				for(int Idx = 1; Idx < Arguments.Length; Idx++)
+				{
+					if(Idx == Arguments.Length - 1)
+					{
+						Result.Append(" and ");
+					}
+					else
+					{
+						Result.Append(", ");
+					}
+					Result.Append(Arguments[Idx]);
+				}
+			}
+			return Result.ToString();
+		}
+
+		/// <summary>
+		/// Formats a list of strings in the style "1, 2, 3 and 4"
+		/// </summary>
+		/// <param name="Arguments">List of strings to format</param>
+		/// <returns>Formatted list of strings</returns>
+		public static string FormatList(IEnumerable<string> Arguments)
+		{
+			return FormatList(Arguments.ToArray());
+		}
 	}
 }
