UBT: Remote unused global exit code override.

#rnx

diff --git a/Engine/Source/Programs/UnrealBuildTool/UnrealBuildTool.cs b/Engine/Source/Programs/UnrealBuildTool/UnrealBuildTool.cs
--- a/Engine/Source/Programs/UnrealBuildTool/UnrealBuildTool.cs
+++ b/Engine/Source/Programs/UnrealBuildTool/UnrealBuildTool.cs
@@ -943,14 +943,9 @@ namespace UnrealBuildTool
 			// Print some performance info
 			Log.TraceVerbose("Execution time: {0}", (DateTime.UtcNow - StartTime).TotalSeconds);
 
-			if (ExtendedErrorCode != 0)
-			{
-				return ExtendedErrorCode;
-			}
 			return (int)Result;
 		}
 
-		public static int ExtendedErrorCode = 0;
 		private static int Main(string[] Arguments)
 		{
 			// make sure we catch any exceptions and return an appropriate error code.
@@ -966,10 +961,6 @@ namespace UnrealBuildTool
 				{
 					Log.TraceError("UnrealBuildTool Exception: " + Exception.ToString());
 				}
-				if (ExtendedErrorCode != 0)
-				{
-					return ExtendedErrorCode;
-				}
 				return (int)ECompilationResult.OtherCompilationError;
 			}
 		}
