Revert change to make BuildException internal; used by lots of .build.cs files.

#rnx

diff --git a/Engine/Source/Programs/UnrealBuildTool/System/BuildException.cs b/Engine/Source/Programs/UnrealBuildTool/System/BuildException.cs
--- a/Engine/Source/Programs/UnrealBuildTool/System/BuildException.cs
+++ b/Engine/Source/Programs/UnrealBuildTool/System/BuildException.cs
@@ -9,7 +9,7 @@ namespace UnrealBuildTool
 	/// <summary>
 	/// Base class for exceptions thrown by UBT
 	/// </summary>
-	class BuildException : Exception
+	public class BuildException : Exception
 	{
 		/// <summary>
 		/// Constructor
