UBT: Remove unused variable.

#rnx

diff --git a/Engine/Source/Programs/UnrealBuildTool/UnrealBuildTool.cs b/Engine/Source/Programs/UnrealBuildTool/UnrealBuildTool.cs
--- a/Engine/Source/Programs/UnrealBuildTool/UnrealBuildTool.cs
+++ b/Engine/Source/Programs/UnrealBuildTool/UnrealBuildTool.cs
@@ -113,11 +113,6 @@ namespace UnrealBuildTool
 		/// </summary>
 		static string RemoteIniPath = null;
 
-		/// <summary>
-		/// Cached array of all the supported target platforms
-		/// </summary>
-		static public UnrealTargetPlatform[] AllPlatforms = (UnrealTargetPlatform[])Enum.GetValues(typeof(UnrealTargetPlatform));
-
 		/// <summary>
 		/// Whether to print debug information out to the log
 		/// </summary>
