UBT: Add support for running UBT in different modes with separate command line parsing and entry points. Different modes derive from the ToolMode class, and use a [ToolMode] attribute to specify their name. Executing a mode is done with the -Mode=XXX parameter on the command line. The default mode implements traditional UBT functionality, but the WriteDocumentation mode can be used to output HTML documentation for BuildConfiguration, ModuleRules and TargetRules classes.

diff --git a/Engine/Source/Programs/DotNETCommon/DotNETUtilities/Log.cs b/Engine/Source/Programs/DotNETCommon/DotNETUtilities/Log.cs
--- a/Engine/Source/Programs/DotNETCommon/DotNETUtilities/Log.cs
+++ b/Engine/Source/Programs/DotNETCommon/DotNETUtilities/Log.cs
@@ -226,6 +226,42 @@ namespace Tools.DotNETCommon
 			ColorConsoleOutput = true;
 		}
 
+		/// <summary>
+		/// Adds a trace listener that writes to a log file
+		/// </summary>
+		/// <param name="OutputFile">The file to write to</param>
+		/// <returns>The created trace listener</returns>
+		public static TextWriterTraceListener AddFileWriter(string Name, FileReference OutputFile)
+		{
+			try
+			{
+				DirectoryReference.CreateDirectory(OutputFile.Directory);
+				TextWriterTraceListener LogTraceListener = new TextWriterTraceListener(new StreamWriter(OutputFile.FullName), Name);
+				Trace.Listeners.Add(LogTraceListener);
+				return LogTraceListener;
+			}
+			catch (Exception Ex)
+			{
+				throw new Exception(String.Format("Unable to open log file for writing ({0})", OutputFile), Ex);
+			}
+		}
+
+		/// <summary>
+		/// Determines if a TextWriterTraceListener has been added to the list of trace listeners
+		/// </summary>
+		/// <returns>True if a TextWriterTraceListener has been added</returns>
+		public static bool HasFileWriter()
+		{
+			foreach(TraceListener Listener in Trace.Listeners)
+			{
+				if(Listener is TextWriterTraceListener)
+				{
+					return true;
+				}
+			}
+			return false;
+		}
+
 		/// <summary>
 		/// Gets the name of the Method N levels deep in the stack frame. Used to trap what method actually made the logging call.
 		/// Only used when bLogSources is true.
diff --git a/Engine/Source/Programs/UnrealBuildTool/Modes/WriteDocumentationMode.cs b/Engine/Source/Programs/UnrealBuildTool/Modes/WriteDocumentationMode.cs
new file mode 100644
--- /dev/null
+++ b/Engine/Source/Programs/UnrealBuildTool/Modes/WriteDocumentationMode.cs
@@ -0,0 +1,64 @@
+using System;
+using System.Collections.Generic;
+using System.Linq;
+using System.Text;
+using System.Threading.Tasks;
+using Tools.DotNETCommon;
+
+namespace UnrealBuildTool
+{
+	/// <summary>
+	/// Generates documentation from reflection data
+	/// </summary>
+	[ToolMode("WriteDocumentation")]
+	class WriteDocumentationMode : ToolMode
+	{
+		/// <summary>
+		/// Enum for the type of documentation to generate
+		/// </summary>
+		enum DocumentationType
+		{
+			BuildConfiguration,
+			ModuleRules,
+			TargetRules,
+		}
+
+		/// <summary>
+		/// Type of documentation to generate
+		/// </summary>
+		[CommandLine(Required = true)]
+		DocumentationType Type = DocumentationType.BuildConfiguration;
+
+		/// <summary>
+		/// The HTML file to write to
+		/// </summary>
+		[CommandLine(Required = true)]
+		FileReference OutputFile = null;
+
+		/// <summary>
+		/// Entry point for this command
+		/// </summary>
+		/// <returns></returns>
+		public override int Execute(List<string> Arguments)
+		{
+			CommandLine.ParseAndRemoveArguments(Arguments, this);
+			CommandLine.CheckNoRemainingArguments(Arguments);
+
+			switch(Type)
+			{
+				case DocumentationType.BuildConfiguration:
+					XmlConfig.WriteDocumentation(OutputFile);
+					break;
+				case DocumentationType.ModuleRules:
+					RulesDocumentation.WriteDocumentation(typeof(ModuleRules), OutputFile);
+					break;
+				case DocumentationType.TargetRules:
+					RulesDocumentation.WriteDocumentation(typeof(TargetRules), OutputFile);
+					break;
+				default:
+					throw new BuildException("Invalid documentation type: {0}", Type);
+			}
+			return 0;
+		}
+	}
+}
diff --git a/Engine/Source/Programs/UnrealBuildTool/System/CommandLine.cs b/Engine/Source/Programs/UnrealBuildTool/System/CommandLine.cs
--- a/Engine/Source/Programs/UnrealBuildTool/System/CommandLine.cs
+++ b/Engine/Source/Programs/UnrealBuildTool/System/CommandLine.cs
@@ -157,6 +157,25 @@ namespace UnrealBuildTool
 			}
 		}
 
+		/// <summary>
+		/// Checks that the list of arguments is empty. If not, throws an exception listing them.
+		/// </summary>
+		/// <param name="RemainingArguments">List of remaining arguments</param>
+		public static void CheckNoRemainingArguments(List<string> RemainingArguments)
+		{
+			if(RemainingArguments.Count > 0)
+			{
+				if(RemainingArguments.Count == 1)
+				{
+					Log.TraceWarning("Invalid argument: {0}", RemainingArguments[0]);
+				}
+				else
+				{
+					Log.TraceWarning("Invalid arguments:\n{0}", String.Join("\n", RemainingArguments));
+				}
+			}
+		}
+
 		/// <summary>
 		/// Parses and assigns a value to a field
 		/// </summary>
diff --git a/Engine/Source/Programs/UnrealBuildTool/System/ExternalExecution.cs b/Engine/Source/Programs/UnrealBuildTool/System/ExternalExecution.cs
--- a/Engine/Source/Programs/UnrealBuildTool/System/ExternalExecution.cs
+++ b/Engine/Source/Programs/UnrealBuildTool/System/ExternalExecution.cs
@@ -1188,7 +1188,7 @@ namespace UnrealBuildTool
 						}
 						else
 						{
-							UBTArguments.AppendFormat(" -log=\"{0}\"", Path.Combine(Path.GetDirectoryName(BuildConfiguration.LogFileName), Path.GetFileNameWithoutExtension(BuildConfiguration.LogFileName) + "_UHT.txt"));
+							UBTArguments.Append(" -logsuffix=UHT");
 						}
 
 						if (RunExternalDotNETExecutable(UnrealBuildTool.GetUBTPath(), UBTArguments.ToString()) != 0)
diff --git a/Engine/Source/Programs/UnrealBuildTool/System/ToolMode.cs b/Engine/Source/Programs/UnrealBuildTool/System/ToolMode.cs
new file mode 100644
--- /dev/null
+++ b/Engine/Source/Programs/UnrealBuildTool/System/ToolMode.cs
@@ -0,0 +1,43 @@
+using System;
+using System.Collections.Generic;
+using System.Linq;
+using System.Text;
+using System.Threading.Tasks;
+using Tools.DotNETCommon;
+
+namespace UnrealBuildTool
+{
+	/// <summary>
+	/// Attribute used to specify options for a UBT mode.
+	/// </summary>
+	class ToolModeAttribute : Attribute
+	{
+		/// <summary>
+		/// Name of this mode
+		/// </summary>
+		public string Name;
+
+		/// <summary>
+		/// Constructor
+		/// </summary>
+		/// <param name="Name">Name of the mode</param>
+		public ToolModeAttribute(string Name)
+		{
+			this.Name = Name;
+		}
+	}
+
+	/// <summary>
+	/// Base class for standalone UBT modes. Different modes can be invoked using the -Mode=[Name] argument on the command line, where [Name] is determined by 
+	/// the ToolModeAttribute on a ToolMode derived class. The log system will be initialized before calling the mode, but little else.
+	/// </summary>
+	abstract class ToolMode
+	{
+		/// <summary>
+		/// Entry point for this command.
+		/// </summary>
+		/// <param name="Arguments">List of command line arguments</param>
+		/// <returns>Exit code for the process</returns>
+		public abstract int Execute(List<string> Arguments);
+	}
+}
diff --git a/Engine/Source/Programs/UnrealBuildTool/UnrealBuildTool.cs b/Engine/Source/Programs/UnrealBuildTool/UnrealBuildTool.cs
--- a/Engine/Source/Programs/UnrealBuildTool/UnrealBuildTool.cs
+++ b/Engine/Source/Programs/UnrealBuildTool/UnrealBuildTool.cs
@@ -320,28 +320,15 @@ namespace UnrealBuildTool
 		/// This should be very minimal startup code.
 		/// </summary>
 		/// <param name="Arguments">Cmdline arguments</param>
-		private static int GuardedMain(string[] Arguments)
+		internal static int GuardedMain(string[] Arguments)
 		{
 			DateTime StartTime = DateTime.UtcNow;
 			ECompilationResult Result = ECompilationResult.Succeeded;
 
-			// Parse the log level argument
-			if (Arguments.Any(x => x.Equals("-Verbose", StringComparison.InvariantCultureIgnoreCase)))
-			{
-				Log.OutputLevel = LogEventType.Verbose;
-			}
-			else if (Arguments.Any(x => x.Equals("-VeryVerbose", StringComparison.InvariantCultureIgnoreCase)))
-			{
-				Log.OutputLevel = LogEventType.VeryVerbose;
-			}
-
 			// Initialize the log system, buffering the output until we can create the log file
 			StartupTraceListener StartupListener = new StartupTraceListener();
 			Trace.Listeners.Add(StartupListener);
 
-			// If we're running from MsBuild, we need to include the program name to avoid it formatting messages as EXEC : UnrealBuildTool : ERROR
-			Log.IncludeProgramNameWithSeverityPrefix = Arguments.Any(x => x.Equals("-FromMsBuild", StringComparison.InvariantCultureIgnoreCase));
-
 			// Write the command line
 			Log.TraceLog("Command line: {0}", Environment.CommandLine);
 
@@ -392,7 +379,6 @@ namespace UnrealBuildTool
 				}
 			}
 
-
 			// Don't allow simultaneous execution of Unreal Built Tool. Multi-selection in the UI e.g. causes this and you want serial
 			// execution in this case to avoid contention issues with shared produced items.
 			bool bCreatedMutex = false;
@@ -503,12 +489,13 @@ namespace UnrealBuildTool
 					}
 
 					// Then let the command lines override any configs necessary.
+					string LogSuffix = "";
 					foreach (string Argument in Arguments)
 					{
 						string LowercaseArg = Argument.ToLowerInvariant();
-						if (LowercaseArg.StartsWith("-log="))
+						if (LowercaseArg.StartsWith("-LogSuffix=", StringComparison.OrdinalIgnoreCase))
 						{
-							BuildConfiguration.LogFileName = Argument.Substring("-log=".Length);
+							LogSuffix += "_" + Argument.Substring("-LogSuffix=".Length);
 						}
 						else if (LowercaseArg == "-nolog")
 						{
@@ -524,38 +511,19 @@ namespace UnrealBuildTool
 							BuildConfiguration.bUseUBTMakefiles = false;
 							BuildConfiguration.SingleFileToCompile = LowercaseArg.Replace("-singlefile=", "");
 						}
-						else if (LowercaseArg.StartsWith("-buildconfigurationdoc="))
-						{
-							XmlConfig.WriteDocumentation(new FileReference(Argument.Substring("-buildconfigurationdoc=".Length)));
-							return 0;
-						}
-						else if (LowercaseArg.StartsWith("-modulerulesdoc="))
-						{
-							RulesDocumentation.WriteDocumentation(typeof(ModuleRules), new FileReference(Argument.Substring("-modulerulesdoc=".Length)));
-							return 0;
-						}
-						else if (LowercaseArg.StartsWith("-targetrulesdoc="))
-						{
-							RulesDocumentation.WriteDocumentation(typeof(TargetRules), new FileReference(Argument.Substring("-targetrulesdoc=".Length)));
-							return 0;
-						}
 					}
 
 					// Create the log file, and flush the startup listener to it
-					if (!String.IsNullOrEmpty(BuildConfiguration.LogFileName))
+					if (!String.IsNullOrEmpty(BuildConfiguration.LogFileName) && !Log.HasFileWriter())
 					{
 						FileReference LogLocation = new FileReference(BuildConfiguration.LogFileName);
-						try
+						if(LogSuffix.Length > 0)
 						{
-							DirectoryReference.CreateDirectory(LogLocation.Directory);
-							TextWriterTraceListener LogTraceListener = new TextWriterTraceListener(new StreamWriter(LogLocation.FullName), "LogTraceListener");
-							StartupListener.CopyTo(LogTraceListener);
-							Trace.Listeners.Add(LogTraceListener);
-						}
-						catch (Exception Ex)
-						{
-							throw new BuildException(Ex, "Unable to open log file for writing ({0})", LogLocation);
+							LogLocation = LogLocation.ChangeExtension(null) + LogSuffix + LogLocation.GetExtension();
 						}
+
+						TextWriterTraceListener LogTraceListener = Log.AddFileWriter("DefaultLogTraceListener", LogLocation);
+						StartupListener.CopyTo(LogTraceListener);
 					}
 					Trace.Listeners.Remove(StartupListener);
 
@@ -930,20 +898,116 @@ namespace UnrealBuildTool
 			return (int)Result;
 		}
 
+		/// <summary>
+		/// Global options for UBT (any modes)
+		/// </summary>
+		class GlobalOptions
+		{
+			/// <summary>
+			/// The amount of detail to write to the log
+			/// </summary>
+			[CommandLine(Prefix = "-Verbose", Value ="Verbose")]
+			[CommandLine(Prefix = "-VeryVerbose", Value ="VeryVerbose")]
+			public LogEventType LogOutputLevel = LogEventType.Log;
+
+			/// <summary>
+			/// Specifies the path to a log file to write. Note that the default mode (eg. building, generating project files) will create a log file by default if this not specified.
+			/// </summary>
+			[CommandLine(Prefix = "-Log")]
+			public FileReference LogFileName = null;
+
+			/// <summary>
+			/// Whether to include timestamps in the log
+			/// </summary>
+			[CommandLine(Prefix = "-Timestamps")]
+			public bool bLogTimestamps = false;
+
+			/// <summary>
+			/// Whether to format messages in MsBuild format
+			/// </summary>
+			[CommandLine(Prefix = "-FromMsBuild")]
+			public bool bLogFromMsBuild = false;
+
+			/// <summary>
+			/// The mode to execute
+			/// </summary>
+			[CommandLine]
+			public string Mode = "Default";
+		}
+
+		/// <summary>
+		/// Main entry point. Parses any global options and initializes the logging system, then invokes the appropriate command.
+		/// </summary>
+		/// <param name="Arguments">Command line arguments</param>
+		/// <returns>Zero on success, non-zero on error</returns>
 		private static int Main(string[] Arguments)
 		{
-			// make sure we catch any exceptions and return an appropriate error code.
-			// Some inner code already does this (to ensure the Mutex is released),
-			// but we need something to cover all outer code as well.
 			try
 			{
-				return GuardedMain(Arguments);
+				// Copy the arguments into a list. We remove entries from this list as they are parsed, so after we've parsed all the options we can we can identify any incorrect arguments and warn the user appropriately.
+				List<string> RemainingArguments = new List<string>(Arguments);
+
+				// Parse the global options
+				GlobalOptions Options = new GlobalOptions();
+				CommandLine.ParseAndRemoveArguments(RemainingArguments, Options);
+
+				// Configure the log system
+				Log.OutputLevel = Options.LogOutputLevel;
+				Log.IncludeTimestamps = Options.bLogTimestamps;
+				Log.IncludeProgramNameWithSeverityPrefix = Options.bLogFromMsBuild;
+
+				// Add the log writer if requested. When building a target, we'll create the writer for the default log file later.
+				if(Options.LogFileName != null)
+				{
+					Log.AddFileWriter("LogTraceListener", Options.LogFileName);
+				}
+
+				// Find all the valid modes
+				Dictionary<string, Type> ModeNameToType = new Dictionary<string, Type>(StringComparer.OrdinalIgnoreCase);
+				foreach(Type Type in Assembly.GetExecutingAssembly().GetTypes())
+				{
+					if(Type.IsClass && !Type.IsAbstract && Type.IsSubclassOf(typeof(ToolMode)))
+					{
+						ToolModeAttribute Attribute = Type.GetCustomAttribute<ToolModeAttribute>();
+						if(Attribute == null)
+						{
+							throw new BuildException("Class '{0}' should have a ToolModeAttribute", Type.Name);
+						}
+						ModeNameToType.Add(Attribute.Name, Type);
+					}
+				}
+
+				// Try to get the correct mode
+				Type ModeType;
+				if(!ModeNameToType.TryGetValue(Options.Mode, out ModeType))
+				{
+					Log.TraceError("No mode named '{0}'. Available modes are:\n  {1}", Options.Mode, String.Join("\n  ", ModeNameToType.Keys));
+					return 1;
+				}
+
+				// Create the appropriate handler
+				ToolMode Mode = (ToolMode)Activator.CreateInstance(ModeType);
+				return Mode.Execute(RemainingArguments);
+			}
+			catch (BuildException Ex)
+			{
+				// BuildExceptions should have nicely formatted messages. We can log these directly.
+				Log.TraceError(Ex.Message.ToString());
+				Log.TraceLog(ExceptionUtils.FormatExceptionDetails(Ex));
+				return (int)ECompilationResult.OtherCompilationError;
 			}
-			catch (Exception Exception)
+			catch (Exception Ex)
 			{
-				Log.TraceError("UnrealBuildTool Exception: " + Exception.ToString());
+				// Unhandled exception. 
+				Log.TraceError("Unhandled exception: {0}", Ex);
+				Log.TraceLog(ExceptionUtils.FormatExceptionDetails(Ex));
 				return (int)ECompilationResult.OtherCompilationError;
 			}
+			finally
+			{
+				// Make sure we flush the logs however we exit
+				Trace.Close();
+			}
 		}
 
 		/// <summary>
@@ -2259,4 +2323,13 @@ namespace UnrealBuildTool
 			}
 		}
 	}
+
+	[ToolMode("Default")]
+	class DefaultToolMode : ToolMode
+	{
+		public override int Execute(List<string> Arguments)
+		{
+			return UnrealBuildTool.GuardedMain(Arguments.ToArray());
+		}
+	}
 }
diff --git a/Engine/Source/Programs/UnrealBuildTool/UnrealBuildTool.csproj b/Engine/Source/Programs/UnrealBuildTool/UnrealBuildTool.csproj
--- a/Engine/Source/Programs/UnrealBuildTool/UnrealBuildTool.csproj
+++ b/Engine/Source/Programs/UnrealBuildTool/UnrealBuildTool.csproj
@@ -140,6 +140,7 @@
     </When>
   </Choose>
   <ItemGroup>
+    <Compile Include="Modes\WriteDocumentationMode.cs" />
     <Compile Include="Configuration\ModuleRules.cs" />
     <Compile Include="Configuration\PlatformTargetRules.cs" />
     <Compile Include="Configuration\TargetDescriptor.cs" />
@@ -208,6 +209,7 @@
     <Compile Include="System\SupportedConfigurationsAttribute.cs" />
     <Compile Include="System\SupportedPlatformsAttribute.cs" />
     <Compile Include="System\TargetReceipt.cs" />
+    <Compile Include="System\ToolMode.cs" />
     <Compile Include="System\UBTMakefile.cs" />
     <Compile Include="System\VersionNumber.cs" />
     <Compile Include="System\XmlConfig.cs" />
