UBT: Add a class stores a list of command line arguments and allow efficient ad-hoc queries of particular options (eg. "-Flag") and retreival of typed values (eg. "-Foo=Bar"), as well as attribute-driven application to fields with the [CommandLine] attribute applied. Also tracks which arguments have been retrieved, allowing the display of diagnostic messages for invalid arguments.

diff --git a/Engine/Source/Programs/UnrealBuildTool/Modes/WriteDocumentationMode.cs b/Engine/Source/Programs/UnrealBuildTool/Modes/WriteDocumentationMode.cs
--- a/Engine/Source/Programs/UnrealBuildTool/Modes/WriteDocumentationMode.cs
+++ b/Engine/Source/Programs/UnrealBuildTool/Modes/WriteDocumentationMode.cs
@@ -39,10 +39,10 @@ namespace UnrealBuildTool
 		/// Entry point for this command
 		/// </summary>
 		/// <returns></returns>
-		public override int Execute(List<string> Arguments)
+		public override int Execute(CommandLineArguments Arguments)
 		{
-			CommandLine.ParseAndRemoveArguments(Arguments, this);
-			CommandLine.CheckNoRemainingArguments(Arguments);
+			Arguments.ApplyTo(this);
+			Arguments.CheckAllArgumentsUsed();
 
 			switch(Type)
 			{
diff --git a/Engine/Source/Programs/UnrealBuildTool/System/CommandLineArguments.cs b/Engine/Source/Programs/UnrealBuildTool/System/CommandLineArguments.cs
new file mode 100644
--- /dev/null
+++ b/Engine/Source/Programs/UnrealBuildTool/System/CommandLineArguments.cs
@@ -0,0 +1,875 @@
+using System;
+using System.Collections;
+using System.Collections.Generic;
+using System.Diagnostics;
+using System.Linq;
+using System.Reflection;
+using System.Text;
+using System.Threading.Tasks;
+using Tools.DotNETCommon;
+
+namespace UnrealBuildTool
+{
+	/// <summary>
+	/// Helper class to visualize an argument list
+	/// </summary>
+	class CommandLineArgumentListView
+	{
+		/// <summary>
+		/// The list of arguments
+		/// </summary>
+		[DebuggerBrowsable(DebuggerBrowsableState.RootHidden)]
+		public string[] Arguments;
+
+		/// <summary>
+		/// Constructor
+		/// </summary>
+		/// <param name="ArgumentList">The argument list to proxy</param>
+		public CommandLineArgumentListView(CommandLineArguments ArgumentList)
+		{
+			Arguments = ArgumentList.GetRawArray();
+		}
+	}
+
+	/// <summary>
+	/// Exception thrown for invalid command line arguments
+	/// </summary>
+	class CommandLineArgumentException : Exception
+	{
+		/// <summary>
+		/// Constructor
+		/// </summary>
+		/// <param name="Message">Message to display for this exception</param>
+		public CommandLineArgumentException(string Message)
+			: base(Message)
+		{
+		}
+
+		/// <summary>
+		/// Constructor
+		/// </summary>
+		/// <param name="Message">Message to display for this exception</param>
+		/// <param name="InnerException">The inner exception</param>
+		public CommandLineArgumentException(string Message, Exception InnerException)
+			: base(Message, InnerException)
+		{
+		}
+
+		/// <summary>
+		/// Converts this exception to a string
+		/// </summary>
+		/// <returns>Exception message</returns>
+		public override string ToString()
+		{
+			return Message;
+		}
+	}
+
+	/// <summary>
+	/// Stores a list of command line arguments, allowing efficient ad-hoc queries of particular options (eg. "-Flag") and retreival of typed values (eg. "-Foo=Bar"), as
+	/// well as attribute-driven application to fields with the [CommandLine] attribute applied.
+	/// 
+	/// Also tracks which arguments have been retrieved, allowing the display of diagnostic messages for invalid arguments.
+	/// </summary>
+	[DebuggerDisplay("Count = {Count}")]
+	[DebuggerTypeProxy(typeof(CommandLineArgumentListView))]
+	class CommandLineArguments : IReadOnlyList<string>, IReadOnlyCollection<string>, IEnumerable<string>, IEnumerable
+	{
+		/// <summary>
+		/// The raw array of arguments
+		/// </summary>
+		string[] Arguments;
+
+		/// <summary>
+		/// Bitmask indicating which arguments are flags rather than values
+		/// </summary>
+		BitArray FlagArguments;
+
+		/// <summary>
+		/// Bitmask indicating which arguments have been used, via calls to GetOption(), GetValues() etc...
+		/// </summary>
+		BitArray UsedArguments;
+
+		/// <summary>
+		/// Dictionary of argument names (or prefixes, in the case of "-Foo=" style arguments) to their index into the arguments array.
+		/// </summary>
+		Dictionary<string, int> ArgumentToFirstIndex;
+
+		/// <summary>
+		/// For each argument which is seen more than once, keeps a list of indices for the second and subsequent arguments.
+		/// </summary>
+		int[] NextArgumentIndex;
+
+		/// <summary>
+		/// Array of characters that separate argument names from values
+		/// </summary>
+		static readonly char[] ValueSeparators = { '=', ':' };
+
+		/// <summary>
+		/// Constructor
+		/// </summary>
+		/// <param name="Arguments">The raw list of arguments</param>
+		public CommandLineArguments(string[] Arguments)
+		{
+			this.Arguments = Arguments;
+			this.FlagArguments = new BitArray(Arguments.Length);
+			this.UsedArguments = new BitArray(Arguments.Length);
+
+			// Clear the linked list of identical arguments
+			NextArgumentIndex = new int[Arguments.Length];
+			for(int Idx = 0; Idx < Arguments.Length; Idx++)
+			{
+				NextArgumentIndex[Idx] = -1;
+			}
+
+			// Temporarily store the index of the last matching argument
+			int[] LastArgumentIndex = new int[Arguments.Length];
+
+			// Parse the argument array and build a lookup
+			ArgumentToFirstIndex = new Dictionary<string, int>(Arguments.Length, StringComparer.OrdinalIgnoreCase);
+			for(int Idx = 0; Idx < Arguments.Length; Idx++)
+			{
+				int SeparatorIdx = Arguments[Idx].IndexOfAny(ValueSeparators);
+				if(SeparatorIdx == -1)
+				{
+					// Ignore duplicate -Option flags; they are harmless.
+					if(ArgumentToFirstIndex.ContainsKey(Arguments[Idx]))
+					{
+						UsedArguments.Set(Idx, true);
+					}
+					else
+					{
+						ArgumentToFirstIndex.Add(Arguments[Idx], Idx);
+					}
+
+					// Mark this argument as a flag
+					FlagArguments.Set(Idx, true);
+				}
+				else
+				{
+					// Just take the part up to and including the separator character
+					string Prefix = Arguments[Idx].Substring(0, SeparatorIdx + 1);
+
+					// Add the prefix to the argument lookup, or update the appropriate matching argument list if it's been seen before
+					int ExistingArgumentIndex;
+					if(ArgumentToFirstIndex.TryGetValue(Prefix, out ExistingArgumentIndex))
+					{
+						NextArgumentIndex[LastArgumentIndex[ExistingArgumentIndex]] = Idx;
+						LastArgumentIndex[ExistingArgumentIndex] = Idx;
+					}
+					else
+					{
+						ArgumentToFirstIndex.Add(Prefix, Idx);
+						LastArgumentIndex[Idx] = Idx;
+					}
+				}
+			}
+		}
+
+		/// <summary>
+		/// The number of arguments in this list
+		/// </summary>
+		public int Count
+		{
+			get { return Arguments.Length; }
+		}
+
+		/// <summary>
+		/// Access an argument by index
+		/// </summary>
+		/// <param name="Index">Index of the argument</param>
+		/// <returns>The argument at the given index</returns>
+		public string this[int Index]
+		{
+			get { return Arguments[Index]; }
+		}
+		
+		/// <summary>
+		/// Determines if an argument has been used
+		/// </summary>
+		/// <param name="Index">Index of the argument</param>
+		/// <returns>True if the argument has been used, false otherwise</returns>
+		public bool HasBeenUsed(int Index)
+		{
+			return UsedArguments.Get(Index);
+		}
+
+		/// <summary>
+		/// Marks an argument as having been used
+		/// </summary>
+		/// <param name="Index">Index of the argument to mark as used</param>
+		public void MarkAsUsed(int Index)
+		{
+			UsedArguments.Set(Index, true);
+		}
+
+		/// <summary>
+		/// Marks an argument as not having been used
+		/// </summary>
+		/// <param name="Index">Index of the argument to mark as being unused</param>
+		public void MarkAsUnused(int Index)
+		{
+			UsedArguments.Set(Index, true);
+		}
+
+		/// <summary>
+		/// Checks if the given option (eg. "-Foo") was specified on the command line.
+		/// </summary>
+		/// <param name="Option">The option to look for</param>
+		/// <returns>True if the option was found, false otherwise.</returns>
+		public bool HasOption(string Option)
+		{
+			int Index;
+			if(ArgumentToFirstIndex.TryGetValue(Option, out Index))
+			{
+				UsedArguments.Set(Index, true);
+				return true;
+			}
+			return false;
+		}
+
+		/// <summary>
+		/// Checks for an argument prefixed with the given string is present.
+		/// </summary>
+		/// <param name="Prefix">The argument prefix (eg. "-Foo="). Must end with an '=' character.</param>
+		/// <returns>True if an argument with the given prefix was specified</returns>
+		public bool HasValue(string Prefix)
+		{
+			CheckValidPrefix(Prefix);
+			return ArgumentToFirstIndex.ContainsKey(Prefix);
+		}
+
+		/// <summary>
+		/// Gets the value specified by an argument with the given prefix. Throws an exception if the argument was not specified.
+		/// </summary>
+		/// <param name="Prefix">The argument prefix (eg. "-Foo="). Must end with an '=' character.</param>
+		/// <returns>Value of the argument</returns>
+		public string GetString(string Prefix)
+		{
+			string Value;
+			if(!TryGetValue(Prefix, out Value))
+			{
+				throw new CommandLineArgumentException(String.Format("Missing '{0}...' argument", Prefix));
+			}
+			return Value;
+		}
+
+		/// <summary>
+		/// Gets the value specified by an argument with the given prefix. Throws an exception if the argument was not specified.
+		/// </summary>
+		/// <param name="Prefix">The argument prefix (eg. "-Foo="). Must end with an '=' character.</param>
+		/// <returns>Value of the argument</returns>
+		public int GetInteger(string Prefix)
+		{
+			int Value;
+			if(!TryGetValue(Prefix, out Value))
+			{
+				throw new CommandLineArgumentException(String.Format("Missing '{0}...' argument", Prefix));
+			}
+			return Value;
+		}
+
+		/// <summary>
+		/// Gets the value specified by an argument with the given prefix. Throws an exception if the argument was not specified.
+		/// </summary>
+		/// <param name="Prefix">The argument prefix (eg. "-Foo="). Must end with an '=' character.</param>
+		/// <returns>Value of the argument</returns>
+		public FileReference GetFileReference(string Prefix)
+		{
+			FileReference Value;
+			if(!TryGetValue(Prefix, out Value))
+			{
+				throw new CommandLineArgumentException(String.Format("Missing '{0}...' argument", Prefix));
+			}
+			return Value;
+		}
+
+		/// <summary>
+		/// Gets the value specified by an argument with the given prefix. Throws an exception if the argument was not specified.
+		/// </summary>
+		/// <param name="Prefix">The argument prefix (eg. "-Foo="). Must end with an '=' character.</param>
+		/// <returns>Value of the argument</returns>
+		public DirectoryReference GetDirectoryReference(string Prefix)
+		{
+			DirectoryReference Value;
+			if(!TryGetValue(Prefix, out Value))
+			{
+				throw new CommandLineArgumentException(String.Format("Missing '{0}...' argument", Prefix));
+			}
+			return Value;
+		}
+
+		/// <summary>
+		/// Gets the value specified by an argument with the given prefix. Throws an exception if the argument was not specified.
+		/// </summary>
+		/// <param name="Prefix">The argument prefix (eg. "-Foo="). Must end with an '=' character.</param>
+		/// <returns>Value of the argument</returns>
+		public T GetEnum<T>(string Prefix) where T : struct
+		{
+			T Value;
+			if(!TryGetValue(Prefix, out Value))
+			{
+				throw new CommandLineArgumentException(String.Format("Missing '{0}...' argument", Prefix));
+			}
+			return Value;
+		}
+
+		/// <summary>
+		/// Gets the value specified by an argument with the given prefix, or a default value.
+		/// </summary>
+		/// <param name="Prefix">The argument prefix (eg. "-Foo="). Must end with an '=' character.</param>
+		/// <param name="DefaultValue">Default value for the argument</param>
+		/// <returns>Value of the argument</returns>
+		public string GetStringOrDefault(string Prefix, string DefaultValue)
+		{
+			string Value;
+			if(!TryGetValue(Prefix, out Value))
+			{
+				Value = DefaultValue;
+			}
+			return Value;
+		}
+
+		/// <summary>
+		/// Gets the value specified by an argument with the given prefix, or a default value.
+		/// </summary>
+		/// <param name="Prefix">The argument prefix (eg. "-Foo="). Must end with an '=' character.</param>
+		/// <param name="DefaultValue">Default value for the argument</param>
+		/// <returns>Value of the argument</returns>
+		public int GetIntegerOrDefault(string Prefix, int DefaultValue)
+		{
+			int Value;
+			if(!TryGetValue(Prefix, out Value))
+			{
+				Value = DefaultValue;
+			}
+			return Value;
+		}
+
+		/// <summary>
+		/// Gets the value specified by an argument with the given prefix, or a default value.
+		/// </summary>
+		/// <param name="Prefix">The argument prefix (eg. "-Foo="). Must end with an '=' character.</param>
+		/// <param name="DefaultValue">Default value for the argument</param>
+		/// <returns>Value of the argument</returns>
+		public FileReference GetFileReferenceOrDefault(string Prefix, FileReference DefaultValue)
+		{
+			FileReference Value;
+			if(!TryGetValue(Prefix, out Value))
+			{
+				Value = DefaultValue;
+			}
+			return Value;
+		}
+
+		/// <summary>
+		/// Gets the value specified by an argument with the given prefix, or a default value.
+		/// </summary>
+		/// <param name="Prefix">The argument prefix (eg. "-Foo="). Must end with an '=' character.</param>
+		/// <param name="DefaultValue">Default value for the argument</param>
+		/// <returns>Value of the argument</returns>
+		public DirectoryReference GetDirectoryReferenceOrDefault(string Prefix, DirectoryReference DefaultValue)
+		{
+			DirectoryReference Value;
+			if(!TryGetValue(Prefix, out Value))
+			{
+				Value = DefaultValue;
+			}
+			return Value;
+		}
+
+		/// <summary>
+		/// Gets the value specified by an argument with the given prefix, or a default value.
+		/// </summary>
+		/// <param name="Prefix">The argument prefix (eg. "-Foo="). Must end with an '=' character.</param>
+		/// <param name="DefaultValue">Default value for the argument</param>
+		/// <returns>Value of the argument</returns>
+		public T GetEnumOrDefault<T>(string Prefix, T DefaultValue) where T : struct
+		{
+			T Value;
+			if(!TryGetValue(Prefix, out Value))
+			{
+				Value = DefaultValue;
+			}
+			return Value;
+		}
+
+		/// <summary>
+		/// Tries to gets the value specified by an argument with the given prefix.
+		/// </summary>
+		/// <param name="Prefix">The argument prefix (eg. "-Foo="). Must end with an '=' character.</param>
+		/// <param name="Value">Value of the argument, if found</param>
+		/// <returns>True if the argument was found (and Value was set), false otherwise.</returns>
+		public bool TryGetValue(string Prefix, out string Value)
+		{
+			CheckValidPrefix(Prefix);
+
+			int Index;
+			if(!ArgumentToFirstIndex.TryGetValue(Prefix, out Index))
+			{
+				Value = null;
+				return false;
+			}
+
+			if(NextArgumentIndex[Index] != -1)
+			{
+				throw new CommandLineArgumentException(String.Format("Multiple {0}... arguments are specified", Prefix));
+			}
+
+			UsedArguments.Set(Index, true);
+			Value = Arguments[Index].Substring(Prefix.Length);
+			return true;
+		}
+
+		/// <summary>
+		/// Tries to gets the value specified by an argument with the given prefix.
+		/// </summary>
+		/// <param name="Prefix">The argument prefix (eg. "-Foo="). Must end with an '=' character.</param>
+		/// <param name="Value">Value of the argument, if found</param>
+		/// <returns>True if the argument was found (and Value was set), false otherwise.</returns>
+		public bool TryGetValue(string Prefix, out int Value)
+		{
+			// Try to get the string value of this argument
+			string StringValue;
+			if(!TryGetValue(Prefix, out StringValue))
+			{
+				Value = 0;
+				return false;
+			}
+
+			// Try to parse it. If it fails, throw an exception.
+			try
+			{
+				Value = int.Parse(StringValue);
+				return true;
+			}
+			catch(Exception Ex)
+			{
+				throw new CommandLineArgumentException(String.Format("The argument '{0}{1}' does not specify a valid integer", Prefix, StringValue), Ex);
+			}
+		}
+
+		/// <summary>
+		/// Tries to gets the value specified by an argument with the given prefix.
+		/// </summary>
+		/// <param name="Prefix">The argument prefix (eg. "-Foo="). Must end with an '=' character.</param>
+		/// <param name="Value">Value of the argument, if found</param>
+		/// <returns>True if the argument was found (and Value was set), false otherwise.</returns>
+		public bool TryGetValue(string Prefix, out FileReference Value)
+		{
+			// Try to get the string value of this argument
+			string StringValue;
+			if(!TryGetValue(Prefix, out StringValue))
+			{
+				Value = null;
+				return false;
+			}
+
+			// Try to parse it. If it fails, throw an exception.
+			try
+			{
+				Value = new FileReference(StringValue);
+				return true;
+			}
+			catch(Exception Ex)
+			{
+				throw new CommandLineArgumentException(String.Format("The argument '{0}{1}' does not specify a valid file name", Prefix, StringValue), Ex);
+			}
+		}
+
+		/// <summary>
+		/// Tries to gets the value specified by an argument with the given prefix.
+		/// </summary>
+		/// <param name="Prefix">The argument prefix (eg. "-Foo="). Must end with an '=' character.</param>
+		/// <param name="Value">Value of the argument, if found</param>
+		/// <returns>True if the argument was found (and Value was set), false otherwise.</returns>
+		public bool TryGetValue(string Prefix, out DirectoryReference Value)
+		{
+			// Try to get the string value of this argument
+			string StringValue;
+			if(!TryGetValue(Prefix, out StringValue))
+			{
+				Value = null;
+				return false;
+			}
+
+			// Try to parse it. If it fails, throw an exception.
+			try
+			{
+				Value = new DirectoryReference(StringValue);
+				return true;
+			}
+			catch(Exception Ex)
+			{
+				throw new CommandLineArgumentException(String.Format("The argument '{0}{1}' does not specify a valid directory name", Prefix, StringValue), Ex);
+			}
+		}
+
+		/// <summary>
+		/// Tries to gets the value specified by an argument with the given prefix.
+		/// </summary>
+		/// <param name="Prefix">The argument prefix (eg. "-Foo="). Must end with an '=' character.</param>
+		/// <param name="Value">Value of the argument, if found</param>
+		/// <returns>True if the argument was found (and Value was set), false otherwise.</returns>
+		public bool TryGetValue<T>(string Prefix, out T Value) where T : struct
+		{
+			// Try to get the string value of this argument
+			string StringValue;
+			if(!TryGetValue(Prefix, out StringValue))
+			{
+				Value = new T();
+				return false;
+			}
+
+			// Try to parse it. If it fails, throw an exception.
+			try
+			{
+				Value = (T)Enum.Parse(typeof(T), StringValue, true);
+				return true;
+			}
+			catch(Exception Ex)
+			{
+				throw new CommandLineArgumentException(String.Format("The argument '{0}{1}' does not specify a valid {2}", Prefix, StringValue, typeof(T).Name), Ex);
+			}
+		}
+
+		/// <summary>
+		/// Returns all arguments with the given prefix.
+		/// </summary>
+		/// <param name="Prefix">The argument prefix (eg. "-Foo="). Must end with an '=' character.</param>
+		/// <returns>Sequence of values for the given prefix.</returns>
+		public IEnumerable<string> GetValues(string Prefix)
+		{
+			CheckValidPrefix(Prefix);
+
+			int Index;
+			if(ArgumentToFirstIndex.TryGetValue(Prefix, out Index))
+			{
+				for(; Index != -1; Index = NextArgumentIndex[Index])
+				{
+					UsedArguments.Set(Index, true);
+					yield return Arguments[Index].Substring(Prefix.Length);
+				}
+			}
+		}
+
+		/// <summary>
+		/// Returns all arguments with the given prefix, allowing multiple arguments to be specified in a single argument with a separator character.
+		/// </summary>
+		/// <param name="Prefix">The argument prefix (eg. "-Foo="). Must end with an '=' character.</param>
+		/// <param name="Separator">The separator character (eg. '+')</param>
+		/// <returns>Sequence of values for the given prefix.</returns>
+		public IEnumerable<string> GetValues(string Prefix, char Separator)
+		{
+			foreach(string Value in GetValues(Prefix))
+			{
+				foreach(string SplitValue in Value.Split(Separator))
+				{
+					yield return SplitValue;
+				}
+			}
+		}
+
+		/// <summary>
+		/// Applies these arguments to fields with the [CommandLine] attribute in the given object.
+		/// </summary>
+		/// <param name="TargetObject">The object to configure</param>
+		public void ApplyTo(object TargetObject)
+		{
+			// Build a mapping from name to field and attribute for this object
+			List<string> MissingArguments = new List<string>();
+			for(Type TargetType = TargetObject.GetType(); TargetType != typeof(object); TargetType = TargetType.BaseType)
+			{
+				foreach(FieldInfo FieldInfo in TargetType.GetFields(BindingFlags.Instance | BindingFlags.GetField | BindingFlags.GetProperty | BindingFlags.Public | BindingFlags.NonPublic | BindingFlags.DeclaredOnly))
+				{
+					// If any attribute is required, keep track of it so we can include an error for it
+					string RequiredPrefix = null;
+
+					// Keep track of whether a value has already been assigned to this field
+					string AssignedArgument = null;
+
+					// Loop through all the attributes for different command line options that can modify it
+					IEnumerable<CommandLineAttribute> Attributes = FieldInfo.GetCustomAttributes<CommandLineAttribute>();
+					foreach(CommandLineAttribute Attribute in Attributes)
+					{
+						// Get the appropriate prefix for this attribute
+						string Prefix = Attribute.Prefix;
+						if(Prefix == null)
+						{
+							if(FieldInfo.FieldType == typeof(bool))
+							{
+								Prefix = String.Format("-{0}", FieldInfo.Name);
+							}
+							else
+							{
+								Prefix = String.Format("-{0}=", FieldInfo.Name);
+							}
+						}
+						else
+						{
+							if(FieldInfo.FieldType != typeof(bool) && Attribute.Value == null && !Prefix.EndsWith("=") && !Prefix.EndsWith(":"))
+							{
+								Prefix = Prefix + "=";
+							}
+						}
+
+						// Get the value with the correct prefix
+						int FirstIndex;
+						if(ArgumentToFirstIndex.TryGetValue(Prefix, out FirstIndex))
+						{
+							for(int Index = FirstIndex; Index != -1; Index = NextArgumentIndex[Index])
+							{
+								// Get the argument text
+								string Argument = Arguments[Index];
+
+								// Get the text for this value
+								string ValueText;
+								if(FlagArguments.Get(Index))
+								{
+									ValueText = Attribute.Value ?? "true";
+								}
+								else
+								{
+									ValueText = Argument.Substring(Prefix.Length);
+								}
+
+								// Apply the value to the field
+								if(ApplyArgument(TargetObject, FieldInfo, Argument, ValueText, AssignedArgument))
+								{
+									AssignedArgument = Argument;
+								}
+
+								// Mark this argument as used
+								UsedArguments.Set(Index, true);
+							}
+						}
+
+						// If this attribute is marked as required, keep track of it so we can warn if the field is not assigned to
+						if(Attribute.Required && RequiredPrefix == null)
+						{
+							RequiredPrefix = Prefix;
+						}
+					}
+
+					// Make sure that this field has been assigned to
+					if(AssignedArgument == null && RequiredPrefix != null)
+					{
+						MissingArguments.Add(RequiredPrefix);
+					}
+				}
+			}
+
+			// If any arguments were missing, print an error about them
+			if(MissingArguments.Count > 0)
+			{
+				if(MissingArguments.Count == 1)
+				{
+					throw new CommandLineArgumentException(String.Format("Missing {0} argument", MissingArguments[0].Replace("=", "=...")));
+				}
+				else
+				{
+					throw new CommandLineArgumentException(String.Format("Missing {0} arguments", Utils.FormatList(MissingArguments.Select(x => x.Replace("=", "=...")))));
+				}
+			}
+		}
+
+		/// <summary>
+		/// Checks that there are no unused arguments (and warns if there are)
+		/// </summary>
+		public void CheckAllArgumentsUsed()
+		{
+			// Find all the unused arguments
+			List<string> RemainingArguments = new List<string>();
+			for(int Idx = 0; Idx < Arguments.Length; Idx++)
+			{
+				if(!UsedArguments[Idx])
+				{
+					RemainingArguments.Add(Arguments[Idx]);
+				}
+			}
+
+			// Output a warning
+			if(RemainingArguments.Count > 0)
+			{
+				if(RemainingArguments.Count == 1)
+				{
+					Log.TraceWarning("Invalid argument: {0}", RemainingArguments[0]);
+				}
+				else
+				{
+					Log.TraceWarning("Invalid arguments:\n{0}", String.Join("\n", RemainingArguments));
+				}
+			}
+		}
+
+		/// <summary>
+		/// Checks that a given string is a valid argument prefix
+		/// </summary>
+		/// <param name="Prefix">The prefix to check</param>
+		private static void CheckValidPrefix(string Prefix)
+		{
+			if(Prefix.Length == 0)
+			{
+				throw new ArgumentException("Argument prefix cannot be empty.");
+			}
+			else if(Prefix[0] != '-')
+			{
+				throw new ArgumentException("Argument prefix must begin with a hyphen.");
+			}
+			else if(!ValueSeparators.Contains(Prefix[Prefix.Length - 1]))
+			{
+				throw new ArgumentException(String.Format("Argument prefix must end with '{0}'", String.Join("' or '", ValueSeparators)));
+			}
+		}
+
+		/// <summary>
+		/// Parses and assigns a value to a field
+		/// </summary>
+		/// <param name="TargetObject">The target object to assign values to</param>
+		/// <param name="Field">The field to assign the value to</param>
+		/// <param name="ArgumentText">The full argument text</param>
+		/// <param name="ValueText">Argument text</param>
+		/// <param name="PreviousArgumentText">The previous text used to configure this field</param>
+		/// <returns>True if the value was assigned to the field, false otherwise</returns>
+		private static bool ApplyArgument(object TargetObject, FieldInfo Field, string ArgumentText, string ValueText, string PreviousArgumentText)
+		{
+			// Try to parse the value
+			object Value;
+			if(!TryParseValue(Field.FieldType, ValueText, out Value))
+			{
+				Log.TraceWarning("Unable to parse value for argument '{0}'.", ArgumentText);
+				return false;
+			}
+
+			// Check if the field type implements ICollection<>. If so, we can take multiple values.
+			Type CollectionType = null;
+			foreach (Type InterfaceType in Field.FieldType.GetInterfaces())
+			{
+				if (InterfaceType.IsGenericType && InterfaceType.GetGenericTypeDefinition() == typeof(ICollection<>))
+				{
+					CollectionType = InterfaceType;
+					break;
+				}
+			}
+
+			// Try to assign values to the target field
+			if (CollectionType == null)
+			{
+				// Check if this field has already been assigned to. Output a warning if the previous value is in conflict with the new one.
+				if(PreviousArgumentText != null)
+				{
+					object PreviousValue = Field.GetValue(TargetObject);
+					if(!PreviousValue.Equals(Value))
+					{
+						Log.TraceWarning("Argument '{0}' conflicts with '{1}'; ignoring.", ArgumentText, PreviousArgumentText);
+					}
+					return false;
+				}
+
+				// Set the value on the target object
+				Field.SetValue(TargetObject, Value);
+				return true;
+			}
+			else
+			{
+				// Call the 'Add' method on the collection
+				CollectionType.InvokeMember("Add", BindingFlags.InvokeMethod, null, Field.GetValue(TargetObject), new object[] { Value });
+				return true;
+			}
+		}
+
+		/// <summary>
+		/// Attempts to parse the given string to a value
+		/// </summary>
+		/// <param name="FieldType">Type of the field to convert to</param>
+		/// <param name="Text">The value text</param>
+		/// <param name="Value">On success, contains the parsed object</param>
+		/// <returns>True if the text could be parsed, false otherwise</returns>
+		private static bool TryParseValue(Type FieldType, string Text, out object Value)
+		{
+			if(FieldType.IsEnum)
+			{
+				// Special handling for enums; parse the value ignoring case.
+				try
+				{
+					Value = Enum.Parse(FieldType, Text, true);
+					return true;
+				}
+				catch(ArgumentException)
+				{
+					Value = null;
+					return false;
+				}
+			}
+			else if(FieldType == typeof(FileReference))
+			{
+				// Construct a file reference from the string
+				try
+				{
+					Value = new FileReference(Text);
+					return true;
+				}
+				catch
+				{
+					Value = null;
+					return false;
+				}
+			}
+			else if(FieldType == typeof(DirectoryReference))
+			{
+				// Construct a file reference from the string
+				try
+				{
+					Value = new DirectoryReference(Text);
+					return true;
+				}
+				catch
+				{
+					Value = null;
+					return false;
+				}
+			}
+			else
+			{
+				// Otherwise let the framework convert between types
+				try
+				{
+					Value = Convert.ChangeType(Text, FieldType);
+					return true;
+				}
+				catch(InvalidCastException)
+				{
+					Value = null;
+					return false;
+				}
+			}
+		}
+
+		/// <summary>
+		/// Obtains an enumerator for the argument list
+		/// </summary>
+		/// <returns>IEnumerator interface</returns>
+		IEnumerator IEnumerable.GetEnumerator()
+		{
+			return Arguments.GetEnumerator();
+		}
+
+		/// <summary>
+		/// Obtains an enumerator for the argument list
+		/// </summary>
+		/// <returns>Generic IEnumerator interface</returns>
+		public IEnumerator<string> GetEnumerator()
+		{
+			return ((IEnumerable<string>)Arguments).GetEnumerator();
+		}
+
+		/// <summary>
+		/// Gets the raw argument array
+		/// </summary>
+		/// <returns>Array of arguments</returns>
+		public string[] GetRawArray()
+		{
+			return Arguments;
+		}
+	}
+}
diff --git a/Engine/Source/Programs/UnrealBuildTool/System/ToolMode.cs b/Engine/Source/Programs/UnrealBuildTool/System/ToolMode.cs
--- a/Engine/Source/Programs/UnrealBuildTool/System/ToolMode.cs
+++ b/Engine/Source/Programs/UnrealBuildTool/System/ToolMode.cs
@@ -38,6 +38,6 @@ namespace UnrealBuildTool
 		/// </summary>
 		/// <param name="Arguments">List of command line arguments</param>
 		/// <returns>Exit code for the process</returns>
-		public abstract int Execute(List<string> Arguments);
+		public abstract int Execute(CommandLineArguments Arguments);
 	}
 }
diff --git a/Engine/Source/Programs/UnrealBuildTool/UnrealBuildTool.cs b/Engine/Source/Programs/UnrealBuildTool/UnrealBuildTool.cs
--- a/Engine/Source/Programs/UnrealBuildTool/UnrealBuildTool.cs
+++ b/Engine/Source/Programs/UnrealBuildTool/UnrealBuildTool.cs
@@ -933,23 +933,31 @@ namespace UnrealBuildTool
 			/// </summary>
 			[CommandLine]
 			public string Mode = "Default";
+
+			/// <summary>
+			/// Initialize the options with the given commnad line arguments
+			/// </summary>
+			/// <param name="Arguments"></param>
+			public GlobalOptions(CommandLineArguments Arguments)
+			{
+				Arguments.ApplyTo(this);
+			}
 		}
 
 		/// <summary>
 		/// Main entry point. Parses any global options and initializes the logging system, then invokes the appropriate command.
 		/// </summary>
-		/// <param name="Arguments">Command line arguments</param>
+		/// <param name="ArgumentsArray">Command line arguments</param>
 		/// <returns>Zero on success, non-zero on error</returns>
-		private static int Main(string[] Arguments)
+		private static int Main(string[] ArgumentsArray)
 		{
 			try
 			{
-				// Copy the arguments into a list. We remove entries from this list as they are parsed, so after we've parsed all the options we can we can identify any incorrect arguments and warn the user appropriately.
-				List<string> RemainingArguments = new List<string>(Arguments);
+				// Parse the command line arguments
+				CommandLineArguments Arguments = new CommandLineArguments(ArgumentsArray);
 
 				// Parse the global options
-				GlobalOptions Options = new GlobalOptions();
-				CommandLine.ParseAndRemoveArguments(RemainingArguments, Options);
+				GlobalOptions Options = new GlobalOptions(Arguments);
 
 				// Configure the log system
 				Log.OutputLevel = Options.LogOutputLevel;
@@ -987,7 +995,7 @@ namespace UnrealBuildTool
 
 				// Create the appropriate handler
 				ToolMode Mode = (ToolMode)Activator.CreateInstance(ModeType);
-				return Mode.Execute(RemainingArguments);
+				return Mode.Execute(Arguments);
 			}
 			catch (BuildException Ex)
 			{
@@ -2327,9 +2335,9 @@ namespace UnrealBuildTool
 	[ToolMode("Default")]
 	class DefaultToolMode : ToolMode
 	{
-		public override int Execute(List<string> Arguments)
+		public override int Execute(CommandLineArguments Arguments)
 		{
-			return UnrealBuildTool.GuardedMain(Arguments.ToArray());
+			return UnrealBuildTool.GuardedMain(Arguments.GetRawArray());
 		}
 	}
 }
diff --git a/Engine/Source/Programs/UnrealBuildTool/UnrealBuildTool.csproj b/Engine/Source/Programs/UnrealBuildTool/UnrealBuildTool.csproj
--- a/Engine/Source/Programs/UnrealBuildTool/UnrealBuildTool.csproj
+++ b/Engine/Source/Programs/UnrealBuildTool/UnrealBuildTool.csproj
@@ -175,6 +175,7 @@
     <Compile Include="ProjectFiles\Xcode\XcodeProjectFileGenerator.cs" />
     <Compile Include="System\Asn.cs" />
     <Compile Include="System\BuildVersion.cs" />
+    <Compile Include="System\CommandLineArguments.cs" />
     <Compile Include="System\CppCompileEnvironment.cs" />
     <Compile Include="System\CppIncludeParser.cs" />
     <Compile Include="System\DataDrivenPlatformInfo.cs" />
@@ -204,6 +205,7 @@
     <Compile Include="System\ProgressWriter.cs" />
     <Compile Include="System\RulesAssembly.cs" />
     <Compile Include="System\ScopedTimer.cs" />
+    <Compile Include="System\SingleInstanceMutex.cs" />
     <Compile Include="System\SourceFileWorkingSet.cs" />
     <Compile Include="System\StartupTraceListener.cs" />
     <Compile Include="System\SupportedConfigurationsAttribute.cs" />
