UBT: Remove reference to local file from project.

#rnx

diff --git a/Engine/Source/Programs/UnrealBuildTool/UnrealBuildTool.csproj b/Engine/Source/Programs/UnrealBuildTool/UnrealBuildTool.csproj
--- a/Engine/Source/Programs/UnrealBuildTool/UnrealBuildTool.csproj
+++ b/Engine/Source/Programs/UnrealBuildTool/UnrealBuildTool.csproj
@@ -205,7 +205,6 @@
     <Compile Include="System\ProgressWriter.cs" />
     <Compile Include="System\RulesAssembly.cs" />
     <Compile Include="System\ScopedTimer.cs" />
-    <Compile Include="System\SingleInstanceMutex.cs" />
     <Compile Include="System\SourceFileWorkingSet.cs" />
     <Compile Include="System\StartupTraceListener.cs" />
     <Compile Include="System\SupportedConfigurationsAttribute.cs" />
