UBT: Remove separate setting for SDK output verbosity. Now we have genuine log output from UBT, we can just send all output there.

#rnx

diff --git a/Engine/Source/Programs/UnrealBuildTool/Configuration/UEBuildPlatformFactory.cs b/Engine/Source/Programs/UnrealBuildTool/Configuration/UEBuildPlatformFactory.cs
--- a/Engine/Source/Programs/UnrealBuildTool/Configuration/UEBuildPlatformFactory.cs
+++ b/Engine/Source/Programs/UnrealBuildTool/Configuration/UEBuildPlatformFactory.cs
@@ -1,4 +1,4 @@
-﻿// Copyright 1998-2018 Epic Games, Inc. All Rights Reserved.
+// Copyright 1998-2018 Epic Games, Inc. All Rights Reserved.
 
 using System;
 using System.Collections.Generic;
@@ -16,12 +16,12 @@ namespace UnrealBuildTool
 		/// <summary>
 		/// Attempt to register a build platform, checking whether it is a valid platform in installed builds
 		/// </summary>
-		public void TryRegisterBuildPlatforms(SDKOutputLevel OutputLevel, bool bValidatingPlatforms)
+		public void TryRegisterBuildPlatforms(bool bValidatingPlatforms)
 		{
 			// We need all platforms to be registered when we run -validateplatform command to check SDK status of each
 			if (bValidatingPlatforms || InstalledPlatformInfo.IsValidPlatform(TargetPlatform))
 			{
-				RegisterBuildPlatforms(OutputLevel);
+				RegisterBuildPlatforms();
 			}
 		}
 
@@ -36,6 +36,6 @@ namespace UnrealBuildTool
 		/// <summary>
 		/// Register the platform with the UEBuildPlatform class
 		/// </summary>
-		protected abstract void RegisterBuildPlatforms(SDKOutputLevel OutputLevel);
+		protected abstract void RegisterBuildPlatforms();
 	}
 }
diff --git a/Engine/Source/Programs/UnrealBuildTool/Configuration/UEBuildPlatformSDK.cs b/Engine/Source/Programs/UnrealBuildTool/Configuration/UEBuildPlatformSDK.cs
--- a/Engine/Source/Programs/UnrealBuildTool/Configuration/UEBuildPlatformSDK.cs
+++ b/Engine/Source/Programs/UnrealBuildTool/Configuration/UEBuildPlatformSDK.cs
@@ -21,34 +21,12 @@ namespace UnrealBuildTool
 		/// </summary>
 		Valid,
 
-
 		/// <summary>
 		/// Could not find the desired SDK, SDK setup failed, etc.		
 		/// </summary>
 		Invalid,
 	};
 
-	/// <summary>
-	/// Output level for diagnostic messages when initializing the platform SDK
-	/// </summary>
-	enum SDKOutputLevel
-	{
-		/// <summary>
-		/// No output
-		/// </summary>
-		Quiet,
-
-		/// <summary>
-		/// Minimal output
-		/// </summary>
-		Minimal,
-
-		/// <summary>
-		/// All messages
-		/// </summary>
-		Verbose,
-	}
-
 	/// <summary>
 	/// SDK for a platform
 	/// </summary>
@@ -353,20 +331,16 @@ namespace UnrealBuildTool
 		/// <summary>
 		/// Runs install/uninstall hooks for SDK
 		/// </summary>
-		/// <param name="OutputLevel">The output level for diagnostic messages</param>
 		/// <param name="PlatformSDKRoot">absolute path to platform SDK root</param>
 		/// <param name="SDKVersionString">version string to run for (can be empty!)</param>
 		/// <param name="Hook">which one of hooks to run</param>
 		/// <param name="bHookCanBeNonExistent">whether a non-existing hook means failure</param>
 		/// <returns>true if succeeded</returns>
-		protected virtual bool RunAutoSDKHooks(SDKOutputLevel OutputLevel, string PlatformSDKRoot, string SDKVersionString, SDKHookType Hook, bool bHookCanBeNonExistent = true)
+		protected virtual bool RunAutoSDKHooks(string PlatformSDKRoot, string SDKVersionString, SDKHookType Hook, bool bHookCanBeNonExistent = true)
 		{
 			if (!IsAutoSDKSafe())
 			{
-				Console.ForegroundColor = ConsoleColor.Red;
-				LogAutoSDK(OutputLevel, GetSDKTargetPlatformName() + " attempted to run SDK hook which could have damaged manual SDK install!");
-				Console.ResetColor();
-
+				Log.TraceLog(GetSDKTargetPlatformName() + " attempted to run SDK hook which could have damaged manual SDK install!");
 				return false;
 			}
 			if (SDKVersionString != "")
@@ -376,7 +350,7 @@ namespace UnrealBuildTool
 
 				if (File.Exists(HookExe))
 				{
-					LogAutoSDK(OutputLevel, "Running {0} hook {1}", Hook, HookExe);
+					Log.TraceLog("Running {0} hook {1}", Hook, HookExe);
 
 					// run it
 					Process HookProcess = new Process();
@@ -390,7 +364,7 @@ namespace UnrealBuildTool
 					//HookProcess.StartInfo.RedirectStandardOutput = true;
 					//HookProcess.StartInfo.RedirectStandardError = true;					
 
-					using (ScopedTimer HookTimer = new ScopedTimer("Time to run hook: ", (OutputLevel >= SDKOutputLevel.Minimal) ? LogEventType.Log : LogEventType.Verbose))
+					using (ScopedTimer HookTimer = new ScopedTimer("Time to run hook: ", LogEventType.Log))
 					{
 						//installers may require administrator access to succeed. so run as an admmin.
 						HookProcess.StartInfo.Verb = "runas";
@@ -402,7 +376,7 @@ namespace UnrealBuildTool
 					//LogAutoSDK(HookProcess.StandardError.ReadToEnd());
 					if (HookProcess.ExitCode != 0)
 					{
-						LogAutoSDK(OutputLevel, "Hook exited uncleanly (returned {0}), considering it failed.", HookProcess.ExitCode);
+						Log.TraceLog("Hook exited uncleanly (returned {0}), considering it failed.", HookProcess.ExitCode);
 						return false;
 					}
 
@@ -410,12 +384,12 @@ namespace UnrealBuildTool
 				}
 				else
 				{
-					LogAutoSDK(OutputLevel, "File {0} does not exist", HookExe);
+					Log.TraceLog("File {0} does not exist", HookExe);
 				}
 			}
 			else
 			{
-				LogAutoSDK(OutputLevel, "Version string is blank for {0}. Can't determine {1} hook.", PlatformSDKRoot, Hook.ToString());
+				Log.TraceLog("Version string is blank for {0}. Can't determine {1} hook.", PlatformSDKRoot, Hook.ToString());
 			}
 
 			return bHookCanBeNonExistent;
@@ -427,9 +401,8 @@ namespace UnrealBuildTool
 		/// TargetPlatformManagerModule.cpp
 		/// </summary>
 		/// <param name="PlatformSDKRoot">absolute path to platform SDK</param>
-		/// <param name="OutputLevel">The output level for diagnostic messages</param>
 		/// <returns>true if succeeded</returns>
-		protected bool SetupEnvironmentFromAutoSDK(string PlatformSDKRoot, SDKOutputLevel OutputLevel)
+		protected bool SetupEnvironmentFromAutoSDK(string PlatformSDKRoot)
 		{
 			string EnvVarFile = Path.Combine(PlatformSDKRoot, SDKEnvironmentVarsFile);
 			if (File.Exists(EnvVarFile))
@@ -455,8 +428,8 @@ namespace UnrealBuildTool
 						string[] Parts = VariableString.Split('=');
 						if (Parts.Length != 2)
 						{
-							LogAutoSDK(OutputLevel, "Incorrect environment variable declaration:");
-							LogAutoSDK(OutputLevel, VariableString);
+							Log.TraceLog("Incorrect environment variable declaration:");
+							Log.TraceLog(VariableString);
 							return false;
 						}
 
@@ -487,10 +460,7 @@ namespace UnrealBuildTool
 					{
 						string EnvVarName = EnvVarNames[i];
 						string EnvVarValue = EnvVarValues[i];
-						if (OutputLevel >= SDKOutputLevel.Verbose)
-						{
-							LogAutoSDK(OutputLevel, "Setting variable '{0}' to '{1}'", EnvVarName, EnvVarValue);
-						}
+						Log.TraceVerbose("Setting variable '{0}' to '{1}'", EnvVarName, EnvVarValue);
 						Environment.SetEnvironmentVariable(EnvVarName, EnvVarValue);
 					}
 
@@ -505,7 +475,7 @@ namespace UnrealBuildTool
                     }
                     else
                     {
-                        LogAutoSDK(OutputLevel, "Path environment variable is null during AutoSDK");
+                        Log.TraceVerbose("Path environment variable is null during AutoSDK");
                     }
 
 					List<String> ModifiedPathVars = new List<string>();
@@ -518,7 +488,7 @@ namespace UnrealBuildTool
 						{
 							if (PathVar.IndexOf(PathRemove, StringComparison.OrdinalIgnoreCase) >= 0)
 							{
-								LogAutoSDK(OutputLevel, "Removing Path: '{0}'", PathVar);
+								Log.TraceVerbose("Removing Path: '{0}'", PathVar);
 								ModifiedPathVars.Remove(PathVar);
 							}
 						}
@@ -532,7 +502,7 @@ namespace UnrealBuildTool
 						{
 							if (String.Compare(PathAdd, PathVar, true) == 0)
 							{
-								LogAutoSDK(OutputLevel, "Removing Path: '{0}'", PathVar);
+								Log.TraceVerbose("Removing Path: '{0}'", PathVar);
 								ModifiedPathVars.Remove(PathVar);
 							}
 						}
@@ -543,7 +513,7 @@ namespace UnrealBuildTool
 					{
 						if (!ModifiedPathVars.Contains(PathAdd))
 						{
-							LogAutoSDK(OutputLevel, "Adding Path: '{0}'", PathAdd);
+							Log.TraceVerbose("Adding Path: '{0}'", PathAdd);
 							ModifiedPathVars.Add(PathAdd);
 						}
 					}
@@ -573,7 +543,7 @@ namespace UnrealBuildTool
 			}
 			else
 			{
-				LogAutoSDK(OutputLevel, "Cannot set up environment for {1} because command file {1} does not exist.", PlatformSDKRoot, EnvVarFile);
+				Log.TraceLog("Cannot set up environment for {1} because command file {2} does not exist.", PlatformSDKRoot, EnvVarFile);
 			}
 
 			return false;
@@ -706,7 +676,7 @@ namespace UnrealBuildTool
 		/// AutoSDKs are SDKs that have not been setup through a formal installer, but rather come from
 		/// a source control directory, or other local copy.
 		/// </summary>
-		private void SetupAutoSDK(SDKOutputLevel OutputLevel)
+		private void SetupAutoSDK()
 		{
 			if (IsAutoSDKSafe() && PlatformSupportsAutoSDKs() && HasAutoSDKSystemEnabled())
 			{
@@ -721,27 +691,27 @@ namespace UnrealBuildTool
 					GetCurrentlyInstalledSDKString(AutoSDKRoot, out CurrentSDKString);
 
 					// switch over (note that version string can be empty)
-					if (!RunAutoSDKHooks(OutputLevel, AutoSDKRoot, CurrentSDKString, SDKHookType.Uninstall))
+					if (!RunAutoSDKHooks(AutoSDKRoot, CurrentSDKString, SDKHookType.Uninstall))
 					{
-						LogAutoSDK(OutputLevel, "Failed to uninstall currently installed SDK {0}", CurrentSDKString);
+						Log.TraceLog("Failed to uninstall currently installed SDK {0}", CurrentSDKString);
 						InvalidateCurrentlyInstalledAutoSDK();
 						return;
 					}
 					// delete Manifest file to avoid multiple uninstalls
 					InvalidateCurrentlyInstalledAutoSDK();
 
-					if (!RunAutoSDKHooks(OutputLevel, AutoSDKRoot, GetRequiredSDKString(), SDKHookType.Install, false))
+					if (!RunAutoSDKHooks(AutoSDKRoot, GetRequiredSDKString(), SDKHookType.Install, false))
 					{
-						LogAutoSDK(OutputLevel, "Failed to install required SDK {0}.  Attemping to uninstall", GetRequiredSDKString());
-						RunAutoSDKHooks(OutputLevel, AutoSDKRoot, GetRequiredSDKString(), SDKHookType.Uninstall, false);
+						Log.TraceLog("Failed to install required SDK {0}.  Attemping to uninstall", GetRequiredSDKString());
+						RunAutoSDKHooks(AutoSDKRoot, GetRequiredSDKString(), SDKHookType.Uninstall, false);
 						return;
 					}
 
 					string EnvVarFile = Path.Combine(AutoSDKRoot, SDKEnvironmentVarsFile);
 					if (!File.Exists(EnvVarFile))
 					{
-						LogAutoSDK(OutputLevel, "Installation of required SDK {0}.  Did not generate Environment file {1}", GetRequiredSDKString(), EnvVarFile);
-						RunAutoSDKHooks(OutputLevel, AutoSDKRoot, GetRequiredSDKString(), SDKHookType.Uninstall, false);
+						Log.TraceLog("Installation of required SDK {0}.  Did not generate Environment file {1}", GetRequiredSDKString(), EnvVarFile);
+						RunAutoSDKHooks(AutoSDKRoot, GetRequiredSDKString(), SDKHookType.Uninstall, false);
 						return;
 					}
 
@@ -750,7 +720,7 @@ namespace UnrealBuildTool
 				}
 
 				// fixup process environment to match autosdk
-				SetupEnvironmentFromAutoSDK(OutputLevel);
+				SetupEnvironmentFromAutoSDK();
 			}
 		}
 
@@ -779,14 +749,14 @@ namespace UnrealBuildTool
 		/* Whether or not we should try to automatically switch SDKs when asked to validate the platform's SDK state. */
 		public static bool bAllowAutoSDKSwitching = true;
 
-		public SDKStatus SetupEnvironmentFromAutoSDK(SDKOutputLevel OutputLevel)
+		public SDKStatus SetupEnvironmentFromAutoSDK()
 		{
 			string PlatformSDKRoot = GetPathToPlatformAutoSDKs();
 
 			// load environment variables from current SDK
-			if (!SetupEnvironmentFromAutoSDK(PlatformSDKRoot, OutputLevel))
+			if (!SetupEnvironmentFromAutoSDK(PlatformSDKRoot))
 			{
-				LogAutoSDK(OutputLevel, "Failed to load environment from required SDK {0}", GetRequiredSDKString());
+				Log.TraceLog("Failed to load environment from required SDK {0}", GetRequiredSDKString());
 				InvalidateCurrentlyInstalledAutoSDK();
 				return SDKStatus.Invalid;
 			}
@@ -812,7 +782,7 @@ namespace UnrealBuildTool
 		}
 
 		// Arbitrates between manual SDKs and setting up AutoSDK based on program options and platform preferences.
-		public void ManageAndValidateSDK(SDKOutputLevel OutputLevel)
+		public void ManageAndValidateSDK()
 		{
 			// do not modify installed manifests if parent process has already set everything up.
 			// this avoids problems with determining IsAutoSDKSafe and doing an incorrect invalidate.
@@ -822,7 +792,7 @@ namespace UnrealBuildTool
 				bool bHasRequiredManualSDK = HasRequiredManualSDK() == SDKStatus.Valid;
 				if (IsAutoSDKSafe() && (PreferAutoSDK() || !bHasRequiredManualSDK))
 				{
-					SetupAutoSDK(OutputLevel);
+					SetupAutoSDK();
 					bSetSomeSDK = true;
 				}
 
@@ -839,14 +809,10 @@ namespace UnrealBuildTool
 				}
 			}
 
-
-			if (OutputLevel != SDKOutputLevel.Quiet)
-			{
-				PrintSDKInfo(OutputLevel);
-			}
+			PrintSDKInfo();
 		}
 
-		public void PrintSDKInfo(SDKOutputLevel OutputLevel)
+		public void PrintSDKInfo()
 		{
 			if (HasRequiredSDKsInstalled() == SDKStatus.Valid)
 			{
@@ -854,36 +820,20 @@ namespace UnrealBuildTool
 				if (HasSetupAutoSDK())
 				{
 					string PlatformSDKRoot = GetPathToPlatformAutoSDKs();
-					LogAutoSDK(OutputLevel, GetSDKTargetPlatformName() + " using SDK from: " + Path.Combine(PlatformSDKRoot, GetRequiredSDKString()));
+					Log.TraceLog(GetSDKTargetPlatformName() + " using SDK from: " + Path.Combine(PlatformSDKRoot, GetRequiredSDKString()));
 				}
 				else if (bHasRequiredManualSDK)
 				{
-					LogAutoSDK(OutputLevel, this.ToString() + " using manually installed SDK " + GetRequiredSDKString());
+					Log.TraceLog(this.ToString() + " using manually installed SDK " + GetRequiredSDKString());
 				}
 				else
 				{
-					LogAutoSDK(OutputLevel, this.ToString() + " setup error.  Inform platform team.");
+					Log.TraceLog(this.ToString() + " setup error.  Inform platform team.");
 				}
 			}
 			else
 			{
-				LogAutoSDK(OutputLevel, this.ToString() + " has no valid SDK");
-			}
-		}
-
-		protected static void LogAutoSDK(SDKOutputLevel OutputLevel, string Format, params object[] Args)
-		{
-			if (OutputLevel != SDKOutputLevel.Quiet)
-			{
-				Log.WriteLine(LogEventType.Log, Format, Args);
-			}
-		}
-
-		protected static void LogAutoSDK(SDKOutputLevel OutputLevel, String Message)
-		{
-			if (OutputLevel != SDKOutputLevel.Quiet)
-			{
-				Log.WriteLine(LogEventType.Log, Message);
+				Log.TraceLog(this.ToString() + " has no valid SDK");
 			}
 		}
 
diff --git a/Engine/Source/Programs/UnrealBuildTool/Platform/Android/UEBuildAndroid.cs b/Engine/Source/Programs/UnrealBuildTool/Platform/Android/UEBuildAndroid.cs
--- a/Engine/Source/Programs/UnrealBuildTool/Platform/Android/UEBuildAndroid.cs
+++ b/Engine/Source/Programs/UnrealBuildTool/Platform/Android/UEBuildAndroid.cs
@@ -661,10 +661,10 @@ namespace UnrealBuildTool
 			get { return UnrealTargetPlatform.Android; }
 		}
 
-		protected override void RegisterBuildPlatforms(SDKOutputLevel OutputLevel)
+		protected override void RegisterBuildPlatforms()
 		{
 			AndroidPlatformSDK SDK = new AndroidPlatformSDK();
-			SDK.ManageAndValidateSDK(OutputLevel);
+			SDK.ManageAndValidateSDK();
 
 			if ((ProjectFileGenerator.bGenerateProjectFiles == true) || (SDK.HasRequiredSDKsInstalled() == SDKStatus.Valid) || Environment.GetEnvironmentVariable("IsBuildMachine") == "1")
 			{
diff --git a/Engine/Source/Programs/UnrealBuildTool/Platform/HTML5/UEBuildHTML5.cs b/Engine/Source/Programs/UnrealBuildTool/Platform/HTML5/UEBuildHTML5.cs
--- a/Engine/Source/Programs/UnrealBuildTool/Platform/HTML5/UEBuildHTML5.cs
+++ b/Engine/Source/Programs/UnrealBuildTool/Platform/HTML5/UEBuildHTML5.cs
@@ -274,10 +274,10 @@ namespace UnrealBuildTool
 		/// <summary>
 		/// Register the platform with the UEBuildPlatform class
 		/// </summary>
-		protected override void RegisterBuildPlatforms(SDKOutputLevel OutputLevel)
+		protected override void RegisterBuildPlatforms()
 		{
 			HTML5PlatformSDK SDK = new HTML5PlatformSDK();
-			SDK.ManageAndValidateSDK(OutputLevel);
+			SDK.ManageAndValidateSDK();
 
 			// Make sure the SDK is installed
 			if ((ProjectFileGenerator.bGenerateProjectFiles == true) || (SDK.HasRequiredSDKsInstalled() == SDKStatus.Valid))
diff --git a/Engine/Source/Programs/UnrealBuildTool/Platform/IOS/UEBuildIOS.cs b/Engine/Source/Programs/UnrealBuildTool/Platform/IOS/UEBuildIOS.cs
--- a/Engine/Source/Programs/UnrealBuildTool/Platform/IOS/UEBuildIOS.cs
+++ b/Engine/Source/Programs/UnrealBuildTool/Platform/IOS/UEBuildIOS.cs
@@ -991,10 +991,10 @@ namespace UnrealBuildTool
 		/// <summary>
 		/// Register the platform with the UEBuildPlatform class
 		/// </summary>
-		protected override void RegisterBuildPlatforms(SDKOutputLevel OutputLevel)
+		protected override void RegisterBuildPlatforms()
 		{
 			IOSPlatformSDK SDK = new IOSPlatformSDK();
-			SDK.ManageAndValidateSDK(OutputLevel);
+			SDK.ManageAndValidateSDK();
 
 			// Register this build platform for IOS
 			Log.TraceVerbose("        Registering for {0}", UnrealTargetPlatform.IOS.ToString());
diff --git a/Engine/Source/Programs/UnrealBuildTool/Platform/Linux/UEBuildLinux.cs b/Engine/Source/Programs/UnrealBuildTool/Platform/Linux/UEBuildLinux.cs
--- a/Engine/Source/Programs/UnrealBuildTool/Platform/Linux/UEBuildLinux.cs
+++ b/Engine/Source/Programs/UnrealBuildTool/Platform/Linux/UEBuildLinux.cs
@@ -758,10 +758,10 @@ namespace UnrealBuildTool
 		/// <summary>
 		/// Register the platform with the UEBuildPlatform class
 		/// </summary>
-		protected override void RegisterBuildPlatforms(SDKOutputLevel OutputLevel)
+		protected override void RegisterBuildPlatforms()
 		{
 			LinuxPlatformSDK SDK = new LinuxPlatformSDK();
-			SDK.ManageAndValidateSDK(OutputLevel);
+			SDK.ManageAndValidateSDK();
 
 			if ((ProjectFileGenerator.bGenerateProjectFiles == true) || (SDK.HasRequiredSDKsInstalled() == SDKStatus.Valid))
 			{
diff --git a/Engine/Source/Programs/UnrealBuildTool/Platform/Lumin/UEBuildLumin.cs b/Engine/Source/Programs/UnrealBuildTool/Platform/Lumin/UEBuildLumin.cs
--- a/Engine/Source/Programs/UnrealBuildTool/Platform/Lumin/UEBuildLumin.cs
+++ b/Engine/Source/Programs/UnrealBuildTool/Platform/Lumin/UEBuildLumin.cs
@@ -382,10 +382,10 @@ namespace UnrealBuildTool
 			get { return UnrealTargetPlatform.Lumin; }
 		}
 
-		protected override void RegisterBuildPlatforms(SDKOutputLevel OutputLevel)
+		protected override void RegisterBuildPlatforms()
 		{
 			LuminPlatformSDK SDK = new LuminPlatformSDK();
-			SDK.ManageAndValidateSDK(OutputLevel);
+			SDK.ManageAndValidateSDK();
 
 			if ((ProjectFileGenerator.bGenerateProjectFiles == true) || (SDK.HasRequiredSDKsInstalled() == SDKStatus.Valid) || Environment.GetEnvironmentVariable("IsBuildMachine") == "1")
 			{
diff --git a/Engine/Source/Programs/UnrealBuildTool/Platform/Mac/UEBuildMac.cs b/Engine/Source/Programs/UnrealBuildTool/Platform/Mac/UEBuildMac.cs
--- a/Engine/Source/Programs/UnrealBuildTool/Platform/Mac/UEBuildMac.cs
+++ b/Engine/Source/Programs/UnrealBuildTool/Platform/Mac/UEBuildMac.cs
@@ -373,10 +373,10 @@ namespace UnrealBuildTool
 		/// <summary>
 		/// Register the platform with the UEBuildPlatform class
 		/// </summary>
-		protected override void RegisterBuildPlatforms(SDKOutputLevel OutputLevel)
+		protected override void RegisterBuildPlatforms()
 		{
 			MacPlatformSDK SDK = new MacPlatformSDK();
-			SDK.ManageAndValidateSDK(OutputLevel);
+			SDK.ManageAndValidateSDK();
 
 			// Register this build platform for Mac
 			Log.TraceVerbose("        Registering for {0}", UnrealTargetPlatform.Mac.ToString());
diff --git a/Engine/Source/Programs/UnrealBuildTool/Platform/TVOS/UEBuildTVOS.cs b/Engine/Source/Programs/UnrealBuildTool/Platform/TVOS/UEBuildTVOS.cs
--- a/Engine/Source/Programs/UnrealBuildTool/Platform/TVOS/UEBuildTVOS.cs
+++ b/Engine/Source/Programs/UnrealBuildTool/Platform/TVOS/UEBuildTVOS.cs
@@ -146,10 +146,10 @@ namespace UnrealBuildTool
 		/// <summary>
 		/// Register the platform with the UEBuildPlatform class
 		/// </summary>
-		protected override void RegisterBuildPlatforms(SDKOutputLevel OutputLevel)
+		protected override void RegisterBuildPlatforms()
 		{
 			IOSPlatformSDK SDK = new IOSPlatformSDK();
-			SDK.ManageAndValidateSDK(OutputLevel);
+			SDK.ManageAndValidateSDK();
 
 			// Register this build platform for IOS
 			Log.TraceVerbose("        Registering for {0}", UnrealTargetPlatform.TVOS.ToString());
diff --git a/Engine/Source/Programs/UnrealBuildTool/Platform/Windows/UEBuildWindows.cs b/Engine/Source/Programs/UnrealBuildTool/Platform/Windows/UEBuildWindows.cs
--- a/Engine/Source/Programs/UnrealBuildTool/Platform/Windows/UEBuildWindows.cs
+++ b/Engine/Source/Programs/UnrealBuildTool/Platform/Windows/UEBuildWindows.cs
@@ -1674,10 +1674,10 @@ namespace UnrealBuildTool
 		/// <summary>
 		/// Register the platform with the UEBuildPlatform class
 		/// </summary>
-		protected override void RegisterBuildPlatforms(SDKOutputLevel OutputLevel)
+		protected override void RegisterBuildPlatforms()
 		{
 			WindowsPlatformSDK SDK = new WindowsPlatformSDK();
-			SDK.ManageAndValidateSDK(OutputLevel);
+			SDK.ManageAndValidateSDK();
 
 			// Register this build platform for both Win64 and Win32
 			Log.TraceVerbose("        Registering for {0}", UnrealTargetPlatform.Win64.ToString());
diff --git a/Engine/Source/Programs/UnrealBuildTool/System/PlatformExports.cs b/Engine/Source/Programs/UnrealBuildTool/System/PlatformExports.cs
--- a/Engine/Source/Programs/UnrealBuildTool/System/PlatformExports.cs
+++ b/Engine/Source/Programs/UnrealBuildTool/System/PlatformExports.cs
@@ -162,7 +162,7 @@ namespace UnrealBuildTool
 			}
 
 			// Register all the platform classes
-			UnrealBuildTool.RegisterAllUBTClasses(SDKOutputLevel.Quiet, false);
+			UnrealBuildTool.RegisterAllUBTClasses(false);
 			return true;
 		}
 	}
diff --git a/Engine/Source/Programs/UnrealBuildTool/UnrealBuildTool.cs b/Engine/Source/Programs/UnrealBuildTool/UnrealBuildTool.cs
--- a/Engine/Source/Programs/UnrealBuildTool/UnrealBuildTool.cs
+++ b/Engine/Source/Programs/UnrealBuildTool/UnrealBuildTool.cs
@@ -215,7 +215,7 @@ namespace UnrealBuildTool
 				InDirectory.IsUnderDirectory(UnrealBuildTool.EnterprisePluginsDirectory) || InDirectory.IsUnderDirectory(UnrealBuildTool.EnterpriseIntermediateDirectory);
 		}
 
-		public static void RegisterAllUBTClasses(SDKOutputLevel OutputLevel, bool bValidatingPlatforms)
+		public static void RegisterAllUBTClasses(bool bValidatingPlatforms)
 		{
 			// Find and register all tool chains and build platforms that are present
 			Assembly UBTAssembly = Assembly.GetExecutingAssembly();
@@ -235,7 +235,7 @@ namespace UnrealBuildTool
 						{
 							Log.TraceVerbose("    Registering build platform: {0}", CheckType.ToString());
 							UEBuildPlatformFactory TempInst = (UEBuildPlatformFactory)(UBTAssembly.CreateInstance(CheckType.FullName, true));
-							TempInst.TryRegisterBuildPlatforms(OutputLevel, bValidatingPlatforms);
+							TempInst.TryRegisterBuildPlatforms(bValidatingPlatforms);
 						}
 					}
 				}
@@ -726,20 +726,7 @@ namespace UnrealBuildTool
 					}
 
 					// Find and register all tool chains, build platforms, etc. that are present
-					SDKOutputLevel OutputLevel;
-					if (UnrealBuildTool.bPrintDebugInfo)
-					{
-						OutputLevel = SDKOutputLevel.Verbose;
-					}
-					else if (Environment.GetEnvironmentVariable("IsBuildMachine") == "1")
-					{
-						OutputLevel = SDKOutputLevel.Minimal;
-					}
-					else
-					{
-						OutputLevel = SDKOutputLevel.Quiet;
-					}
-					RegisterAllUBTClasses(OutputLevel, bValidatePlatforms);
+					RegisterAllUBTClasses(bValidatePlatforms);
 
 					if (UnrealBuildTool.bPrintPerformanceInfo)
 					{
