UBT: Move logic particular to the -validateplatforms mode into UnrealBuildTool.cs, rather than letting the platform factory worry about it.

#rnx

diff --git a/Engine/Source/Programs/UnrealBuildTool/Configuration/UEBuildPlatformFactory.cs b/Engine/Source/Programs/UnrealBuildTool/Configuration/UEBuildPlatformFactory.cs
--- a/Engine/Source/Programs/UnrealBuildTool/Configuration/UEBuildPlatformFactory.cs
+++ b/Engine/Source/Programs/UnrealBuildTool/Configuration/UEBuildPlatformFactory.cs
@@ -13,22 +13,10 @@ namespace UnrealBuildTool
 	/// </summary>
 	abstract class UEBuildPlatformFactory
 	{
-		/// <summary>
-		/// Attempt to register a build platform, checking whether it is a valid platform in installed builds
-		/// </summary>
-		public void TryRegisterBuildPlatforms(bool bValidatingPlatforms)
-		{
-			// We need all platforms to be registered when we run -validateplatform command to check SDK status of each
-			if (bValidatingPlatforms || InstalledPlatformInfo.IsValidPlatform(TargetPlatform))
-			{
-				RegisterBuildPlatforms();
-			}
-		}
-
 		/// <summary>
 		/// Gets the target platform for an individual factory
 		/// </summary>
-		protected abstract UnrealTargetPlatform TargetPlatform
+		public abstract UnrealTargetPlatform TargetPlatform
 		{
 			get;
 		}
@@ -36,6 +24,6 @@ namespace UnrealBuildTool
 		/// <summary>
 		/// Register the platform with the UEBuildPlatform class
 		/// </summary>
-		protected abstract void RegisterBuildPlatforms();
+		public abstract void RegisterBuildPlatforms();
 	}
 }
diff --git a/Engine/Source/Programs/UnrealBuildTool/Platform/Android/UEBuildAndroid.cs b/Engine/Source/Programs/UnrealBuildTool/Platform/Android/UEBuildAndroid.cs
--- a/Engine/Source/Programs/UnrealBuildTool/Platform/Android/UEBuildAndroid.cs
+++ b/Engine/Source/Programs/UnrealBuildTool/Platform/Android/UEBuildAndroid.cs
@@ -656,12 +656,12 @@ namespace UnrealBuildTool
 
 	class AndroidPlatformFactory : UEBuildPlatformFactory
 	{
-		protected override UnrealTargetPlatform TargetPlatform
+		public override UnrealTargetPlatform TargetPlatform
 		{
 			get { return UnrealTargetPlatform.Android; }
 		}
 
-		protected override void RegisterBuildPlatforms()
+		public override void RegisterBuildPlatforms()
 		{
 			AndroidPlatformSDK SDK = new AndroidPlatformSDK();
 			SDK.ManageAndValidateSDK();
diff --git a/Engine/Source/Programs/UnrealBuildTool/Platform/HTML5/UEBuildHTML5.cs b/Engine/Source/Programs/UnrealBuildTool/Platform/HTML5/UEBuildHTML5.cs
--- a/Engine/Source/Programs/UnrealBuildTool/Platform/HTML5/UEBuildHTML5.cs
+++ b/Engine/Source/Programs/UnrealBuildTool/Platform/HTML5/UEBuildHTML5.cs
@@ -266,7 +266,7 @@ namespace UnrealBuildTool
 
 	class HTML5PlatformFactory : UEBuildPlatformFactory
 	{
-		protected override UnrealTargetPlatform TargetPlatform
+		public override UnrealTargetPlatform TargetPlatform
 		{
 			get { return UnrealTargetPlatform.HTML5; }
 		}
@@ -274,7 +274,7 @@ namespace UnrealBuildTool
 		/// <summary>
 		/// Register the platform with the UEBuildPlatform class
 		/// </summary>
-		protected override void RegisterBuildPlatforms()
+		public override void RegisterBuildPlatforms()
 		{
 			HTML5PlatformSDK SDK = new HTML5PlatformSDK();
 			SDK.ManageAndValidateSDK();
diff --git a/Engine/Source/Programs/UnrealBuildTool/Platform/IOS/UEBuildIOS.cs b/Engine/Source/Programs/UnrealBuildTool/Platform/IOS/UEBuildIOS.cs
--- a/Engine/Source/Programs/UnrealBuildTool/Platform/IOS/UEBuildIOS.cs
+++ b/Engine/Source/Programs/UnrealBuildTool/Platform/IOS/UEBuildIOS.cs
@@ -983,7 +983,7 @@ namespace UnrealBuildTool
 
 	class IOSPlatformFactory : UEBuildPlatformFactory
 	{
-		protected override UnrealTargetPlatform TargetPlatform
+		public override UnrealTargetPlatform TargetPlatform
 		{
 			get { return UnrealTargetPlatform.IOS; }
 		}
@@ -991,7 +991,7 @@ namespace UnrealBuildTool
 		/// <summary>
 		/// Register the platform with the UEBuildPlatform class
 		/// </summary>
-		protected override void RegisterBuildPlatforms()
+		public override void RegisterBuildPlatforms()
 		{
 			IOSPlatformSDK SDK = new IOSPlatformSDK();
 			SDK.ManageAndValidateSDK();
diff --git a/Engine/Source/Programs/UnrealBuildTool/Platform/Linux/UEBuildLinux.cs b/Engine/Source/Programs/UnrealBuildTool/Platform/Linux/UEBuildLinux.cs
--- a/Engine/Source/Programs/UnrealBuildTool/Platform/Linux/UEBuildLinux.cs
+++ b/Engine/Source/Programs/UnrealBuildTool/Platform/Linux/UEBuildLinux.cs
@@ -750,7 +750,7 @@ namespace UnrealBuildTool
 
 	class LinuxPlatformFactory : UEBuildPlatformFactory
 	{
-		protected override UnrealTargetPlatform TargetPlatform
+		public override UnrealTargetPlatform TargetPlatform
 		{
 			get { return UnrealTargetPlatform.Linux; }
 		}
@@ -758,7 +758,7 @@ namespace UnrealBuildTool
 		/// <summary>
 		/// Register the platform with the UEBuildPlatform class
 		/// </summary>
-		protected override void RegisterBuildPlatforms()
+		public override void RegisterBuildPlatforms()
 		{
 			LinuxPlatformSDK SDK = new LinuxPlatformSDK();
 			SDK.ManageAndValidateSDK();
diff --git a/Engine/Source/Programs/UnrealBuildTool/Platform/Lumin/UEBuildLumin.cs b/Engine/Source/Programs/UnrealBuildTool/Platform/Lumin/UEBuildLumin.cs
--- a/Engine/Source/Programs/UnrealBuildTool/Platform/Lumin/UEBuildLumin.cs
+++ b/Engine/Source/Programs/UnrealBuildTool/Platform/Lumin/UEBuildLumin.cs
@@ -377,12 +377,12 @@ namespace UnrealBuildTool
 
 	class LuminPlatformFactory : UEBuildPlatformFactory
 	{
-		protected override UnrealTargetPlatform TargetPlatform
+		public override UnrealTargetPlatform TargetPlatform
 		{
 			get { return UnrealTargetPlatform.Lumin; }
 		}
 
-		protected override void RegisterBuildPlatforms()
+		public override void RegisterBuildPlatforms()
 		{
 			LuminPlatformSDK SDK = new LuminPlatformSDK();
 			SDK.ManageAndValidateSDK();
diff --git a/Engine/Source/Programs/UnrealBuildTool/Platform/Mac/UEBuildMac.cs b/Engine/Source/Programs/UnrealBuildTool/Platform/Mac/UEBuildMac.cs
--- a/Engine/Source/Programs/UnrealBuildTool/Platform/Mac/UEBuildMac.cs
+++ b/Engine/Source/Programs/UnrealBuildTool/Platform/Mac/UEBuildMac.cs
@@ -365,7 +365,7 @@ namespace UnrealBuildTool
 
 	class MacPlatformFactory : UEBuildPlatformFactory
 	{
-		protected override UnrealTargetPlatform TargetPlatform
+		public override UnrealTargetPlatform TargetPlatform
 		{
 			get { return UnrealTargetPlatform.Mac; }
 		}
@@ -373,7 +373,7 @@ namespace UnrealBuildTool
 		/// <summary>
 		/// Register the platform with the UEBuildPlatform class
 		/// </summary>
-		protected override void RegisterBuildPlatforms()
+		public override void RegisterBuildPlatforms()
 		{
 			MacPlatformSDK SDK = new MacPlatformSDK();
 			SDK.ManageAndValidateSDK();
diff --git a/Engine/Source/Programs/UnrealBuildTool/Platform/TVOS/UEBuildTVOS.cs b/Engine/Source/Programs/UnrealBuildTool/Platform/TVOS/UEBuildTVOS.cs
--- a/Engine/Source/Programs/UnrealBuildTool/Platform/TVOS/UEBuildTVOS.cs
+++ b/Engine/Source/Programs/UnrealBuildTool/Platform/TVOS/UEBuildTVOS.cs
@@ -138,7 +138,7 @@ namespace UnrealBuildTool
 
 	class TVOSPlatformFactory : UEBuildPlatformFactory
 	{
-		protected override UnrealTargetPlatform TargetPlatform
+		public override UnrealTargetPlatform TargetPlatform
 		{
 			get { return UnrealTargetPlatform.TVOS; }
 		}
@@ -146,7 +146,7 @@ namespace UnrealBuildTool
 		/// <summary>
 		/// Register the platform with the UEBuildPlatform class
 		/// </summary>
-		protected override void RegisterBuildPlatforms()
+		public override void RegisterBuildPlatforms()
 		{
 			IOSPlatformSDK SDK = new IOSPlatformSDK();
 			SDK.ManageAndValidateSDK();
diff --git a/Engine/Source/Programs/UnrealBuildTool/Platform/Windows/UEBuildWindows.cs b/Engine/Source/Programs/UnrealBuildTool/Platform/Windows/UEBuildWindows.cs
--- a/Engine/Source/Programs/UnrealBuildTool/Platform/Windows/UEBuildWindows.cs
+++ b/Engine/Source/Programs/UnrealBuildTool/Platform/Windows/UEBuildWindows.cs
@@ -1666,7 +1666,7 @@ namespace UnrealBuildTool
 
 	class WindowsPlatformFactory : UEBuildPlatformFactory
 	{
-		protected override UnrealTargetPlatform TargetPlatform
+		public override UnrealTargetPlatform TargetPlatform
 		{
 			get { return UnrealTargetPlatform.Win64; }
 		}
@@ -1674,7 +1674,7 @@ namespace UnrealBuildTool
 		/// <summary>
 		/// Register the platform with the UEBuildPlatform class
 		/// </summary>
-		protected override void RegisterBuildPlatforms()
+		public override void RegisterBuildPlatforms()
 		{
 			WindowsPlatformSDK SDK = new WindowsPlatformSDK();
 			SDK.ManageAndValidateSDK();
diff --git a/Engine/Source/Programs/UnrealBuildTool/UnrealBuildTool.cs b/Engine/Source/Programs/UnrealBuildTool/UnrealBuildTool.cs
--- a/Engine/Source/Programs/UnrealBuildTool/UnrealBuildTool.cs
+++ b/Engine/Source/Programs/UnrealBuildTool/UnrealBuildTool.cs
@@ -235,7 +235,12 @@ namespace UnrealBuildTool
 						{
 							Log.TraceVerbose("    Registering build platform: {0}", CheckType.ToString());
 							UEBuildPlatformFactory TempInst = (UEBuildPlatformFactory)(UBTAssembly.CreateInstance(CheckType.FullName, true));
-							TempInst.TryRegisterBuildPlatforms(bValidatingPlatforms);
+
+							// We need all platforms to be registered when we run -validateplatform command to check SDK status of each
+							if (bValidatingPlatforms || InstalledPlatformInfo.IsValidPlatform(TempInst.TargetPlatform))
+							{
+								TempInst.RegisterBuildPlatforms();
+							}
 						}
 					}
 				}
