Allow targets to specify additional build products produced by post-build steps, which will be included in the manifest and target receipt.

diff --git a/Engine/Source/Programs/ShaderCompileWorker/ShaderCompileWorker.Target.cs b/Engine/Source/Programs/ShaderCompileWorker/ShaderCompileWorker.Target.cs
--- a/Engine/Source/Programs/ShaderCompileWorker/ShaderCompileWorker.Target.cs
+++ b/Engine/Source/Programs/ShaderCompileWorker/ShaderCompileWorker.Target.cs
@@ -14,7 +14,7 @@ public class ShaderCompileWorkerTarget : TargetRules
 
 		LaunchModuleName = "ShaderCompileWorker";
 
-        if (bUseXGEController && (Target.Platform == UnrealTargetPlatform.Win32 || Target.Platform == UnrealTargetPlatform.Win64))
+        if (bUseXGEController && (Target.Platform == UnrealTargetPlatform.Win32 || Target.Platform == UnrealTargetPlatform.Win64) && Configuration == UnrealTargetConfiguration.Development)
         {
             // The interception interface in XGE requires that the parent and child processes have different filenames on disk.
             // To avoid building an entire separate worker just for this, we duplicate the ShaderCompileWorker in a post build step.
@@ -22,7 +22,9 @@ public class ShaderCompileWorkerTarget : TargetRules
             const string DestPath = "$(EngineDir)\\Binaries\\$(TargetPlatform)\\XGEControlWorker.exe";
 
             PostBuildSteps.Add(string.Format("echo Copying {0} to {1}", SrcPath, DestPath));
-            PostBuildSteps.Add(string.Format("copy /Y /B \"{0}\" /B \"{1}\"", SrcPath, DestPath));
+            PostBuildSteps.Add(string.Format("copy /Y /B \"{0}\" /B \"{1}\" >nul:", SrcPath, DestPath));
+
+			AdditionalBuildProducts.Add(DestPath);
         }
 
 		// Turn off various third party features we don't need
diff --git a/Engine/Source/Programs/UnrealBuildTool/Configuration/TargetRules.cs b/Engine/Source/Programs/UnrealBuildTool/Configuration/TargetRules.cs
--- a/Engine/Source/Programs/UnrealBuildTool/Configuration/TargetRules.cs
+++ b/Engine/Source/Programs/UnrealBuildTool/Configuration/TargetRules.cs
@@ -1045,6 +1045,11 @@ namespace UnrealBuildTool
 		/// </summary>
 		public List<string> PostBuildSteps = new List<string>();
 
+		/// <summary>
+		/// Specifies additional build products produced as part of this target.
+		/// </summary>
+		public List<string> AdditionalBuildProducts = new List<string>();
+
 		/// <summary>
 		/// Additional arguments to pass to the compiler
 		/// </summary>
@@ -2060,6 +2065,11 @@ namespace UnrealBuildTool
 			get { return Inner.PostBuildSteps; }
 		}
 
+		public IReadOnlyList<string> AdditionalBuildProducts
+		{
+			get { return Inner.AdditionalBuildProducts; }
+		}
+
 		public string AdditionalCompilerArguments
 		{
 			get { return Inner.AdditionalCompilerArguments; }
diff --git a/Engine/Source/Programs/UnrealBuildTool/Configuration/UEBuildTarget.cs b/Engine/Source/Programs/UnrealBuildTool/Configuration/UEBuildTarget.cs
--- a/Engine/Source/Programs/UnrealBuildTool/Configuration/UEBuildTarget.cs
+++ b/Engine/Source/Programs/UnrealBuildTool/Configuration/UEBuildTarget.cs
@@ -2178,6 +2178,17 @@ namespace UnrealBuildTool
 			}
 			BuildProducts.AddRange(RuntimeDependencyTargetFileToSourceFile.Select(x => new KeyValuePair<FileReference, BuildProductType>(x.Key, BuildProductType.RequiredResource)));
 
+			// Also add any explicitly specified build products
+			if(Rules.AdditionalBuildProducts.Count > 0)
+			{
+				Dictionary<string, string> Variables = GetTargetVariables(null);
+				foreach(string AdditionalBuildProduct in Rules.AdditionalBuildProducts)
+				{
+					FileReference BuildProductFile = new FileReference(Utils.ExpandVariables(AdditionalBuildProduct, Variables));
+					BuildProducts.Add(new KeyValuePair<FileReference, BuildProductType>(BuildProductFile, BuildProductType.RequiredResource));
+				}
+			}
+
 			// Create a receipt for the target
 			if (!ProjectFileGenerator.bGenerateProjectFiles)
 			{
@@ -2601,6 +2612,33 @@ namespace UnrealBuildTool
 			}
 		}
 
+		/// <summary>
+		/// Gets a list of variables that can be expanded in paths referenced by this target
+		/// </summary>
+		/// <param name="Plugin">The current plugin</param>
+		/// <returns>Map of variable names to values</returns>
+		private Dictionary<string, string> GetTargetVariables(UEBuildPlugin Plugin)
+		{
+			Dictionary<string, string> Variables = new Dictionary<string,string>();
+			Variables.Add("RootDir", UnrealBuildTool.RootDirectory.FullName);
+			Variables.Add("EngineDir", UnrealBuildTool.EngineDirectory.FullName);
+			Variables.Add("EnterpriseDir", UnrealBuildTool.EnterpriseDirectory.FullName);
+			Variables.Add("ProjectDir", ProjectDirectory.FullName);
+			Variables.Add("TargetName", TargetName);
+			Variables.Add("TargetPlatform", Platform.ToString());
+			Variables.Add("TargetConfiguration", Configuration.ToString());
+			Variables.Add("TargetType", TargetType.ToString());
+			if(ProjectFile != null)
+			{
+				Variables.Add("ProjectFile", ProjectFile.FullName);
+			}
+			if(Plugin != null)
+			{
+				Variables.Add("PluginDir", Plugin.Directory.FullName);
+			}
+			return Variables;
+		}
+
 		/// <summary>
 		/// Write scripts containing the custom build steps for the given host platform
 		/// </summary>
@@ -2615,23 +2653,7 @@ namespace UnrealBuildTool
 			foreach(Tuple<string[], UEBuildPlugin> CommandBatch in CommandBatches)
 			{
 				// Find all the standard variables
-				Dictionary<string, string> Variables = new Dictionary<string,string>();
-				Variables.Add("RootDir", UnrealBuildTool.RootDirectory.FullName);
-				Variables.Add("EngineDir", UnrealBuildTool.EngineDirectory.FullName);
-				Variables.Add("EnterpriseDir", UnrealBuildTool.EnterpriseDirectory.FullName);
-				Variables.Add("ProjectDir", ProjectDirectory.FullName);
-				Variables.Add("TargetName", TargetName);
-				Variables.Add("TargetPlatform", Platform.ToString());
-				Variables.Add("TargetConfiguration", Configuration.ToString());
-				Variables.Add("TargetType", TargetType.ToString());
-				if(ProjectFile != null)
-				{
-					Variables.Add("ProjectFile", ProjectFile.FullName);
-				}
-				if(CommandBatch.Item2 != null)
-				{
-					Variables.Add("PluginDir", CommandBatch.Item2.Directory.FullName);
-				}
+				Dictionary<string, string> Variables = GetTargetVariables(CommandBatch.Item2);
 
 				// Get the output path to the script
 				string ScriptExtension = (HostPlatform == UnrealTargetPlatform.Win64)? ".bat" : ".sh";
