UBT: Convert PVS-Studio to use a dedicated UBT mode to gather output rather than a custom action handler, so that it can use any available executor.

#rnx

diff --git a/Engine/Source/Programs/UnrealBuildTool/Executors/LocalExecutor.cs b/Engine/Source/Programs/UnrealBuildTool/Executors/LocalExecutor.cs
--- a/Engine/Source/Programs/UnrealBuildTool/Executors/LocalExecutor.cs
+++ b/Engine/Source/Programs/UnrealBuildTool/Executors/LocalExecutor.cs
@@ -89,156 +89,132 @@ namespace UnrealBuildTool
 			// thread start time
 			Action.StartTime = DateTimeOffset.Now;
 
-			if (Action.ActionHandler != null)
-			{
-				// call the function and get the ExitCode and an output string
-				string Output;
-				Action.ActionHandler(Action, out ExitCode, out Output);
+			// batch files (.bat, .cmd) need to be run via or cmd /c or shellexecute, 
+			// the latter which we can't use because we want to redirect input/output
+			bool bLaunchViaCmdExe = !Utils.IsRunningOnMono && !Path.GetExtension(Action.CommandPath).ToLower().EndsWith("exe");
+
+			// Create the action's process.
+			ProcessStartInfo ActionStartInfo = new ProcessStartInfo();
+			ActionStartInfo.WorkingDirectory = ExpandEnvironmentVariables(Action.WorkingDirectory);
 
-				// Output status description (file name) when no output is returned
-				if (string.IsNullOrEmpty(Output))
-				{
-					if (Action.bShouldOutputStatusDescription)
-					{
-						Output = string.Format("[{0}/{1}] {2} {3}", JobNumber, TotalJobs, Action.CommandDescription, Action.StatusDescription);
-					}
-					else
-					{
-						Output = Action.StatusDescription;
-					}
-				}
 
-				Log.TraceInformation(Output);
+			string ExpandedCommandPath = ExpandEnvironmentVariables(Action.CommandPath);
+			if (bLaunchViaCmdExe)
+			{
+				ActionStartInfo.FileName = "cmd.exe";
+				ActionStartInfo.Arguments = string.Format
+				(
+					"/c \"{0} {1}\"",
+					ExpandEnvironmentVariables(ExpandedCommandPath),
+					ExpandEnvironmentVariables(Action.CommandArguments)
+				);
 			}
 			else
 			{
-				// batch files (.bat, .cmd) need to be run via or cmd /c or shellexecute, 
-				// the latter which we can't use because we want to redirect input/output
-				bool bLaunchViaCmdExe = !Utils.IsRunningOnMono && !Path.GetExtension(Action.CommandPath).ToLower().EndsWith("exe");
-
-				// Create the action's process.
-				ProcessStartInfo ActionStartInfo = new ProcessStartInfo();
-				ActionStartInfo.WorkingDirectory = ExpandEnvironmentVariables(Action.WorkingDirectory);
+				ActionStartInfo.FileName = ExpandedCommandPath;
+				ActionStartInfo.Arguments = ExpandEnvironmentVariables(Action.CommandArguments);
+			}
 
+			ActionStartInfo.UseShellExecute = false;
+			ActionStartInfo.RedirectStandardInput = false;
+			ActionStartInfo.RedirectStandardOutput = false;
+			ActionStartInfo.RedirectStandardError = false;
 
-				string ExpandedCommandPath = ExpandEnvironmentVariables(Action.CommandPath);
-				if (bLaunchViaCmdExe)
+			// Log command-line used to execute task if debug info printing is enabled.
+			if (UnrealBuildTool.bPrintDebugInfo)
+			{
+				Log.TraceVerbose("Executing: {0} {1}", ExpandedCommandPath, ActionStartInfo.Arguments);
+			}
+			// Log summary if wanted.
+			else if (Action.bShouldOutputStatusDescription)
+			{
+				string CommandDescription = Action.CommandDescription != null ? Action.CommandDescription : Path.GetFileName(ExpandedCommandPath);
+				if (string.IsNullOrEmpty(CommandDescription))
 				{
-					ActionStartInfo.FileName = "cmd.exe";
-					ActionStartInfo.Arguments = string.Format
-					(
-						"/c \"{0} {1}\"",
-						ExpandEnvironmentVariables(ExpandedCommandPath),
-						ExpandEnvironmentVariables(Action.CommandArguments)
-					);
+					Log.TraceInformation(Action.StatusDescription);
 				}
 				else
 				{
-					ActionStartInfo.FileName = ExpandedCommandPath;
-					ActionStartInfo.Arguments = ExpandEnvironmentVariables(Action.CommandArguments);
+					Log.TraceInformation("[{0}/{1}] {2} {3}", JobNumber, TotalJobs, CommandDescription, Action.StatusDescription);
 				}
+			}
 
-				ActionStartInfo.UseShellExecute = false;
-				ActionStartInfo.RedirectStandardInput = false;
-				ActionStartInfo.RedirectStandardOutput = false;
-				ActionStartInfo.RedirectStandardError = false;
-
-				// Log command-line used to execute task if debug info printing is enabled.
-				if (UnrealBuildTool.bPrintDebugInfo)
+			// Try to launch the action's process, and produce a friendly error message if it fails.
+			Process ActionProcess = null;
+			try
+			{
+				try
 				{
-					Log.TraceVerbose("Executing: {0} {1}", ExpandedCommandPath, ActionStartInfo.Arguments);
+					ActionProcess = new Process();
+					ActionProcess.StartInfo = ActionStartInfo;
+					ActionStartInfo.RedirectStandardOutput = true;
+					ActionStartInfo.RedirectStandardError = true;
+					ActionProcess.OutputDataReceived += new DataReceivedEventHandler(ActionDebugOutput);
+					ActionProcess.ErrorDataReceived += new DataReceivedEventHandler(ActionDebugOutput);
+					ActionProcess.Start();
+
+					ActionProcess.BeginOutputReadLine();
+					ActionProcess.BeginErrorReadLine();
 				}
-				// Log summary if wanted.
-				else if (Action.bShouldOutputStatusDescription)
+				catch (Exception ex)
 				{
-					string CommandDescription = Action.CommandDescription != null ? Action.CommandDescription : Path.GetFileName(ExpandedCommandPath);
-					if (string.IsNullOrEmpty(CommandDescription))
-					{
-						Log.TraceInformation(Action.StatusDescription);
-					}
-					else
-					{
-						Log.TraceInformation("[{0}/{1}] {2} {3}", JobNumber, TotalJobs, CommandDescription, Action.StatusDescription);
-					}
+					Log.TraceError("Failed to start local process for action: {0} {1}", Action.CommandPath, Action.CommandArguments);
+					Log.WriteException(ex, null);
+					ExitCode = 1;
+					bComplete = true;
+					return;
 				}
 
-				// Try to launch the action's process, and produce a friendly error message if it fails.
-				Process ActionProcess = null;
-				try
+				// wait for process to start
+				// NOTE: this may or may not be necessary; seems to depend on whether the system UBT is running on start the process in a timely manner.
+				int checkIterations = 0;
+				bool haveConfiguredProcess = false;
+				do
 				{
-					try
-					{
-						ActionProcess = new Process();
-						ActionProcess.StartInfo = ActionStartInfo;
-						ActionStartInfo.RedirectStandardOutput = true;
-						ActionStartInfo.RedirectStandardError = true;
-						ActionProcess.OutputDataReceived += new DataReceivedEventHandler(ActionDebugOutput);
-						ActionProcess.ErrorDataReceived += new DataReceivedEventHandler(ActionDebugOutput);
-						ActionProcess.Start();
-
-						ActionProcess.BeginOutputReadLine();
-						ActionProcess.BeginErrorReadLine();
-					}
-					catch (Exception ex)
+					if (ActionProcess.HasExited)
 					{
-						Log.TraceError("Failed to start local process for action: {0} {1}", Action.CommandPath, Action.CommandArguments);
-						Log.WriteException(ex, null);
-						ExitCode = 1;
-						bComplete = true;
-						return;
+						if (haveConfiguredProcess == false)
+							Debug.WriteLine("Process for action exited before able to configure!");
+						break;
 					}
 
-					// wait for process to start
-					// NOTE: this may or may not be necessary; seems to depend on whether the system UBT is running on start the process in a timely manner.
-					int checkIterations = 0;
-					bool haveConfiguredProcess = false;
-					do
+					if (!haveConfiguredProcess)
 					{
-						if (ActionProcess.HasExited)
+						try
 						{
-							if (haveConfiguredProcess == false)
-								Debug.WriteLine("Process for action exited before able to configure!");
-							break;
+							ActionProcess.PriorityClass = ProcessPriorityClass.BelowNormal;
+							haveConfiguredProcess = true;
 						}
-
-						if (!haveConfiguredProcess)
+						catch (Exception)
 						{
-							try
-							{
-								ActionProcess.PriorityClass = ProcessPriorityClass.BelowNormal;
-								haveConfiguredProcess = true;
-							}
-							catch (Exception)
-							{
-							}
-							break;
 						}
-
-						Thread.Sleep(10);
-
-						checkIterations++;
-					} while (checkIterations < 100);
-					if (checkIterations == 100)
-					{
-						throw new BuildException("Failed to configure local process for action: {0} {1}", Action.CommandPath, Action.CommandArguments);
+						break;
 					}
 
-					// block until it's complete
-					// @todo iosmerge: UBT had started looking at:	if (Utils.IsValidProcess(Process))
-					//    do we need to check that in the thread model?
-					ActionProcess.WaitForExit();
+					Thread.Sleep(10);
 
-					// capture exit code
-					ExitCode = ActionProcess.ExitCode;
-				}
-				finally
+					checkIterations++;
+				} while (checkIterations < 100);
+				if (checkIterations == 100)
 				{
-					// As the process has finished now, free its resources. On non-Windows platforms, processes depend 
-					// on POSIX/BSD threading and these are limited per application. Disposing the Process releases 
-					// these thread resources.
-					if (ActionProcess != null)
-						ActionProcess.Close();
+					throw new BuildException("Failed to configure local process for action: {0} {1}", Action.CommandPath, Action.CommandArguments);
 				}
+
+				// block until it's complete
+				// @todo iosmerge: UBT had started looking at:	if (Utils.IsValidProcess(Process))
+				//    do we need to check that in the thread model?
+				ActionProcess.WaitForExit();
+
+				// capture exit code
+				ExitCode = ActionProcess.ExitCode;
+			}
+			finally
+			{
+				// As the process has finished now, free its resources. On non-Windows platforms, processes depend 
+				// on POSIX/BSD threading and these are limited per application. Disposing the Process releases 
+				// these thread resources.
+				if (ActionProcess != null)
+					ActionProcess.Close();
 			}
 
 			// track how long it took
diff --git a/Engine/Source/Programs/UnrealBuildTool/Platform/Windows/PVSToolChain.cs b/Engine/Source/Programs/UnrealBuildTool/Platform/Windows/PVSToolChain.cs
--- a/Engine/Source/Programs/UnrealBuildTool/Platform/Windows/PVSToolChain.cs
+++ b/Engine/Source/Programs/UnrealBuildTool/Platform/Windows/PVSToolChain.cs
@@ -35,6 +35,100 @@
 		public string SerialNumber;
 	}
 
+	/// <summary>
+	/// Special mode for gathering all the messages into a single output file
+	/// </summary>
+	[ToolMode("PVSGather")]
+	class PVSGatherMode : ToolMode
+	{
+		/// <summary>
+		/// Path to the input file list
+		/// </summary>
+		[CommandLine("-Input", Required = true)]
+		FileReference InputFileList = null;
+
+		/// <summary>
+		/// Output file to generate
+		/// </summary>
+		[CommandLine("-Output", Required = true)]
+		FileReference OutputFile = null;
+
+		/// <summary>
+		/// Execute the command
+		/// </summary>
+		/// <param name="Arguments">List of command line arguments</param>
+		/// <returns>Always zero, or throws an exception</returns>
+		public override int Execute(CommandLineArguments Arguments)
+		{
+			Arguments.ApplyTo(this);
+			Arguments.CheckAllArgumentsUsed();
+
+			Log.TraceInformation("{0}", OutputFile.GetFileName());
+
+			// Read the input files
+			string[] InputFileLines = FileReference.ReadAllLines(InputFileList);
+			FileReference[] InputFiles = InputFileLines.Select(x => x.Trim()).Where(x => x.Length > 0).Select(x => new FileReference(x)).ToArray();
+
+			// Create the combined output file, and print the diagnostics to the log
+			HashSet<string> UniqueItems = new HashSet<string>();
+			using (StreamWriter RawWriter = new StreamWriter(OutputFile.FullName))
+			{
+				foreach (FileReference InputFile in InputFiles)
+				{
+					string[] Lines = File.ReadAllLines(InputFile.FullName);
+					for(int LineIdx = 0; LineIdx < Lines.Length; LineIdx++)
+					{
+						string Line = Lines[LineIdx];
+						if (!String.IsNullOrWhiteSpace(Line) && UniqueItems.Add(Line))
+						{
+							bool bCanParse = false;
+
+							string[] Tokens = Line.Split(new string[] { "<#~>" }, StringSplitOptions.None);
+							if(Tokens.Length >= 9)
+							{
+								string Trial = Tokens[1];
+								string LineNumberStr = Tokens[2];
+								string FileName = Tokens[3];
+								string WarningCode = Tokens[5];
+								string WarningMessage = Tokens[6];
+								string FalseAlarmStr = Tokens[7];
+								string LevelStr = Tokens[8];
+
+								int LineNumber;
+								bool bFalseAlarm;
+								int Level;
+								if(int.TryParse(LineNumberStr, out LineNumber) && bool.TryParse(FalseAlarmStr, out bFalseAlarm) && int.TryParse(LevelStr, out Level))
+								{
+									bCanParse = true;
+
+									// Ignore anything in ThirdParty folders
+									if(FileName.Replace('/', '\\').IndexOf("\\ThirdParty\\", StringComparison.InvariantCultureIgnoreCase) == -1)
+									{
+										// Output the line to the raw output file
+										RawWriter.WriteLine(Line);
+
+										// Output the line to the log
+										if (!bFalseAlarm && Level == 1)
+										{
+											Log.WriteLine(LogEventType.Warning, LogFormatOptions.NoSeverityPrefix, "{0}({1}): warning {2}: {3}", FileName, LineNumber, WarningCode, WarningMessage);
+										}
+									}
+								}
+							}
+
+							if(!bCanParse)
+							{
+								Log.WriteLine(LogEventType.Warning, LogFormatOptions.NoSeverityPrefix, "{0}({1}): warning: Unable to parse PVS output line '{2}' (tokens=|{3}|)", InputFile, LineIdx + 1, Line, String.Join("|", Tokens));
+							}
+						}
+					}
+				}
+			}
+			Log.TraceInformation("Written {0} {1} to {2}.", UniqueItems.Count, (UniqueItems.Count == 1)? "diagnostic" : "diagnostics", OutputFile.FullName);
+			return 0;
+		}
+	}
+
 	class PVSToolChain : UEToolChain
 	{
 		VCEnvironment EnvVars;
@@ -244,84 +338,18 @@
 
 			List<FileReference> InputFiles = OutputItems.Select(x => x.Location).Where(x => x.HasExtension(".pvslog")).ToList();
 
+			FileItem InputFileListItem = FileItem.CreateIntermediateTextFile(OutputFile.ChangeExtension(".input"), InputFiles.Select(x => x.FullName));
+
 			Action AnalyzeAction = ActionGraph.Add(ActionType.Compile);
-			AnalyzeAction.CommandPath = "Dummy.exe";
-			AnalyzeAction.CommandArguments = "";
-			AnalyzeAction.CommandDescription = "Combining output from PVS-Studio";
+			AnalyzeAction.CommandPath = UnrealBuildTool.GetUBTPath();
+			AnalyzeAction.CommandArguments = String.Format("-Mode=PVSGather -Input=\"{0}\" -Output=\"{1}\"", InputFileListItem.Location, OutputFile);
 			AnalyzeAction.WorkingDirectory = UnrealBuildTool.EngineSourceDirectory.FullName;
+			AnalyzeAction.PrerequisiteItems.Add(InputFileListItem);
 			AnalyzeAction.PrerequisiteItems.AddRange(OutputItems);
 			AnalyzeAction.ProducedItems.Add(FileItem.GetItemByFileReference(OutputFile));
-			AnalyzeAction.ActionHandler = (Action Action, out int ExitCode, out string Output) => WriteResults(OutputFile, InputFiles, out ExitCode, out Output);
 			AnalyzeAction.DeleteItems.AddRange(AnalyzeAction.ProducedItems);
 
 			OutputItems.AddRange(AnalyzeAction.ProducedItems);
 		}
-
-		void WriteResults(FileReference OutputFile, List<FileReference> InputFiles, out int ExitCode, out string Output)
-		{
-			StringBuilder OutputBuilder = new StringBuilder();
-			OutputBuilder.Append("Processing PVS-Studio output...\n");
-
-			// Create the combined output file, and print the diagnostics to the log
-			HashSet<string> UniqueItems = new HashSet<string>();
-			using (StreamWriter RawWriter = new StreamWriter(OutputFile.FullName))
-			{
-				foreach (FileReference InputFile in InputFiles)
-				{
-					string[] Lines = File.ReadAllLines(InputFile.FullName);
-					for(int LineIdx = 0; LineIdx < Lines.Length; LineIdx++)
-					{
-						string Line = Lines[LineIdx];
-						if (!String.IsNullOrWhiteSpace(Line) && UniqueItems.Add(Line))
-						{
-							bool bCanParse = false;
-
-							string[] Tokens = Line.Split(new string[] { "<#~>" }, StringSplitOptions.None);
-							if(Tokens.Length >= 9)
-							{
-								string Trial = Tokens[1];
-								string LineNumberStr = Tokens[2];
-								string FileName = Tokens[3];
-								string WarningCode = Tokens[5];
-								string WarningMessage = Tokens[6];
-								string FalseAlarmStr = Tokens[7];
-								string LevelStr = Tokens[8];
-
-								int LineNumber;
-								bool bFalseAlarm;
-								int Level;
-								if(int.TryParse(LineNumberStr, out LineNumber) && bool.TryParse(FalseAlarmStr, out bFalseAlarm) && int.TryParse(LevelStr, out Level))
-								{
-									bCanParse = true;
-
-									// Ignore anything in ThirdParty folders
-									if(FileName.Replace('/', '\\').IndexOf("\\ThirdParty\\", StringComparison.InvariantCultureIgnoreCase) == -1)
-									{
-										// Output the line to the raw output file
-										RawWriter.WriteLine(Line);
-
-										// Output the line to the log
-										if (!bFalseAlarm && Level == 1)
-										{
-											OutputBuilder.AppendFormat("{0}({1}): warning {2}: {3}\n", FileName, LineNumber, WarningCode, WarningMessage);
-										}
-									}
-								}
-							}
-
-							if(!bCanParse)
-							{
-								Log.WriteLine(LogEventType.Warning, LogFormatOptions.NoSeverityPrefix, "{0}({1}): warning: Unable to parse PVS output line '{2}' (tokens=|{3}|)", InputFile, LineIdx + 1, Line, String.Join("|", Tokens));
-							}
-						}
-					}
-				}
-			}
-			OutputBuilder.AppendFormat("Written {0} {1} to {2}.", UniqueItems.Count, (UniqueItems.Count == 1)? "diagnostic" : "diagnostics", OutputFile.FullName);
-
-			// Return the text to be output. We don't have a custom output handler, so we can just return an empty string.
-			Output = OutputBuilder.ToString();
-			ExitCode = 0;
-		}
 	}
 }
diff --git a/Engine/Source/Programs/UnrealBuildTool/System/ActionGraph.cs b/Engine/Source/Programs/UnrealBuildTool/System/ActionGraph.cs
--- a/Engine/Source/Programs/UnrealBuildTool/System/ActionGraph.cs
+++ b/Engine/Source/Programs/UnrealBuildTool/System/ActionGraph.cs
@@ -117,12 +117,6 @@ namespace UnrealBuildTool
 		/// </summary>
 		public bool bProducesImportLibrary = false;
 
-		/// <summary>
-		/// Callback used to perform a special action instead of a generic command line
-		/// </summary>
-		public delegate void BlockingActionHandler(Action Action, out int ExitCode, out string Output);
-		public BlockingActionHandler ActionHandler = null;		// @todo ubtmake urgent: Delegate variables are not saved, but we are comparing against this in ExecuteActions() for XGE!
-
 
 
 		///
@@ -436,11 +430,7 @@ namespace UnrealBuildTool
 				DateTime StartTime = DateTime.UtcNow;
 
 				ActionExecutor Executor;
-				if(ActionsToExecute.Any(x => x.ActionHandler != null))
-				{
-					Executor = new LocalExecutor();
-				}
-				else if ((XGE.IsAvailable() && BuildConfiguration.bAllowXGE) || BuildConfiguration.bXGEExport)
+				if ((XGE.IsAvailable() && BuildConfiguration.bAllowXGE) || BuildConfiguration.bXGEExport)
 				{
 					Executor = new XGE();
 				}
diff --git a/Engine/Source/Programs/UnrealBuildTool/UnrealBuildTool.cs b/Engine/Source/Programs/UnrealBuildTool/UnrealBuildTool.cs
--- a/Engine/Source/Programs/UnrealBuildTool/UnrealBuildTool.cs
+++ b/Engine/Source/Programs/UnrealBuildTool/UnrealBuildTool.cs
@@ -1579,7 +1579,7 @@ namespace UnrealBuildTool
 							UBTMakefile.SourceFileWorkingSet = Unity.SourceFileWorkingSet;
 							UBTMakefile.CandidateSourceFilesForWorkingSet = Unity.CandidateSourceFilesForWorkingSet;
 
-							if (BuildConfiguration.bUseUBTMakefiles && !UBTMakefile.PrerequisiteActions.Any(x => x.ActionHandler != null))
+							if (BuildConfiguration.bUseUBTMakefiles)
 							{
 								if(UBTMakefile.PrerequisiteActions.Length == 0)
 								{
