UBT: Fix project files being generated for VS2017 if the user only has AutoSDKs, and does not have the IDE installed.

diff --git a/Engine/Source/Programs/UnrealBuildTool/Platform/Windows/UEBuildWindows.cs b/Engine/Source/Programs/UnrealBuildTool/Platform/Windows/UEBuildWindows.cs
--- a/Engine/Source/Programs/UnrealBuildTool/Platform/Windows/UEBuildWindows.cs
+++ b/Engine/Source/Programs/UnrealBuildTool/Platform/Windows/UEBuildWindows.cs
@@ -807,6 +807,16 @@ namespace UnrealBuildTool
 			return FileReference.Exists(FileReference.Combine(ToolChainDir, "bin", "Hostx86", "x64", "cl.exe")) || FileReference.Exists(FileReference.Combine(ToolChainDir, "bin", "Hostx64", "x64", "cl.exe"));
 		}
 
+		/// <summary>
+		/// Determines if an IDE for the given compiler is installed.
+		/// </summary>
+		/// <param name="Compiler">Compiler to check for</param>
+		/// <returns>True if the given compiler is installed</returns>
+		public static bool HasIDE(WindowsCompiler Compiler)
+		{
+			return FindVSInstallDirs(Compiler).Count > 0;
+		}
+
 		/// <summary>
 		/// Determines if a given compiler is installed
 		/// </summary>
diff --git a/Engine/Source/Programs/UnrealBuildTool/ProjectFiles/VisualStudio/VCProjectFileGenerator.cs b/Engine/Source/Programs/UnrealBuildTool/ProjectFiles/VisualStudio/VCProjectFileGenerator.cs
--- a/Engine/Source/Programs/UnrealBuildTool/ProjectFiles/VisualStudio/VCProjectFileGenerator.cs
+++ b/Engine/Source/Programs/UnrealBuildTool/ProjectFiles/VisualStudio/VCProjectFileGenerator.cs
@@ -288,11 +288,11 @@ namespace UnrealBuildTool
 			if (ProjectFileFormat == VCProjectFileFormat.Default)
 			{
 				// Pick the best platform installed by default
-				if (WindowsPlatform.HasCompiler(WindowsCompiler.VisualStudio2017))
+				if (WindowsPlatform.HasCompiler(WindowsCompiler.VisualStudio2017) && WindowsPlatform.HasIDE(WindowsCompiler.VisualStudio2017))
 				{
 					ProjectFileFormat = VCProjectFileFormat.VisualStudio2017;
 				}
-				else if (WindowsPlatform.HasCompiler(WindowsCompiler.VisualStudio2015))
+				else if (WindowsPlatform.HasCompiler(WindowsCompiler.VisualStudio2015) && WindowsPlatform.HasIDE(WindowsCompiler.VisualStudio2015))
 				{
 					ProjectFileFormat = VCProjectFileFormat.VisualStudio2015;
 				}
