Expose the bIsEngineInstalled flag through ReadOnlyTargetRules.

diff --git a/Engine/Source/Programs/UnrealBuildTool/Configuration/TargetRules.cs b/Engine/Source/Programs/UnrealBuildTool/Configuration/TargetRules.cs
--- a/Engine/Source/Programs/UnrealBuildTool/Configuration/TargetRules.cs
+++ b/Engine/Source/Programs/UnrealBuildTool/Configuration/TargetRules.cs
@@ -2143,6 +2143,11 @@ namespace UnrealBuildTool
 			get { return Inner.bGenerateProjectFiles; }
 		}
 
+		public bool bIsEngineInstalled
+		{
+			get { return Inner.bIsEngineInstalled; }
+		}
+
 		#if !__MonoCS__
 		#pragma warning restore C1591
 		#endif
