UBT: Encapsulate functionality for creating a system-wide mutex for UBT in a separate class, so it can be used by different tool modes.

#rnx

diff --git a/Engine/Source/Programs/UnrealBuildTool/System/SingleInstanceMutex.cs b/Engine/Source/Programs/UnrealBuildTool/System/SingleInstanceMutex.cs
new file mode 100644
--- /dev/null
+++ b/Engine/Source/Programs/UnrealBuildTool/System/SingleInstanceMutex.cs
@@ -0,0 +1,114 @@
+using System;
+using System.Collections.Generic;
+using System.Linq;
+using System.Reflection;
+using System.Text;
+using System.Threading;
+using System.Threading.Tasks;
+
+namespace UnrealBuildTool
+{
+	/// <summary>
+	/// The type of mutex to aquire
+	/// </summary>
+	enum SingleInstanceMutexType
+	{
+		/// <summary>
+		/// Prevent any instance running in the system
+		/// </summary>
+		Global,
+
+		/// <summary>
+		/// Prevents any instance running in the current branch
+		/// </summary>
+		PerBranch,
+	}
+
+	/// <summary>
+	/// System-wide mutex allowing only one instance of the program to run at a time
+	/// </summary>
+	class SingleInstanceMutex : IDisposable
+	{
+		/// <summary>
+		/// The global mutex instance
+		/// </summary>
+		Mutex GlobalMutex;
+
+		/// <summary>
+		/// Constructor. Attempts to acquire the global mutex
+		/// </summary>
+		/// <param name="Type">The type of mutex to acquire</param>
+		/// <param name="bWaitMutex"></param>
+		public SingleInstanceMutex(SingleInstanceMutexType Type, bool bWaitMutex)
+		{
+			// Get the mutex name
+			string MutexName;
+			if (Type == SingleInstanceMutexType.Global)
+			{
+				MutexName = "Global\\UnrealBuildTool_Mutex_AutoSDKS";
+			}
+			else
+			{
+				MutexName = "Global\\UnrealBuildTool_Mutex_" + Assembly.GetExecutingAssembly().CodeBase.GetHashCode().ToString();
+			}
+
+			// Try to create the mutex, with it initially locked
+			bool bCreatedMutex;
+			GlobalMutex = new Mutex(true, MutexName, out bCreatedMutex);
+
+			// If we didn't create the mutex, we can wait for it or fail immediately
+			if (!bCreatedMutex)
+			{
+				if (bWaitMutex)
+				{
+					try
+					{
+						GlobalMutex.WaitOne();
+					}
+					catch (AbandonedMutexException)
+					{
+					}
+				}
+				else if (Type == SingleInstanceMutexType.Global)
+				{
+					throw new BuildException("A conflicting instance of UnrealBuildTool is already running.");
+				}
+				else
+				{
+					throw new BuildException("A conflicting instance of UnrealBuildTool is already running at {0}.", Assembly.GetExecutingAssembly().CodeBase);
+				}
+			}
+		}
+
+		/// <summary>
+		/// Attempt to aquire the single instance mutex, using settings determined from the command line
+		/// </summary>
+		/// <param name="Type">The type of mutex to aquire</param>
+		/// <param name="Arguments">Command line arguments</param>
+		/// <returns>The single instance mutex, or null if the -NoMutex argument is specified on the command line</returns>
+		public static SingleInstanceMutex Acquire(SingleInstanceMutexType Type, CommandLineArguments Arguments)
+		{
+			if(Arguments.HasOption("-NoMutex"))
+			{
+				return null;
+			}
+			else
+			{
+				return new SingleInstanceMutex(Type, Arguments.HasOption("-WaitMutex"));
+			}
+		}
+
+		/// <summary>
+		/// Release the mutex and dispose of the object
+		/// </summary>
+		public void Dispose()
+		{
+			if (GlobalMutex != null)
+			{
+				GlobalMutex.ReleaseMutex();
+				GlobalMutex.Dispose();
+				GlobalMutex = null;
+			}
+		}
+	}
+}
diff --git a/Engine/Source/Programs/UnrealBuildTool/UnrealBuildTool.cs b/Engine/Source/Programs/UnrealBuildTool/UnrealBuildTool.cs
--- a/Engine/Source/Programs/UnrealBuildTool/UnrealBuildTool.cs
+++ b/Engine/Source/Programs/UnrealBuildTool/UnrealBuildTool.cs
@@ -165,7 +165,7 @@ namespace UnrealBuildTool
 		{
 			if (!bIsProjectInstalled.HasValue)
 			{
-				bIsProjectInstalled = FileReference.Exists(FileReference.Combine(RootDirectory, "Engine", "Build", "InstalledProjectBuild.txt"));
+				bIsProjectInstalled = FileReference.Exists(FileReference.Combine(EngineDirectory, "Build", "InstalledProjectBuild.txt"));
 			}
 			return bIsProjectInstalled.Value;
 		}
@@ -315,8 +315,8 @@ namespace UnrealBuildTool
 		/// This function is for super early startup stuff that should not access Configuration classes (anything loaded by XmlConfg).
 		/// This should be very minimal startup code.
 		/// </summary>
-		/// <param name="Arguments">Cmdline arguments</param>
-		internal static int GuardedMain(string[] Arguments)
+		/// <param name="CmdLine">Cmdline arguments</param>
+		internal static int GuardedMain(CommandLineArguments CmdLine)
 		{
 			DateTime StartTime = DateTime.UtcNow;
 			ECompilationResult Result = ECompilationResult.Succeeded;
@@ -328,9 +328,6 @@ namespace UnrealBuildTool
 			// Write the command line
 			Log.TraceLog("Command line: {0}", Environment.CommandLine);
 
-			// ensure we can resolve any external assemblies that are not in the same folder as our assembly.
-			AssemblyUtils.InstallAssemblyResolver(Path.GetDirectoryName(Assembly.GetEntryAssembly().GetOriginalLocation()));
-
 			// Grab the environment.
 			InitialEnvironment = Environment.GetEnvironmentVariables();
 			if (InitialEnvironment.Count < 1)
@@ -338,111 +335,19 @@ namespace UnrealBuildTool
 				throw new BuildException("Environment could not be read");
 			}
 
-			// @todo: Ideally we never need to Mutex unless we are invoked with the same target project,
-			// in the same branch/path!  This would allow two clientspecs to build at the same time (even though we need
-			// to make sure that XGE is only used once at a time)
-			bool bUseMutex = true;
-			{
-				int NoMutexArgumentIndex;
-				if (Utils.ParseCommandLineFlag(Arguments, "-NoMutex", out NoMutexArgumentIndex))
-				{
-					bUseMutex = false;
-				}
-			}
-			bool bWaitMutex = false;
-			{
-				int WaitMutexArgumentIndex;
-				if (Utils.ParseCommandLineFlag(Arguments, "-WaitMutex", out WaitMutexArgumentIndex))
-				{
-					bWaitMutex = true;
-				}
-			}
+			bool bAutoSDKOnly = CmdLine.HasOption("-AutoSDKOnly");
+			bool bValidatePlatforms = CmdLine.HasOption("-ValidatePlatform");
 
-			bool bAutoSDKOnly = false;
+			SingleInstanceMutexType MutexType = (bAutoSDKOnly || bValidatePlatforms)? SingleInstanceMutexType.Global : SingleInstanceMutexType.PerBranch;
+			using(SingleInstanceMutex.Acquire(MutexType, CmdLine))
 			{
-				int AutoSDKOnlyArgumentIndex;
-				if (Utils.ParseCommandLineFlag(Arguments, "-autosdkonly", out AutoSDKOnlyArgumentIndex))
-				{
-					bAutoSDKOnly = true;
-				}
-			}
-			bool bValidatePlatforms = false;
-			{
-				int ValidatePlatformsArgumentIndex;
-				if (Utils.ParseCommandLineFlag(Arguments, "-validateplatform", out ValidatePlatformsArgumentIndex))
-				{
-					bValidatePlatforms = true;
-				}
-			}
-
-			// Don't allow simultaneous execution of Unreal Built Tool. Multi-selection in the UI e.g. causes this and you want serial
-			// execution in this case to avoid contention issues with shared produced items.
-			bool bCreatedMutex = false;
-			{
-				Mutex SingleInstanceMutex = null;
-				if (bUseMutex)
-				{
-					int LocationHash = Assembly.GetEntryAssembly().GetOriginalLocation().GetHashCode();
-
-					String MutexName;
-					if (bAutoSDKOnly || bValidatePlatforms)
-					{
-						// this mutex has to be truly global because AutoSDK installs may change global state (like a console copying files into VS directories,
-						// or installing Target Management services
-						MutexName = "Global\\UnrealBuildTool_Mutex_AutoSDKS";
-					}
-					else
-					{
-						MutexName = "Global\\UnrealBuildTool_Mutex_" + LocationHash.ToString();
-					}
-					SingleInstanceMutex = new Mutex(true, MutexName, out bCreatedMutex);
-					if (!bCreatedMutex)
-					{
-						if (bWaitMutex)
-						{
-							try
-							{
-								SingleInstanceMutex.WaitOne();
-							}
-							catch (AbandonedMutexException)
-							{
-							}
-						}
-						else if (bAutoSDKOnly || bValidatePlatforms)
-						{
-							throw new BuildException("A conflicting instance of UnrealBuildTool is already running. Either -autosdkonly or -validateplatform was passed, which allows only a single instance to be run globally. Therefore, the conflicting instance may be in another location or the current location: {0}. A process manager may be used to determine the conflicting process and what tool may have launched it.", Assembly.GetEntryAssembly().GetOriginalLocation());
-						}
-						else
-						{
-							throw new BuildException("A conflicting instance of UnrealBuildTool is already running at this location: {0}. A process manager may be used to determine the conflicting process and what tool may have launched it.", Assembly.GetEntryAssembly().GetOriginalLocation());
-						}
-					}
-				}
-
+				string[] Arguments = CmdLine.GetRawArray();
 				try
 				{
+
 					// Change the working directory to be the Engine/Source folder. We are likely running from Engine/Binaries/DotNET
 					// This is critical to be done early so any code that relies on the current directory being Engine/Source will work.
-					// UEBuildConfiguration.PostReset is susceptible to this, so we must do this before configs are loaded.
-					string EngineSourceDirectory = Path.Combine(Path.GetDirectoryName(Assembly.GetEntryAssembly().GetOriginalLocation()), "..", "..", "..", "Engine", "Source");
-
-					[REDACTED_EMAIL]: This is a workaround for recompiling game code in editor
-					// The working directory when launching is *not* what we would expect
-					if (Directory.Exists(EngineSourceDirectory) == false)
-					{
-						// We are assuming UBT always runs from <>/Engine/Binaries/DotNET/...
-						EngineSourceDirectory = Assembly.GetExecutingAssembly().GetOriginalLocation();
-						EngineSourceDirectory = EngineSourceDirectory.Replace("\\", "/");
-						Int32 EngineIdx = EngineSourceDirectory.IndexOf("/Engine/Binaries/DotNET/", StringComparison.InvariantCultureIgnoreCase);
-						if (EngineIdx > 0)
-						{
-							EngineSourceDirectory = Path.Combine(EngineSourceDirectory.Substring(0, EngineIdx), "Engine", "Source");
-						}
-					}
-					if (Directory.Exists(EngineSourceDirectory)) // only set the directory if it exists, this should only happen if we are launching the editor from an artist sync
-					{
-						Directory.SetCurrentDirectory(EngineSourceDirectory);
-					}
+					DirectoryReference.SetCurrentDirectory(EngineSourceDirectory);
 
 					// Read the XML configuration files
 					bool bForceXmlConfigCache = Arguments.Any(x => x.Equals("-ForceXmlConfigCache", StringComparison.InvariantCultureIgnoreCase));
@@ -847,14 +752,6 @@ namespace UnrealBuildTool
 					Log.TraceError("UnrealBuildTool Exception: " + Exception);
 					Result = ECompilationResult.OtherCompilationError;
 				}
-
-				if (bUseMutex)
-				{
-					// Release the mutex to avoid the abandoned mutex timeout.
-					SingleInstanceMutex.ReleaseMutex();
-					SingleInstanceMutex.Dispose();
-					SingleInstanceMutex = null;
-				}
 			}
 
 			// Print some performance info
@@ -950,6 +847,9 @@ namespace UnrealBuildTool
 					}
 				}
 
+				// Ensure we can resolve any external assemblies that are not in the same folder as our assembly.
+				AssemblyUtils.InstallAssemblyResolver(Path.GetDirectoryName(Assembly.GetEntryAssembly().GetOriginalLocation()));
+
 				// Try to get the correct mode
 				Type ModeType;
 				if(!ModeNameToType.TryGetValue(Options.Mode, out ModeType))
@@ -2308,7 +2208,7 @@ namespace UnrealBuildTool
 	{
 		public override int Execute(CommandLineArguments Arguments)
 		{
-			return UnrealBuildTool.GuardedMain(Arguments.GetRawArray());
+			return UnrealBuildTool.GuardedMain(Arguments);
 		}
 	}
 }
diff --git a/Engine/Source/Programs/UnrealBuildTool/UnrealBuildTool.csproj b/Engine/Source/Programs/UnrealBuildTool/UnrealBuildTool.csproj
--- a/Engine/Source/Programs/UnrealBuildTool/UnrealBuildTool.csproj
+++ b/Engine/Source/Programs/UnrealBuildTool/UnrealBuildTool.csproj
@@ -205,6 +205,7 @@
     <Compile Include="System\ProgressWriter.cs" />
     <Compile Include="System\RulesAssembly.cs" />
     <Compile Include="System\ScopedTimer.cs" />
+    <Compile Include="System\SingleInstanceMutex.cs" />
     <Compile Include="System\SourceFileWorkingSet.cs" />
     <Compile Include="System\StartupTraceListener.cs" />
     <Compile Include="System\SupportedConfigurationsAttribute.cs" />
