UBT: Create a separate tool mode for validating platforms, rather than handling in the regular UBT loop. Invoked by specifying -Mode=ValidatePlatforms rather than specifying -ValidatePlatforms on the command line.

diff --git a/Engine/Source/Developer/TargetPlatform/Private/TargetPlatformManagerModule.cpp b/Engine/Source/Developer/TargetPlatform/Private/TargetPlatformManagerModule.cpp
--- a/Engine/Source/Developer/TargetPlatform/Private/TargetPlatformManagerModule.cpp
+++ b/Engine/Source/Developer/TargetPlatform/Private/TargetPlatformManagerModule.cpp
@@ -923,18 +923,18 @@ RETRY_SETUPANDVALIDATE:
 		{
 			CmdExe = TEXT("/bin/sh");
 			FString ScriptPath = FPaths::ConvertRelativePathToFull(FPaths::EngineDir() / TEXT("Build/BatchFiles/Mac/RunMono.sh"));
-			CommandLine = TEXT("\"") + ScriptPath + TEXT("\" \"") + FPaths::ConvertRelativePathToFull(FPaths::EngineDir() / TEXT("Binaries/DotNET/UnrealBuildTool.exe")) + TEXT("\" -validateplatform");
+			CommandLine = TEXT("\"") + ScriptPath + TEXT("\" \"") + FPaths::ConvertRelativePathToFull(FPaths::EngineDir() / TEXT("Binaries/DotNET/UnrealBuildTool.exe")) + TEXT("\" -Mode=ValidatePlatforms");
 		}
 		else if (PLATFORM_WINDOWS)
 		{
 			CmdExe = FPaths::ConvertRelativePathToFull(FPaths::EngineDir() / TEXT("Binaries/DotNET/UnrealBuildTool.exe"));
-			CommandLine = TEXT("-validateplatform");
+			CommandLine = TEXT("-Mode=ValidatePlatforms");
 		}
 		else if (PLATFORM_LINUX)
 		{
 			CmdExe = TEXT("/bin/bash");	// bash and not sh because of pushd
 			FString ScriptPath = FPaths::ConvertRelativePathToFull(FPaths::EngineDir() / TEXT("Build/BatchFiles/Linux/RunMono.sh"));
-			CommandLine = TEXT("\"") + ScriptPath + TEXT("\" \"") + FPaths::ConvertRelativePathToFull(FPaths::EngineDir() / TEXT("Binaries/DotNET/UnrealBuildTool.exe")) + TEXT("\" -validateplatform");
+			CommandLine = TEXT("\"") + ScriptPath + TEXT("\" \"") + FPaths::ConvertRelativePathToFull(FPaths::EngineDir() / TEXT("Binaries/DotNET/UnrealBuildTool.exe")) + TEXT("\" -Mode=ValidatePlatforms");
 		}
 		else
 		{
diff --git a/Engine/Source/Programs/UnrealBuildTool/Configuration/UEBuildPlatform.cs b/Engine/Source/Programs/UnrealBuildTool/Configuration/UEBuildPlatform.cs
--- a/Engine/Source/Programs/UnrealBuildTool/Configuration/UEBuildPlatform.cs
+++ b/Engine/Source/Programs/UnrealBuildTool/Configuration/UEBuildPlatform.cs
@@ -503,14 +503,6 @@ namespace UnrealBuildTool
 			}
 		}
 
-		/// <summary>
-		/// Returns the name that should be returned in the output when doing -validateplatforms
-		/// </summary>
-		public virtual string GetPlatformValidationName()
-		{
-			return Platform.ToString();
-		}
-
 		/// <summary>
 		/// If this platform can be compiled with XGE
 		/// </summary>
diff --git a/Engine/Source/Programs/UnrealBuildTool/UnrealBuildTool.cs b/Engine/Source/Programs/UnrealBuildTool/UnrealBuildTool.cs
--- a/Engine/Source/Programs/UnrealBuildTool/UnrealBuildTool.cs
+++ b/Engine/Source/Programs/UnrealBuildTool/UnrealBuildTool.cs
@@ -336,9 +336,8 @@ namespace UnrealBuildTool
 			}
 
 			bool bAutoSDKOnly = CmdLine.HasOption("-AutoSDKOnly");
-			bool bValidatePlatforms = CmdLine.HasOption("-ValidatePlatform");
 
-			SingleInstanceMutexType MutexType = (bAutoSDKOnly || bValidatePlatforms)? SingleInstanceMutexType.Global : SingleInstanceMutexType.PerBranch;
+			SingleInstanceMutexType MutexType = bAutoSDKOnly? SingleInstanceMutexType.Global : SingleInstanceMutexType.PerBranch;
 			using(SingleInstanceMutex.Acquire(MutexType, CmdLine))
 			{
 				string[] Arguments = CmdLine.GetRawArray();
@@ -366,7 +365,7 @@ namespace UnrealBuildTool
 					bPrintPerformanceInfo = BuildConfiguration.bPrintPerformanceInfo;
 
 					// Don't run junk deleter if we're just setting up autosdks
-					if (bAutoSDKOnly || bValidatePlatforms)
+					if (bAutoSDKOnly)
 					{
 						BuildConfiguration.bIgnoreJunk = true;
 					}
@@ -609,7 +608,7 @@ namespace UnrealBuildTool
 					}
 
 					// Find and register all tool chains, build platforms, etc. that are present
-					RegisterAllUBTClasses(bValidatePlatforms);
+					RegisterAllUBTClasses(false);
 
 					if (UnrealBuildTool.bPrintPerformanceInfo)
 					{
@@ -714,10 +713,6 @@ namespace UnrealBuildTool
 						// RegisterAllUBTClasses has already done all the SDK setup.
 						Result = ECompilationResult.Succeeded;
 					}
-					else if (bValidatePlatforms)
-					{
-						ValidatePlatforms(Arguments);
-					}
 					else if (BuildConfiguration.DeployTargetFile != null)
 					{
 						UEBuildDeployTarget DeployTarget = new UEBuildDeployTarget(BuildConfiguration.DeployTargetFile);
@@ -898,93 +893,6 @@ namespace UnrealBuildTool
 			return bSuccess;
 		}
 
-
-
-		/// <summary>
-		/// Validates the various platforms to determine if they are ready for building
-		/// </summary>
-		public static void ValidatePlatforms(string[] Arguments)
-		{
-			List<UnrealTargetPlatform> Platforms = new List<UnrealTargetPlatform>();
-			foreach (string CurArgument in Arguments)
-			{
-				if (CurArgument.StartsWith("-"))
-				{
-					if (CurArgument.StartsWith("-Platforms=", StringComparison.InvariantCultureIgnoreCase))
-					{
-						// Parse the list... will be in Foo+Bar+New format
-						string PlatformList = CurArgument.Substring(11);
-						while (PlatformList.Length > 0)
-						{
-							string PlatformString = PlatformList;
-							Int32 PlusIdx = PlatformList.IndexOf("+");
-							if (PlusIdx != -1)
-							{
-								PlatformString = PlatformList.Substring(0, PlusIdx);
-								PlatformList = PlatformList.Substring(PlusIdx + 1);
-							}
-							else
-							{
-								// We are on the last platform... clear the list to exit the loop
-								PlatformList = "";
-							}
-
-							// Is the string a valid platform? If so, add it to the list
-							UnrealTargetPlatform SpecifiedPlatform = UnrealTargetPlatform.Unknown;
-							foreach (UnrealTargetPlatform PlatformParam in Enum.GetValues(typeof(UnrealTargetPlatform)))
-							{
-								if (PlatformString.Equals(PlatformParam.ToString(), StringComparison.InvariantCultureIgnoreCase))
-								{
-									SpecifiedPlatform = PlatformParam;
-									break;
-								}
-							}
-
-							if (SpecifiedPlatform != UnrealTargetPlatform.Unknown)
-							{
-								if (Platforms.Contains(SpecifiedPlatform) == false)
-								{
-									Platforms.Add(SpecifiedPlatform);
-								}
-							}
-							else
-							{
-								Log.TraceWarning("ValidatePlatforms invalid platform specified: {0}", PlatformString);
-							}
-						}
-					}
-					else switch (CurArgument.ToUpperInvariant())
-						{
-							case "-ALLPLATFORMS":
-								foreach (UnrealTargetPlatform platform in Enum.GetValues(typeof(UnrealTargetPlatform)))
-								{
-									if (platform != UnrealTargetPlatform.Unknown)
-									{
-										if (Platforms.Contains(platform) == false)
-										{
-											Platforms.Add(platform);
-										}
-									}
-								}
-								break;
-						}
-				}
-			}
-
-			foreach (UnrealTargetPlatform platform in Platforms)
-			{
-				UEBuildPlatform BuildPlatform = UEBuildPlatform.GetBuildPlatform(platform, true);
-				if (BuildPlatform != null && BuildPlatform.HasRequiredSDKsInstalled() == SDKStatus.Valid)
-				{
-					Log.TraceInformation("##PlatformValidate: {0} VALID", BuildPlatform.GetPlatformValidationName());
-				}
-				else
-				{
-					Log.TraceInformation("##PlatformValidate: {0} INVALID", BuildPlatform != null ? BuildPlatform.GetPlatformValidationName() : platform.ToString());
-				}
-			}
-		}
-
 		internal static ECompilationResult RunUBT(BuildConfiguration BuildConfiguration, string[] Arguments, FileReference ProjectFile, bool bCatchExceptions)
 		{
 			bool bSuccess = true;
diff --git a/Engine/Source/Programs/UnrealBuildTool/UnrealBuildTool.csproj b/Engine/Source/Programs/UnrealBuildTool/UnrealBuildTool.csproj
--- a/Engine/Source/Programs/UnrealBuildTool/UnrealBuildTool.csproj
+++ b/Engine/Source/Programs/UnrealBuildTool/UnrealBuildTool.csproj
@@ -140,6 +140,7 @@
     </When>
   </Choose>
   <ItemGroup>
+    <Compile Include="Modes\ValidatePlatformsMode.cs" />
     <Compile Include="Modes\WriteDocumentationMode.cs" />
     <Compile Include="Configuration\ModuleRules.cs" />
     <Compile Include="Configuration\PlatformTargetRules.cs" />
