UBT: Move the -AutoSDKOnly functionality into a separate mode, invoked by -Mode=SetupPlatforms.

#rnx

diff --git a/Engine/Source/Developer/TargetPlatform/Private/TargetPlatformManagerModule.cpp b/Engine/Source/Developer/TargetPlatform/Private/TargetPlatformManagerModule.cpp
--- a/Engine/Source/Developer/TargetPlatform/Private/TargetPlatformManagerModule.cpp
+++ b/Engine/Source/Developer/TargetPlatform/Private/TargetPlatformManagerModule.cpp
@@ -92,7 +92,7 @@ public:
 			// amortize UBT cost by calling it once for all platforms, rather than once per platform.
 			if (FParse::Param(FCommandLine::Get(), TEXT("Multiprocess"))==false)
 			{
-				FString UBTParams(TEXT("-autosdkonly"));
+				FString UBTParams(TEXT("-SetupPlatforms"));
 				int32 UBTReturnCode = -1;
 				FString UBTOutput;
 				if (!FDesktopPlatformModule::Get()->InvokeUnrealBuildToolSync(UBTParams, *GLog, true, UBTReturnCode, UBTOutput))
diff --git a/Engine/Source/Programs/UnrealBuildTool/Modes/SetupPlatformsMode.cs b/Engine/Source/Programs/UnrealBuildTool/Modes/SetupPlatformsMode.cs
new file mode 100644
--- /dev/null
+++ b/Engine/Source/Programs/UnrealBuildTool/Modes/SetupPlatformsMode.cs
@@ -0,0 +1,37 @@
+using System;
+using System.Collections.Generic;
+using System.Linq;
+using System.Text;
+using System.Threading.Tasks;
+using Tools.DotNETCommon;
+
+namespace UnrealBuildTool
+{
+	/// <summary>
+	/// Register all platforms (and in the process, configure all autosdks)
+	/// </summary>
+	[ToolMode("SetupPlatforms")]
+	class SetupPlatforms : ToolMode
+	{
+		/// <summary>
+		/// Execute the tool mode
+		/// </summary>
+		/// <param name="Arguments">Command line arguments</param>
+		/// <returns>Exit code</returns>
+		public override int Execute(CommandLineArguments Arguments)
+		{
+			using(SingleInstanceMutex.Acquire(SingleInstanceMutexType.Global, Arguments))
+			{
+				// Output a warning if there are any arguments that are still unused
+				Arguments.CheckAllArgumentsUsed();
+
+				// Read the XML configuration files
+				XmlConfig.ReadConfigFiles(false);
+
+				// Find and register all tool chains, build platforms, etc. that are present
+				UnrealBuildTool.RegisterAllUBTClasses(false);
+			}
+			return 0;
+		}
+	}
+}
diff --git a/Engine/Source/Programs/UnrealBuildTool/System/PlatformExports.cs b/Engine/Source/Programs/UnrealBuildTool/System/PlatformExports.cs
--- a/Engine/Source/Programs/UnrealBuildTool/System/PlatformExports.cs
+++ b/Engine/Source/Programs/UnrealBuildTool/System/PlatformExports.cs
@@ -153,10 +153,7 @@ namespace UnrealBuildTool
 		public static bool Initialize()
 		{
 			// Read the XML configuration files
-			if(!XmlConfig.ReadConfigFiles(false))
-			{
-				return false;
-			}
+			XmlConfig.ReadConfigFiles(false);
 
 			// Register all the platform classes
 			UnrealBuildTool.RegisterAllUBTClasses(false);
diff --git a/Engine/Source/Programs/UnrealBuildTool/System/XmlConfig.cs b/Engine/Source/Programs/UnrealBuildTool/System/XmlConfig.cs
--- a/Engine/Source/Programs/UnrealBuildTool/System/XmlConfig.cs
+++ b/Engine/Source/Programs/UnrealBuildTool/System/XmlConfig.cs
@@ -49,7 +49,7 @@ namespace UnrealBuildTool
 		/// Initialize the config system with the given types
 		/// </summary>
 		/// <param name="bForceLoadCache">Force use of the cached XML config without checking if it's valid (useful for remote builds)</param>
-		public static bool ReadConfigFiles(bool bForceLoadCache)
+		public static void ReadConfigFiles(bool bForceLoadCache)
 		{
 			// Find all the configurable types
 			List<Type> ConfigTypes = FindConfigurableTypes();
@@ -112,8 +112,7 @@ namespace UnrealBuildTool
 					{
 						if(!TryReadFile(InputFile, CategoryToFields, TypeToValues, Schema))
 						{
-							Log.TraceError("Failed to properly read XML file : {0}", InputFile.FullName);
-							return false;
+							throw new BuildException("Failed to properly read XML file : {0}", InputFile.FullName);
 						}
 					}
 
@@ -138,7 +137,6 @@ namespace UnrealBuildTool
 					}
 				}
 			}
-			return true;
 		}
 
 		/// <summary>
diff --git a/Engine/Source/Programs/UnrealBuildTool/UnrealBuildTool.cs b/Engine/Source/Programs/UnrealBuildTool/UnrealBuildTool.cs
--- a/Engine/Source/Programs/UnrealBuildTool/UnrealBuildTool.cs
+++ b/Engine/Source/Programs/UnrealBuildTool/UnrealBuildTool.cs
@@ -335,10 +335,7 @@ namespace UnrealBuildTool
 				throw new BuildException("Environment could not be read");
 			}
 
-			bool bAutoSDKOnly = CmdLine.HasOption("-AutoSDKOnly");
-
-			SingleInstanceMutexType MutexType = bAutoSDKOnly? SingleInstanceMutexType.Global : SingleInstanceMutexType.PerBranch;
-			using(SingleInstanceMutex.Acquire(MutexType, CmdLine))
+			using(SingleInstanceMutex.Acquire(SingleInstanceMutexType.PerBranch, CmdLine))
 			{
 				string[] Arguments = CmdLine.GetRawArray();
 				try
@@ -350,10 +347,7 @@ namespace UnrealBuildTool
 
 					// Read the XML configuration files
 					bool bForceXmlConfigCache = Arguments.Any(x => x.Equals("-ForceXmlConfigCache", StringComparison.InvariantCultureIgnoreCase));
-					if (!XmlConfig.ReadConfigFiles(bForceXmlConfigCache))
-					{
-						return 1;
-					}
+					XmlConfig.ReadConfigFiles(bForceXmlConfigCache);
 
 					// Create the build configuration object, and read the settings
 					BuildConfiguration BuildConfiguration = new BuildConfiguration();
@@ -364,12 +358,6 @@ namespace UnrealBuildTool
 					bPrintDebugInfo = BuildConfiguration.bPrintDebugInfo || Log.OutputLevel == LogEventType.Verbose || Log.OutputLevel == LogEventType.VeryVerbose;
 					bPrintPerformanceInfo = BuildConfiguration.bPrintPerformanceInfo;
 
-					// Don't run junk deleter if we're just setting up autosdks
-					if (bAutoSDKOnly)
-					{
-						BuildConfiguration.bIgnoreJunk = true;
-					}
-
 					// Then let the command lines override any configs necessary.
 					string LogSuffix = "";
 					foreach (string Argument in Arguments)
@@ -708,11 +696,6 @@ namespace UnrealBuildTool
 							}
 						}
 					}
-					else if (bAutoSDKOnly)
-					{
-						// RegisterAllUBTClasses has already done all the SDK setup.
-						Result = ECompilationResult.Succeeded;
-					}
 					else if (BuildConfiguration.DeployTargetFile != null)
 					{
 						UEBuildDeployTarget DeployTarget = new UEBuildDeployTarget(BuildConfiguration.DeployTargetFile);
diff --git a/Engine/Source/Programs/UnrealBuildTool/UnrealBuildTool.csproj b/Engine/Source/Programs/UnrealBuildTool/UnrealBuildTool.csproj
--- a/Engine/Source/Programs/UnrealBuildTool/UnrealBuildTool.csproj
+++ b/Engine/Source/Programs/UnrealBuildTool/UnrealBuildTool.csproj
@@ -140,6 +140,7 @@
     </When>
   </Choose>
   <ItemGroup>
+    <Compile Include="Modes\SetupPlatformsMode.cs" />
     <Compile Include="Modes\ValidatePlatformsMode.cs" />
     <Compile Include="Modes\WriteDocumentationMode.cs" />
     <Compile Include="Configuration\ModuleRules.cs" />
