Missing file.

#rnx

diff --git a/Engine/Source/Programs/UnrealBuildTool/Modes/ValidatePlatformsMode.cs b/Engine/Source/Programs/UnrealBuildTool/Modes/ValidatePlatformsMode.cs
new file mode 100644
--- /dev/null
+++ b/Engine/Source/Programs/UnrealBuildTool/Modes/ValidatePlatformsMode.cs
@@ -0,0 +1,70 @@
+using System;
+using System.Collections.Generic;
+using System.Linq;
+using System.Text;
+using System.Threading.Tasks;
+using Tools.DotNETCommon;
+
+namespace UnrealBuildTool
+{
+	/// <summary>
+	/// Validates the various platforms to determine if they are ready for building
+	/// </summary>
+	[ToolMode("ValidatePlatforms")]
+	class ValidatePlatformsMode : ToolMode
+	{
+		/// <summary>
+		/// Platforms to validate
+		/// </summary>
+		[CommandLine("-Platforms=", ListSeparator = '+')]
+		HashSet<UnrealTargetPlatform> Platforms = new HashSet<UnrealTargetPlatform>();
+
+		/// <summary>
+		/// Whether to validate all platforms
+		/// </summary>
+		[CommandLine("-AllPlatforms")]
+		bool bAllPlatforms = false;
+
+		/// <summary>
+		/// Executes the tool with the given arguments
+		/// </summary>
+		/// <param name="Arguments">Command line arguments</param>
+		/// <returns>Exit code</returns>
+		public override int Execute(CommandLineArguments Arguments)
+		{
+			Arguments.ApplyTo(this);
+			using(SingleInstanceMutex.Acquire(SingleInstanceMutexType.Global, Arguments))
+			{
+				// Output a message if there are any arguments that are still unused
+				Arguments.CheckAllArgumentsUsed();
+
+				// If the -AllPlatforms argument is specified, add all the known platforms into the list
+				if(bAllPlatforms)
+				{
+					Platforms.UnionWith(Enum.GetValues(typeof(UnrealTargetPlatform)).OfType<UnrealTargetPlatform>().Where(x => x != UnrealTargetPlatform.Unknown));
+				}
+
+				// Read all the config files
+				XmlConfig.ReadConfigFiles(false);
+				
+				// Register all the platforms
+				UnrealBuildTool.RegisterAllUBTClasses(true);
+
+				// Output a line for each registered platform
+				foreach (UnrealTargetPlatform Platform in Platforms)
+				{
+					UEBuildPlatform BuildPlatform = UEBuildPlatform.GetBuildPlatform(Platform, true);
+					if (BuildPlatform != null && BuildPlatform.HasRequiredSDKsInstalled() == SDKStatus.Valid)
+					{
+						Log.TraceInformation("##PlatformValidate: {0} VALID", Platform.ToString());
+					}
+					else
+					{
+						Log.TraceInformation("##PlatformValidate: {0} INVALID", Platform.ToString());
+					}
+				} 
+			}
+			return 0;
+		}
+	}
+}
