UBT: Fix a couple of bugs in the attribute-driven command line parser.

* If the attribute specifies a value, it should always be used rather than the specified value (required for when a key/value pair implies a specific value for another option)
* For container types, we should parse each value as the type of the generic type argument rather than the collection as a whole.

#rnx

diff --git a/Engine/Source/Programs/UnrealBuildTool/System/CommandLineArguments.cs b/Engine/Source/Programs/UnrealBuildTool/System/CommandLineArguments.cs
--- a/Engine/Source/Programs/UnrealBuildTool/System/CommandLineArguments.cs
+++ b/Engine/Source/Programs/UnrealBuildTool/System/CommandLineArguments.cs
@@ -624,9 +624,13 @@ namespace UnrealBuildTool
 
 								// Get the text for this value
 								string ValueText;
-								if(FlagArguments.Get(Index))
+								if(Attribute.Value != null)
 								{
-									ValueText = Attribute.Value ?? "true";
+									ValueText = Attribute.Value;
+								}
+								else if(FlagArguments.Get(Index))
+								{
+									ValueText = "true";
 								}
 								else
 								{
@@ -733,13 +737,8 @@ namespace UnrealBuildTool
 		/// <returns>True if the value was assigned to the field, false otherwise</returns>
 		private static bool ApplyArgument(object TargetObject, FieldInfo Field, string ArgumentText, string ValueText, string PreviousArgumentText)
 		{
-			// Try to parse the value
-			object Value;
-			if(!TryParseValue(Field.FieldType, ValueText, out Value))
-			{
-				Log.TraceWarning("Unable to parse value for argument '{0}'.", ArgumentText);
-				return false;
-			}
+			// The value type for items of this field
+			Type ValueType = Field.FieldType;
 
 			// Check if the field type implements ICollection<>. If so, we can take multiple values.
 			Type CollectionType = null;
@@ -747,11 +746,20 @@ namespace UnrealBuildTool
 			{
 				if (InterfaceType.IsGenericType && InterfaceType.GetGenericTypeDefinition() == typeof(ICollection<>))
 				{
+					ValueType = InterfaceType.GetGenericArguments()[0];
 					CollectionType = InterfaceType;
 					break;
 				}
 			}
 
+			// Try to parse the value
+			object Value;
+			if(!TryParseValue(ValueType, ValueText, out Value))
+			{
+				Log.TraceWarning("Unable to parse value for argument '{0}'.", ArgumentText);
+				return false;
+			}
+
 			// Try to assign values to the target field
 			if (CollectionType == null)
 			{
