UBT: Move functionality for generating project files into a separate mode.

diff --git a/Engine/Source/Programs/UnrealBuildTool/Modes/GenerateProjectFilesMode.cs b/Engine/Source/Programs/UnrealBuildTool/Modes/GenerateProjectFilesMode.cs
new file mode 100644
--- /dev/null
+++ b/Engine/Source/Programs/UnrealBuildTool/Modes/GenerateProjectFilesMode.cs
@@ -0,0 +1,177 @@
+using System;
+using System.Collections.Generic;
+using System.IO;
+using System.Linq;
+using System.Text;
+using System.Threading.Tasks;
+using Tools.DotNETCommon;
+
+namespace UnrealBuildTool
+{
+	/// <summary>
+	/// Generates project files for one or more projects
+	/// </summary>
+	[ToolMode("GenerateProjectFiles")]
+	class GenerateProjectFilesMode : ToolMode
+	{
+		/// <summary>
+		/// Types of project files to generate
+		/// </summary>
+		[CommandLine("-ProjectFileFormat")]
+		[CommandLine("-2012unsupported", Value = nameof(ProjectFileFormat.VisualStudio2012))]
+		[CommandLine("-2013unsupported", Value = nameof(ProjectFileFormat.VisualStudio2013))]
+		[CommandLine("-2015", Value = nameof(ProjectFileFormat.VisualStudio2015))] // + override compiler
+		[CommandLine("-2017", Value = nameof(ProjectFileFormat.VisualStudio2017))] // + override compiler
+		[CommandLine("-Makefile", Value = nameof(ProjectFileFormat.Make))]
+		[CommandLine("-CMakefile", Value = nameof(ProjectFileFormat.CMake))]
+		[CommandLine("-QMakefile", Value = nameof(ProjectFileFormat.QMake))]
+		[CommandLine("-KDevelopfile", Value = nameof(ProjectFileFormat.KDevelop))]
+		[CommandLine("-CodeLiteFiles", Value = nameof(ProjectFileFormat.CodeLite))]
+		[CommandLine("-XCodeProjectFiles", Value = nameof(ProjectFileFormat.XCode))]
+		[CommandLine("-EddieProjectFiles", Value = nameof(ProjectFileFormat.Eddie))]
+		[CommandLine("-VSCode", Value = nameof(ProjectFileFormat.VisualStudioCode))]
+		[CommandLine("-VSMac", Value = nameof(ProjectFileFormat.VisualStudioMac))]
+		[CommandLine("-CLion", Value = nameof(ProjectFileFormat.CLion))]
+		HashSet<ProjectFileFormat> ProjectFileFormats = new HashSet<ProjectFileFormat>();
+
+		/// <summary>
+		/// Execute the tool mode
+		/// </summary>
+		/// <param name="Arguments">Command line arguments</param>
+		/// <returns>Exit code</returns>
+		public override int Execute(CommandLineArguments Arguments)
+		{
+			using(SingleInstanceMutex.Acquire(SingleInstanceMutexType.PerBranch, Arguments))
+			{
+				// Apply any command line arguments to this class
+				Arguments.ApplyTo(this);
+
+				// Parse rocket-specific arguments.
+				FileReference ProjectFile;
+				UnrealBuildTool.TryParseProjectFileArgument(Arguments, out ProjectFile);
+
+				// Change the working directory to be the Engine/Source folder. We are likely running from Engine/Binaries/DotNET
+				// This is critical to be done early so any code that relies on the current directory being Engine/Source will work.
+				DirectoryReference.SetCurrentDirectory(UnrealBuildTool.EngineSourceDirectory);
+
+				// Read the XML configuration files
+				XmlConfig.ReadConfigFiles(false);
+
+				// Find and register all tool chains, build platforms, etc. that are present
+				UnrealBuildTool.RegisterAllUBTClasses(false);
+
+				// HACK
+				if (Arguments.HasOption("-vsdebugandroid"))
+				{
+					AndroidProjectGenerator.VSDebugCommandLineOptionPresent = true;
+				}
+
+				// If there aren't any formats set, read the default project file format from the config file
+				if (ProjectFileFormats.Count == 0)
+				{
+					// Read from the XML config
+					if (!String.IsNullOrEmpty(ProjectFileGeneratorSettings.Format))
+					{
+						ProjectFileFormats.UnionWith(ProjectFileGeneratorSettings.ParseFormatList(ProjectFileGeneratorSettings.Format));
+					}
+
+					// Read from the editor config
+					ProjectFileFormat PreferredSourceCodeAccessor;
+					if (ProjectFileGenerator.GetPreferredSourceCodeAccessor(ProjectFile, out PreferredSourceCodeAccessor))
+					{
+						ProjectFileFormats.Add(PreferredSourceCodeAccessor);
+					}
+
+					// If there's still nothing set, get the default project file format for this platform
+					if (ProjectFileFormats.Count == 0)
+					{
+						ProjectFileFormats.UnionWith(BuildHostPlatform.Current.GetDefaultProjectFileFormats());
+					}
+				}
+
+				// Create each project generator and run it
+				List<ProjectFileGenerator> Generators = new List<ProjectFileGenerator>();
+				foreach (ProjectFileFormat ProjectFileFormat in ProjectFileFormats.Distinct())
+				{
+					ProjectFileGenerator Generator;
+					switch (ProjectFileFormat)
+					{
+						case ProjectFileFormat.Make:
+							Generator = new MakefileGenerator(ProjectFile);
+							break;
+						case ProjectFileFormat.CMake:
+							Generator = new CMakefileGenerator(ProjectFile);
+							break;
+						case ProjectFileFormat.QMake:
+							Generator = new QMakefileGenerator(ProjectFile);
+							break;
+						case ProjectFileFormat.KDevelop:
+							Generator = new KDevelopGenerator(ProjectFile);
+							break;
+						case ProjectFileFormat.CodeLite:
+							Generator = new CodeLiteGenerator(ProjectFile, Arguments);
+							break;
+						case ProjectFileFormat.VisualStudio:
+							Generator = new VCProjectFileGenerator(ProjectFile, VCProjectFileFormat.Default, Arguments);
+							break;
+						case ProjectFileFormat.VisualStudio2012:
+							Generator = new VCProjectFileGenerator(ProjectFile, VCProjectFileFormat.VisualStudio2012, Arguments);
+							break;
+						case ProjectFileFormat.VisualStudio2013:
+							Generator = new VCProjectFileGenerator(ProjectFile, VCProjectFileFormat.VisualStudio2013, Arguments);
+							break;
+						case ProjectFileFormat.VisualStudio2015:
+							Generator = new VCProjectFileGenerator(ProjectFile, VCProjectFileFormat.VisualStudio2015, Arguments);
+							break;
+						case ProjectFileFormat.VisualStudio2017:
+							Generator = new VCProjectFileGenerator(ProjectFile, VCProjectFileFormat.VisualStudio2017, Arguments);
+							break;
+						case ProjectFileFormat.XCode:
+							Generator = new XcodeProjectFileGenerator(ProjectFile);
+							break;
+						case ProjectFileFormat.Eddie:
+							Generator = new EddieProjectFileGenerator(ProjectFile);
+							break;
+						case ProjectFileFormat.VisualStudioCode:
+							Generator = new VSCodeProjectFileGenerator(ProjectFile);
+							break;
+						case ProjectFileFormat.CLion:
+							Generator = new CLionGenerator(ProjectFile);
+							break;
+						case ProjectFileFormat.VisualStudioMac:
+							Generator = new VCMacProjectFileGenerator(ProjectFile, Arguments);
+							break;
+						default:
+							throw new BuildException("Unhandled project file type '{0}", ProjectFileFormat);
+					}
+					Generators.Add(Generator);
+				}
+
+				// Check there are no superfluous command line arguments
+				Arguments.CheckAllArgumentsUsed();
+
+				// Build the list of game projects that we know about. When building from the editor (for hot-reload) or for projects from installed builds, we require the 
+				// project file to be passed in. Otherwise we scan for projects in directories named in UE4Games.uprojectdirs.
+				if (ProjectFile != null)
+				{
+					UProjectInfo.AddProject(ProjectFile);
+				}
+				else
+				{
+					UProjectInfo.FillProjectInfo();
+				}
+
+				// Now generate project files
+				ProjectFileGenerator.bGenerateProjectFiles = true;
+				foreach(ProjectFileGenerator Generator in Generators)
+				{
+					if (!Generator.GenerateProjectFiles(Arguments.GetRawArray()))
+					{
+						return (int)ECompilationResult.OtherCompilationError;
+					}
+				}
+			}
+			return (int)ECompilationResult.Succeeded;
+		}
+	}
+}
diff --git a/Engine/Source/Programs/UnrealBuildTool/ProjectFiles/CodeLite/CodeLiteGenerator.cs b/Engine/Source/Programs/UnrealBuildTool/ProjectFiles/CodeLite/CodeLiteGenerator.cs
--- a/Engine/Source/Programs/UnrealBuildTool/ProjectFiles/CodeLite/CodeLiteGenerator.cs
+++ b/Engine/Source/Programs/UnrealBuildTool/ProjectFiles/CodeLite/CodeLiteGenerator.cs
@@ -32,19 +32,16 @@ namespace UnrealBuildTool
 		public string CodeCompletionPreProcessorFileName = "CodeLitePreProcessor.txt";
 		CodeliteProjectFileFormat ProjectFileFormat = CodeliteProjectFileFormat.CodeLite10;
 
-		public CodeLiteGenerator(FileReference InOnlyGameProject, string[] Arguments)
+		public CodeLiteGenerator(FileReference InOnlyGameProject, CommandLineArguments CommandLine)
 			: base(InOnlyGameProject)
 		{
-			foreach(string Argument in Arguments)
+			if(CommandLine.HasOption("-cl10"))
 			{
-				if(Argument.Equals("-cl10", StringComparison.InvariantCultureIgnoreCase))
-				{
-					ProjectFileFormat = CodeliteProjectFileFormat.CodeLite10;
-				}
-				else if (Argument.Equals("-cl9", StringComparison.InvariantCultureIgnoreCase))
-				{
-					ProjectFileFormat = CodeliteProjectFileFormat.CodeLite9;
-				}
+				ProjectFileFormat = CodeliteProjectFileFormat.CodeLite10;
+			}
+			else if (CommandLine.HasOption("-cl9"))
+			{
+				ProjectFileFormat = CodeliteProjectFileFormat.CodeLite9;
 			}
 		}
 
diff --git a/Engine/Source/Programs/UnrealBuildTool/ProjectFiles/VisualStudio/VCProjectFileGenerator.cs b/Engine/Source/Programs/UnrealBuildTool/ProjectFiles/VisualStudio/VCProjectFileGenerator.cs
--- a/Engine/Source/Programs/UnrealBuildTool/ProjectFiles/VisualStudio/VCProjectFileGenerator.cs
+++ b/Engine/Source/Programs/UnrealBuildTool/ProjectFiles/VisualStudio/VCProjectFileGenerator.cs
@@ -91,8 +91,8 @@ namespace UnrealBuildTool
 		/// </summary>
 		/// <param name="InOnlyGameProject">The single project to generate project files for, or null</param>
 		/// <param name="InProjectFileFormat">Override the project file format to use</param>
-		/// <param name="InOverrideCompiler">Override the compiler version to use</param>
-		public VCProjectFileGenerator(FileReference InOnlyGameProject, VCProjectFileFormat InProjectFileFormat, WindowsCompiler InOverrideCompiler)
+		/// <param name="InArguments">Additional command line arguments</param>
+		public VCProjectFileGenerator(FileReference InOnlyGameProject, VCProjectFileFormat InProjectFileFormat, CommandLineArguments InArguments)
 			: base(InOnlyGameProject)
 		{
 			XmlConfig.ApplyTo(this);
@@ -102,11 +102,11 @@ namespace UnrealBuildTool
 				ProjectFileFormat = InProjectFileFormat;
 			}
 
-			if(InOverrideCompiler == WindowsCompiler.VisualStudio2015)
+			if(InArguments.HasOption("-2015"))
 			{
 				BuildToolOverride = "-2015";
 			}
-			else if(InOverrideCompiler == WindowsCompiler.VisualStudio2017)
+			else if(InArguments.HasOption("-2017"))
 			{
 				BuildToolOverride = "-2017";
 			}
diff --git a/Engine/Source/Programs/UnrealBuildTool/ProjectFiles/VisualStudioMac/VisualStudioMacProjectFileGenerator.cs b/Engine/Source/Programs/UnrealBuildTool/ProjectFiles/VisualStudioMac/VisualStudioMacProjectFileGenerator.cs
--- a/Engine/Source/Programs/UnrealBuildTool/ProjectFiles/VisualStudioMac/VisualStudioMacProjectFileGenerator.cs
+++ b/Engine/Source/Programs/UnrealBuildTool/ProjectFiles/VisualStudioMac/VisualStudioMacProjectFileGenerator.cs
@@ -21,9 +21,9 @@ namespace UnrealBuildTool
         /// Default constructor
         /// </summary>
         /// <param name="InOnlyGameProject">The single project to generate project files for, or null</param>
-        /// <param name="InOverrideCompiler">Override the compiler version to use</param>
-        public VCMacProjectFileGenerator(FileReference InOnlyGameProject, WindowsCompiler InOverrideCompiler)
-			: base(InOnlyGameProject, VCProjectFileFormat.Default, InOverrideCompiler)
+        /// <param name="InArguments">Additional command line arguments</param>
+        public VCMacProjectFileGenerator(FileReference InOnlyGameProject, CommandLineArguments InArguments)
+			: base(InOnlyGameProject, VCProjectFileFormat.Default, InArguments)
         {
             // no suo file, requires ole32
             bWriteSolutionOptionFile = false;
diff --git a/Engine/Source/Programs/UnrealBuildTool/System/BuildHostPlatform.cs b/Engine/Source/Programs/UnrealBuildTool/System/BuildHostPlatform.cs
--- a/Engine/Source/Programs/UnrealBuildTool/System/BuildHostPlatform.cs
+++ b/Engine/Source/Programs/UnrealBuildTool/System/BuildHostPlatform.cs
@@ -217,7 +217,7 @@ namespace UnrealBuildTool
 		/// Determines the default project file formats for this platform
 		/// </summary>
 		/// <returns>Sequence of project file formats</returns>
-		internal abstract void GetDefaultProjectFileFormats(List<ProjectFileFormat> Formats);
+		internal abstract IEnumerable<ProjectFileFormat> GetDefaultProjectFileFormats();
 	}
 
 	class WindowsBuildHostPlatform : BuildHostPlatform
@@ -312,9 +312,9 @@ namespace UnrealBuildTool
 			return Result;
 		}
 
-		internal override void GetDefaultProjectFileFormats(List<ProjectFileFormat> Formats)
+		internal override IEnumerable<ProjectFileFormat> GetDefaultProjectFileFormats()
 		{
-			Formats.Add(ProjectFileFormat.VisualStudio);
+			yield return ProjectFileFormat.VisualStudio;
 		}
 	}
 
@@ -432,9 +432,9 @@ namespace UnrealBuildTool
 			}
 		}
 
-		internal override void GetDefaultProjectFileFormats(List<ProjectFileFormat> Formats)
+		internal override IEnumerable<ProjectFileFormat> GetDefaultProjectFileFormats()
 		{
-			Formats.Add(ProjectFileFormat.XCode);
+			yield return ProjectFileFormat.XCode;
 		}
 	}
 
@@ -475,14 +475,14 @@ namespace UnrealBuildTool
 			return new List<string>().ToArray();
 		}
 
-		internal override void GetDefaultProjectFileFormats(List<ProjectFileFormat> Formats)
+		internal override IEnumerable<ProjectFileFormat> GetDefaultProjectFileFormats()
 		{
-			Formats.Add(ProjectFileFormat.Make);
-			Formats.Add(ProjectFileFormat.VisualStudioCode);
-			Formats.Add(ProjectFileFormat.KDevelop);
-			Formats.Add(ProjectFileFormat.QMake);
-			Formats.Add(ProjectFileFormat.CMake);
-			Formats.Add(ProjectFileFormat.CodeLite);
+			yield return ProjectFileFormat.Make;
+			yield return ProjectFileFormat.VisualStudioCode;
+			yield return ProjectFileFormat.KDevelop;
+			yield return ProjectFileFormat.QMake;
+			yield return ProjectFileFormat.CMake;
+			yield return ProjectFileFormat.CodeLite;
 		}
 	}
 }
diff --git a/Engine/Source/Programs/UnrealBuildTool/UnrealBuildTool.cs b/Engine/Source/Programs/UnrealBuildTool/UnrealBuildTool.cs
--- a/Engine/Source/Programs/UnrealBuildTool/UnrealBuildTool.cs
+++ b/Engine/Source/Programs/UnrealBuildTool/UnrealBuildTool.cs
@@ -258,6 +258,45 @@ namespace UnrealBuildTool
 			}
 		}
 
+		/// <summary>
+		/// Try to parse the project file from the command line
+		/// </summary>
+		/// <param name="Arguments">The command line arguments</param>
+		/// <param name="ProjectFile">The project file that was parsed</param>
+		/// <returns>True if the project file was parsed, false otherwise</returns>
+		public static bool TryParseProjectFileArgument(CommandLineArguments Arguments, out FileReference ProjectFile)
+		{
+			FileReference ExplicitProjectFile;
+			if(Arguments.TryGetValue("-Project=", out ExplicitProjectFile))
+			{
+				ProjectFile = ExplicitProjectFile;
+				return true;
+			}
+
+			for(int Idx = 0; Idx < Arguments.Count; Idx++)
+			{
+				if(Arguments[Idx][0] != '-' && Arguments[Idx].EndsWith(".uproject"))
+				{
+					Arguments.MarkAsUsed(Idx);
+					ProjectFile = new FileReference(Arguments[Idx]);
+					return true;
+				}
+			}
+
+			if(IsProjectInstalled())
+			{
+				FileReference InstalledProjectFile = FileReference.Combine(UnrealBuildTool.RootDirectory, "Engine", "Build", "InstalledProjectBuild.txt");
+				if (FileReference.Exists(InstalledProjectFile))
+				{
+					ProjectFile = FileReference.Combine(UnrealBuildTool.RootDirectory, File.ReadAllText(InstalledProjectFile.FullName).Trim());
+					return true;
+				}
+			}
+
+			ProjectFile = null;
+			return false;
+		}
+
 		private static bool ParseRocketCommandlineArg(string InArg, ref string OutGameName, ref FileReference ProjectFile)
 		{
 			string LowercaseArg = InArg.ToLowerInvariant();
@@ -460,9 +499,6 @@ namespace UnrealBuildTool
 						BuildConfiguration.bUseUBTMakefiles = false;
 					}
 
-					bool bGenerateProjectFiles = false;
-					WindowsCompiler OverrideWindowsCompiler = WindowsCompiler.Default;
-					List<ProjectFileFormat> ProjectFileFormats = new List<ProjectFileFormat>();
 					foreach (string Arg in Arguments)
 					{
 						string LowercaseArg = Arg.ToLowerInvariant();
@@ -471,87 +507,6 @@ namespace UnrealBuildTool
 							// Already handled at startup. Calling now just to properly set the game name
 							continue;
 						}
-						else if (LowercaseArg == "-projectfiles")
-						{
-							bGenerateProjectFiles = true;
-						}
-						else if (LowercaseArg.StartsWith("-projectfileformat="))
-						{
-							ProjectFileFormats.AddRange(ProjectFileGeneratorSettings.ParseFormatList(Arg.Substring("-projectfileformat=".Length)));
-							bGenerateProjectFiles = true;
-						}
-						else if (LowercaseArg == "-2012unsupported")
-						{
-							// May be for compiling; don't set bGenerateProjectFiles by default 
-							ProjectFileFormats.Add(ProjectFileFormat.VisualStudio2012);
-						}
-						else if (LowercaseArg == "-2013unsupported")
-						{
-							// May be for compiling; don't set bGenerateProjectFiles by default 
-							ProjectFileFormats.Add(ProjectFileFormat.VisualStudio2013);
-						}
-						else if (LowercaseArg == "-2015")
-						{
-							// May be for compiling; don't set bGenerateProjectFiles by default, but do override the compiler if it is.
-							ProjectFileFormats.Add(ProjectFileFormat.VisualStudio2015);
-							OverrideWindowsCompiler = WindowsCompiler.VisualStudio2015;
-						}
-						else if (LowercaseArg == "-2017")
-						{
-							// May be for compiling; don't set bGenerateProjectFiles by default, but do override the compiler if it is. 
-							ProjectFileFormats.Add(ProjectFileFormat.VisualStudio2017);
-							OverrideWindowsCompiler = WindowsCompiler.VisualStudio2017;
-						}
-						else if (LowercaseArg.StartsWith("-makefile"))
-						{
-							bGenerateProjectFiles = true;
-							ProjectFileFormats.Add(ProjectFileFormat.Make);
-						}
-						else if (LowercaseArg.StartsWith("-cmakefile"))
-						{
-							bGenerateProjectFiles = true;
-							ProjectFileFormats.Add(ProjectFileFormat.CMake);
-						}
-						else if (LowercaseArg.StartsWith("-qmakefile"))
-						{
-							bGenerateProjectFiles = true;
-							ProjectFileFormats.Add(ProjectFileFormat.QMake);
-						}
-						else if (LowercaseArg.StartsWith("-kdevelopfile"))
-						{
-							bGenerateProjectFiles = true;
-							ProjectFileFormats.Add(ProjectFileFormat.KDevelop);
-						}
-						else if (LowercaseArg == "-codelitefiles")
-						{
-							bGenerateProjectFiles = true;
-							ProjectFileFormats.Add(ProjectFileFormat.CodeLite);
-						}
-						else if (LowercaseArg == "-xcodeprojectfiles")
-						{
-							bGenerateProjectFiles = true;
-							ProjectFileFormats.Add(ProjectFileFormat.XCode);
-						}
-						else if (LowercaseArg == "-eddieprojectfiles")
-						{
-							bGenerateProjectFiles = true;
-							ProjectFileFormats.Add(ProjectFileFormat.Eddie);
-						}
-						else if (LowercaseArg == "-vscode")
-						{
-							bGenerateProjectFiles = true;
-							ProjectFileFormats.Add(ProjectFileFormat.VisualStudioCode);
-						}
-						else if (LowercaseArg == "-vsmac")
-						{
-							bGenerateProjectFiles = true;
-							ProjectFileFormats.Add(ProjectFileFormat.VisualStudioMac);
-						}
-						else if (LowercaseArg == "-clion")
-						{
-							bGenerateProjectFiles = true;
-							ProjectFileFormats.Add(ProjectFileFormat.CLion);
-						}
 						else if (LowercaseArg == "development" || LowercaseArg == "debug" || LowercaseArg == "shipping" || LowercaseArg == "test" || LowercaseArg == "debuggame")
 						{
 							//ConfigName = Arg;
@@ -569,10 +524,6 @@ namespace UnrealBuildTool
 						{
 							// Skip -define
 						}
-						else if (LowercaseArg == "-progress")
-						{
-							ProgressWriter.bWriteMarkup = true;
-						}
 						else if (CheckPlatform.ToString().ToLowerInvariant() == LowercaseArg)
 						{
 							// It's the platform set...
@@ -610,93 +561,7 @@ namespace UnrealBuildTool
 						JunkDeleter.DeleteJunk();
 					}
 
-					if (bGenerateProjectFiles)
-					{
-						// If there aren't any formats set, read the default project file format from the config file
-						if (ProjectFileFormats.Count == 0)
-						{
-							// Read from the XML config
-							if (!String.IsNullOrEmpty(ProjectFileGeneratorSettings.Format))
-							{
-								ProjectFileFormats.AddRange(ProjectFileGeneratorSettings.ParseFormatList(ProjectFileGeneratorSettings.Format));
-							}
-
-							// Read from the editor config
-							ProjectFileFormat PreferredSourceCodeAccessor;
-							if (ProjectFileGenerator.GetPreferredSourceCodeAccessor(ProjectFile, out PreferredSourceCodeAccessor))
-							{
-								ProjectFileFormats.Add(PreferredSourceCodeAccessor);
-							}
-
-							// If there's still nothing set, get the default project file format for this platform
-							if (ProjectFileFormats.Count == 0)
-							{
-								BuildHostPlatform.Current.GetDefaultProjectFileFormats(ProjectFileFormats);
-							}
-						}
-
-						// Create each project generator and run it
-						ProjectFileGenerator.bGenerateProjectFiles = true;
-						foreach (ProjectFileFormat ProjectFileFormat in ProjectFileFormats.Distinct())
-						{
-							ProjectFileGenerator Generator;
-							switch (ProjectFileFormat)
-							{
-								case ProjectFileFormat.Make:
-									Generator = new MakefileGenerator(ProjectFile);
-									break;
-								case ProjectFileFormat.CMake:
-									Generator = new CMakefileGenerator(ProjectFile);
-									break;
-								case ProjectFileFormat.QMake:
-									Generator = new QMakefileGenerator(ProjectFile);
-									break;
-								case ProjectFileFormat.KDevelop:
-									Generator = new KDevelopGenerator(ProjectFile);
-									break;
-								case ProjectFileFormat.CodeLite:
-									Generator = new CodeLiteGenerator(ProjectFile, Arguments);
-									break;
-								case ProjectFileFormat.VisualStudio:
-									Generator = new VCProjectFileGenerator(ProjectFile, VCProjectFileFormat.Default, OverrideWindowsCompiler);
-									break;
-								case ProjectFileFormat.VisualStudio2012:
-									Generator = new VCProjectFileGenerator(ProjectFile, VCProjectFileFormat.VisualStudio2012, OverrideWindowsCompiler);
-									break;
-								case ProjectFileFormat.VisualStudio2013:
-									Generator = new VCProjectFileGenerator(ProjectFile, VCProjectFileFormat.VisualStudio2013, OverrideWindowsCompiler);
-									break;
-								case ProjectFileFormat.VisualStudio2015:
-									Generator = new VCProjectFileGenerator(ProjectFile, VCProjectFileFormat.VisualStudio2015, OverrideWindowsCompiler);
-									break;
-								case ProjectFileFormat.VisualStudio2017:
-									Generator = new VCProjectFileGenerator(ProjectFile, VCProjectFileFormat.VisualStudio2017, OverrideWindowsCompiler);
-									break;
-								case ProjectFileFormat.XCode:
-									Generator = new XcodeProjectFileGenerator(ProjectFile);
-									break;
-								case ProjectFileFormat.Eddie:
-									Generator = new EddieProjectFileGenerator(ProjectFile);
-									break;
-								case ProjectFileFormat.VisualStudioCode:
-									Generator = new VSCodeProjectFileGenerator(ProjectFile);
-									break;
-								case ProjectFileFormat.CLion:
-									Generator = new CLionGenerator(ProjectFile);
-									break;
-								case ProjectFileFormat.VisualStudioMac:
-									Generator = new VCMacProjectFileGenerator(ProjectFile, OverrideWindowsCompiler);
-									break;
-								default:
-									throw new BuildException("Unhandled project file type '{0}", ProjectFileFormat);
-							}
-							if (!Generator.GenerateProjectFiles(Arguments))
-							{
-								Result = ECompilationResult.OtherCompilationError;
-							}
-						}
-					}
-					else if (BuildConfiguration.DeployTargetFile != null)
+					if (BuildConfiguration.DeployTargetFile != null)
 					{
 						UEBuildDeployTarget DeployTarget = new UEBuildDeployTarget(BuildConfiguration.DeployTargetFile);
 						Log.WriteLine(LogEventType.Console, "Deploying {0} {1} {2}...", DeployTarget.TargetName, DeployTarget.Platform, DeployTarget.Configuration);
@@ -768,10 +633,28 @@ namespace UnrealBuildTool
 			[CommandLine(Prefix = "-FromMsBuild")]
 			public bool bLogFromMsBuild = false;
 
+			/// <summary>
+			/// Whether to write progress markup in a format that can be parsed by other programs
+			/// </summary>
+			[CommandLine(Prefix = "-Progress")]
+			public bool bWriteProgressMarkup = false;
+
 			/// <summary>
 			/// The mode to execute
 			/// </summary>
 			[CommandLine]
+			[CommandLine("-ProjectFiles", Value="GenerateProjectFiles")]
+			[CommandLine("-ProjectFileFormat=", Value="GenerateProjectFiles")]
+			[CommandLine("-Makefile", Value="GenerateProjectFiles")]
+			[CommandLine("-CMakefile", Value="GenerateProjectFiles")]
+			[CommandLine("-QMakefile", Value="GenerateProjectFiles")]
+			[CommandLine("-KDevelopfile", Value="GenerateProjectFiles")]
+			[CommandLine("-CodeliteFiles", Value="GenerateProjectFiles")]
+			[CommandLine("-XCodeProjectFiles", Value="GenerateProjectFiles")]
+			[CommandLine("-EdditProjectFiles", Value="GenerateProjectFiles")]
+			[CommandLine("-VSCode", Value="GenerateProjectFiles")]
+			[CommandLine("-VSMac", Value="GenerateProjectFiles")]
+			[CommandLine("-CLion", Value="GenerateProjectFiles")]
 			public string Mode = "Default";
 
 			/// <summary>
@@ -803,6 +686,9 @@ namespace UnrealBuildTool
 				Log.OutputLevel = Options.LogOutputLevel;
 				Log.IncludeTimestamps = Options.bLogTimestamps;
 				Log.IncludeProgramNameWithSeverityPrefix = Options.bLogFromMsBuild;
+				
+				// Configure the progress writer
+				ProgressWriter.bWriteMarkup = Options.bWriteProgressMarkup;
 
 				// Add the log writer if requested. When building a target, we'll create the writer for the default log file later.
 				if(Options.LogFileName != null)
diff --git a/Engine/Source/Programs/UnrealBuildTool/UnrealBuildTool.csproj b/Engine/Source/Programs/UnrealBuildTool/UnrealBuildTool.csproj
--- a/Engine/Source/Programs/UnrealBuildTool/UnrealBuildTool.csproj
+++ b/Engine/Source/Programs/UnrealBuildTool/UnrealBuildTool.csproj
@@ -140,6 +140,7 @@
     </When>
   </Choose>
   <ItemGroup>
+    <Compile Include="Modes\GenerateProjectFilesMode.cs" />
     <Compile Include="Modes\SetupPlatformsMode.cs" />
     <Compile Include="Modes\ValidatePlatformsMode.cs" />
     <Compile Include="Modes\WriteDocumentationMode.cs" />
