UBT: Rename UEPlatformProjectGenerator to PlatformProjectGenerator.

#rnx

diff --git a/Engine/Source/Programs/UnrealBuildTool/Platform/Android/AndroidProjectGenerator.cs b/Engine/Source/Programs/UnrealBuildTool/Platform/Android/AndroidProjectGenerator.cs
--- a/Engine/Source/Programs/UnrealBuildTool/Platform/Android/AndroidProjectGenerator.cs
+++ b/Engine/Source/Programs/UnrealBuildTool/Platform/Android/AndroidProjectGenerator.cs
@@ -12,7 +12,7 @@ namespace UnrealBuildTool
 	/// <summary>
 	/// Base class for platform-specific project generators
 	/// </summary>
-	class AndroidProjectGenerator : UEPlatformProjectGenerator
+	class AndroidProjectGenerator : PlatformProjectGenerator
 	{
 		static bool CheckedForNsight = false;		// whether we have checked for a recent enough version of Nsight yet
 		static bool NsightInstalled = false;		// true if a recent enough version of Nsight is installed
@@ -147,7 +147,7 @@ namespace UnrealBuildTool
 		public override void RegisterPlatformProjectGenerator()
 		{
 			Log.TraceVerbose("        Registering for {0}", UnrealTargetPlatform.Android.ToString());
-			UEPlatformProjectGenerator.RegisterPlatformProjectGenerator(UnrealTargetPlatform.Android, this);
+			PlatformProjectGenerator.RegisterPlatformProjectGenerator(UnrealTargetPlatform.Android, this);
 		}
 
 		/// <summary>
diff --git a/Engine/Source/Programs/UnrealBuildTool/Platform/HTML5/HTML5ProjectGenerator.cs b/Engine/Source/Programs/UnrealBuildTool/Platform/HTML5/HTML5ProjectGenerator.cs
--- a/Engine/Source/Programs/UnrealBuildTool/Platform/HTML5/HTML5ProjectGenerator.cs
+++ b/Engine/Source/Programs/UnrealBuildTool/Platform/HTML5/HTML5ProjectGenerator.cs
@@ -12,7 +12,7 @@ namespace UnrealBuildTool
 	/// <summary>
 	/// Base class for platform-specific project generators
 	/// </summary>
-	class HTML5ProjectGenerator : UEPlatformProjectGenerator
+	class HTML5ProjectGenerator : PlatformProjectGenerator
 	{
 		/// <summary>
 		/// Register the platform with the UEPlatformProjectGenerator class
@@ -21,7 +21,7 @@ namespace UnrealBuildTool
 		{
 			// Register this project generator for HTML5
 			Log.TraceVerbose("        Registering for {0}", UnrealTargetPlatform.HTML5.ToString());
-			UEPlatformProjectGenerator.RegisterPlatformProjectGenerator(UnrealTargetPlatform.HTML5, this);
+			PlatformProjectGenerator.RegisterPlatformProjectGenerator(UnrealTargetPlatform.HTML5, this);
 		}
 
 		///
diff --git a/Engine/Source/Programs/UnrealBuildTool/Platform/IOS/IOSProjectGenerator.cs b/Engine/Source/Programs/UnrealBuildTool/Platform/IOS/IOSProjectGenerator.cs
--- a/Engine/Source/Programs/UnrealBuildTool/Platform/IOS/IOSProjectGenerator.cs
+++ b/Engine/Source/Programs/UnrealBuildTool/Platform/IOS/IOSProjectGenerator.cs
@@ -12,7 +12,7 @@ namespace UnrealBuildTool
 	/// <summary>
 	/// Base class for platform-specific project generators
 	/// </summary>
-	class IOSProjectGenerator : UEPlatformProjectGenerator
+	class IOSProjectGenerator : PlatformProjectGenerator
 	{
 		/// <summary>
 		/// Register the platform with the UEPlatformProjectGenerator class
@@ -21,7 +21,7 @@ namespace UnrealBuildTool
 		{
 			// Register this project generator for Mac
 			Log.TraceVerbose("        Registering for {0}", UnrealTargetPlatform.IOS.ToString());
-			UEPlatformProjectGenerator.RegisterPlatformProjectGenerator(UnrealTargetPlatform.IOS, this);
+			PlatformProjectGenerator.RegisterPlatformProjectGenerator(UnrealTargetPlatform.IOS, this);
 		}
 
 		///
diff --git a/Engine/Source/Programs/UnrealBuildTool/Platform/Linux/LinuxProjectGenerator.cs b/Engine/Source/Programs/UnrealBuildTool/Platform/Linux/LinuxProjectGenerator.cs
--- a/Engine/Source/Programs/UnrealBuildTool/Platform/Linux/LinuxProjectGenerator.cs
+++ b/Engine/Source/Programs/UnrealBuildTool/Platform/Linux/LinuxProjectGenerator.cs
@@ -12,7 +12,7 @@ namespace UnrealBuildTool
 	/// <summary>
 	/// Base class for platform-specific project generators
 	/// </summary>
-	class LinuxProjectGenerator : UEPlatformProjectGenerator
+	class LinuxProjectGenerator : PlatformProjectGenerator
 	{
 		/// <summary>
 		/// Register the platform with the UEPlatformProjectGenerator class
@@ -21,7 +21,7 @@ namespace UnrealBuildTool
 		{
 			// Register this project generator for both Linux
 			Log.TraceVerbose("        Registering for {0}", UnrealTargetPlatform.Linux.ToString());
-			UEPlatformProjectGenerator.RegisterPlatformProjectGenerator(UnrealTargetPlatform.Linux, this);
+			PlatformProjectGenerator.RegisterPlatformProjectGenerator(UnrealTargetPlatform.Linux, this);
 		}
 
 		///
diff --git a/Engine/Source/Programs/UnrealBuildTool/Platform/Lumin/LuminProjectGenerator.cs b/Engine/Source/Programs/UnrealBuildTool/Platform/Lumin/LuminProjectGenerator.cs
--- a/Engine/Source/Programs/UnrealBuildTool/Platform/Lumin/LuminProjectGenerator.cs
+++ b/Engine/Source/Programs/UnrealBuildTool/Platform/Lumin/LuminProjectGenerator.cs
@@ -12,7 +12,7 @@ namespace UnrealBuildTool
 	/// <summary>
 	/// Base class for platform-specific project generators
 	/// </summary>
-	class LuminProjectGenerator : UEPlatformProjectGenerator
+	class LuminProjectGenerator : PlatformProjectGenerator
 	{
 		static bool VSSupportChecked = false;       // Don't want to check multiple times
 		static bool VSDebuggingEnabled = false;
@@ -34,7 +34,7 @@ namespace UnrealBuildTool
 		public override void RegisterPlatformProjectGenerator()
 		{
 			Log.TraceVerbose("        Registering for {0}", UnrealTargetPlatform.Lumin.ToString());
-			UEPlatformProjectGenerator.RegisterPlatformProjectGenerator(UnrealTargetPlatform.Lumin, this);
+			PlatformProjectGenerator.RegisterPlatformProjectGenerator(UnrealTargetPlatform.Lumin, this);
 		}
 
 		/// <summary>
diff --git a/Engine/Source/Programs/UnrealBuildTool/Platform/Mac/MacProjectGenerator.cs b/Engine/Source/Programs/UnrealBuildTool/Platform/Mac/MacProjectGenerator.cs
--- a/Engine/Source/Programs/UnrealBuildTool/Platform/Mac/MacProjectGenerator.cs
+++ b/Engine/Source/Programs/UnrealBuildTool/Platform/Mac/MacProjectGenerator.cs
@@ -12,7 +12,7 @@ namespace UnrealBuildTool
 	/// <summary>
 	/// Base class for platform-specific project generators
 	/// </summary>
-	class MacProjectGenerator : UEPlatformProjectGenerator
+	class MacProjectGenerator : PlatformProjectGenerator
 	{
 		/// <summary>
 		/// Register the platform with the UEPlatformProjectGenerator class
@@ -21,7 +21,7 @@ namespace UnrealBuildTool
 		{
 			// Register this project generator for Mac
 			Log.TraceVerbose("        Registering for {0}", UnrealTargetPlatform.Mac.ToString());
-			UEPlatformProjectGenerator.RegisterPlatformProjectGenerator(UnrealTargetPlatform.Mac, this);
+			PlatformProjectGenerator.RegisterPlatformProjectGenerator(UnrealTargetPlatform.Mac, this);
 		}
 
 		///
diff --git a/Engine/Source/Programs/UnrealBuildTool/Platform/TVOS/TVOSProjectGenerator.cs b/Engine/Source/Programs/UnrealBuildTool/Platform/TVOS/TVOSProjectGenerator.cs
--- a/Engine/Source/Programs/UnrealBuildTool/Platform/TVOS/TVOSProjectGenerator.cs
+++ b/Engine/Source/Programs/UnrealBuildTool/Platform/TVOS/TVOSProjectGenerator.cs
@@ -18,7 +18,7 @@ namespace UnrealBuildTool
         {
             // Register this project generator for TVOS
             Log.TraceVerbose("        Registering for {0}", UnrealTargetPlatform.TVOS.ToString());
-            UEPlatformProjectGenerator.RegisterPlatformProjectGenerator(UnrealTargetPlatform.TVOS, this);
+            PlatformProjectGenerator.RegisterPlatformProjectGenerator(UnrealTargetPlatform.TVOS, this);
         }
     }
 }
diff --git a/Engine/Source/Programs/UnrealBuildTool/Platform/Windows/WindowsProjectGenerator.cs b/Engine/Source/Programs/UnrealBuildTool/Platform/Windows/WindowsProjectGenerator.cs
--- a/Engine/Source/Programs/UnrealBuildTool/Platform/Windows/WindowsProjectGenerator.cs
+++ b/Engine/Source/Programs/UnrealBuildTool/Platform/Windows/WindowsProjectGenerator.cs
@@ -12,7 +12,7 @@ namespace UnrealBuildTool
 	/// <summary>
 	/// Base class for platform-specific project generators
 	/// </summary>
-	class WindowsProjectGenerator : UEPlatformProjectGenerator
+	class WindowsProjectGenerator : PlatformProjectGenerator
 	{
 		/// <summary>
 		/// Register the platform with the UEPlatformProjectGenerator class
@@ -22,8 +22,8 @@ namespace UnrealBuildTool
 			// Register this project generator for Win32 and Win64
 			Log.TraceVerbose("        Registering for {0}", UnrealTargetPlatform.Win32.ToString());
 			Log.TraceVerbose("        Registering for {0}", UnrealTargetPlatform.Win64.ToString());
-			UEPlatformProjectGenerator.RegisterPlatformProjectGenerator(UnrealTargetPlatform.Win32, this);
-			UEPlatformProjectGenerator.RegisterPlatformProjectGenerator(UnrealTargetPlatform.Win64, this);
+			PlatformProjectGenerator.RegisterPlatformProjectGenerator(UnrealTargetPlatform.Win32, this);
+			PlatformProjectGenerator.RegisterPlatformProjectGenerator(UnrealTargetPlatform.Win64, this);
 		}
 
 		///
diff --git a/Engine/Source/Programs/UnrealBuildTool/ProjectFiles/PlatformProjectGenerator.cs b/Engine/Source/Programs/UnrealBuildTool/ProjectFiles/PlatformProjectGenerator.cs
new file mode 100644
--- /dev/null
+++ b/Engine/Source/Programs/UnrealBuildTool/ProjectFiles/PlatformProjectGenerator.cs
@@ -0,0 +1,357 @@
+// Copyright 1998-2018 Epic Games, Inc. All Rights Reserved.
+
+using System;
+using System.Collections.Generic;
+using System.Text;
+using System.Diagnostics;
+using System.IO;
+using Tools.DotNETCommon;
+
+namespace UnrealBuildTool
+{
+	/// <summary>
+	/// Base class for platform-specific project generators
+	/// </summary>
+	abstract class PlatformProjectGenerator
+	{
+		static Dictionary<UnrealTargetPlatform, PlatformProjectGenerator> ProjectGeneratorDictionary = new Dictionary<UnrealTargetPlatform, PlatformProjectGenerator>();
+
+		/// <summary>
+		/// Register the given platforms UEPlatformProjectGenerator instance
+		/// </summary>
+		/// <param name="InPlatform">  The UnrealTargetPlatform to register with</param>
+		/// <param name="InProjectGenerator">The UEPlatformProjectGenerator instance to use for the InPlatform</param>
+		public static void RegisterPlatformProjectGenerator(UnrealTargetPlatform InPlatform, PlatformProjectGenerator InProjectGenerator)
+		{
+			// Make sure the build platform is legal
+			UEBuildPlatform BuildPlatform = UEBuildPlatform.GetBuildPlatform(InPlatform, true);
+			if (BuildPlatform != null)
+			{
+				if (ProjectGeneratorDictionary.ContainsKey(InPlatform) == true)
+				{
+					Log.TraceInformation("RegisterPlatformProjectGenerator Warning: Registering project generator {0} for {1} when it is already set to {2}",
+						InProjectGenerator.ToString(), InPlatform.ToString(), ProjectGeneratorDictionary[InPlatform].ToString());
+					ProjectGeneratorDictionary[InPlatform] = InProjectGenerator;
+				}
+				else
+				{
+					ProjectGeneratorDictionary.Add(InPlatform, InProjectGenerator);
+				}
+			}
+			else
+			{
+				Log.TraceVerbose("Skipping project file generator registration for {0} due to no valid BuildPlatform.", InPlatform.ToString());
+			}
+		}
+
+		/// <summary>
+		/// Retrieve the UEPlatformProjectGenerator instance for the given TargetPlatform
+		/// </summary>
+		/// <param name="InPlatform">    The UnrealTargetPlatform being built</param>
+		/// <param name="bInAllowFailure">   If true, do not throw an exception and return null</param>
+		/// <returns>UEPlatformProjectGenerator The instance of the project generator</returns>
+		public static PlatformProjectGenerator GetPlatformProjectGenerator(UnrealTargetPlatform InPlatform, bool bInAllowFailure = false)
+		{
+			if (ProjectGeneratorDictionary.ContainsKey(InPlatform) == true)
+			{
+				return ProjectGeneratorDictionary[InPlatform];
+			}
+			if (bInAllowFailure == true)
+			{
+				return null;
+			}
+			throw new BuildException("GetPlatformProjectGenerator: No PlatformProjectGenerator found for {0}", InPlatform.ToString());
+		}
+
+		/// <summary>
+		/// Allow various platform project generators to generate stub projects if required
+		/// </summary>
+		/// <param name="InGenerator"></param>
+		/// <param name="InTargetName"></param>
+		/// <param name="InTargetFilepath"></param>
+		/// <param name="InTargetRules"></param>
+		/// <param name="InPlatforms"></param>
+		/// <param name="InConfigurations"></param>
+		/// <returns></returns>
+		public static bool GenerateGameProjectStubs(ProjectFileGenerator InGenerator, string InTargetName, string InTargetFilepath, TargetRules InTargetRules,
+			List<UnrealTargetPlatform> InPlatforms, List<UnrealTargetConfiguration> InConfigurations)
+		{
+			foreach (KeyValuePair<UnrealTargetPlatform, PlatformProjectGenerator> Entry in ProjectGeneratorDictionary)
+			{
+				PlatformProjectGenerator ProjGen = Entry.Value;
+				ProjGen.GenerateGameProjectStub(InGenerator, InTargetName, InTargetFilepath, InTargetRules, InPlatforms, InConfigurations);
+			}
+			return true;
+		}
+
+		/// <summary>
+		/// Allow various platform project generators to generate any special project properties if required
+		/// </summary>
+		/// <param name="InPlatform"></param>
+		/// <param name="Configuration"></param>
+		/// <param name="TargetType"></param>
+		/// <param name="VCProjectFileContent"></param>
+		/// <param name="RootDirectory"></param>
+		/// <param name="TargetFilePath"></param>
+		/// <returns></returns>
+		public static bool GenerateGamePlatformSpecificProperties(UnrealTargetPlatform InPlatform, UnrealTargetConfiguration Configuration, TargetType TargetType, StringBuilder VCProjectFileContent, DirectoryReference RootDirectory, FileReference TargetFilePath)
+		{
+			if (ProjectGeneratorDictionary.ContainsKey(InPlatform) == true)
+			{
+				ProjectGeneratorDictionary[InPlatform].GenerateGameProperties(Configuration, VCProjectFileContent, TargetType, RootDirectory, TargetFilePath); ;
+			}
+			return true;
+		}
+
+		public static bool PlatformRequiresVSUserFileGeneration(List<UnrealTargetPlatform> InPlatforms, List<UnrealTargetConfiguration> InConfigurations)
+		{
+			bool bRequiresVSUserFileGeneration = false;
+			foreach (KeyValuePair<UnrealTargetPlatform, PlatformProjectGenerator> Entry in ProjectGeneratorDictionary)
+			{
+				if (InPlatforms.Contains(Entry.Key))
+				{
+					PlatformProjectGenerator ProjGen = Entry.Value;
+					bRequiresVSUserFileGeneration |= ProjGen.RequiresVSUserFileGeneration();
+				}
+			}
+			return bRequiresVSUserFileGeneration;
+		}
+
+		/// <summary>
+		/// Register the platform with the UEPlatformProjectGenerator class
+		/// </summary>
+		public abstract void RegisterPlatformProjectGenerator();
+
+		public virtual void GenerateGameProjectStub(ProjectFileGenerator InGenerator, string InTargetName, string InTargetFilepath, TargetRules InTargetRules,
+			List<UnrealTargetPlatform> InPlatforms, List<UnrealTargetConfiguration> InConfigurations)
+		{
+			// Do nothing
+		}
+
+		public virtual void GenerateGameProperties(UnrealTargetConfiguration Configuration, StringBuilder VCProjectFileContent, TargetType TargetType, DirectoryReference RootDirectory, FileReference TargetFilePath)
+		{
+			// Do nothing
+		}
+
+		public virtual bool RequiresVSUserFileGeneration()
+		{
+			return false;
+		}
+
+
+		///
+		///	VisualStudio project generation functions
+		///	
+		/// <summary>
+		/// Whether this build platform has native support for VisualStudio
+		/// </summary>
+		/// <param name="InPlatform">  The UnrealTargetPlatform being built</param>
+		/// <param name="InConfiguration"> The UnrealTargetConfiguration being built</param>
+		/// <param name="ProjectFileFormat">The visual studio project file format being generated</param>
+		/// <returns>bool    true if native VisualStudio support (or custom VSI) is available</returns>
+		public virtual bool HasVisualStudioSupport(UnrealTargetPlatform InPlatform, UnrealTargetConfiguration InConfiguration, VCProjectFileFormat ProjectFileFormat)
+		{
+			// By default, we assume this is true
+			return true;
+		}
+
+		/// <summary>
+		/// Return the VisualStudio platform name for this build platform
+		/// </summary>
+		/// <param name="InPlatform">  The UnrealTargetPlatform being built</param>
+		/// <param name="InConfiguration"> The UnrealTargetConfiguration being built</param>
+		/// <returns>string    The name of the platform that VisualStudio recognizes</returns>
+		public virtual string GetVisualStudioPlatformName(UnrealTargetPlatform InPlatform, UnrealTargetConfiguration InConfiguration)
+		{
+			// By default, return the platform string
+			return InPlatform.ToString();
+		}
+
+		/// <summary>
+		/// Return project configuration settings that must be included before the default props file
+		/// </summary>
+		/// <param name="InPlatform">  The UnrealTargetPlatform being built</param>
+		/// <param name="InConfiguration"> The UnrealTargetConfiguration being built</param>
+		/// <returns>string    The custom configuration section for the project file; Empty string if it doesn't require one</returns>
+		public virtual string GetVisualStudioPreDefaultString(UnrealTargetPlatform InPlatform, UnrealTargetConfiguration InConfiguration)
+		{
+			return "";
+		}
+
+		/// <summary>
+		/// Return the platform toolset string to write into the project configuration
+		/// </summary>
+		/// <param name="InPlatform">  The UnrealTargetPlatform being built</param>
+		/// <param name="InConfiguration"> The UnrealTargetConfiguration being built</param>
+		/// <param name="InProjectFileFormat">The visual studio project file format being generated</param>
+		/// <returns>string    The custom configuration section for the project file; Empty string if it doesn't require one</returns>
+		public virtual string GetVisualStudioPlatformToolsetString(UnrealTargetPlatform InPlatform, UnrealTargetConfiguration InConfiguration, VCProjectFileFormat InProjectFileFormat)
+		{
+			return "";
+		}
+
+		/// <summary>
+		/// Return any custom property group lines
+		/// </summary>
+		/// <param name="InPlatform">  The UnrealTargetPlatform being built</param>
+		/// <param name="InProjectFileFormat">The visual studio project file format being generated</param>
+		/// <returns>string    The custom property import lines for the project file; Empty string if it doesn't require one</returns>
+		public virtual string GetAdditionalVisualStudioPropertyGroups(UnrealTargetPlatform InPlatform, VCProjectFileFormat InProjectFileFormat)
+		{
+			return "";
+		}
+
+		/// <summary>
+		/// Return any custom property group lines
+		/// </summary>
+		/// <param name="InPlatform">  The UnrealTargetPlatform being built</param>
+		/// <param name="InProjectFileFormat">The visual studio project file format being generated</param>
+		/// <returns>string    The platform configuration type.  Defaults to "Makefile" unless overridden</returns>
+		public virtual string GetVisualStudioPlatformConfigurationType(UnrealTargetPlatform InPlatform, VCProjectFileFormat InProjectFileFormat)
+		{
+			return "Makefile";
+		}
+
+		/// <summary>
+		/// Return any custom paths for VisualStudio this platform requires
+		/// This include ReferencePath, LibraryPath, LibraryWPath, IncludePath and ExecutablePath.
+		/// </summary>
+		/// <param name="InPlatform">The UnrealTargetPlatform being built</param>
+		/// <param name="InConfiguration">The configuration being built</param>
+		/// <param name="TargetType">The type of target (game or program)</param>
+		/// <param name="TargetRulesPath">Path to the .target.cs file</param>
+		/// <param name="ProjectFilePath"></param>
+		/// <param name="NMakeOutputPath"></param>
+		/// <param name="InProjectFileFormat">The visual studio project file format being generated</param>
+		/// <returns>The custom path lines for the project file; Empty string if it doesn't require one</returns>
+		public virtual string GetVisualStudioPathsEntries(UnrealTargetPlatform InPlatform, UnrealTargetConfiguration InConfiguration, TargetType TargetType, FileReference TargetRulesPath, FileReference ProjectFilePath, FileReference NMakeOutputPath, VCProjectFileFormat InProjectFileFormat)
+		{
+			// NOTE: We are intentionally overriding defaults for these paths with empty strings.  We never want Visual Studio's
+			//       defaults for these fields to be propagated, since they are version-sensitive paths that may not reflect
+			//       the environment that UBT is building in.  We'll set these environment variables ourselves!
+			// NOTE: We don't touch 'ExecutablePath' because that would result in Visual Studio clobbering the system "Path"
+			//       environment variable
+			string PathsLines =
+				"		<IncludePath />\n" +
+				"		<ReferencePath />\n" +
+				"		<LibraryPath />\n" +
+				"		<LibraryWPath />\n" +
+				"		<SourcePath />\n" +
+				"		<ExcludePath />\n";
+
+			return PathsLines;
+		}
+
+		/// <summary>
+		/// Return any custom property settings. These will be included in the ImportGroup section
+		/// </summary>
+		/// <param name="InPlatform">  The UnrealTargetPlatform being built</param>
+		/// <returns>string    The custom property import lines for the project file; Empty string if it doesn't require one</returns>
+		public virtual string GetVisualStudioImportGroupProperties(UnrealTargetPlatform InPlatform)
+		{
+			return "";
+		}
+
+		/// <summary>
+		/// Return any custom property settings. These will be included right after Global properties to make values available to all other imports.
+		/// </summary>
+		/// <param name="InPlatform">  The UnrealTargetPlatform being built</param>
+		/// <returns>string    The custom property import lines for the project file; Empty string if it doesn't require one</returns>
+		public virtual string GetVisualStudioGlobalProperties(UnrealTargetPlatform InPlatform)
+		{
+			return "";
+		}
+
+		/// <summary>
+		/// Return any custom target overrides. These will be included last in the project file so they have the opportunity to override any existing settings.
+		/// </summary>
+		/// <param name="InPlatform">  The UnrealTargetPlatform being built</param>
+		/// <param name="InProjectFileFormat">The visual studio project file format being generated</param>
+		/// <returns>string    The custom property import lines for the project file; Empty string if it doesn't require one</returns>
+		public virtual string GetVisualStudioTargetOverrides(UnrealTargetPlatform InPlatform, VCProjectFileFormat InProjectFileFormat)
+		{
+			return "";
+		}
+
+		/// <summary>
+		/// Return any custom layout directory sections
+		/// </summary>
+		/// <param name="InPlatform">The UnrealTargetPlatform being built</param>
+		/// <param name="InConfiguration">The configuration being built</param>
+		/// <param name="InConditionString"></param>
+		/// <param name="TargetType">The type of target (game or program)</param>
+		/// <param name="InProjectFileFormat">The visual studio project file format being generated</param>
+		/// <param name="NMakeOutputPath"></param>
+		/// <param name="ProjectFilePath"></param>
+		/// <param name="TargetRulesPath"></param>
+		/// <returns>string    The custom property import lines for the project file; Empty string if it doesn't require one</returns>
+		public virtual string GetVisualStudioLayoutDirSection(UnrealTargetPlatform InPlatform, UnrealTargetConfiguration InConfiguration, string InConditionString, TargetType TargetType, FileReference TargetRulesPath, FileReference ProjectFilePath, FileReference NMakeOutputPath, VCProjectFileFormat InProjectFileFormat)
+		{
+			return "";
+		}
+
+		/// <summary>
+		/// Get the output manifest section, if required
+		/// </summary>
+		/// <param name="InPlatform">The UnrealTargetPlatform being built</param>
+		/// <param name="TargetType">The type of the target being built</param>
+		/// <param name="TargetRulesPath">Path to the .target.cs file</param>
+		/// <param name="ProjectFilePath">Path to the project file</param>
+		/// <param name="InProjectFileFormat">The visual studio project file format being generated</param>
+		/// <returns>The output manifest section for the project file; Empty string if it doesn't require one</returns>
+		public virtual string GetVisualStudioOutputManifestSection(UnrealTargetPlatform InPlatform, TargetType TargetType, FileReference TargetRulesPath, FileReference ProjectFilePath, VCProjectFileFormat InProjectFileFormat)
+		{
+			return "";
+		}
+
+		/// <summary>
+		/// Get whether this platform deploys
+		/// </summary>
+		/// <returns>bool  true if the 'Deploy' option should be enabled</returns>
+		public virtual bool GetVisualStudioDeploymentEnabled(UnrealTargetPlatform InPlatform, UnrealTargetConfiguration InConfiguration)
+		{
+			return false;
+		}
+
+		/// <summary>
+		/// Get the text to insert into the user file for the given platform/configuration/target
+		/// </summary>
+		/// <param name="InPlatform">The platform being added</param>
+		/// <param name="InConfiguration">The configuration being added</param>
+		/// <param name="InConditionString">The condition string </param>
+		/// <param name="InTargetRules">The target rules </param>
+		/// <param name="TargetRulesPath">The target rules path</param>
+		/// <param name="ProjectFilePath">The project file path</param>
+		/// <returns>The string to append to the user file</returns>
+		public virtual string GetVisualStudioUserFileStrings(UnrealTargetPlatform InPlatform, UnrealTargetConfiguration InConfiguration,
+			string InConditionString, TargetRules InTargetRules, FileReference TargetRulesPath, FileReference ProjectFilePath)
+		{
+			return "";
+		}
+
+		/// <summary>
+		/// For Additional Project Property files that need to be written out.  This is currently used only on Android. 
+		/// </summary>
+		public virtual void WriteAdditionalPropFile()
+		{
+		}
+
+		/// <summary>
+		/// For additional Project files (ex. *PROJECTNAME*-AndroidRun.androidproj.user) that needs to be written out.  This is currently used only on Android. 
+		/// </summary>
+		/// <param name="ProjectFile">Project file this will be related to</param>
+		public virtual void WriteAdditionalProjUserFile(ProjectFile ProjectFile)
+		{
+		}
+
+		/// <summary>
+		/// For additional Project files (ex. *PROJECTNAME*-AndroidRun.androidproj) that needs to be written out.  This is currently used only on Android. 
+		/// </summary>
+		/// <param name="ProjectFile">Project file this will be related to</param>
+		/// <returns>Project file written out, Solution folder it should be put in</returns>
+		public virtual Tuple<ProjectFile, string> WriteAdditionalProjFile(ProjectFile ProjectFile)
+		{
+			return null;
+		}
+	}
+}
diff --git a/Engine/Source/Programs/UnrealBuildTool/ProjectFiles/ProjectFileGenerator.cs b/Engine/Source/Programs/UnrealBuildTool/ProjectFiles/ProjectFileGenerator.cs
--- a/Engine/Source/Programs/UnrealBuildTool/ProjectFiles/ProjectFileGenerator.cs
+++ b/Engine/Source/Programs/UnrealBuildTool/ProjectFiles/ProjectFileGenerator.cs
@@ -2146,7 +2146,7 @@ namespace UnrealBuildTool
 					if (TargetRulesObject.Type == TargetType.Game || TargetRulesObject.Type == TargetType.Client || TargetRulesObject.Type == TargetType.Server)
 					{
 						// Allow platforms to generate stub projects here...
-						UEPlatformProjectGenerator.GenerateGameProjectStubs(
+						PlatformProjectGenerator.GenerateGameProjectStubs(
 							InGenerator: this,
 							InTargetName: TargetName,
 							InTargetFilepath: TargetFilePath.FullName,
diff --git a/Engine/Source/Programs/UnrealBuildTool/ProjectFiles/VisualStudio/VCProject.cs b/Engine/Source/Programs/UnrealBuildTool/ProjectFiles/VisualStudio/VCProject.cs
--- a/Engine/Source/Programs/UnrealBuildTool/ProjectFiles/VisualStudio/VCProject.cs
+++ b/Engine/Source/Programs/UnrealBuildTool/ProjectFiles/VisualStudio/VCProject.cs
@@ -155,7 +155,7 @@ namespace UnrealBuildTool
 				}
 			}
 
-			UEPlatformProjectGenerator PlatformProjectGenerator = UEPlatformProjectGenerator.GetPlatformProjectGenerator(Platform, true);
+			PlatformProjectGenerator PlatformProjectGenerator = PlatformProjectGenerator.GetPlatformProjectGenerator(Platform, true);
 			if (PlatformProjectGenerator == null)
 			{
 				return false;
@@ -290,7 +290,7 @@ namespace UnrealBuildTool
 		/// <param name="ProjectConfigurationName">Name of configuration string to use for Visual Studio project</param>
 		public override void MakeProjectPlatformAndConfigurationNames(UnrealTargetPlatform Platform, UnrealTargetConfiguration Configuration, string TargetConfigurationName, out string ProjectPlatformName, out string ProjectConfigurationName)
 		{
-			UEPlatformProjectGenerator PlatformProjectGenerator = UEPlatformProjectGenerator.GetPlatformProjectGenerator(Platform, bInAllowFailure: true);
+			PlatformProjectGenerator PlatformProjectGenerator = PlatformProjectGenerator.GetPlatformProjectGenerator(Platform, bInAllowFailure: true);
 
 			// Check to see if this platform is supported directly by Visual Studio projects.
 			bool HasActualVSPlatform = (PlatformProjectGenerator != null) ? PlatformProjectGenerator.HasVisualStudioSupport(Platform, Configuration, ProjectFileFormat) : false;
@@ -449,7 +449,7 @@ namespace UnrealBuildTool
 			{
 				foreach (UnrealTargetPlatform Platform in ProjectPlatforms)
 				{
-					UEPlatformProjectGenerator ProjGenerator = UEPlatformProjectGenerator.GetPlatformProjectGenerator(Platform, true);
+					PlatformProjectGenerator ProjGenerator = PlatformProjectGenerator.GetPlatformProjectGenerator(Platform, true);
 					if (ProjGenerator != null)
 					{
 						//write out additional prop file
@@ -524,7 +524,7 @@ namespace UnrealBuildTool
 				ProjectFileGenerator.NewLine +
 				"<Project DefaultTargets=\"Build\" ToolsVersion=\"" + VCProjectFileGenerator.GetProjectFileToolVersionString(ProjectFileFormat) + "\" xmlns=\"http://schemas.microsoft.com/developer/msbuild/2003\">" + ProjectFileGenerator.NewLine);
 
-			bool bGenerateUserFileContent = UEPlatformProjectGenerator.PlatformRequiresVSUserFileGeneration(InPlatforms, InConfigurations);
+			bool bGenerateUserFileContent = PlatformProjectGenerator.PlatformRequiresVSUserFileGeneration(InPlatforms, InConfigurations);
 			if (bGenerateUserFileContent)
 			{
 				VCUserFileContent.Append(
@@ -588,7 +588,7 @@ namespace UnrealBuildTool
 			{
 				foreach (UnrealTargetPlatform Platform in ProjectPlatforms)
 				{
-					UEPlatformProjectGenerator ProjGenerator = UEPlatformProjectGenerator.GetPlatformProjectGenerator(Platform, true);
+					PlatformProjectGenerator ProjGenerator = PlatformProjectGenerator.GetPlatformProjectGenerator(Platform, true);
 					if (ProjGenerator != null && ProjGenerator.HasVisualStudioSupport(Platform, UnrealTargetConfiguration.Development, ProjectFileFormat))
 					{
 						AdditionalPropertyGroups.Append(ProjGenerator.GetAdditionalVisualStudioPropertyGroups(Platform, ProjectFileFormat));
@@ -616,7 +616,7 @@ namespace UnrealBuildTool
 			{
 				foreach (UnrealTargetPlatform Platform in ProjectPlatforms)
 				{
-					UEPlatformProjectGenerator ProjGenerator = UEPlatformProjectGenerator.GetPlatformProjectGenerator(Platform, true);
+					PlatformProjectGenerator ProjGenerator = PlatformProjectGenerator.GetPlatformProjectGenerator(Platform, true);
 					if (ProjGenerator != null && ProjGenerator.HasVisualStudioSupport(Platform, UnrealTargetConfiguration.Development, ProjectFileFormat))
 					{
 						VCProjectFileContent.Append(ProjGenerator.GetVisualStudioGlobalProperties(Platform));
@@ -834,7 +834,7 @@ namespace UnrealBuildTool
 			{
 				foreach (UnrealTargetPlatform Platform in ProjectPlatforms)
 				{
-					UEPlatformProjectGenerator ProjGenerator = UEPlatformProjectGenerator.GetPlatformProjectGenerator(Platform, true);
+					PlatformProjectGenerator ProjGenerator = PlatformProjectGenerator.GetPlatformProjectGenerator(Platform, true);
 					if (ProjGenerator != null && ProjGenerator.HasVisualStudioSupport(Platform, UnrealTargetConfiguration.Development, ProjectFileFormat))
 					{
 						AdditionalTargetSettings.Append(ProjGenerator.GetVisualStudioTargetOverrides(Platform, ProjectFileFormat));
@@ -847,7 +847,7 @@ namespace UnrealBuildTool
 			{
 				foreach (UnrealTargetPlatform Platform in ProjectPlatforms)
 				{
-					UEPlatformProjectGenerator ProjGenerator = UEPlatformProjectGenerator.GetPlatformProjectGenerator(Platform, true);
+					PlatformProjectGenerator ProjGenerator = PlatformProjectGenerator.GetPlatformProjectGenerator(Platform, true);
 					if (ProjGenerator != null && ProjGenerator.HasVisualStudioSupport(Platform, UnrealTargetConfiguration.Development, ProjectFileFormat))
 					{
 						// @todo projectfiles: Serious hacks here because we are trying to emit one-time platform-specific sections that need information
@@ -1039,7 +1039,7 @@ namespace UnrealBuildTool
 		// Anonymous function that writes pre-Default.props configuration data
 		private void WritePreDefaultPropsConfiguration(UnrealTargetPlatform TargetPlatform, UnrealTargetConfiguration TargetConfiguration, string ProjectPlatformName, string ProjectConfigurationName, StringBuilder VCProjectFileContent)
 		{
-			UEPlatformProjectGenerator ProjGenerator = UEPlatformProjectGenerator.GetPlatformProjectGenerator(TargetPlatform, true);
+			PlatformProjectGenerator ProjGenerator = PlatformProjectGenerator.GetPlatformProjectGenerator(TargetPlatform, true);
 			if (((ProjGenerator == null) && (TargetPlatform != UnrealTargetPlatform.Unknown)))
 			{
 				return;
@@ -1063,7 +1063,7 @@ namespace UnrealBuildTool
 		// Anonymous function that writes post-Default.props configuration data
 		private void WritePostDefaultPropsConfiguration(UnrealTargetPlatform TargetPlatform, UnrealTargetConfiguration TargetConfiguration, string ProjectPlatformName, string ProjectConfigurationName, StringBuilder VCProjectFileContent)
 		{
-			UEPlatformProjectGenerator ProjGenerator = UEPlatformProjectGenerator.GetPlatformProjectGenerator(TargetPlatform, true);
+			PlatformProjectGenerator ProjGenerator = PlatformProjectGenerator.GetPlatformProjectGenerator(TargetPlatform, true);
 			if (((ProjGenerator == null) && (TargetPlatform != UnrealTargetPlatform.Unknown)))
 			{
 				return;
@@ -1093,7 +1093,7 @@ namespace UnrealBuildTool
 			UnrealTargetPlatform Platform = Combination.Platform;
 			UnrealTargetConfiguration Configuration = Combination.Configuration;
 
-			UEPlatformProjectGenerator ProjGenerator = UEPlatformProjectGenerator.GetPlatformProjectGenerator(Platform, true);
+			PlatformProjectGenerator ProjGenerator = PlatformProjectGenerator.GetPlatformProjectGenerator(Platform, true);
 			if (((ProjGenerator == null) && (Platform != UnrealTargetPlatform.Unknown)))
 			{
 				return;
@@ -1274,7 +1274,7 @@ namespace UnrealBuildTool
 					if (TargetRulesObject.Type == TargetType.Game || TargetRulesObject.Type == TargetType.Client || TargetRulesObject.Type == TargetType.Server)
 					{
 						// Allow platforms to add any special properties they require... like aumid override for Xbox One
-						UEPlatformProjectGenerator.GenerateGamePlatformSpecificProperties(Platform, Configuration, TargetRulesObject.Type, VCProjectFileContent, RootDirectory, TargetFilePath);
+						PlatformProjectGenerator.GenerateGamePlatformSpecificProperties(Platform, Configuration, TargetRulesObject.Type, VCProjectFileContent, RootDirectory, TargetFilePath);
 					}
 
 					VCProjectFileContent.Append("	</PropertyGroup>" + ProjectFileGenerator.NewLine);
diff --git a/Engine/Source/Programs/UnrealBuildTool/ProjectFiles/VisualStudio/VCProjectFileGenerator.cs b/Engine/Source/Programs/UnrealBuildTool/ProjectFiles/VisualStudio/VCProjectFileGenerator.cs
--- a/Engine/Source/Programs/UnrealBuildTool/ProjectFiles/VisualStudio/VCProjectFileGenerator.cs
+++ b/Engine/Source/Programs/UnrealBuildTool/ProjectFiles/VisualStudio/VCProjectFileGenerator.cs
@@ -867,7 +867,7 @@ namespace UnrealBuildTool
 												VCSolutionFileContent.Append(
 														"		" + CurProjectGUID + "." + SolutionConfigCombination.VCSolutionConfigAndPlatformName + ".Build.0 = " + ProjectConfigAndPlatformPair + ProjectFileGenerator.NewLine);
 
-												UEPlatformProjectGenerator ProjGen = UEPlatformProjectGenerator.GetPlatformProjectGenerator(SolutionConfigCombination.Platform, true);
+												PlatformProjectGenerator ProjGen = PlatformProjectGenerator.GetPlatformProjectGenerator(SolutionConfigCombination.Platform, true);
 												if (MatchingProjectTarget.ProjectDeploys ||
 													((ProjGen != null) && (ProjGen.GetVisualStudioDeploymentEnabled(SolutionPlatform, SolutionConfiguration) == true)))
 												{
diff --git a/Engine/Source/Programs/UnrealBuildTool/System/UEPlatformProjectGenerator.cs b/Engine/Source/Programs/UnrealBuildTool/System/UEPlatformProjectGenerator.cs
deleted file mode 100644
diff --git a/Engine/Source/Programs/UnrealBuildTool/UnrealBuildTool.cs b/Engine/Source/Programs/UnrealBuildTool/UnrealBuildTool.cs
--- a/Engine/Source/Programs/UnrealBuildTool/UnrealBuildTool.cs
+++ b/Engine/Source/Programs/UnrealBuildTool/UnrealBuildTool.cs
@@ -240,7 +240,7 @@ namespace UnrealBuildTool
 				{
 					if (CheckType.IsClass && !CheckType.IsAbstract)
 					{
-						if (CheckType.IsSubclassOf(typeof(UEPlatformProjectGenerator)))
+						if (CheckType.IsSubclassOf(typeof(PlatformProjectGenerator)))
 						{
 							ProjectGeneratorList.Add(CheckType);
 						}
@@ -252,7 +252,7 @@ namespace UnrealBuildTool
 				foreach (Type ProjType in ProjectGeneratorList)
 				{
 					Log.TraceVerbose("    Registering project generator: {0}", ProjType.ToString());
-					UEPlatformProjectGenerator TempInst = (UEPlatformProjectGenerator)(UBTAssembly.CreateInstance(ProjType.FullName, true));
+					PlatformProjectGenerator TempInst = (PlatformProjectGenerator)(UBTAssembly.CreateInstance(ProjType.FullName, true));
 					TempInst.RegisterPlatformProjectGenerator();
 				}
 			}
diff --git a/Engine/Source/Programs/UnrealBuildTool/UnrealBuildTool.csproj b/Engine/Source/Programs/UnrealBuildTool/UnrealBuildTool.csproj
--- a/Engine/Source/Programs/UnrealBuildTool/UnrealBuildTool.csproj
+++ b/Engine/Source/Programs/UnrealBuildTool/UnrealBuildTool.csproj
@@ -167,6 +167,7 @@
     <Compile Include="ProjectFiles\Eddie\EddieProject.cs" />
     <Compile Include="ProjectFiles\KDevelop\KDevelopGenerator.cs" />
     <Compile Include="ProjectFiles\Make\MakefileGenerator.cs" />
+    <Compile Include="ProjectFiles\PlatformProjectGenerator.cs" />
     <Compile Include="ProjectFiles\Project.cs" />
     <Compile Include="ProjectFiles\ProjectFileGenerator.cs" />
     <Compile Include="ProjectFiles\QMake\QMakefileGenerator.cs" />
@@ -343,7 +344,6 @@
     <Compile Include="System\Plugins.cs" />
     <Compile Include="System\RulesCompiler.cs" />
     <Compile Include="System\SourceFileSearch.cs" />
-    <Compile Include="System\UEPlatformProjectGenerator.cs" />
     <Compile Include="System\Unity.cs" />
     <Compile Include="System\UProjectInfo.cs" />
     <Compile Include="System\Utils.cs" />
