UBT: Remove the global registration of project file generators. This is now done only when needed (when actually generating project files).

#rnx

diff --git a/Engine/Source/Programs/UnrealBuildTool/Modes/GenerateProjectFilesMode.cs b/Engine/Source/Programs/UnrealBuildTool/Modes/GenerateProjectFilesMode.cs
--- a/Engine/Source/Programs/UnrealBuildTool/Modes/GenerateProjectFilesMode.cs
+++ b/Engine/Source/Programs/UnrealBuildTool/Modes/GenerateProjectFilesMode.cs
@@ -2,6 +2,7 @@ using System;
 using System.Collections.Generic;
 using System.IO;
 using System.Linq;
+using System.Reflection;
 using System.Text;
 using System.Threading.Tasks;
 using Tools.DotNETCommon;
@@ -89,6 +90,21 @@ namespace UnrealBuildTool
 					}
 				}
 
+				// Register all the platform project generators
+				PlatformProjectGeneratorCollection PlatformProjectGenerators = new PlatformProjectGeneratorCollection();
+				foreach (Type CheckType in Assembly.GetExecutingAssembly().GetTypes())
+				{
+					if (CheckType.IsClass && !CheckType.IsAbstract && CheckType.IsSubclassOf(typeof(PlatformProjectGenerator)))
+					{
+						PlatformProjectGenerator Generator = (PlatformProjectGenerator)Activator.CreateInstance(CheckType);
+						foreach(UnrealTargetPlatform Platform in Generator.GetPlatforms())
+						{
+							Log.TraceVerbose("Registering project generator {0} for {1}", CheckType, Platform);
+							PlatformProjectGenerators.RegisterPlatformProjectGenerator(Platform, Generator);
+						}
+					}
+				}
+
 				// Create each project generator and run it
 				List<ProjectFileGenerator> Generators = new List<ProjectFileGenerator>();
 				foreach (ProjectFileFormat ProjectFileFormat in ProjectFileFormats.Distinct())
@@ -165,7 +181,7 @@ namespace UnrealBuildTool
 				ProjectFileGenerator.bGenerateProjectFiles = true;
 				foreach(ProjectFileGenerator Generator in Generators)
 				{
-					if (!Generator.GenerateProjectFiles(Arguments.GetRawArray()))
+					if (!Generator.GenerateProjectFiles(PlatformProjectGenerators, Arguments.GetRawArray()))
 					{
 						return (int)ECompilationResult.OtherCompilationError;
 					}
diff --git a/Engine/Source/Programs/UnrealBuildTool/Platform/Android/AndroidProjectGenerator.cs b/Engine/Source/Programs/UnrealBuildTool/Platform/Android/AndroidProjectGenerator.cs
--- a/Engine/Source/Programs/UnrealBuildTool/Platform/Android/AndroidProjectGenerator.cs
+++ b/Engine/Source/Programs/UnrealBuildTool/Platform/Android/AndroidProjectGenerator.cs
@@ -142,12 +142,11 @@ namespace UnrealBuildTool
 		}
 
 		/// <summary>
-		/// Register the platform with the UEPlatformProjectGenerator class
+		/// Enumerate all the platforms that this generator supports
 		/// </summary>
-		public override void RegisterPlatformProjectGenerator()
+		public override IEnumerable<UnrealTargetPlatform> GetPlatforms()
 		{
-			Log.TraceVerbose("        Registering for {0}", UnrealTargetPlatform.Android.ToString());
-			PlatformProjectGenerator.RegisterPlatformProjectGenerator(UnrealTargetPlatform.Android, this);
+			yield return UnrealTargetPlatform.Android;
 		}
 
 		/// <summary>
@@ -710,9 +709,10 @@ namespace UnrealBuildTool
 		/// <param name="Platform">Actual platform</param>
 		/// <param name="Configuration">Actual configuration</param>
 		/// <param name="TargetConfigurationName">The configuration name from the target rules, or null if we don't have one</param>
+		/// <param name="PlatformProjectGenerators">The registered platform project generators</param>
 		/// <param name="ProjectPlatformName">Name of platform string to use for Visual Studio project</param>
 		/// <param name="ProjectConfigurationName">Name of configuration string to use for Visual Studio project</param>
-		public override void MakeProjectPlatformAndConfigurationNames(UnrealTargetPlatform Platform, UnrealTargetConfiguration Configuration, string TargetConfigurationName, out string ProjectPlatformName, out string ProjectConfigurationName)
+		public override void MakeProjectPlatformAndConfigurationNames(UnrealTargetPlatform Platform, UnrealTargetConfiguration Configuration, string TargetConfigurationName, PlatformProjectGeneratorCollection PlatformProjectGenerators, out string ProjectPlatformName, out string ProjectConfigurationName)
 		{
 			ProjectConfigurationName = UnrealTargetConfiguration.Debug.ToString();
 			ProjectPlatformName = "ARM";
diff --git a/Engine/Source/Programs/UnrealBuildTool/Platform/HTML5/HTML5ProjectGenerator.cs b/Engine/Source/Programs/UnrealBuildTool/Platform/HTML5/HTML5ProjectGenerator.cs
--- a/Engine/Source/Programs/UnrealBuildTool/Platform/HTML5/HTML5ProjectGenerator.cs
+++ b/Engine/Source/Programs/UnrealBuildTool/Platform/HTML5/HTML5ProjectGenerator.cs
@@ -15,13 +15,11 @@ namespace UnrealBuildTool
 	class HTML5ProjectGenerator : PlatformProjectGenerator
 	{
 		/// <summary>
-		/// Register the platform with the UEPlatformProjectGenerator class
+		/// Enumerate all the platforms that this generator supports
 		/// </summary>
-		public override void RegisterPlatformProjectGenerator()
+		public override IEnumerable<UnrealTargetPlatform> GetPlatforms()
 		{
-			// Register this project generator for HTML5
-			Log.TraceVerbose("        Registering for {0}", UnrealTargetPlatform.HTML5.ToString());
-			PlatformProjectGenerator.RegisterPlatformProjectGenerator(UnrealTargetPlatform.HTML5, this);
+			yield return UnrealTargetPlatform.HTML5;
 		}
 
 		///
diff --git a/Engine/Source/Programs/UnrealBuildTool/Platform/IOS/IOSProjectGenerator.cs b/Engine/Source/Programs/UnrealBuildTool/Platform/IOS/IOSProjectGenerator.cs
--- a/Engine/Source/Programs/UnrealBuildTool/Platform/IOS/IOSProjectGenerator.cs
+++ b/Engine/Source/Programs/UnrealBuildTool/Platform/IOS/IOSProjectGenerator.cs
@@ -15,13 +15,11 @@ namespace UnrealBuildTool
 	class IOSProjectGenerator : PlatformProjectGenerator
 	{
 		/// <summary>
-		/// Register the platform with the UEPlatformProjectGenerator class
+		/// Enumerate all the platforms that this generator supports
 		/// </summary>
-		public override void RegisterPlatformProjectGenerator()
+		public override IEnumerable<UnrealTargetPlatform> GetPlatforms()
 		{
-			// Register this project generator for Mac
-			Log.TraceVerbose("        Registering for {0}", UnrealTargetPlatform.IOS.ToString());
-			PlatformProjectGenerator.RegisterPlatformProjectGenerator(UnrealTargetPlatform.IOS, this);
+			yield return UnrealTargetPlatform.IOS;
 		}
 
 		///
diff --git a/Engine/Source/Programs/UnrealBuildTool/Platform/IOS/IOSToolChain.cs b/Engine/Source/Programs/UnrealBuildTool/Platform/IOS/IOSToolChain.cs
--- a/Engine/Source/Programs/UnrealBuildTool/Platform/IOS/IOSToolChain.cs
+++ b/Engine/Source/Programs/UnrealBuildTool/Platform/IOS/IOSToolChain.cs
@@ -1453,6 +1453,24 @@ namespace UnrealBuildTool
             }
         }
 
+		internal static bool GenerateProjectFiles(FileReference ProjectFile, string[] Arguments)
+		{
+			ProjectFileGenerator.bGenerateProjectFiles = true;
+			try
+			{
+				PlatformProjectGeneratorCollection PlatformProjectGenerators = new PlatformProjectGeneratorCollection();
+				PlatformProjectGenerators.RegisterPlatformProjectGenerator(UnrealTargetPlatform.IOS, new IOSProjectGenerator());
+				PlatformProjectGenerators.RegisterPlatformProjectGenerator(UnrealTargetPlatform.TVOS, new TVOSProjectGenerator());
+
+				XcodeProjectFileGenerator Generator = new XcodeProjectFileGenerator(ProjectFile);
+				return Generator.GenerateProjectFiles(PlatformProjectGenerators, Arguments);
+			}
+			finally
+			{
+				ProjectFileGenerator.bGenerateProjectFiles = false;
+			}
+		}
+
         public static void PostBuildSync(UEBuildTarget Target)
 		{
 			if (Target.Rules == null)
@@ -1515,12 +1533,12 @@ namespace UnrealBuildTool
 				DirectoryReference XcodeWorkspaceDir;
 				if (AppName == "UE4Game" || AppName == "UE4Client" || Target.ProjectFile == null || Target.ProjectFile.IsUnderDirectory(UnrealBuildTool.EngineDirectory))
 				{
-					UnrealBuildTool.GenerateProjectFiles(new XcodeProjectFileGenerator(Target.ProjectFile), new string[] { "-platforms=" + (Target.Platform == UnrealTargetPlatform.IOS ? "IOS" : "TVOS"), "-NoIntellIsense", (Target.Platform == UnrealTargetPlatform.IOS ? "-iosdeployonly" : "-tvosdeployonly"), "-ignorejunk" });
+					GenerateProjectFiles(Target.ProjectFile, new string[] { "-platforms=" + (Target.Platform == UnrealTargetPlatform.IOS ? "IOS" : "TVOS"), "-NoIntellIsense", (Target.Platform == UnrealTargetPlatform.IOS ? "-iosdeployonly" : "-tvosdeployonly"), "-ignorejunk" });
 					XcodeWorkspaceDir = DirectoryReference.Combine(UnrealBuildTool.RootDirectory, String.Format("UE4_{0}.xcworkspace", (Target.Platform == UnrealTargetPlatform.IOS ? "IOS" : "TVOS")));
 				}
 				else
 				{
-					UnrealBuildTool.GenerateProjectFiles(new XcodeProjectFileGenerator(Target.ProjectFile), new string[] { "-platforms=" + (Target.Platform == UnrealTargetPlatform.IOS ? "IOS" : "TVOS"), "-NoIntellIsense", (Target.Platform == UnrealTargetPlatform.IOS ? "-iosdeployonly" : "-tvosdeployonly"), "-ignorejunk", String.Format("-project={0}", Target.ProjectFile), "-game" });
+					GenerateProjectFiles(Target.ProjectFile, new string[] { "-platforms=" + (Target.Platform == UnrealTargetPlatform.IOS ? "IOS" : "TVOS"), "-NoIntellIsense", (Target.Platform == UnrealTargetPlatform.IOS ? "-iosdeployonly" : "-tvosdeployonly"), "-ignorejunk", String.Format("-project={0}", Target.ProjectFile), "-game" });
 					XcodeWorkspaceDir = DirectoryReference.Combine(Target.ProjectDirectory, String.Format("{0}_{1}.xcworkspace", AppName, (Target.Platform == UnrealTargetPlatform.IOS ? "IOS" : "TVOS")));
 				}
 
diff --git a/Engine/Source/Programs/UnrealBuildTool/Platform/Linux/LinuxProjectGenerator.cs b/Engine/Source/Programs/UnrealBuildTool/Platform/Linux/LinuxProjectGenerator.cs
--- a/Engine/Source/Programs/UnrealBuildTool/Platform/Linux/LinuxProjectGenerator.cs
+++ b/Engine/Source/Programs/UnrealBuildTool/Platform/Linux/LinuxProjectGenerator.cs
@@ -15,13 +15,11 @@ namespace UnrealBuildTool
 	class LinuxProjectGenerator : PlatformProjectGenerator
 	{
 		/// <summary>
-		/// Register the platform with the UEPlatformProjectGenerator class
+		/// Enumerate all the platforms that this generator supports
 		/// </summary>
-		public override void RegisterPlatformProjectGenerator()
+		public override IEnumerable<UnrealTargetPlatform> GetPlatforms()
 		{
-			// Register this project generator for both Linux
-			Log.TraceVerbose("        Registering for {0}", UnrealTargetPlatform.Linux.ToString());
-			PlatformProjectGenerator.RegisterPlatformProjectGenerator(UnrealTargetPlatform.Linux, this);
+			yield return UnrealTargetPlatform.Linux;
 		}
 
 		///
diff --git a/Engine/Source/Programs/UnrealBuildTool/Platform/Lumin/LuminProjectGenerator.cs b/Engine/Source/Programs/UnrealBuildTool/Platform/Lumin/LuminProjectGenerator.cs
--- a/Engine/Source/Programs/UnrealBuildTool/Platform/Lumin/LuminProjectGenerator.cs
+++ b/Engine/Source/Programs/UnrealBuildTool/Platform/Lumin/LuminProjectGenerator.cs
@@ -29,12 +29,11 @@ namespace UnrealBuildTool
 		}
 		
 		/// <summary>
-		/// Register the platform with the UEPlatformProjectGenerator class
+		/// Enumerate all the platforms that this generator supports
 		/// </summary>
-		public override void RegisterPlatformProjectGenerator()
+		public override IEnumerable<UnrealTargetPlatform> GetPlatforms()
 		{
-			Log.TraceVerbose("        Registering for {0}", UnrealTargetPlatform.Lumin.ToString());
-			PlatformProjectGenerator.RegisterPlatformProjectGenerator(UnrealTargetPlatform.Lumin, this);
+			yield return UnrealTargetPlatform.Lumin;
 		}
 
 		/// <summary>
diff --git a/Engine/Source/Programs/UnrealBuildTool/Platform/Mac/MacProjectGenerator.cs b/Engine/Source/Programs/UnrealBuildTool/Platform/Mac/MacProjectGenerator.cs
--- a/Engine/Source/Programs/UnrealBuildTool/Platform/Mac/MacProjectGenerator.cs
+++ b/Engine/Source/Programs/UnrealBuildTool/Platform/Mac/MacProjectGenerator.cs
@@ -17,11 +17,9 @@ namespace UnrealBuildTool
 		/// <summary>
 		/// Register the platform with the UEPlatformProjectGenerator class
 		/// </summary>
-		public override void RegisterPlatformProjectGenerator()
+		public override IEnumerable<UnrealTargetPlatform> GetPlatforms()
 		{
-			// Register this project generator for Mac
-			Log.TraceVerbose("        Registering for {0}", UnrealTargetPlatform.Mac.ToString());
-			PlatformProjectGenerator.RegisterPlatformProjectGenerator(UnrealTargetPlatform.Mac, this);
+			yield return UnrealTargetPlatform.Mac;
 		}
 
 		///
diff --git a/Engine/Source/Programs/UnrealBuildTool/Platform/TVOS/TVOSProjectGenerator.cs b/Engine/Source/Programs/UnrealBuildTool/Platform/TVOS/TVOSProjectGenerator.cs
--- a/Engine/Source/Programs/UnrealBuildTool/Platform/TVOS/TVOSProjectGenerator.cs
+++ b/Engine/Source/Programs/UnrealBuildTool/Platform/TVOS/TVOSProjectGenerator.cs
@@ -11,14 +11,12 @@ namespace UnrealBuildTool
 {
 	class TVOSProjectGenerator : IOSProjectGenerator
     {
-        /**
-         *	Register the platform with the UEPlatformProjectGenerator class
-         */
-        public override void RegisterPlatformProjectGenerator()
-        {
-            // Register this project generator for TVOS
-            Log.TraceVerbose("        Registering for {0}", UnrealTargetPlatform.TVOS.ToString());
-            PlatformProjectGenerator.RegisterPlatformProjectGenerator(UnrealTargetPlatform.TVOS, this);
-        }
+		/// <summary>
+		/// Enumerate all the platforms that this generator supports
+		/// </summary>
+		public override IEnumerable<UnrealTargetPlatform> GetPlatforms()
+		{
+			yield return UnrealTargetPlatform.TVOS;
+		}
     }
 }
diff --git a/Engine/Source/Programs/UnrealBuildTool/Platform/Windows/WindowsProjectGenerator.cs b/Engine/Source/Programs/UnrealBuildTool/Platform/Windows/WindowsProjectGenerator.cs
--- a/Engine/Source/Programs/UnrealBuildTool/Platform/Windows/WindowsProjectGenerator.cs
+++ b/Engine/Source/Programs/UnrealBuildTool/Platform/Windows/WindowsProjectGenerator.cs
@@ -15,15 +15,12 @@ namespace UnrealBuildTool
 	class WindowsProjectGenerator : PlatformProjectGenerator
 	{
 		/// <summary>
-		/// Register the platform with the UEPlatformProjectGenerator class
+		/// Enumerate all the platforms that this generator supports
 		/// </summary>
-		public override void RegisterPlatformProjectGenerator()
+		public override IEnumerable<UnrealTargetPlatform> GetPlatforms()
 		{
-			// Register this project generator for Win32 and Win64
-			Log.TraceVerbose("        Registering for {0}", UnrealTargetPlatform.Win32.ToString());
-			Log.TraceVerbose("        Registering for {0}", UnrealTargetPlatform.Win64.ToString());
-			PlatformProjectGenerator.RegisterPlatformProjectGenerator(UnrealTargetPlatform.Win32, this);
-			PlatformProjectGenerator.RegisterPlatformProjectGenerator(UnrealTargetPlatform.Win64, this);
+			yield return UnrealTargetPlatform.Win32;
+			yield return UnrealTargetPlatform.Win64;
 		}
 
 		///
diff --git a/Engine/Source/Programs/UnrealBuildTool/ProjectFiles/CMake/CMakefileGenerator.cs b/Engine/Source/Programs/UnrealBuildTool/ProjectFiles/CMake/CMakefileGenerator.cs
--- a/Engine/Source/Programs/UnrealBuildTool/ProjectFiles/CMake/CMakefileGenerator.cs
+++ b/Engine/Source/Programs/UnrealBuildTool/ProjectFiles/CMake/CMakefileGenerator.cs
@@ -229,8 +229,9 @@ namespace UnrealBuildTool
 		/// Writes the master project file (e.g. Visual Studio Solution file)
 		/// </summary>
 		/// <param name="UBTProject">The UnrealBuildTool project</param>
+		/// <param name="PlatformProjectGenerators">The registered platform project generators</param>
 		/// <returns>True if successful</returns>
-		protected override bool WriteMasterProjectFile(ProjectFile UBTProject)
+		protected override bool WriteMasterProjectFile(ProjectFile UBTProject, PlatformProjectGeneratorCollection PlatformProjectGenerators)
 		{
 			return true;
 		}
@@ -780,7 +781,7 @@ namespace UnrealBuildTool
 
 		#region ProjectFileGenerator implementation
 
-		protected override bool WriteProjectFiles()
+		protected override bool WriteProjectFiles(PlatformProjectGeneratorCollection PlatformProjectGenerators)
 		{
 			WriteCLionIgnoreDirs();
 			return WriteCMakeLists();
diff --git a/Engine/Source/Programs/UnrealBuildTool/ProjectFiles/CodeLite/CodeLiteGenerator.cs b/Engine/Source/Programs/UnrealBuildTool/ProjectFiles/CodeLite/CodeLiteGenerator.cs
--- a/Engine/Source/Programs/UnrealBuildTool/ProjectFiles/CodeLite/CodeLiteGenerator.cs
+++ b/Engine/Source/Programs/UnrealBuildTool/ProjectFiles/CodeLite/CodeLiteGenerator.cs
@@ -55,7 +55,7 @@ namespace UnrealBuildTool
 				return ".project";
 			}
 		}
-		protected override bool WriteMasterProjectFile(ProjectFile UBTProject)
+		protected override bool WriteMasterProjectFile(ProjectFile UBTProject, PlatformProjectGeneratorCollection PlatformProjectGenerators)
 		{
 			string SolutionFileName = MasterProjectName + SolutionExtension;
 			string CodeCompletionFile = MasterProjectName + CodeCompletionFileName;
diff --git a/Engine/Source/Programs/UnrealBuildTool/ProjectFiles/CodeLite/CodeLiteProject.cs b/Engine/Source/Programs/UnrealBuildTool/ProjectFiles/CodeLite/CodeLiteProject.cs
--- a/Engine/Source/Programs/UnrealBuildTool/ProjectFiles/CodeLite/CodeLiteProject.cs
+++ b/Engine/Source/Programs/UnrealBuildTool/ProjectFiles/CodeLite/CodeLiteProject.cs
@@ -34,7 +34,7 @@ namespace UnrealBuildTool
 			return diff.ToString();
 		}
 
-		public override bool WriteProjectFile(List<UnrealTargetPlatform> InPlatforms, List<UnrealTargetConfiguration> InConfigurations)
+		public override bool WriteProjectFile(List<UnrealTargetPlatform> InPlatforms, List<UnrealTargetConfiguration> InConfigurations, PlatformProjectGeneratorCollection PlatformProjectGenerators)
 		{
 			bool bSuccess = false;
 			string ProjectNameRaw = ProjectFilePath.GetFileNameWithoutExtension();
diff --git a/Engine/Source/Programs/UnrealBuildTool/ProjectFiles/Eddie/EddieGenerator.cs b/Engine/Source/Programs/UnrealBuildTool/ProjectFiles/Eddie/EddieGenerator.cs
--- a/Engine/Source/Programs/UnrealBuildTool/ProjectFiles/Eddie/EddieGenerator.cs
+++ b/Engine/Source/Programs/UnrealBuildTool/ProjectFiles/Eddie/EddieGenerator.cs
@@ -105,7 +105,7 @@ namespace UnrealBuildTool
 			return bSuccess;
 		}
 		
-		protected override bool WriteMasterProjectFile(ProjectFile UBTProject)
+		protected override bool WriteMasterProjectFile(ProjectFile UBTProject, PlatformProjectGeneratorCollection PlatformProjectGenerators)
 		{
 			return WriteEddieWorkset();
 		}
diff --git a/Engine/Source/Programs/UnrealBuildTool/ProjectFiles/Eddie/EddieProject.cs b/Engine/Source/Programs/UnrealBuildTool/ProjectFiles/Eddie/EddieProject.cs
--- a/Engine/Source/Programs/UnrealBuildTool/ProjectFiles/Eddie/EddieProject.cs
+++ b/Engine/Source/Programs/UnrealBuildTool/ProjectFiles/Eddie/EddieProject.cs
@@ -195,7 +195,7 @@ namespace UnrealBuildTool
 			}
 		}
 		
-		public override bool WriteProjectFile(List<UnrealTargetPlatform> InPlatforms, List<UnrealTargetConfiguration> InConfigurations)
+		public override bool WriteProjectFile(List<UnrealTargetPlatform> InPlatforms, List<UnrealTargetConfiguration> InConfigurations, PlatformProjectGeneratorCollection PlatformProjectGenerators)
 		{
 			bool bSuccess = false;
 			
diff --git a/Engine/Source/Programs/UnrealBuildTool/ProjectFiles/KDevelop/KDevelopGenerator.cs b/Engine/Source/Programs/UnrealBuildTool/ProjectFiles/KDevelop/KDevelopGenerator.cs
--- a/Engine/Source/Programs/UnrealBuildTool/ProjectFiles/KDevelop/KDevelopGenerator.cs
+++ b/Engine/Source/Programs/UnrealBuildTool/ProjectFiles/KDevelop/KDevelopGenerator.cs
@@ -50,7 +50,7 @@ namespace UnrealBuildTool
 			}
 		}
 
-		protected override bool WriteMasterProjectFile(ProjectFile UBTProject)
+		protected override bool WriteMasterProjectFile(ProjectFile UBTProject, PlatformProjectGeneratorCollection PlatformProjectGenerators)
 		{
 			bool bSuccess = true;
 			return bSuccess;
@@ -505,7 +505,7 @@ namespace UnrealBuildTool
 
 		/// ProjectFileGenerator interface
 		//protected override bool WriteMasterProjectFile( ProjectFile UBTProject )
-		protected override bool WriteProjectFiles()
+		protected override bool WriteProjectFiles(PlatformProjectGeneratorCollection PlatformProjectGenerators)
 		{
 			return WriteKDevelopPro();
 		}
diff --git a/Engine/Source/Programs/UnrealBuildTool/ProjectFiles/Make/MakefileGenerator.cs b/Engine/Source/Programs/UnrealBuildTool/ProjectFiles/Make/MakefileGenerator.cs
--- a/Engine/Source/Programs/UnrealBuildTool/ProjectFiles/Make/MakefileGenerator.cs
+++ b/Engine/Source/Programs/UnrealBuildTool/ProjectFiles/Make/MakefileGenerator.cs
@@ -60,7 +60,7 @@ namespace UnrealBuildTool
 			}
 		}
 
-		protected override bool WriteMasterProjectFile(ProjectFile UBTProject)
+		protected override bool WriteMasterProjectFile(ProjectFile UBTProject, PlatformProjectGeneratorCollection PlatformProjectGenerators)
 		{
 			bool bSuccess = true;
 			return bSuccess;
@@ -198,7 +198,7 @@ namespace UnrealBuildTool
 
 		/// ProjectFileGenerator interface
 		//protected override bool WriteMasterProjectFile( ProjectFile UBTProject )
-		protected override bool WriteProjectFiles()
+		protected override bool WriteProjectFiles(PlatformProjectGeneratorCollection PlatformProjectGenerators)
 		{
 			return WriteMakefile();
 		}
diff --git a/Engine/Source/Programs/UnrealBuildTool/ProjectFiles/PlatformProjectGenerator.cs b/Engine/Source/Programs/UnrealBuildTool/ProjectFiles/PlatformProjectGenerator.cs
--- a/Engine/Source/Programs/UnrealBuildTool/ProjectFiles/PlatformProjectGenerator.cs
+++ b/Engine/Source/Programs/UnrealBuildTool/ProjectFiles/PlatformProjectGenerator.cs
@@ -14,113 +14,10 @@ namespace UnrealBuildTool
 	/// </summary>
 	abstract class PlatformProjectGenerator
 	{
-		static Dictionary<UnrealTargetPlatform, PlatformProjectGenerator> ProjectGeneratorDictionary = new Dictionary<UnrealTargetPlatform, PlatformProjectGenerator>();
-
-		/// <summary>
-		/// Register the given platforms UEPlatformProjectGenerator instance
-		/// </summary>
-		/// <param name="InPlatform">  The UnrealTargetPlatform to register with</param>
-		/// <param name="InProjectGenerator">The UEPlatformProjectGenerator instance to use for the InPlatform</param>
-		public static void RegisterPlatformProjectGenerator(UnrealTargetPlatform InPlatform, PlatformProjectGenerator InProjectGenerator)
-		{
-			// Make sure the build platform is legal
-			UEBuildPlatform BuildPlatform = UEBuildPlatform.GetBuildPlatform(InPlatform, true);
-			if (BuildPlatform != null)
-			{
-				if (ProjectGeneratorDictionary.ContainsKey(InPlatform) == true)
-				{
-					Log.TraceInformation("RegisterPlatformProjectGenerator Warning: Registering project generator {0} for {1} when it is already set to {2}",
-						InProjectGenerator.ToString(), InPlatform.ToString(), ProjectGeneratorDictionary[InPlatform].ToString());
-					ProjectGeneratorDictionary[InPlatform] = InProjectGenerator;
-				}
-				else
-				{
-					ProjectGeneratorDictionary.Add(InPlatform, InProjectGenerator);
-				}
-			}
-			else
-			{
-				Log.TraceVerbose("Skipping project file generator registration for {0} due to no valid BuildPlatform.", InPlatform.ToString());
-			}
-		}
-
-		/// <summary>
-		/// Retrieve the UEPlatformProjectGenerator instance for the given TargetPlatform
-		/// </summary>
-		/// <param name="InPlatform">    The UnrealTargetPlatform being built</param>
-		/// <param name="bInAllowFailure">   If true, do not throw an exception and return null</param>
-		/// <returns>UEPlatformProjectGenerator The instance of the project generator</returns>
-		public static PlatformProjectGenerator GetPlatformProjectGenerator(UnrealTargetPlatform InPlatform, bool bInAllowFailure = false)
-		{
-			if (ProjectGeneratorDictionary.ContainsKey(InPlatform) == true)
-			{
-				return ProjectGeneratorDictionary[InPlatform];
-			}
-			if (bInAllowFailure == true)
-			{
-				return null;
-			}
-			throw new BuildException("GetPlatformProjectGenerator: No PlatformProjectGenerator found for {0}", InPlatform.ToString());
-		}
-
-		/// <summary>
-		/// Allow various platform project generators to generate stub projects if required
-		/// </summary>
-		/// <param name="InGenerator"></param>
-		/// <param name="InTargetName"></param>
-		/// <param name="InTargetFilepath"></param>
-		/// <param name="InTargetRules"></param>
-		/// <param name="InPlatforms"></param>
-		/// <param name="InConfigurations"></param>
-		/// <returns></returns>
-		public static bool GenerateGameProjectStubs(ProjectFileGenerator InGenerator, string InTargetName, string InTargetFilepath, TargetRules InTargetRules,
-			List<UnrealTargetPlatform> InPlatforms, List<UnrealTargetConfiguration> InConfigurations)
-		{
-			foreach (KeyValuePair<UnrealTargetPlatform, PlatformProjectGenerator> Entry in ProjectGeneratorDictionary)
-			{
-				PlatformProjectGenerator ProjGen = Entry.Value;
-				ProjGen.GenerateGameProjectStub(InGenerator, InTargetName, InTargetFilepath, InTargetRules, InPlatforms, InConfigurations);
-			}
-			return true;
-		}
-
-		/// <summary>
-		/// Allow various platform project generators to generate any special project properties if required
-		/// </summary>
-		/// <param name="InPlatform"></param>
-		/// <param name="Configuration"></param>
-		/// <param name="TargetType"></param>
-		/// <param name="VCProjectFileContent"></param>
-		/// <param name="RootDirectory"></param>
-		/// <param name="TargetFilePath"></param>
-		/// <returns></returns>
-		public static bool GenerateGamePlatformSpecificProperties(UnrealTargetPlatform InPlatform, UnrealTargetConfiguration Configuration, TargetType TargetType, StringBuilder VCProjectFileContent, DirectoryReference RootDirectory, FileReference TargetFilePath)
-		{
-			if (ProjectGeneratorDictionary.ContainsKey(InPlatform) == true)
-			{
-				ProjectGeneratorDictionary[InPlatform].GenerateGameProperties(Configuration, VCProjectFileContent, TargetType, RootDirectory, TargetFilePath); ;
-			}
-			return true;
-		}
-
-		public static bool PlatformRequiresVSUserFileGeneration(List<UnrealTargetPlatform> InPlatforms, List<UnrealTargetConfiguration> InConfigurations)
-		{
-			bool bRequiresVSUserFileGeneration = false;
-			foreach (KeyValuePair<UnrealTargetPlatform, PlatformProjectGenerator> Entry in ProjectGeneratorDictionary)
-			{
-				if (InPlatforms.Contains(Entry.Key))
-				{
-					PlatformProjectGenerator ProjGen = Entry.Value;
-					bRequiresVSUserFileGeneration |= ProjGen.RequiresVSUserFileGeneration();
-				}
-			}
-			return bRequiresVSUserFileGeneration;
-		}
-
 		/// <summary>
 		/// Register the platform with the UEPlatformProjectGenerator class
 		/// </summary>
-		public abstract void RegisterPlatformProjectGenerator();
+		public abstract IEnumerable<UnrealTargetPlatform> GetPlatforms();
 
 		public virtual void GenerateGameProjectStub(ProjectFileGenerator InGenerator, string InTargetName, string InTargetFilepath, TargetRules InTargetRules,
 			List<UnrealTargetPlatform> InPlatforms, List<UnrealTargetConfiguration> InConfigurations)
diff --git a/Engine/Source/Programs/UnrealBuildTool/ProjectFiles/PlatformProjectGeneratorCollection.cs b/Engine/Source/Programs/UnrealBuildTool/ProjectFiles/PlatformProjectGeneratorCollection.cs
new file mode 100644
--- /dev/null
+++ b/Engine/Source/Programs/UnrealBuildTool/ProjectFiles/PlatformProjectGeneratorCollection.cs
@@ -0,0 +1,119 @@
+using System;
+using System.Collections.Generic;
+using System.Linq;
+using System.Reflection;
+using System.Text;
+using System.Threading.Tasks;
+using Tools.DotNETCommon;
+
+namespace UnrealBuildTool
+{
+	/// <summary>
+	/// Stores all the registered platform project generators
+	/// </summary>
+	class PlatformProjectGeneratorCollection
+	{
+		Dictionary<UnrealTargetPlatform, PlatformProjectGenerator> ProjectGeneratorDictionary = new Dictionary<UnrealTargetPlatform, PlatformProjectGenerator>();
+
+		/// <summary>
+		/// Register the given platforms UEPlatformProjectGenerator instance
+		/// </summary>
+		/// <param name="InPlatform">  The UnrealTargetPlatform to register with</param>
+		/// <param name="InProjectGenerator">The UEPlatformProjectGenerator instance to use for the InPlatform</param>
+		public void RegisterPlatformProjectGenerator(UnrealTargetPlatform InPlatform, PlatformProjectGenerator InProjectGenerator)
+		{
+			// Make sure the build platform is legal
+			UEBuildPlatform BuildPlatform = UEBuildPlatform.GetBuildPlatform(InPlatform, true);
+			if (BuildPlatform != null)
+			{
+				if (ProjectGeneratorDictionary.ContainsKey(InPlatform) == true)
+				{
+					Log.TraceInformation("RegisterPlatformProjectGenerator Warning: Registering project generator {0} for {1} when it is already set to {2}",
+						InProjectGenerator.ToString(), InPlatform.ToString(), ProjectGeneratorDictionary[InPlatform].ToString());
+					ProjectGeneratorDictionary[InPlatform] = InProjectGenerator;
+				}
+				else
+				{
+					ProjectGeneratorDictionary.Add(InPlatform, InProjectGenerator);
+				}
+			}
+			else
+			{
+				Log.TraceVerbose("Skipping project file generator registration for {0} due to no valid BuildPlatform.", InPlatform.ToString());
+			}
+		}
+
+		/// <summary>
+		/// Retrieve the UEPlatformProjectGenerator instance for the given TargetPlatform
+		/// </summary>
+		/// <param name="InPlatform">    The UnrealTargetPlatform being built</param>
+		/// <param name="bInAllowFailure">   If true, do not throw an exception and return null</param>
+		/// <returns>UEPlatformProjectGenerator The instance of the project generator</returns>
+		public PlatformProjectGenerator GetPlatformProjectGenerator(UnrealTargetPlatform InPlatform, bool bInAllowFailure = false)
+		{
+			if (ProjectGeneratorDictionary.ContainsKey(InPlatform) == true)
+			{
+				return ProjectGeneratorDictionary[InPlatform];
+			}
+			if (bInAllowFailure == true)
+			{
+				return null;
+			}
+			throw new BuildException("GetPlatformProjectGenerator: No PlatformProjectGenerator found for {0}", InPlatform.ToString());
+		}
+
+		/// <summary>
+		/// Allow various platform project generators to generate stub projects if required
+		/// </summary>
+		/// <param name="InGenerator"></param>
+		/// <param name="InTargetName"></param>
+		/// <param name="InTargetFilepath"></param>
+		/// <param name="InTargetRules"></param>
+		/// <param name="InPlatforms"></param>
+		/// <param name="InConfigurations"></param>
+		/// <returns></returns>
+		public bool GenerateGameProjectStubs(ProjectFileGenerator InGenerator, string InTargetName, string InTargetFilepath, TargetRules InTargetRules,
+			List<UnrealTargetPlatform> InPlatforms, List<UnrealTargetConfiguration> InConfigurations)
+		{
+			foreach (KeyValuePair<UnrealTargetPlatform, PlatformProjectGenerator> Entry in ProjectGeneratorDictionary)
+			{
+				PlatformProjectGenerator ProjGen = Entry.Value;
+				ProjGen.GenerateGameProjectStub(InGenerator, InTargetName, InTargetFilepath, InTargetRules, InPlatforms, InConfigurations);
+			}
+			return true;
+		}
+
+		/// <summary>
+		/// Allow various platform project generators to generate any special project properties if required
+		/// </summary>
+		/// <param name="InPlatform"></param>
+		/// <param name="Configuration"></param>
+		/// <param name="TargetType"></param>
+		/// <param name="VCProjectFileContent"></param>
+		/// <param name="RootDirectory"></param>
+		/// <param name="TargetFilePath"></param>
+		/// <returns></returns>
+		public bool GenerateGamePlatformSpecificProperties(UnrealTargetPlatform InPlatform, UnrealTargetConfiguration Configuration, TargetType TargetType, StringBuilder VCProjectFileContent, DirectoryReference RootDirectory, FileReference TargetFilePath)
+		{
+			if (ProjectGeneratorDictionary.ContainsKey(InPlatform) == true)
+			{
+				ProjectGeneratorDictionary[InPlatform].GenerateGameProperties(Configuration, VCProjectFileContent, TargetType, RootDirectory, TargetFilePath); ;
+			}
+			return true;
+		}
+
+		public bool PlatformRequiresVSUserFileGeneration(List<UnrealTargetPlatform> InPlatforms, List<UnrealTargetConfiguration> InConfigurations)
+		{
+			bool bRequiresVSUserFileGeneration = false;
+			foreach (KeyValuePair<UnrealTargetPlatform, PlatformProjectGenerator> Entry in ProjectGeneratorDictionary)
+			{
+				if (InPlatforms.Contains(Entry.Key))
+				{
+					PlatformProjectGenerator ProjGen = Entry.Value;
+					bRequiresVSUserFileGeneration |= ProjGen.RequiresVSUserFileGeneration();
+				}
+			}
+			return bRequiresVSUserFileGeneration;
+		}
+	}
+}
diff --git a/Engine/Source/Programs/UnrealBuildTool/ProjectFiles/Project.cs b/Engine/Source/Programs/UnrealBuildTool/ProjectFiles/Project.cs
--- a/Engine/Source/Programs/UnrealBuildTool/ProjectFiles/Project.cs
+++ b/Engine/Source/Programs/UnrealBuildTool/ProjectFiles/Project.cs
@@ -1,4 +1,4 @@
-﻿// Copyright 1998-2018 Epic Games, Inc. All Rights Reserved.
+// Copyright 1998-2018 Epic Games, Inc. All Rights Reserved.
 
 using System;
 using System.Collections.Generic;
@@ -380,8 +380,9 @@ namespace UnrealBuildTool
 		/// </summary>
 		/// <param name="InPlatforms">The platforms to write the project files for</param>
 		/// <param name="InConfigurations">The configurations to add to the project files</param>
+		/// <param name="PlatformProjectGenerators">The registered platform project generators</param>
 		/// <returns>True on success</returns>
-		public virtual bool WriteProjectFile(List<UnrealTargetPlatform> InPlatforms, List<UnrealTargetConfiguration> InConfigurations)
+		public virtual bool WriteProjectFile(List<UnrealTargetPlatform> InPlatforms, List<UnrealTargetConfiguration> InConfigurations, PlatformProjectGeneratorCollection PlatformProjectGenerators)
 		{
 			throw new BuildException("UnrealBuildTool cannot automatically generate this project type because WriteProjectFile() was not overridden.");
 		}
@@ -391,8 +392,9 @@ namespace UnrealBuildTool
 		/// </summary>
 		/// <param name="InPlatforms">The platforms to write the project files for</param>
 		/// <param name="InConfigurations">The configurations to add to the project files</param>
+		/// <param name="PlatformProjectGenerators">The registered platform project generators</param>
 		/// <returns>List of project files written</returns>
-		public virtual List<Tuple<ProjectFile, string>> WriteDebugProjectFiles(List<UnrealTargetPlatform> InPlatforms, List<UnrealTargetConfiguration> InConfigurations)
+		public virtual List<Tuple<ProjectFile, string>> WriteDebugProjectFiles(List<UnrealTargetPlatform> InPlatforms, List<UnrealTargetConfiguration> InConfigurations, PlatformProjectGeneratorCollection PlatformProjectGenerators)
 		{
 			return null;
 		}
diff --git a/Engine/Source/Programs/UnrealBuildTool/ProjectFiles/ProjectFileGenerator.cs b/Engine/Source/Programs/UnrealBuildTool/ProjectFiles/ProjectFileGenerator.cs
--- a/Engine/Source/Programs/UnrealBuildTool/ProjectFiles/ProjectFileGenerator.cs
+++ b/Engine/Source/Programs/UnrealBuildTool/ProjectFiles/ProjectFileGenerator.cs
@@ -573,8 +573,9 @@ namespace UnrealBuildTool
 		/// Generates a Visual Studio solution file and Visual C++ project files for all known engine and game targets.
 		/// Does not actually build anything.
 		/// </summary>
+		/// <param name="PlatformProjectGenerators">The registered platform project generators</param>
 		/// <param name="Arguments">Command-line arguments</param>
-		public virtual bool GenerateProjectFiles( String[] Arguments )
+		public virtual bool GenerateProjectFiles( PlatformProjectGeneratorCollection PlatformProjectGenerators, String[] Arguments )
 		{
 			bool bSuccess = true;
 
@@ -681,7 +682,7 @@ namespace UnrealBuildTool
 			Dictionary<DirectoryReference, ProjectFile> SampleGameProjects = null;
 			{
 				// Setup buildable projects for all targets
-				AddProjectsForAllTargets( AllGameProjects, out EngineProject, out EnterpriseProject, out GameProjects, out ModProjects, out ProgramProjects, out TemplateGameProjects, out SampleGameProjects );
+				AddProjectsForAllTargets( PlatformProjectGenerators, AllGameProjects, out EngineProject, out EnterpriseProject, out GameProjects, out ModProjects, out ProgramProjects, out TemplateGameProjects, out SampleGameProjects );
 
 				// Add all game projects and game config files
 				AddAllGameProjects(GameProjects, SupportedPlatformNames, RootFolder);
@@ -765,7 +766,7 @@ namespace UnrealBuildTool
 							AddEngineDocumentation( EngineProject );
 						}
 
-						List<Tuple<ProjectFile, string>> NewProjectFiles = EngineProject.WriteDebugProjectFiles(InPlatforms: SupportedPlatforms, InConfigurations: SupportedConfigurations);
+						List<Tuple<ProjectFile, string>> NewProjectFiles = EngineProject.WriteDebugProjectFiles(SupportedPlatforms, SupportedConfigurations, PlatformProjectGenerators);
 
 						if (NewProjectFiles != null)
 						{
@@ -816,7 +817,7 @@ namespace UnrealBuildTool
 							RootFolder.AddSubFolder( "Games" ).ChildProjects.Add( CurGameProject );
 						}
 
-						List<Tuple<ProjectFile, string>> NewProjectFiles = CurGameProject.WriteDebugProjectFiles(InPlatforms: SupportedPlatforms, InConfigurations: SupportedConfigurations);
+						List<Tuple<ProjectFile, string>> NewProjectFiles = CurGameProject.WriteDebugProjectFiles(SupportedPlatforms, SupportedConfigurations, PlatformProjectGenerators);
 
 						if (NewProjectFiles != null)
 						{
@@ -960,7 +961,7 @@ namespace UnrealBuildTool
 				// the targets so that we can extra the compiler defines, include paths, etc.
 				if(GenerateIntelliSenseData(Arguments, IntelliSenseTargetFiles))
 				{
-					WriteProjectFiles();
+					WriteProjectFiles(PlatformProjectGenerators);
 					Log.TraceVerbose( "Project generation complete ({0} generated, {1} imported)", GeneratedProjectFiles.Count, OtherProjectFiles.Count );
 				}
 			}
@@ -2014,6 +2015,7 @@ namespace UnrealBuildTool
 		/// <summary>
 		/// Creates project entries for all known targets (*.Target.cs files)
 		/// </summary>
+		/// <param name="PlatformProjectGenerators">The registered platform project generators</param>
 		/// <param name="AllGames">All game folders</param>
 		/// <param name="EngineProject">The engine project we created</param>
 		/// <param name="EnterpriseProject">The enterprise project we created</param>
@@ -2022,7 +2024,7 @@ namespace UnrealBuildTool
 		/// <param name="ProgramProjects">Map of program names to all of the program projects we created</param>
 		/// <param name="TemplateGameProjects">Set of template game projects we found.  These will also be in the GameProjects map</param>
 		/// <param name="SampleGameProjects">Set of sample game projects that were found</param>
-		private void AddProjectsForAllTargets( List<FileReference> AllGames, out ProjectFile EngineProject, out ProjectFile EnterpriseProject,
+		private void AddProjectsForAllTargets( PlatformProjectGeneratorCollection PlatformProjectGenerators, List<FileReference> AllGames, out ProjectFile EngineProject, out ProjectFile EnterpriseProject,
 			out Dictionary<DirectoryReference, ProjectFile> GameProjects, out Dictionary<DirectoryReference, ProjectFile> ModProjects,
 			out Dictionary<FileReference, ProjectFile> ProgramProjects, out Dictionary<DirectoryReference, ProjectFile> TemplateGameProjects, out Dictionary<DirectoryReference, ProjectFile> SampleGameProjects )
 		{
@@ -2146,7 +2148,7 @@ namespace UnrealBuildTool
 					if (TargetRulesObject.Type == TargetType.Game || TargetRulesObject.Type == TargetType.Client || TargetRulesObject.Type == TargetType.Server)
 					{
 						// Allow platforms to generate stub projects here...
-						PlatformProjectGenerator.GenerateGameProjectStubs(
+						PlatformProjectGenerators.GenerateGameProjectStubs(
 							InGenerator: this,
 							InTargetName: TargetName,
 							InTargetFilepath: TargetFilePath.FullName,
@@ -2466,7 +2468,7 @@ namespace UnrealBuildTool
 		/// Writes the project files to disk
 		/// </summary>
 		/// <returns>True if successful</returns>
-		protected virtual bool WriteProjectFiles()
+		protected virtual bool WriteProjectFiles(PlatformProjectGeneratorCollection PlatformProjectGenerators)
 		{
             using(ProgressWriter Progress = new ProgressWriter("Writing project files...", true))
 			{
@@ -2475,9 +2477,7 @@ namespace UnrealBuildTool
 				for(int ProjectFileIndex = 0 ; ProjectFileIndex < GeneratedProjectFiles.Count; ++ProjectFileIndex )
 				{
 					ProjectFile CurProject = GeneratedProjectFiles[ ProjectFileIndex ];
-					if( !CurProject.WriteProjectFile(
-								InPlatforms: SupportedPlatforms,
-								InConfigurations: SupportedConfigurations ) )
+					if( !CurProject.WriteProjectFile(SupportedPlatforms, SupportedConfigurations, PlatformProjectGenerators) )
 					{
 						return false;
 					}
@@ -2485,7 +2485,7 @@ namespace UnrealBuildTool
 					Progress.Write(ProjectFileIndex + 1, TotalProjectFileCount);
 				}
 
-				WriteMasterProjectFile( UBTProject: UBTProject );
+				WriteMasterProjectFile( UBTProject, PlatformProjectGenerators );
 				Progress.Write(TotalProjectFileCount, TotalProjectFileCount);
 			}
 			return true;
@@ -2495,8 +2495,9 @@ namespace UnrealBuildTool
 		/// Writes the master project file (e.g. Visual Studio Solution file)
 		/// </summary>
 		/// <param name="UBTProject">The UnrealBuildTool project</param>
+		/// <param name="PlatformProjectGenerators">The platform project file generators</param>
 		/// <returns>True if successful</returns>
-		protected abstract bool WriteMasterProjectFile( ProjectFile UBTProject );
+		protected abstract bool WriteMasterProjectFile( ProjectFile UBTProject, PlatformProjectGeneratorCollection PlatformProjectGenerators );
 
 
 		/// <summary>
diff --git a/Engine/Source/Programs/UnrealBuildTool/ProjectFiles/QMake/QMakefileGenerator.cs b/Engine/Source/Programs/UnrealBuildTool/ProjectFiles/QMake/QMakefileGenerator.cs
--- a/Engine/Source/Programs/UnrealBuildTool/ProjectFiles/QMake/QMakefileGenerator.cs
+++ b/Engine/Source/Programs/UnrealBuildTool/ProjectFiles/QMake/QMakefileGenerator.cs
@@ -51,7 +51,7 @@ namespace UnrealBuildTool
 			}
 		}
 
-		protected override bool WriteMasterProjectFile(ProjectFile UBTProject)
+		protected override bool WriteMasterProjectFile(ProjectFile UBTProject, PlatformProjectGeneratorCollection PlatformProjectGenerators)
 		{
 			bool bSuccess = true;
 			return bSuccess;
@@ -428,7 +428,7 @@ namespace UnrealBuildTool
 
 		/// ProjectFileGenerator interface
 		//protected override bool WriteMasterProjectFile( ProjectFile UBTProject )
-		protected override bool WriteProjectFiles()
+		protected override bool WriteProjectFiles(PlatformProjectGeneratorCollection PlatformProjectGenerators)
 		{
 			return WriteQMakePro();
 		}
diff --git a/Engine/Source/Programs/UnrealBuildTool/ProjectFiles/VisualStudio/VCProject.cs b/Engine/Source/Programs/UnrealBuildTool/ProjectFiles/VisualStudio/VCProject.cs
--- a/Engine/Source/Programs/UnrealBuildTool/ProjectFiles/VisualStudio/VCProject.cs
+++ b/Engine/Source/Programs/UnrealBuildTool/ProjectFiles/VisualStudio/VCProject.cs
@@ -107,9 +107,10 @@ namespace UnrealBuildTool
 		/// <param name="Platform">Actual platform</param>
 		/// <param name="Configuration">Actual configuration</param>
 		/// <param name="TargetConfigurationName">The configuration name from the target rules, or null if we don't have one</param>
+		/// <param name="PlatformProjectGenerators">The registered platform project generators</param>
 		/// <param name="ProjectPlatformName">Name of platform string to use for Visual Studio project</param>
 		/// <param name="ProjectConfigurationName">Name of configuration string to use for Visual Studio project</param>
-		public abstract void MakeProjectPlatformAndConfigurationNames(UnrealTargetPlatform Platform, UnrealTargetConfiguration Configuration, string TargetConfigurationName, out string ProjectPlatformName, out string ProjectConfigurationName);
+		public abstract void MakeProjectPlatformAndConfigurationNames(UnrealTargetPlatform Platform, UnrealTargetConfiguration Configuration, string TargetConfigurationName, PlatformProjectGeneratorCollection PlatformProjectGenerators, out string ProjectPlatformName, out string ProjectConfigurationName);
 
 		static UnrealTargetConfiguration[] GetSupportedConfigurations(TargetRules Rules)
 		{
@@ -144,8 +145,9 @@ namespace UnrealBuildTool
 		/// <param name="ProjectTarget">The target that we're checking for a valid platform/config combination</param>
 		/// <param name="Platform">Platform</param>
 		/// <param name="Configuration">Configuration</param>
+		/// <param name="PlatformProjectGenerators">The registered platform project generators</param>
 		/// <returns>True if this is a valid combination for this project, otherwise false</returns>
-		public static bool IsValidProjectPlatformAndConfiguration(ProjectTarget ProjectTarget, UnrealTargetPlatform Platform, UnrealTargetConfiguration Configuration)
+		public static bool IsValidProjectPlatformAndConfiguration(ProjectTarget ProjectTarget, UnrealTargetPlatform Platform, UnrealTargetConfiguration Configuration, PlatformProjectGeneratorCollection PlatformProjectGenerators)
 		{
 			if (!ProjectFileGenerator.bIncludeTestAndShippingConfigs)
 			{
@@ -155,7 +157,7 @@ namespace UnrealBuildTool
 				}
 			}
 
-			PlatformProjectGenerator PlatformProjectGenerator = PlatformProjectGenerator.GetPlatformProjectGenerator(Platform, true);
+			PlatformProjectGenerator PlatformProjectGenerator = PlatformProjectGenerators.GetPlatformProjectGenerator(Platform, true);
 			if (PlatformProjectGenerator == null)
 			{
 				return false;
@@ -286,11 +288,12 @@ namespace UnrealBuildTool
 		/// <param name="Platform">Actual platform</param>
 		/// <param name="Configuration">Actual configuration</param>
 		/// <param name="TargetConfigurationName">The configuration name from the target rules, or null if we don't have one</param>
+		/// <param name="PlatformProjectGenerators">The registered platform project generators</param>
 		/// <param name="ProjectPlatformName">Name of platform string to use for Visual Studio project</param>
 		/// <param name="ProjectConfigurationName">Name of configuration string to use for Visual Studio project</param>
-		public override void MakeProjectPlatformAndConfigurationNames(UnrealTargetPlatform Platform, UnrealTargetConfiguration Configuration, string TargetConfigurationName, out string ProjectPlatformName, out string ProjectConfigurationName)
+		public override void MakeProjectPlatformAndConfigurationNames(UnrealTargetPlatform Platform, UnrealTargetConfiguration Configuration, string TargetConfigurationName, PlatformProjectGeneratorCollection PlatformProjectGenerators, out string ProjectPlatformName, out string ProjectConfigurationName)
 		{
-			PlatformProjectGenerator PlatformProjectGenerator = PlatformProjectGenerator.GetPlatformProjectGenerator(Platform, bInAllowFailure: true);
+			PlatformProjectGenerator PlatformProjectGenerator = PlatformProjectGenerators.GetPlatformProjectGenerator(Platform, bInAllowFailure: true);
 
 			// Check to see if this platform is supported directly by Visual Studio projects.
 			bool HasActualVSPlatform = (PlatformProjectGenerator != null) ? PlatformProjectGenerator.HasVisualStudioSupport(Platform, Configuration, ProjectFileFormat) : false;
@@ -362,7 +365,7 @@ namespace UnrealBuildTool
 
 		List<ProjectConfigAndTargetCombination> ProjectConfigAndTargetCombinations = new List<ProjectConfigAndTargetCombination>();
 
-		private void BuildProjectConfigAndTargetCombinations(List<UnrealTargetPlatform> InPlatforms, List<UnrealTargetConfiguration> InConfigurations)
+		private void BuildProjectConfigAndTargetCombinations(List<UnrealTargetPlatform> InPlatforms, List<UnrealTargetConfiguration> InConfigurations, PlatformProjectGeneratorCollection PlatformProjectGenerators)
 		{
 			//no need to do this more than once
 			if(ProjectConfigAndTargetCombinations.Count > 0)
@@ -406,10 +409,10 @@ namespace UnrealBuildTool
 
 							foreach (ProjectTarget ProjectTarget in ProjectTargets)
 							{
-								if (IsValidProjectPlatformAndConfiguration(ProjectTarget, Platform, Configuration))
+								if (IsValidProjectPlatformAndConfiguration(ProjectTarget, Platform, Configuration, PlatformProjectGenerators))
 								{
 									string ProjectPlatformName, ProjectConfigurationName;
-									MakeProjectPlatformAndConfigurationNames(Platform, Configuration, ProjectTarget.TargetRules.Type.ToString(), out ProjectPlatformName, out ProjectConfigurationName);
+									MakeProjectPlatformAndConfigurationNames(Platform, Configuration, ProjectTarget.TargetRules.Type.ToString(), PlatformProjectGenerators, out ProjectPlatformName, out ProjectConfigurationName);
 
 									ProjectConfigAndTargetCombination Combination = new ProjectConfigAndTargetCombination(Platform, Configuration, ProjectPlatformName, ProjectConfigurationName, ProjectTarget);
 									ProjectConfigAndTargetCombinations.Add(Combination);
@@ -426,14 +429,14 @@ namespace UnrealBuildTool
 		/// If found writes a debug project file to disk
 		/// </summary>
 		/// <returns>True on success</returns>
-		public override List<Tuple<ProjectFile, string>> WriteDebugProjectFiles(List<UnrealTargetPlatform> InPlatforms, List<UnrealTargetConfiguration> InConfigurations)
+		public override List<Tuple<ProjectFile, string>> WriteDebugProjectFiles(List<UnrealTargetPlatform> InPlatforms, List<UnrealTargetConfiguration> InConfigurations, PlatformProjectGeneratorCollection PlatformProjectGenerators)
 		{
 			string ProjectName = ProjectFilePath.GetFileNameWithoutExtension();
 
 			List<UnrealTargetPlatform> ProjectPlatforms = new List<UnrealTargetPlatform>();
 			List<Tuple<ProjectFile, string>> ProjectFiles = new List<Tuple<ProjectFile, string>>();
 
-			BuildProjectConfigAndTargetCombinations(InPlatforms, InConfigurations);
+			BuildProjectConfigAndTargetCombinations(InPlatforms, InConfigurations, PlatformProjectGenerators);
 
 
 			foreach (ProjectConfigAndTargetCombination Combination in ProjectConfigAndTargetCombinations)
@@ -449,7 +452,7 @@ namespace UnrealBuildTool
 			{
 				foreach (UnrealTargetPlatform Platform in ProjectPlatforms)
 				{
-					PlatformProjectGenerator ProjGenerator = PlatformProjectGenerator.GetPlatformProjectGenerator(Platform, true);
+					PlatformProjectGenerator ProjGenerator = PlatformProjectGenerators.GetPlatformProjectGenerator(Platform, true);
 					if (ProjGenerator != null)
 					{
 						//write out additional prop file
@@ -472,7 +475,7 @@ namespace UnrealBuildTool
 		}
 
 		/// Implements Project interface
-		public override bool WriteProjectFile(List<UnrealTargetPlatform> InPlatforms, List<UnrealTargetConfiguration> InConfigurations)
+		public override bool WriteProjectFile(List<UnrealTargetPlatform> InPlatforms, List<UnrealTargetConfiguration> InConfigurations, PlatformProjectGeneratorCollection PlatformProjectGenerators)
 		{
 			string ProjectName = ProjectFilePath.GetFileNameWithoutExtension();
 
@@ -524,7 +527,7 @@ namespace UnrealBuildTool
 				ProjectFileGenerator.NewLine +
 				"<Project DefaultTargets=\"Build\" ToolsVersion=\"" + VCProjectFileGenerator.GetProjectFileToolVersionString(ProjectFileFormat) + "\" xmlns=\"http://schemas.microsoft.com/developer/msbuild/2003\">" + ProjectFileGenerator.NewLine);
 
-			bool bGenerateUserFileContent = PlatformProjectGenerator.PlatformRequiresVSUserFileGeneration(InPlatforms, InConfigurations);
+			bool bGenerateUserFileContent = PlatformProjectGenerators.PlatformRequiresVSUserFileGeneration(InPlatforms, InConfigurations);
 			if (bGenerateUserFileContent)
 			{
 				VCUserFileContent.Append(
@@ -534,7 +537,7 @@ namespace UnrealBuildTool
 					);
 			}
 
-			BuildProjectConfigAndTargetCombinations(InPlatforms, InConfigurations);
+			BuildProjectConfigAndTargetCombinations(InPlatforms, InConfigurations, PlatformProjectGenerators);
 
 			VCProjectFileContent.Append(
 				"	<ItemGroup Label=\"ProjectConfigurations\">" + ProjectFileGenerator.NewLine);
@@ -588,7 +591,7 @@ namespace UnrealBuildTool
 			{
 				foreach (UnrealTargetPlatform Platform in ProjectPlatforms)
 				{
-					PlatformProjectGenerator ProjGenerator = PlatformProjectGenerator.GetPlatformProjectGenerator(Platform, true);
+					PlatformProjectGenerator ProjGenerator = PlatformProjectGenerators.GetPlatformProjectGenerator(Platform, true);
 					if (ProjGenerator != null && ProjGenerator.HasVisualStudioSupport(Platform, UnrealTargetConfiguration.Development, ProjectFileFormat))
 					{
 						AdditionalPropertyGroups.Append(ProjGenerator.GetAdditionalVisualStudioPropertyGroups(Platform, ProjectFileFormat));
@@ -616,7 +619,7 @@ namespace UnrealBuildTool
 			{
 				foreach (UnrealTargetPlatform Platform in ProjectPlatforms)
 				{
-					PlatformProjectGenerator ProjGenerator = PlatformProjectGenerator.GetPlatformProjectGenerator(Platform, true);
+					PlatformProjectGenerator ProjGenerator = PlatformProjectGenerators.GetPlatformProjectGenerator(Platform, true);
 					if (ProjGenerator != null && ProjGenerator.HasVisualStudioSupport(Platform, UnrealTargetConfiguration.Development, ProjectFileFormat))
 					{
 						VCProjectFileContent.Append(ProjGenerator.GetVisualStudioGlobalProperties(Platform));
@@ -645,7 +648,7 @@ namespace UnrealBuildTool
 				{
 					string ProjectPlatformName = PlatformTuple.Item1;
 					UnrealTargetPlatform TargetPlatform = PlatformTuple.Item2;
-					WritePreDefaultPropsConfiguration(TargetPlatform, TargetConfiguration, ProjectPlatformName, ProjectConfigurationName, VCProjectFileContent);
+					WritePreDefaultPropsConfiguration(TargetPlatform, TargetConfiguration, ProjectPlatformName, ProjectConfigurationName, PlatformProjectGenerators, VCProjectFileContent);
 				}
 			}
 
@@ -661,7 +664,7 @@ namespace UnrealBuildTool
 				{
 					string ProjectPlatformName = PlatformTuple.Item1;
 					UnrealTargetPlatform TargetPlatform = PlatformTuple.Item2;
-					WritePostDefaultPropsConfiguration(TargetPlatform, TargetConfiguration, ProjectPlatformName, ProjectConfigurationName, VCProjectFileContent);
+					WritePostDefaultPropsConfiguration(TargetPlatform, TargetConfiguration, ProjectPlatformName, ProjectConfigurationName, PlatformProjectGenerators, VCProjectFileContent);
 				}
 			}
 
@@ -674,7 +677,7 @@ namespace UnrealBuildTool
 			// Write each project configuration
 			foreach (ProjectConfigAndTargetCombination Combination in ProjectConfigAndTargetCombinations)
 			{
-				WriteConfiguration(ProjectName, Combination, VCProjectFileContent, bGenerateUserFileContent ? VCUserFileContent : null);
+				WriteConfiguration(ProjectName, Combination, PlatformProjectGenerators, VCProjectFileContent, bGenerateUserFileContent ? VCUserFileContent : null);
 			}
 
 			// Source folders and files
@@ -834,7 +837,7 @@ namespace UnrealBuildTool
 			{
 				foreach (UnrealTargetPlatform Platform in ProjectPlatforms)
 				{
-					PlatformProjectGenerator ProjGenerator = PlatformProjectGenerator.GetPlatformProjectGenerator(Platform, true);
+					PlatformProjectGenerator ProjGenerator = PlatformProjectGenerators.GetPlatformProjectGenerator(Platform, true);
 					if (ProjGenerator != null && ProjGenerator.HasVisualStudioSupport(Platform, UnrealTargetConfiguration.Development, ProjectFileFormat))
 					{
 						AdditionalTargetSettings.Append(ProjGenerator.GetVisualStudioTargetOverrides(Platform, ProjectFileFormat));
@@ -847,7 +850,7 @@ namespace UnrealBuildTool
 			{
 				foreach (UnrealTargetPlatform Platform in ProjectPlatforms)
 				{
-					PlatformProjectGenerator ProjGenerator = PlatformProjectGenerator.GetPlatformProjectGenerator(Platform, true);
+					PlatformProjectGenerator ProjGenerator = PlatformProjectGenerators.GetPlatformProjectGenerator(Platform, true);
 					if (ProjGenerator != null && ProjGenerator.HasVisualStudioSupport(Platform, UnrealTargetConfiguration.Development, ProjectFileFormat))
 					{
 						// @todo projectfiles: Serious hacks here because we are trying to emit one-time platform-specific sections that need information
@@ -1037,9 +1040,9 @@ namespace UnrealBuildTool
 		}
 
 		// Anonymous function that writes pre-Default.props configuration data
-		private void WritePreDefaultPropsConfiguration(UnrealTargetPlatform TargetPlatform, UnrealTargetConfiguration TargetConfiguration, string ProjectPlatformName, string ProjectConfigurationName, StringBuilder VCProjectFileContent)
+		private void WritePreDefaultPropsConfiguration(UnrealTargetPlatform TargetPlatform, UnrealTargetConfiguration TargetConfiguration, string ProjectPlatformName, string ProjectConfigurationName, PlatformProjectGeneratorCollection PlatformProjectGenerators, StringBuilder VCProjectFileContent)
 		{
-			PlatformProjectGenerator ProjGenerator = PlatformProjectGenerator.GetPlatformProjectGenerator(TargetPlatform, true);
+			PlatformProjectGenerator ProjGenerator = PlatformProjectGenerators.GetPlatformProjectGenerator(TargetPlatform, true);
 			if (((ProjGenerator == null) && (TargetPlatform != UnrealTargetPlatform.Unknown)))
 			{
 				return;
@@ -1061,9 +1064,9 @@ namespace UnrealBuildTool
 		}
 
 		// Anonymous function that writes post-Default.props configuration data
-		private void WritePostDefaultPropsConfiguration(UnrealTargetPlatform TargetPlatform, UnrealTargetConfiguration TargetConfiguration, string ProjectPlatformName, string ProjectConfigurationName, StringBuilder VCProjectFileContent)
+		private void WritePostDefaultPropsConfiguration(UnrealTargetPlatform TargetPlatform, UnrealTargetConfiguration TargetConfiguration, string ProjectPlatformName, string ProjectConfigurationName, PlatformProjectGeneratorCollection PlatformProjectGenerators, StringBuilder VCProjectFileContent)
 		{
-			PlatformProjectGenerator ProjGenerator = PlatformProjectGenerator.GetPlatformProjectGenerator(TargetPlatform, true);
+			PlatformProjectGenerator ProjGenerator = PlatformProjectGenerators.GetPlatformProjectGenerator(TargetPlatform, true);
 			if (((ProjGenerator == null) && (TargetPlatform != UnrealTargetPlatform.Unknown)))
 			{
 				return;
@@ -1088,12 +1091,12 @@ namespace UnrealBuildTool
 		}
 
 		// Anonymous function that writes project configuration data
-		private void WriteConfiguration(string ProjectName, ProjectConfigAndTargetCombination Combination, StringBuilder VCProjectFileContent, StringBuilder VCUserFileContent)
+		private void WriteConfiguration(string ProjectName, ProjectConfigAndTargetCombination Combination, PlatformProjectGeneratorCollection PlatformProjectGenerators, StringBuilder VCProjectFileContent, StringBuilder VCUserFileContent)
 		{
 			UnrealTargetPlatform Platform = Combination.Platform;
 			UnrealTargetConfiguration Configuration = Combination.Configuration;
 
-			PlatformProjectGenerator ProjGenerator = PlatformProjectGenerator.GetPlatformProjectGenerator(Platform, true);
+			PlatformProjectGenerator ProjGenerator = PlatformProjectGenerators.GetPlatformProjectGenerator(Platform, true);
 			if (((ProjGenerator == null) && (Platform != UnrealTargetPlatform.Unknown)))
 			{
 				return;
@@ -1274,7 +1277,7 @@ namespace UnrealBuildTool
 					if (TargetRulesObject.Type == TargetType.Game || TargetRulesObject.Type == TargetType.Client || TargetRulesObject.Type == TargetType.Server)
 					{
 						// Allow platforms to add any special properties they require... like aumid override for Xbox One
-						PlatformProjectGenerator.GenerateGamePlatformSpecificProperties(Platform, Configuration, TargetRulesObject.Type, VCProjectFileContent, RootDirectory, TargetFilePath);
+						PlatformProjectGenerators.GenerateGamePlatformSpecificProperties(Platform, Configuration, TargetRulesObject.Type, VCProjectFileContent, RootDirectory, TargetFilePath);
 					}
 
 					VCProjectFileContent.Append("	</PropertyGroup>" + ProjectFileGenerator.NewLine);
@@ -1430,9 +1433,10 @@ namespace UnrealBuildTool
 		/// <param name="Platform">Actual platform</param>
 		/// <param name="Configuration">Actual configuration</param>
 		/// <param name="TargetConfigurationName">The configuration name from the target rules, or null if we don't have one</param>
+		/// <param name="PlatformProjectGenerators">The registered platform project generators</param>
 		/// <param name="ProjectPlatformName">Name of platform string to use for Visual Studio project</param>
 		/// <param name="ProjectConfigurationName">Name of configuration string to use for Visual Studio project</param>
-		public override void MakeProjectPlatformAndConfigurationNames(UnrealTargetPlatform Platform, UnrealTargetConfiguration Configuration, string TargetConfigurationName, out string ProjectPlatformName, out string ProjectConfigurationName)
+		public override void MakeProjectPlatformAndConfigurationNames(UnrealTargetPlatform Platform, UnrealTargetConfiguration Configuration, string TargetConfigurationName, PlatformProjectGeneratorCollection PlatformProjectGenerators, out string ProjectPlatformName, out string ProjectConfigurationName)
 		{
 			ProjectConfigurationName = Configuration.ToString();
 			ProjectPlatformName = VCProjectFileGenerator.DotNetPlatformName;
@@ -1443,8 +1447,9 @@ namespace UnrealBuildTool
 		/// </summary>
 		/// <param name="InPlatforms">Not used.</param>
 		/// <param name="InConfigurations">Not Used.</param>
+		/// <param name="PlatformProjectGenerators">The registered platform project generators</param>
 		/// <returns>true if the opration was successful, false otherwise</returns>
-		public override bool WriteProjectFile(List<UnrealTargetPlatform> InPlatforms, List<UnrealTargetConfiguration> InConfigurations)
+		public override bool WriteProjectFile(List<UnrealTargetPlatform> InPlatforms, List<UnrealTargetConfiguration> InConfigurations, PlatformProjectGeneratorCollection PlatformProjectGenerators)
 		{
 			throw new BuildException("Support for writing C# projects from UnrealBuildTool has been removed.");
 		}
diff --git a/Engine/Source/Programs/UnrealBuildTool/ProjectFiles/VisualStudio/VCProjectFileGenerator.cs b/Engine/Source/Programs/UnrealBuildTool/ProjectFiles/VisualStudio/VCProjectFileGenerator.cs
--- a/Engine/Source/Programs/UnrealBuildTool/ProjectFiles/VisualStudio/VCProjectFileGenerator.cs
+++ b/Engine/Source/Programs/UnrealBuildTool/ProjectFiles/VisualStudio/VCProjectFileGenerator.cs
@@ -372,9 +372,9 @@ namespace UnrealBuildTool
 		/// Writes the project files to disk
 		/// </summary>
 		/// <returns>True if successful</returns>
-		protected override bool WriteProjectFiles()
+		protected override bool WriteProjectFiles(PlatformProjectGeneratorCollection PlatformProjectGenerators)
 		{
-			if(!base.WriteProjectFiles())
+			if(!base.WriteProjectFiles(PlatformProjectGenerators))
 			{
 				return false;
 			}
@@ -406,7 +406,7 @@ namespace UnrealBuildTool
 		}
 
 
-		protected override bool WriteMasterProjectFile(ProjectFile UBTProject)
+		protected override bool WriteMasterProjectFile(ProjectFile UBTProject, PlatformProjectGeneratorCollection PlatformProjectGenerators)
 		{
 			bool bSuccess = true;
 
@@ -622,7 +622,7 @@ namespace UnrealBuildTool
 												// Figure out the set of valid target configuration names
 												foreach (ProjectTarget ProjectTarget in CurProject.ProjectTargets)
 												{
-													if (VCProjectFile.IsValidProjectPlatformAndConfiguration(ProjectTarget, CurPlatform, CurConfiguration))
+													if (VCProjectFile.IsValidProjectPlatformAndConfiguration(ProjectTarget, CurPlatform, CurConfiguration, PlatformProjectGenerators))
 													{
 														PlatformsValidForProjects.Add(CurPlatform);
 
@@ -735,7 +735,7 @@ namespace UnrealBuildTool
 										ProjectTarget MatchingProjectTarget = null;
 										foreach (ProjectTarget ProjectTarget in CurProject.ProjectTargets)
 										{
-											bool IsMatchingCombination = VCProjectFile.IsValidProjectPlatformAndConfiguration(ProjectTarget, SolutionConfigCombination.Platform, SolutionConfigCombination.Configuration);
+											bool IsMatchingCombination = VCProjectFile.IsValidProjectPlatformAndConfiguration(ProjectTarget, SolutionConfigCombination.Platform, SolutionConfigCombination.Configuration, PlatformProjectGenerators);
 											if (ProjectTarget.TargetRules != null)
 											{
 												if (TargetConfigurationName != ProjectTarget.TargetRules.Type.ToString())
@@ -841,7 +841,7 @@ namespace UnrealBuildTool
 										}
 										else
 										{
-											CurProject.MakeProjectPlatformAndConfigurationNames(SolutionPlatform, SolutionConfiguration, TargetConfigurationName, out ProjectPlatformName, out ProjectConfigName);
+											CurProject.MakeProjectPlatformAndConfigurationNames(SolutionPlatform, SolutionConfiguration, TargetConfigurationName, PlatformProjectGenerators, out ProjectPlatformName, out ProjectConfigName);
 										}
 
 										string ProjectConfigAndPlatformPair = ProjectConfigName.ToString() + "|" + ProjectPlatformName.ToString();
@@ -867,7 +867,7 @@ namespace UnrealBuildTool
 												VCSolutionFileContent.Append(
 														"		" + CurProjectGUID + "." + SolutionConfigCombination.VCSolutionConfigAndPlatformName + ".Build.0 = " + ProjectConfigAndPlatformPair + ProjectFileGenerator.NewLine);
 
-												PlatformProjectGenerator ProjGen = PlatformProjectGenerator.GetPlatformProjectGenerator(SolutionConfigCombination.Platform, true);
+												PlatformProjectGenerator ProjGen = PlatformProjectGenerators.GetPlatformProjectGenerator(SolutionConfigCombination.Platform, true);
 												if (MatchingProjectTarget.ProjectDeploys ||
 													((ProjGen != null) && (ProjGen.GetVisualStudioDeploymentEnabled(SolutionPlatform, SolutionConfiguration) == true)))
 												{
diff --git a/Engine/Source/Programs/UnrealBuildTool/ProjectFiles/VisualStudioCode/VSCodeProjectFileGenerator.cs b/Engine/Source/Programs/UnrealBuildTool/ProjectFiles/VisualStudioCode/VSCodeProjectFileGenerator.cs
--- a/Engine/Source/Programs/UnrealBuildTool/ProjectFiles/VisualStudioCode/VSCodeProjectFileGenerator.cs
+++ b/Engine/Source/Programs/UnrealBuildTool/ProjectFiles/VisualStudioCode/VSCodeProjectFileGenerator.cs
@@ -24,7 +24,7 @@ namespace UnrealBuildTool
 		{
 		}
 
-		public override bool WriteProjectFile(List<UnrealTargetPlatform> InPlatforms, List<UnrealTargetConfiguration> InConfigurations)
+		public override bool WriteProjectFile(List<UnrealTargetPlatform> InPlatforms, List<UnrealTargetConfiguration> InConfigurations, PlatformProjectGeneratorCollection PlatformProjectGenerators)
 		{
 			return true;
 		}
@@ -240,7 +240,7 @@ namespace UnrealBuildTool
 			return new VSCodeProjectFolder(InitOwnerProjectFileGenerator, InitFolderName);
 		} 
 
-		protected override bool WriteMasterProjectFile(ProjectFile UBTProject)
+		protected override bool WriteMasterProjectFile(ProjectFile UBTProject, PlatformProjectGeneratorCollection PlatformProjectGenerators)
 		{
 			VSCodeDir = DirectoryReference.Combine(MasterProjectPath, ".vscode");
 			DirectoryReference.CreateDirectory(VSCodeDir);
@@ -269,7 +269,7 @@ namespace UnrealBuildTool
 			}
 			Projects.Sort((A, B) => { return A.ProjectFilePath.GetFileName().CompareTo(B.ProjectFilePath.GetFileName()); });
 
-			ProjectData ProjectData = GatherProjectData(Projects);
+			ProjectData ProjectData = GatherProjectData(Projects, PlatformProjectGenerators);
 
 			WriteTasksFile(ProjectData);
 			WriteLaunchFile(ProjectData);
@@ -347,7 +347,7 @@ namespace UnrealBuildTool
 			public List<string> CombinedIncludePaths = new List<string>();
 		}
 
-		private ProjectData GatherProjectData(List<ProjectFile> InProjects)
+		private ProjectData GatherProjectData(List<ProjectFile> InProjects, PlatformProjectGeneratorCollection PlatformProjectGenerators)
 		{
 			ProjectData ProjectData = new ProjectData();
 
@@ -383,7 +383,7 @@ namespace UnrealBuildTool
 							{
 								foreach (UnrealTargetConfiguration Config in Configs)
 								{
-									if (MSBuildProjectFile.IsValidProjectPlatformAndConfiguration(Target, Platform, Config))
+									if (MSBuildProjectFile.IsValidProjectPlatformAndConfiguration(Target, Platform, Config, PlatformProjectGenerators))
 									{
 										NewTarget.BuildProducts.Add(new ProjectData.BuildProduct
 										{
diff --git a/Engine/Source/Programs/UnrealBuildTool/ProjectFiles/VisualStudioMac/VisualStudioMacProjectFileGenerator.cs b/Engine/Source/Programs/UnrealBuildTool/ProjectFiles/VisualStudioMac/VisualStudioMacProjectFileGenerator.cs
--- a/Engine/Source/Programs/UnrealBuildTool/ProjectFiles/VisualStudioMac/VisualStudioMacProjectFileGenerator.cs
+++ b/Engine/Source/Programs/UnrealBuildTool/ProjectFiles/VisualStudioMac/VisualStudioMacProjectFileGenerator.cs
@@ -39,7 +39,7 @@ namespace UnrealBuildTool
         /// Writes the project files to disk
         /// </summary>
         /// <returns>True if successful</returns>
-        protected override bool WriteProjectFiles()
+        protected override bool WriteProjectFiles(PlatformProjectGeneratorCollection PlatformProjectGenerators)
         {
             // This can be reset by higher level code when it detects that we don't have
             // VS2015 installed (TODO - add custom format for Mac?)
@@ -49,7 +49,7 @@ namespace UnrealBuildTool
             // write out OtherProjectFiles and AutomationProjectFiles
             GeneratedProjectFiles.Clear();
 
-            if (!base.WriteProjectFiles())
+            if (!base.WriteProjectFiles(PlatformProjectGenerators))
             {
                 return false;
             }
diff --git a/Engine/Source/Programs/UnrealBuildTool/ProjectFiles/Xcode/XcodeProject.cs b/Engine/Source/Programs/UnrealBuildTool/ProjectFiles/Xcode/XcodeProject.cs
--- a/Engine/Source/Programs/UnrealBuildTool/ProjectFiles/Xcode/XcodeProject.cs
+++ b/Engine/Source/Programs/UnrealBuildTool/ProjectFiles/Xcode/XcodeProject.cs
@@ -1038,7 +1038,7 @@ namespace UnrealBuildTool
 			public UnrealTargetConfiguration BuildConfig;
 		};
 
-		private List<XcodeBuildConfig> GetSupportedBuildConfigs(List<UnrealTargetPlatform> Platforms, List<UnrealTargetConfiguration> Configurations)
+		private List<XcodeBuildConfig> GetSupportedBuildConfigs(List<UnrealTargetPlatform> Platforms, List<UnrealTargetConfiguration> Configurations, PlatformProjectGeneratorCollection PlatformProjectGenerators)
 		{
 			List<XcodeBuildConfig> BuildConfigs = new List<XcodeBuildConfig>();
 
@@ -1063,7 +1063,7 @@ namespace UnrealBuildTool
 
 								foreach (ProjectTarget ProjectTarget in ProjectTargets)
 								{
-									if (MSBuildProjectFile.IsValidProjectPlatformAndConfiguration(ProjectTarget, Platform, Configuration))
+									if (MSBuildProjectFile.IsValidProjectPlatformAndConfiguration(ProjectTarget, Platform, Configuration, PlatformProjectGenerators))
 									{
 										// Figure out if this is a monolithic build
 										bool bShouldCompileMonolithic = BuildPlatform.ShouldCompileMonolithicBinary(Platform);
@@ -1349,7 +1349,7 @@ namespace UnrealBuildTool
 		}
 
 		/// Implements Project interface
-		public override bool WriteProjectFile(List<UnrealTargetPlatform> InPlatforms, List<UnrealTargetConfiguration> InConfigurations)
+		public override bool WriteProjectFile(List<UnrealTargetPlatform> InPlatforms, List<UnrealTargetConfiguration> InConfigurations, PlatformProjectGeneratorCollection PlatformProjectGenerators)
 		{
 			bool bSuccess = true;
 
@@ -1372,7 +1372,7 @@ namespace UnrealBuildTool
 			string SourcesBuildPhaseGuid = XcodeProjectFileGenerator.MakeXcodeGuid();
 
 			// Figure out all the desired configurations
-			List<XcodeBuildConfig> BuildConfigs = GetSupportedBuildConfigs(InPlatforms, InConfigurations);
+			List<XcodeBuildConfig> BuildConfigs = GetSupportedBuildConfigs(InPlatforms, InConfigurations, PlatformProjectGenerators);
 			if (BuildConfigs.Count == 0)
 			{
 				return true;
diff --git a/Engine/Source/Programs/UnrealBuildTool/ProjectFiles/Xcode/XcodeProjectFileGenerator.cs b/Engine/Source/Programs/UnrealBuildTool/ProjectFiles/Xcode/XcodeProjectFileGenerator.cs
--- a/Engine/Source/Programs/UnrealBuildTool/ProjectFiles/Xcode/XcodeProjectFileGenerator.cs
+++ b/Engine/Source/Programs/UnrealBuildTool/ProjectFiles/Xcode/XcodeProjectFileGenerator.cs
@@ -193,7 +193,7 @@ namespace UnrealBuildTool
 			return bSuccess;
 		}
 
-		protected override bool WriteMasterProjectFile(ProjectFile UBTProject)
+		protected override bool WriteMasterProjectFile(ProjectFile UBTProject, PlatformProjectGeneratorCollection PlatformProjectGenerators)
 		{
 			return WriteXcodeWorkspace();
 		}
diff --git a/Engine/Source/Programs/UnrealBuildTool/UnrealBuildTool.cs b/Engine/Source/Programs/UnrealBuildTool/UnrealBuildTool.cs
--- a/Engine/Source/Programs/UnrealBuildTool/UnrealBuildTool.cs
+++ b/Engine/Source/Programs/UnrealBuildTool/UnrealBuildTool.cs
@@ -235,26 +235,6 @@ namespace UnrealBuildTool
 						}
 					}
 				}
-				// register all the other classes
-				foreach (Type CheckType in AllTypes)
-				{
-					if (CheckType.IsClass && !CheckType.IsAbstract)
-					{
-						if (CheckType.IsSubclassOf(typeof(PlatformProjectGenerator)))
-						{
-							ProjectGeneratorList.Add(CheckType);
-						}
-					}
-				}
-
-				// We need to process the ProjectGenerators AFTER all BuildPlatforms have been 
-				// registered as they use the BuildPlatform to verify they are legitimate.
-				foreach (Type ProjType in ProjectGeneratorList)
-				{
-					Log.TraceVerbose("    Registering project generator: {0}", ProjType.ToString());
-					PlatformProjectGenerator TempInst = (PlatformProjectGenerator)(UBTAssembly.CreateInstance(ProjType.FullName, true));
-					TempInst.RegisterPlatformProjectGenerator();
-				}
 			}
 		}
 
@@ -747,21 +727,6 @@ namespace UnrealBuildTool
 			}
 		}
 
-		/// <summary>
-		/// Generates project files.  Can also be used to generate projects "in memory", to allow for building
-		/// targets without even having project files at all.
-		/// </summary>
-		/// <param name="Generator">Project generator to use</param>
-		/// <param name="Arguments">Command-line arguments</param>
-		/// <returns>True if successful</returns>
-		internal static bool GenerateProjectFiles(ProjectFileGenerator Generator, string[] Arguments)
-		{
-			ProjectFileGenerator.bGenerateProjectFiles = true;
-			bool bSuccess = Generator.GenerateProjectFiles(Arguments);
-			ProjectFileGenerator.bGenerateProjectFiles = false;
-			return bSuccess;
-		}
-
 		internal static ECompilationResult RunUBT(BuildConfiguration BuildConfiguration, string[] Arguments, FileReference ProjectFile, bool bCatchExceptions)
 		{
 			bool bSuccess = true;
diff --git a/Engine/Source/Programs/UnrealBuildTool/UnrealBuildTool.csproj b/Engine/Source/Programs/UnrealBuildTool/UnrealBuildTool.csproj
--- a/Engine/Source/Programs/UnrealBuildTool/UnrealBuildTool.csproj
+++ b/Engine/Source/Programs/UnrealBuildTool/UnrealBuildTool.csproj
@@ -168,6 +168,7 @@
     <Compile Include="ProjectFiles\KDevelop\KDevelopGenerator.cs" />
     <Compile Include="ProjectFiles\Make\MakefileGenerator.cs" />
     <Compile Include="ProjectFiles\PlatformProjectGenerator.cs" />
+    <Compile Include="ProjectFiles\PlatformProjectGeneratorCollection.cs" />
     <Compile Include="ProjectFiles\Project.cs" />
     <Compile Include="ProjectFiles\ProjectFileGenerator.cs" />
     <Compile Include="ProjectFiles\QMake\QMakefileGenerator.cs" />
