UBT: Remove unused function.

#rnx

diff --git a/Engine/Source/Programs/UnrealBuildTool/Configuration/UEBuildModule.cs b/Engine/Source/Programs/UnrealBuildTool/Configuration/UEBuildModule.cs
--- a/Engine/Source/Programs/UnrealBuildTool/Configuration/UEBuildModule.cs
+++ b/Engine/Source/Programs/UnrealBuildTool/Configuration/UEBuildModule.cs
@@ -495,16 +495,6 @@ namespace UnrealBuildTool
 			}
 		}
 
-		static Regex VCMacroRegex = new Regex(@"\$\([A-Za-z0-9_]+\)");
-
-		/// <summary>
-		/// Checks if path contains a VC macro
-		/// </summary>
-		protected bool DoesPathContainVCMacro(string Path)
-		{
-			return VCMacroRegex.IsMatch(Path);
-		}
-
 		/// <summary>
 		/// Sets up the environment for compiling this module.
 		/// </summary>
