UBT: Allow specifying additional compiler arguments via the command line.

diff --git a/Engine/Source/Programs/UnrealBuildTool/Configuration/TargetRules.cs b/Engine/Source/Programs/UnrealBuildTool/Configuration/TargetRules.cs
--- a/Engine/Source/Programs/UnrealBuildTool/Configuration/TargetRules.cs
+++ b/Engine/Source/Programs/UnrealBuildTool/Configuration/TargetRules.cs
@@ -1054,12 +1054,14 @@ namespace UnrealBuildTool
 		/// Additional arguments to pass to the compiler
 		/// </summary>
 		[RequiresUniqueBuildEnvironment]
+		[CommandLine("-CompilerArguments=")]
 		public string AdditionalCompilerArguments;
 
 		/// <summary>
 		/// Additional arguments to pass to the linker
 		/// </summary>
 		[RequiresUniqueBuildEnvironment]
+		[CommandLine("-LinkerArguments=")]
 		public string AdditionalLinkerArguments;
 
 		/// <summary>
