Redo changes that were (deliberately) clobbered in Dev-Main merge.

diff --git a/Engine/Source/Programs/DotNETCommon/DotNETUtilities/ResXResource/UEResXReader.cs b/Engine/Source/Programs/DotNETCommon/DotNETUtilities/ResXResource/UEResXReader.cs
--- a/Engine/Source/Programs/DotNETCommon/DotNETUtilities/ResXResource/UEResXReader.cs
+++ b/Engine/Source/Programs/DotNETCommon/DotNETUtilities/ResXResource/UEResXReader.cs
@@ -6,7 +6,7 @@ using System.Linq;
 using System.Text;
 using System.Threading.Tasks;
 
-namespace Tools.DotNETUtilities
+namespace Tools.DotNETCommon
 {
 	/// <summary>
 	/// Intentionally unimplemented, but declared so that it's type can be included in generated resx files
diff --git a/Engine/Source/Programs/DotNETCommon/DotNETUtilities/ResXResource/UEResXWriter.cs b/Engine/Source/Programs/DotNETCommon/DotNETUtilities/ResXResource/UEResXWriter.cs
--- a/Engine/Source/Programs/DotNETCommon/DotNETUtilities/ResXResource/UEResXWriter.cs
+++ b/Engine/Source/Programs/DotNETCommon/DotNETUtilities/ResXResource/UEResXWriter.cs
@@ -7,7 +7,7 @@ using System.Text;
 using System.Threading.Tasks;
 using System.Xml;
 
-namespace Tools.DotNETUtilities
+namespace Tools.DotNETCommon
 {
 	public class UEResXWriter
 	{
diff --git a/Engine/Source/Programs/UnrealBuildTool/ProjectFiles/VisualStudio/VCProject.cs b/Engine/Source/Programs/UnrealBuildTool/ProjectFiles/VisualStudio/VCProject.cs
--- a/Engine/Source/Programs/UnrealBuildTool/ProjectFiles/VisualStudio/VCProject.cs
+++ b/Engine/Source/Programs/UnrealBuildTool/ProjectFiles/VisualStudio/VCProject.cs
@@ -108,15 +108,16 @@ namespace UnrealBuildTool
 		/// <param name="Platform">Actual platform</param>
 		/// <param name="Configuration">Actual configuration</param>
 		/// <param name="TargetConfigurationName">The configuration name from the target rules, or null if we don't have one</param>
+		/// <param name="PlatformProjectGenerators">Set of platform project generators</param>
 		/// <param name="ProjectPlatformName">Name of platform string to use for Visual Studio project</param>
 		/// <param name="ProjectConfigurationName">Name of configuration string to use for Visual Studio project</param>
-		public abstract void MakeProjectPlatformAndConfigurationNames(UnrealTargetPlatform Platform, UnrealTargetConfiguration Configuration, TargetType TargetConfigurationName, out string ProjectPlatformName, out string ProjectConfigurationName);
+		public abstract void MakeProjectPlatformAndConfigurationNames(UnrealTargetPlatform Platform, UnrealTargetConfiguration Configuration, TargetType TargetConfigurationName, PlatformProjectGeneratorCollection PlatformProjectGenerators, out string ProjectPlatformName, out string ProjectConfigurationName);
 
-		public string MakeConfigurationAndPlatformPair(UnrealTargetPlatform Platform, UnrealTargetConfiguration Configuration, TargetType TargetConfigurationName)
+		public string MakeConfigurationAndPlatformPair(UnrealTargetPlatform Platform, UnrealTargetConfiguration Configuration, TargetType TargetConfigurationName, PlatformProjectGeneratorCollection PlatformProjectGenerators)
 		{
 			string ProjectPlatformName;
 			string ProjectConfigName;
-			MakeProjectPlatformAndConfigurationNames(Platform, Configuration, TargetConfigurationName, out ProjectPlatformName, out ProjectConfigName);
+			MakeProjectPlatformAndConfigurationNames(Platform, Configuration, TargetConfigurationName, PlatformProjectGenerators, out ProjectPlatformName, out ProjectConfigName);
 			return String.Format("{0}|{1}", ProjectConfigName, ProjectPlatformName);
 		}
 
@@ -149,8 +150,9 @@ namespace UnrealBuildTool
 		/// <param name="ProjectTarget">The target that we're checking for a valid platform/config combination</param>
 		/// <param name="Platform">Platform</param>
 		/// <param name="Configuration">Configuration</param>
+		/// <param name="PlatformProjectGenerators">Set of platform project generators</param>
 		/// <returns>True if this is a valid combination for this project, otherwise false</returns>
-		public static bool IsValidProjectPlatformAndConfiguration(ProjectTarget ProjectTarget, UnrealTargetPlatform Platform, UnrealTargetConfiguration Configuration)
+		public static bool IsValidProjectPlatformAndConfiguration(ProjectTarget ProjectTarget, UnrealTargetPlatform Platform, UnrealTargetConfiguration Configuration, PlatformProjectGeneratorCollection PlatformProjectGenerators)
 		{
 			if (!ProjectFileGenerator.bIncludeTestAndShippingConfigs)
 			{
@@ -160,7 +162,7 @@ namespace UnrealBuildTool
 				}
 			}
 
-			UEPlatformProjectGenerator PlatformProjectGenerator = UEPlatformProjectGenerator.GetPlatformProjectGenerator(Platform, true);
+			PlatformProjectGenerator PlatformProjectGenerator = PlatformProjectGenerators.GetPlatformProjectGenerator(Platform, true);
 			if (PlatformProjectGenerator == null)
 			{
 				return false;
@@ -280,11 +282,12 @@ namespace UnrealBuildTool
 		/// <param name="Platform">Actual platform</param>
 		/// <param name="Configuration">Actual configuration</param>
 		/// <param name="TargetConfigurationName">The configuration name from the target rules, or null if we don't have one</param>
+		/// <param name="PlatformProjectGenerators">Set of platform project generators</param>
 		/// <param name="ProjectPlatformName">Name of platform string to use for Visual Studio project</param>
 		/// <param name="ProjectConfigurationName">Name of configuration string to use for Visual Studio project</param>
-		public override void MakeProjectPlatformAndConfigurationNames(UnrealTargetPlatform Platform, UnrealTargetConfiguration Configuration, TargetType TargetConfigurationName, out string ProjectPlatformName, out string ProjectConfigurationName)
+		public override void MakeProjectPlatformAndConfigurationNames(UnrealTargetPlatform Platform, UnrealTargetConfiguration Configuration, TargetType TargetConfigurationName, PlatformProjectGeneratorCollection PlatformProjectGenerators, out string ProjectPlatformName, out string ProjectConfigurationName)
 		{
-			UEPlatformProjectGenerator PlatformProjectGenerator = UEPlatformProjectGenerator.GetPlatformProjectGenerator(Platform, bInAllowFailure: true);
+			PlatformProjectGenerator PlatformProjectGenerator = PlatformProjectGenerators.GetPlatformProjectGenerator(Platform, bInAllowFailure: true);
 
 			// Check to see if this platform is supported directly by Visual Studio projects.
 			bool HasActualVSPlatform = (PlatformProjectGenerator != null) ? PlatformProjectGenerator.HasVisualStudioSupport(Platform, Configuration, ProjectFileFormat) : false;
@@ -353,7 +356,7 @@ namespace UnrealBuildTool
 
 		List<ProjectConfigAndTargetCombination> ProjectConfigAndTargetCombinations = new List<ProjectConfigAndTargetCombination>();
 
-		private void BuildProjectConfigAndTargetCombinations(List<UnrealTargetPlatform> InPlatforms, List<UnrealTargetConfiguration> InConfigurations)
+		private void BuildProjectConfigAndTargetCombinations(List<UnrealTargetPlatform> InPlatforms, List<UnrealTargetConfiguration> InConfigurations, PlatformProjectGeneratorCollection PlatformProjectGenerators)
 		{
 			//no need to do this more than once
 			if(ProjectConfigAndTargetCombinations.Count > 0)
@@ -397,10 +400,10 @@ namespace UnrealBuildTool
 
 							foreach (ProjectTarget ProjectTarget in ProjectTargets)
 							{
-								if (IsValidProjectPlatformAndConfiguration(ProjectTarget, Platform, Configuration))
+								if (IsValidProjectPlatformAndConfiguration(ProjectTarget, Platform, Configuration, PlatformProjectGenerators))
 								{
 									string ProjectPlatformName, ProjectConfigurationName;
-									MakeProjectPlatformAndConfigurationNames(Platform, Configuration, ProjectTarget.TargetRules.Type, out ProjectPlatformName, out ProjectConfigurationName);
+									MakeProjectPlatformAndConfigurationNames(Platform, Configuration, ProjectTarget.TargetRules.Type, PlatformProjectGenerators, out ProjectPlatformName, out ProjectConfigurationName);
 
 									ProjectConfigAndTargetCombination Combination = new ProjectConfigAndTargetCombination(Platform, Configuration, ProjectPlatformName, ProjectConfigurationName, ProjectTarget);
 									ProjectConfigAndTargetCombinations.Add(Combination);
@@ -417,14 +420,14 @@ namespace UnrealBuildTool
 		/// If found writes a debug project file to disk
 		/// </summary>
 		/// <returns>True on success</returns>
-		public override List<Tuple<ProjectFile, string>> WriteDebugProjectFiles(List<UnrealTargetPlatform> InPlatforms, List<UnrealTargetConfiguration> InConfigurations)
+		public override List<Tuple<ProjectFile, string>> WriteDebugProjectFiles(List<UnrealTargetPlatform> InPlatforms, List<UnrealTargetConfiguration> InConfigurations, PlatformProjectGeneratorCollection PlatformProjectGenerators)
 		{
 			string ProjectName = ProjectFilePath.GetFileNameWithoutExtension();
 
 			List<UnrealTargetPlatform> ProjectPlatforms = new List<UnrealTargetPlatform>();
 			List<Tuple<ProjectFile, string>> ProjectFiles = new List<Tuple<ProjectFile, string>>();
 
-			BuildProjectConfigAndTargetCombinations(InPlatforms, InConfigurations);
+			BuildProjectConfigAndTargetCombinations(InPlatforms, InConfigurations, PlatformProjectGenerators);
 
 
 			foreach (ProjectConfigAndTargetCombination Combination in ProjectConfigAndTargetCombinations)
@@ -440,7 +443,7 @@ namespace UnrealBuildTool
 			{
 				foreach (UnrealTargetPlatform Platform in ProjectPlatforms)
 				{
-					UEPlatformProjectGenerator ProjGenerator = UEPlatformProjectGenerator.GetPlatformProjectGenerator(Platform, true);
+					PlatformProjectGenerator ProjGenerator = PlatformProjectGenerators.GetPlatformProjectGenerator(Platform, true);
 					if (ProjGenerator != null)
 					{
 						//write out additional prop file
@@ -463,7 +466,7 @@ namespace UnrealBuildTool
 		}
 
 		/// Implements Project interface
-		public override bool WriteProjectFile(List<UnrealTargetPlatform> InPlatforms, List<UnrealTargetConfiguration> InConfigurations)
+		public override bool WriteProjectFile(List<UnrealTargetPlatform> InPlatforms, List<UnrealTargetConfiguration> InConfigurations, PlatformProjectGeneratorCollection PlatformProjectGenerators)
 		{
 			string ProjectName = ProjectFilePath.GetFileNameWithoutExtension();
 
@@ -513,14 +516,14 @@ namespace UnrealBuildTool
 			VCProjectFileContent.AppendLine("<?xml version=\"1.0\" encoding=\"utf-8\"?>");
 			VCProjectFileContent.AppendLine("<Project DefaultTargets=\"Build\" ToolsVersion=\"{0}\" xmlns=\"http://schemas.microsoft.com/developer/msbuild/2003\">", VCProjectFileGenerator.GetProjectFileToolVersionString(ProjectFileFormat));
 
-			bool bGenerateUserFileContent = UEPlatformProjectGenerator.PlatformRequiresVSUserFileGeneration(InPlatforms, InConfigurations);
+			bool bGenerateUserFileContent = PlatformProjectGenerators.PlatformRequiresVSUserFileGeneration(InPlatforms, InConfigurations);
 			if (bGenerateUserFileContent)
 			{
 				VCUserFileContent.AppendLine("<?xml version=\"1.0\" encoding=\"utf-8\"?>");
 				VCUserFileContent.AppendLine("<Project ToolsVersion=\"{0}\" xmlns=\"http://schemas.microsoft.com/developer/msbuild/2003\">", VCProjectFileGenerator.GetProjectFileToolVersionString(ProjectFileFormat));
 			}
 
-			BuildProjectConfigAndTargetCombinations(InPlatforms, InConfigurations);
+			BuildProjectConfigAndTargetCombinations(InPlatforms, InConfigurations, PlatformProjectGenerators);
 
 			VCProjectFileContent.AppendLine("  <ItemGroup Label=\"ProjectConfigurations\">");
 
@@ -574,7 +577,7 @@ namespace UnrealBuildTool
 			{
 				foreach (UnrealTargetPlatform Platform in ProjectPlatforms)
 				{
-					UEPlatformProjectGenerator ProjGenerator = UEPlatformProjectGenerator.GetPlatformProjectGenerator(Platform, true);
+					PlatformProjectGenerator ProjGenerator = PlatformProjectGenerators.GetPlatformProjectGenerator(Platform, true);
 					if (ProjGenerator != null && ProjGenerator.HasVisualStudioSupport(Platform, UnrealTargetConfiguration.Development, ProjectFileFormat))
 					{
 						ProjGenerator.GetAdditionalVisualStudioPropertyGroups(Platform, ProjectFileFormat, VCProjectFileContent);
@@ -599,7 +602,7 @@ namespace UnrealBuildTool
 			{
 				foreach (UnrealTargetPlatform Platform in ProjectPlatforms)
 				{
-					UEPlatformProjectGenerator ProjGenerator = UEPlatformProjectGenerator.GetPlatformProjectGenerator(Platform, true);
+					PlatformProjectGenerator ProjGenerator = PlatformProjectGenerators.GetPlatformProjectGenerator(Platform, true);
 					if (ProjGenerator != null && ProjGenerator.HasVisualStudioSupport(Platform, UnrealTargetConfiguration.Development, ProjectFileFormat))
 					{
 						ProjGenerator.GetVisualStudioGlobalProperties(Platform, VCProjectFileContent);
@@ -626,7 +629,7 @@ namespace UnrealBuildTool
 				{
 					string ProjectPlatformName = PlatformTuple.Item1;
 					UnrealTargetPlatform TargetPlatform = PlatformTuple.Item2;
-					WritePreDefaultPropsConfiguration(TargetPlatform, TargetConfiguration, ProjectPlatformName, ProjectConfigurationName, VCProjectFileContent);
+					WritePreDefaultPropsConfiguration(TargetPlatform, TargetConfiguration, ProjectPlatformName, ProjectConfigurationName, PlatformProjectGenerators, VCProjectFileContent);
 				}
 			}
 
@@ -646,7 +649,7 @@ namespace UnrealBuildTool
 				{
 					string ProjectPlatformName = PlatformTuple.Item1;
 					UnrealTargetPlatform TargetPlatform = PlatformTuple.Item2;
-					WritePostDefaultPropsConfiguration(TargetPlatform, TargetConfiguration, ProjectPlatformName, ProjectConfigurationName, VCProjectFileContent);
+					WritePostDefaultPropsConfiguration(TargetPlatform, TargetConfiguration, ProjectPlatformName, ProjectConfigurationName, PlatformProjectGenerators, VCProjectFileContent);
 				}
 			}
 
@@ -665,7 +668,7 @@ namespace UnrealBuildTool
 			// Write each project configuration
 			foreach (ProjectConfigAndTargetCombination Combination in ProjectConfigAndTargetCombinations)
 			{
-				WriteConfiguration(ProjectName, Combination, VCProjectFileContent, bGenerateUserFileContent ? VCUserFileContent : null);
+				WriteConfiguration(ProjectName, Combination, VCProjectFileContent, PlatformProjectGenerators, bGenerateUserFileContent ? VCUserFileContent : null);
 			}
 
 			// Source folders and files
@@ -815,7 +818,7 @@ namespace UnrealBuildTool
 			{
 				foreach (UnrealTargetPlatform Platform in ProjectPlatforms)
 				{
-					UEPlatformProjectGenerator ProjGenerator = UEPlatformProjectGenerator.GetPlatformProjectGenerator(Platform, true);
+					PlatformProjectGenerator ProjGenerator = PlatformProjectGenerators.GetPlatformProjectGenerator(Platform, true);
 					if (ProjGenerator != null && ProjGenerator.HasVisualStudioSupport(Platform, UnrealTargetConfiguration.Development, ProjectFileFormat))
 					{
 						// @todo projectfiles: Serious hacks here because we are trying to emit one-time platform-specific sections that need information
@@ -850,7 +853,7 @@ namespace UnrealBuildTool
 			{
 				foreach (UnrealTargetPlatform Platform in ProjectPlatforms)
 				{
-					UEPlatformProjectGenerator ProjGenerator = UEPlatformProjectGenerator.GetPlatformProjectGenerator(Platform, true);
+					PlatformProjectGenerator ProjGenerator = PlatformProjectGenerators.GetPlatformProjectGenerator(Platform, true);
 					if (ProjGenerator != null && ProjGenerator.HasVisualStudioSupport(Platform, UnrealTargetConfiguration.Development, ProjectFileFormat))
 					{
 						ProjGenerator.GetVisualStudioTargetOverrides(Platform, ProjectFileFormat, VCProjectFileContent);
@@ -1010,9 +1013,9 @@ namespace UnrealBuildTool
 		}
 
 		// Anonymous function that writes pre-Default.props configuration data
-		private void WritePreDefaultPropsConfiguration(UnrealTargetPlatform TargetPlatform, UnrealTargetConfiguration TargetConfiguration, string ProjectPlatformName, string ProjectConfigurationName, StringBuilder VCProjectFileContent)
+		private void WritePreDefaultPropsConfiguration(UnrealTargetPlatform TargetPlatform, UnrealTargetConfiguration TargetConfiguration, string ProjectPlatformName, string ProjectConfigurationName, PlatformProjectGeneratorCollection PlatformProjectGenerators, StringBuilder VCProjectFileContent)
 		{
-			UEPlatformProjectGenerator ProjGenerator = UEPlatformProjectGenerator.GetPlatformProjectGenerator(TargetPlatform, true);
+			PlatformProjectGenerator ProjGenerator = PlatformProjectGenerators.GetPlatformProjectGenerator(TargetPlatform, true);
 			if (((ProjGenerator == null) && (TargetPlatform != UnrealTargetPlatform.Unknown)))
 			{
 				return;
@@ -1036,9 +1039,9 @@ namespace UnrealBuildTool
 		}
 
 		// Anonymous function that writes post-Default.props configuration data
-		private void WritePostDefaultPropsConfiguration(UnrealTargetPlatform TargetPlatform, UnrealTargetConfiguration TargetConfiguration, string ProjectPlatformName, string ProjectConfigurationName, StringBuilder VCProjectFileContent)
+		private void WritePostDefaultPropsConfiguration(UnrealTargetPlatform TargetPlatform, UnrealTargetConfiguration TargetConfiguration, string ProjectPlatformName, string ProjectConfigurationName, PlatformProjectGeneratorCollection PlatformProjectGenerators, StringBuilder VCProjectFileContent)
 		{
-			UEPlatformProjectGenerator ProjGenerator = UEPlatformProjectGenerator.GetPlatformProjectGenerator(TargetPlatform, true);
+			PlatformProjectGenerator ProjGenerator = PlatformProjectGenerators.GetPlatformProjectGenerator(TargetPlatform, true);
 			if (((ProjGenerator == null) && (TargetPlatform != UnrealTargetPlatform.Unknown)))
 			{
 				return;
@@ -1065,12 +1068,12 @@ namespace UnrealBuildTool
 		}
 
 		// Anonymous function that writes project configuration data
-		private void WriteConfiguration(string ProjectName, ProjectConfigAndTargetCombination Combination, StringBuilder VCProjectFileContent, StringBuilder VCUserFileContent)
+		private void WriteConfiguration(string ProjectName, ProjectConfigAndTargetCombination Combination, StringBuilder VCProjectFileContent, PlatformProjectGeneratorCollection PlatformProjectGenerators, StringBuilder VCUserFileContent)
 		{
 			UnrealTargetPlatform Platform = Combination.Platform;
 			UnrealTargetConfiguration Configuration = Combination.Configuration;
 
-			UEPlatformProjectGenerator ProjGenerator = UEPlatformProjectGenerator.GetPlatformProjectGenerator(Platform, true);
+			PlatformProjectGenerator ProjGenerator = PlatformProjectGenerators.GetPlatformProjectGenerator(Platform, true);
 			if (((ProjGenerator == null) && (Platform != UnrealTargetPlatform.Unknown)))
 			{
 				return;
@@ -1245,7 +1248,7 @@ namespace UnrealBuildTool
 					if (TargetRulesObject.Type == TargetType.Game || TargetRulesObject.Type == TargetType.Client || TargetRulesObject.Type == TargetType.Server)
 					{
 						// Allow platforms to add any special properties they require... like aumid override for Xbox One
-						UEPlatformProjectGenerator.GenerateGamePlatformSpecificProperties(Platform, Configuration, TargetRulesObject.Type, VCProjectFileContent, RootDirectory, TargetFilePath);
+						PlatformProjectGenerators.GenerateGamePlatformSpecificProperties(Platform, Configuration, TargetRulesObject.Type, VCProjectFileContent, RootDirectory, TargetFilePath);
 					}
 
 					VCProjectFileContent.AppendLine("  </PropertyGroup>");
@@ -1399,9 +1402,10 @@ namespace UnrealBuildTool
 		/// <param name="Platform">Actual platform</param>
 		/// <param name="Configuration">Actual configuration</param>
 		/// <param name="TargetConfigurationName">The configuration name from the target rules, or null if we don't have one</param>
+		/// <param name="PlatformProjectGenerators">Set of platform project generators</param>
 		/// <param name="ProjectPlatformName">Name of platform string to use for Visual Studio project</param>
 		/// <param name="ProjectConfigurationName">Name of configuration string to use for Visual Studio project</param>
-		public override void MakeProjectPlatformAndConfigurationNames(UnrealTargetPlatform Platform, UnrealTargetConfiguration Configuration, TargetType TargetConfigurationName, out string ProjectPlatformName, out string ProjectConfigurationName)
+		public override void MakeProjectPlatformAndConfigurationNames(UnrealTargetPlatform Platform, UnrealTargetConfiguration Configuration, TargetType TargetConfigurationName, PlatformProjectGeneratorCollection PlatformProjectGenerators, out string ProjectPlatformName, out string ProjectConfigurationName)
 		{
 			// Find the matching platform name
 			if(Platform == UnrealTargetPlatform.Win32 && Platforms.Contains("x86"))
@@ -1438,8 +1442,9 @@ namespace UnrealBuildTool
 		/// </summary>
 		/// <param name="InPlatforms">Not used.</param>
 		/// <param name="InConfigurations">Not Used.</param>
+		/// <param name="PlatformProjectGenerators">Set of platform project generators</param>
 		/// <returns>true if the opration was successful, false otherwise</returns>
-		public override bool WriteProjectFile(List<UnrealTargetPlatform> InPlatforms, List<UnrealTargetConfiguration> InConfigurations)
+		public override bool WriteProjectFile(List<UnrealTargetPlatform> InPlatforms, List<UnrealTargetConfiguration> InConfigurations, PlatformProjectGeneratorCollection PlatformProjectGenerators)
 		{
 			throw new BuildException("Support for writing C# projects from UnrealBuildTool has been removed.");
 		}
diff --git a/Engine/Source/Programs/UnrealBuildTool/ProjectFiles/VisualStudio/VCProjectFileGenerator.cs b/Engine/Source/Programs/UnrealBuildTool/ProjectFiles/VisualStudio/VCProjectFileGenerator.cs
--- a/Engine/Source/Programs/UnrealBuildTool/ProjectFiles/VisualStudio/VCProjectFileGenerator.cs
+++ b/Engine/Source/Programs/UnrealBuildTool/ProjectFiles/VisualStudio/VCProjectFileGenerator.cs
@@ -7,7 +7,6 @@ using System.IO;
 using System.Linq;
 using System.Xml.Linq;
 using Tools.DotNETCommon;
-using DotNETUtilities;
 
 namespace UnrealBuildTool
 {
@@ -92,8 +91,8 @@ namespace UnrealBuildTool
 		/// </summary>
 		/// <param name="InOnlyGameProject">The single project to generate project files for, or null</param>
 		/// <param name="InProjectFileFormat">Override the project file format to use</param>
-		/// <param name="InOverrideCompiler">Override the compiler version to use</param>
-		public VCProjectFileGenerator(FileReference InOnlyGameProject, VCProjectFileFormat InProjectFileFormat, WindowsCompiler InOverrideCompiler)
+		/// <param name="InArguments">Additional command line arguments</param>
+		public VCProjectFileGenerator(FileReference InOnlyGameProject, VCProjectFileFormat InProjectFileFormat, CommandLineArguments InArguments)
 			: base(InOnlyGameProject)
 		{
 			XmlConfig.ApplyTo(this);
@@ -103,11 +102,11 @@ namespace UnrealBuildTool
 				ProjectFileFormat = InProjectFileFormat;
 			}
 
-			if(InOverrideCompiler == WindowsCompiler.VisualStudio2015)
+			if(InArguments.HasOption("-2015"))
 			{
 				BuildToolOverride = "-2015";
 			}
-			else if(InOverrideCompiler == WindowsCompiler.VisualStudio2017)
+			else if(InArguments.HasOption("-2017"))
 			{
 				BuildToolOverride = "-2017";
 			}
@@ -360,9 +359,9 @@ namespace UnrealBuildTool
 		/// Writes the project files to disk
 		/// </summary>
 		/// <returns>True if successful</returns>
-		protected override bool WriteProjectFiles()
+		protected override bool WriteProjectFiles(PlatformProjectGeneratorCollection PlatformProjectGenerators)
 		{
-			if(!base.WriteProjectFiles())
+			if(!base.WriteProjectFiles(PlatformProjectGenerators))
 			{
 				return false;
 			}
@@ -394,7 +393,7 @@ namespace UnrealBuildTool
 		}
 
 
-		protected override bool WriteMasterProjectFile(ProjectFile UBTProject)
+		protected override bool WriteMasterProjectFile(ProjectFile UBTProject, PlatformProjectGeneratorCollection PlatformProjectGenerators)
 		{
 			bool bSuccess = true;
 
@@ -594,7 +593,7 @@ namespace UnrealBuildTool
 												// Figure out the set of valid target configuration names
 												foreach (ProjectTarget ProjectTarget in CurProject.ProjectTargets)
 												{
-													if (VCProjectFile.IsValidProjectPlatformAndConfiguration(ProjectTarget, CurPlatform, CurConfiguration))
+													if (VCProjectFile.IsValidProjectPlatformAndConfiguration(ProjectTarget, CurPlatform, CurConfiguration, PlatformProjectGenerators))
 													{
 														PlatformsValidForProjects.Add(CurPlatform);
 
@@ -688,7 +687,7 @@ namespace UnrealBuildTool
 									string ProjectConfigAndPlatformPair;
 									if(CurProject is VCSharpProjectFile)
 									{
-										ProjectConfigAndPlatformPair = CurProject.MakeConfigurationAndPlatformPair(SolutionConfigCombination.Platform, SolutionConfigCombination.Configuration, TargetType.Program);
+										ProjectConfigAndPlatformPair = CurProject.MakeConfigurationAndPlatformPair(SolutionConfigCombination.Platform, SolutionConfigCombination.Configuration, TargetType.Program, PlatformProjectGenerators);
 									}
 									else
 									{
@@ -717,7 +716,7 @@ namespace UnrealBuildTool
 										List<ProjectTarget> MatchingProjectTargets = new List<ProjectTarget>();
 										foreach (ProjectTarget ProjectTarget in CurProject.ProjectTargets)
 										{
-											if(VCProjectFile.IsValidProjectPlatformAndConfiguration(ProjectTarget, SolutionConfigCombination.Platform, SolutionConfigCombination.Configuration))
+											if(VCProjectFile.IsValidProjectPlatformAndConfiguration(ProjectTarget, SolutionConfigCombination.Platform, SolutionConfigCombination.Configuration, PlatformProjectGenerators))
 											{
 												if (ProjectTarget.TargetRules != null)
 												{
@@ -781,7 +780,7 @@ namespace UnrealBuildTool
 											}
 											else
 											{
-												ProjectConfigAndPlatformPair = CurProject.MakeConfigurationAndPlatformPair(ProjectPlatform, ProjectConfiguration, TargetConfigurationName);
+												ProjectConfigAndPlatformPair = CurProject.MakeConfigurationAndPlatformPair(ProjectPlatform, ProjectConfiguration, TargetConfigurationName, PlatformProjectGenerators);
 											}
 
 											// Set whether this project configuration should be built when the user initiates "build solution"
@@ -798,7 +797,7 @@ namespace UnrealBuildTool
 												{
 													bBuildByDefault = true;
 
-													UEPlatformProjectGenerator ProjGen = UEPlatformProjectGenerator.GetPlatformProjectGenerator(SolutionConfigCombination.Platform, true);
+													PlatformProjectGenerator ProjGen = PlatformProjectGenerators.GetPlatformProjectGenerator(SolutionConfigCombination.Platform, true);
 													if (MatchingProjectTarget.ProjectDeploys ||
 														((ProjGen != null) && (ProjGen.GetVisualStudioDeploymentEnabled(ProjectPlatform, ProjectConfiguration) == true)))
 													{
diff --git a/Engine/Source/Programs/UnrealBuildTool/System/SourceFileSearch.cs b/Engine/Source/Programs/UnrealBuildTool/System/SourceFileSearch.cs
--- a/Engine/Source/Programs/UnrealBuildTool/System/SourceFileSearch.cs
+++ b/Engine/Source/Programs/UnrealBuildTool/System/SourceFileSearch.cs
@@ -7,7 +7,6 @@ using System.IO;
 using System.Diagnostics;
 using System.Linq;
 using Tools.DotNETCommon;
-using DotNETUtilities;
 
 namespace UnrealBuildTool
 {
