Improve error message when a runtime dependency cannot be found.

diff --git a/Engine/Source/Programs/UnrealBuildTool/Configuration/UEBuildBinary.cs b/Engine/Source/Programs/UnrealBuildTool/Configuration/UEBuildBinary.cs
--- a/Engine/Source/Programs/UnrealBuildTool/Configuration/UEBuildBinary.cs
+++ b/Engine/Source/Programs/UnrealBuildTool/Configuration/UEBuildBinary.cs
@@ -282,7 +282,8 @@ namespace UnrealBuildTool
 						}
 						catch(FilePatternException Ex)
 						{
-							throw new BuildException(Ex, "While creating runtime dependencies for module {0}", Module.Name);
+							ExceptionUtils.AddContext(Ex, "while creating runtime dependencies for module '{0}'", Module.Name);
+							throw;
 						}
 
 						// Add actions to copy everything
