Missing files from resolve.

#rnx

diff --git a/Engine/Source/Programs/UnrealBuildTool/Modes/GenerateProjectFilesMode.cs b/Engine/Source/Programs/UnrealBuildTool/Modes/GenerateProjectFilesMode.cs
--- a/Engine/Source/Programs/UnrealBuildTool/Modes/GenerateProjectFilesMode.cs
+++ b/Engine/Source/Programs/UnrealBuildTool/Modes/GenerateProjectFilesMode.cs
@@ -56,7 +56,7 @@ namespace UnrealBuildTool
 				DirectoryReference.SetCurrentDirectory(UnrealBuildTool.EngineSourceDirectory);
 
 				// Read the XML configuration files
-				XmlConfig.ReadConfigFiles(false);
+				XmlConfig.ReadConfigFiles(null);
 
 				// Find and register all tool chains, build platforms, etc. that are present
 				UnrealBuildTool.RegisterAllUBTClasses(false);
diff --git a/Engine/Source/Programs/UnrealBuildTool/Modes/SetupPlatformsMode.cs b/Engine/Source/Programs/UnrealBuildTool/Modes/SetupPlatformsMode.cs
--- a/Engine/Source/Programs/UnrealBuildTool/Modes/SetupPlatformsMode.cs
+++ b/Engine/Source/Programs/UnrealBuildTool/Modes/SetupPlatformsMode.cs
@@ -26,7 +26,7 @@ namespace UnrealBuildTool
 				Arguments.CheckAllArgumentsUsed();
 
 				// Read the XML configuration files
-				XmlConfig.ReadConfigFiles(false);
+				XmlConfig.ReadConfigFiles(null);
 
 				// Find and register all tool chains, build platforms, etc. that are present
 				UnrealBuildTool.RegisterAllUBTClasses(false);
diff --git a/Engine/Source/Programs/UnrealBuildTool/Modes/ValidatePlatformsMode.cs b/Engine/Source/Programs/UnrealBuildTool/Modes/ValidatePlatformsMode.cs
--- a/Engine/Source/Programs/UnrealBuildTool/Modes/ValidatePlatformsMode.cs
+++ b/Engine/Source/Programs/UnrealBuildTool/Modes/ValidatePlatformsMode.cs
@@ -45,7 +45,7 @@ namespace UnrealBuildTool
 				}
 
 				// Read all the config files
-				XmlConfig.ReadConfigFiles(false);
+				XmlConfig.ReadConfigFiles(null);
 				
 				// Register all the platforms
 				UnrealBuildTool.RegisterAllUBTClasses(true);
