UBT: Move IsFileInstalled() onto the UnrealBuildTool class, and make it check against the installed project file.

diff --git a/Engine/Source/Programs/UnrealBuildTool/Configuration/UEBuildTarget.cs b/Engine/Source/Programs/UnrealBuildTool/Configuration/UEBuildTarget.cs
--- a/Engine/Source/Programs/UnrealBuildTool/Configuration/UEBuildTarget.cs
+++ b/Engine/Source/Programs/UnrealBuildTool/Configuration/UEBuildTarget.cs
@@ -1706,7 +1706,7 @@ namespace UnrealBuildTool
 				foreach (FileReference VersionManifestFile in FileReferenceToModuleManifestPairs.Select(x => x.Key))
 				{
 					// Make sure the file (and directory) exists before trying to delete it
-					if(FileReference.Exists(VersionManifestFile) && !IsFileInstalled(VersionManifestFile))
+					if(FileReference.Exists(VersionManifestFile) && !UnrealBuildTool.IsFileInstalled(VersionManifestFile))
 					{
 						FileReference.Delete(VersionManifestFile);
 					}
@@ -1748,7 +1748,7 @@ namespace UnrealBuildTool
 				UEBuildPlatform BuildPlatform = UEBuildPlatform.GetBuildPlatform(Platform);
 				if (OnlyModules == null || OnlyModules.Count == 0)
 				{
-					if(!IsFileInstalled(ReceiptFileName))
+					if(!UnrealBuildTool.IsFileInstalled(ReceiptFileName))
 					{
 						DirectoryReference.CreateDirectory(ReceiptFileName.Directory);
 						Receipt.Write(ReceiptFileName, UnrealBuildTool.EngineDirectory, ProjectDirectory);
@@ -1757,7 +1757,7 @@ namespace UnrealBuildTool
 				if (ForceReceiptFileName != null)
 				{
 					FileReference ForceReceiptFile = new FileReference(ForceReceiptFileName);
-					if(!IsFileInstalled(ForceReceiptFile))
+					if(!UnrealBuildTool.IsFileInstalled(ForceReceiptFile))
 					{
 						DirectoryReference.CreateDirectory(ForceReceiptFile.Directory);
 						Receipt.Write(ForceReceiptFile, UnrealBuildTool.EngineDirectory, ProjectDirectory);
@@ -1765,7 +1765,7 @@ namespace UnrealBuildTool
 				}
 				if(VersionFile != null)
 				{
-					if(!IsFileInstalled(VersionFile))
+					if(!UnrealBuildTool.IsFileInstalled(VersionFile))
 					{
 						DirectoryReference.CreateDirectory(VersionFile.Directory);
 
@@ -1784,7 +1784,7 @@ namespace UnrealBuildTool
 			{
 				foreach (KeyValuePair<FileReference, ModuleManifest> FileNameToVersionManifest in FileReferenceToModuleManifestPairs)
 				{
-					if(!IsFileInstalled(FileNameToVersionManifest.Key))
+					if(!UnrealBuildTool.IsFileInstalled(FileNameToVersionManifest.Key))
 					{
 						if(!FileReference.Exists(FileNameToVersionManifest.Key))
 						{
@@ -1830,7 +1830,7 @@ namespace UnrealBuildTool
 			{
 				if (OnlyModules == null || OnlyModules.Count == 0)
 				{
-					if(!IsFileInstalled(ReceiptFileName) && FileReference.Exists(ReceiptFileName))
+					if(!UnrealBuildTool.IsFileInstalled(ReceiptFileName) && FileReference.Exists(ReceiptFileName))
 					{
 						FileReference.Delete(ReceiptFileName);
 					}
@@ -1838,24 +1838,6 @@ namespace UnrealBuildTool
 			}
 		}
 
-		/// <summary>
-		/// Checks whether the given file is under an installed directory, and should not be overridden
-		/// </summary>
-		/// <param name="File">File to test</param>
-		/// <returns>True if the file is part of the installed distribution, false otherwise</returns>
-		bool IsFileInstalled(FileReference File)
-		{
-			if(UnrealBuildTool.IsEngineInstalled() && File.IsUnderDirectory(UnrealBuildTool.EngineDirectory))
-			{
-				return true;
-			}
-			if(UnrealBuildTool.IsProjectInstalled() && ProjectFile != null && File.IsUnderDirectory(ProjectFile.Directory))
-			{
-				return true;
-			}
-			return false;
-		}
-
 		/// <summary>
 		/// Gathers dependency modules for given binaries list.
 		/// </summary>
@@ -2229,7 +2211,7 @@ namespace UnrealBuildTool
 			}
 			foreach(KeyValuePair<FileReference, FileReference> Pair in RuntimeDependencyTargetFileToSourceFile)
 			{
-				if(!IsFileInstalled(Pair.Key))
+				if(!UnrealBuildTool.IsFileInstalled(Pair.Key))
 				{
 					OutputItems.Add(CreateCopyAction(Pair.Value, Pair.Key, ActionGraph));
 				}
diff --git a/Engine/Source/Programs/UnrealBuildTool/UnrealBuildTool.cs b/Engine/Source/Programs/UnrealBuildTool/UnrealBuildTool.cs
--- a/Engine/Source/Programs/UnrealBuildTool/UnrealBuildTool.cs
+++ b/Engine/Source/Programs/UnrealBuildTool/UnrealBuildTool.cs
@@ -47,6 +47,11 @@ namespace UnrealBuildTool
 		/// Whether we're running with an installed project
 		/// </summary>
 		static private bool? bIsProjectInstalled;
+		
+		/// <summary>
+		/// If we are running with an installed project, specifies the path to it
+		/// </summary>
+		static FileReference InstalledProjectFile;
 
 		/// <summary>
 		/// The full name of the Root UE4 directory
@@ -165,11 +170,39 @@ namespace UnrealBuildTool
 		{
 			if (!bIsProjectInstalled.HasValue)
 			{
-				bIsProjectInstalled = FileReference.Exists(FileReference.Combine(EngineDirectory, "Build", "InstalledProjectBuild.txt"));
+				FileReference InstalledProjectLocationFile = FileReference.Combine(UnrealBuildTool.RootDirectory, "Engine", "Build", "InstalledProjectBuild.txt");
+				if (FileReference.Exists(InstalledProjectLocationFile))
+				{
+					InstalledProjectFile = FileReference.Combine(UnrealBuildTool.RootDirectory, File.ReadAllText(InstalledProjectLocationFile.FullName).Trim());
+					bIsProjectInstalled = true;
+				}
+				else
+				{
+					InstalledProjectFile = null;
+					bIsProjectInstalled = false;
+				}
 			}
 			return bIsProjectInstalled.Value;
 		}
 
+		/// <summary>
+		/// Checks whether the given file is under an installed directory, and should not be overridden
+		/// </summary>
+		/// <param name="File">File to test</param>
+		/// <returns>True if the file is part of the installed distribution, false otherwise</returns>
+		static public bool IsFileInstalled(FileReference File)
+		{
+			if(IsEngineInstalled() && File.IsUnderDirectory(EngineDirectory))
+			{
+				return true;
+			}
+			if(IsProjectInstalled() && File.IsUnderDirectory(InstalledProjectFile.Directory))
+			{
+				return true;
+			}
+			return false;
+		}
+
 		/// <summary>
 		/// Gets the absolute path to the UBT assembly.
 		/// </summary>
@@ -265,12 +298,7 @@ namespace UnrealBuildTool
 
 			if(IsProjectInstalled())
 			{
-				FileReference InstalledProjectFile = FileReference.Combine(UnrealBuildTool.RootDirectory, "Engine", "Build", "InstalledProjectBuild.txt");
-				if (FileReference.Exists(InstalledProjectFile))
-				{
-					ProjectFile = FileReference.Combine(UnrealBuildTool.RootDirectory, File.ReadAllText(InstalledProjectFile.FullName).Trim());
-					return true;
-				}
+				ProjectFile = InstalledProjectFile;
 			}
 
 			ProjectFile = null;
@@ -337,6 +365,9 @@ namespace UnrealBuildTool
 		/// <param name="CmdLine">Cmdline arguments</param>
 		internal static int GuardedMain(CommandLineArguments CmdLine)
 		{
+			Console.WriteLine(IsFileInstalled(new FileReference(@"D:\P4 UE4\Templates\TP_Flying\Config\DefaultEditor.ini")));
+			Console.WriteLine(IsFileInstalled(new FileReference(@"D:\P4 UE4\Templates\TP_FlyingBP\Config\DefaultEditor.ini")));
+
 			DateTime StartTime = DateTime.UtcNow;
 			ECompilationResult Result = ECompilationResult.Succeeded;
 
