Remove debug code.

#rnx

diff --git a/Engine/Source/Programs/UnrealBuildTool/UnrealBuildTool.cs b/Engine/Source/Programs/UnrealBuildTool/UnrealBuildTool.cs
--- a/Engine/Source/Programs/UnrealBuildTool/UnrealBuildTool.cs
+++ b/Engine/Source/Programs/UnrealBuildTool/UnrealBuildTool.cs
@@ -365,9 +365,6 @@ namespace UnrealBuildTool
 		/// <param name="CmdLine">Cmdline arguments</param>
 		internal static int GuardedMain(CommandLineArguments CmdLine)
 		{
-			Console.WriteLine(IsFileInstalled(new FileReference(@"D:\P4 UE4\Templates\TP_Flying\Config\DefaultEditor.ini")));
-			Console.WriteLine(IsFileInstalled(new FileReference(@"D:\P4 UE4\Templates\TP_FlyingBP\Config\DefaultEditor.ini")));
-
 			DateTime StartTime = DateTime.UtcNow;
 			ECompilationResult Result = ECompilationResult.Succeeded;
 
