Move command line classes to DotNETCommon.

#rnx

diff --git a/Engine/Source/Programs/DotNETCommon/DotNETUtilities/DotNETUtilities.csproj b/Engine/Source/Programs/DotNETCommon/DotNETUtilities/DotNETUtilities.csproj
--- a/Engine/Source/Programs/DotNETCommon/DotNETUtilities/DotNETUtilities.csproj
+++ b/Engine/Source/Programs/DotNETCommon/DotNETUtilities/DotNETUtilities.csproj
@@ -57,6 +57,8 @@
     <Compile Include="AssemblyUtils.cs" />
     <Compile Include="BinaryReaderExtensions.cs" />
     <Compile Include="BinaryWriterExtensions.cs" />
+    <Compile Include="CommandLineArguments.cs" />
+    <Compile Include="CommandLineAttribute.cs" />
     <Compile Include="CsProjectInfo.cs" />
     <Compile Include="DirectoryReference.cs" />
     <Compile Include="ExceptionUtils.cs" />
diff --git a/Engine/Source/Programs/DotNETCommon/DotNETUtilities/StringUtils.cs b/Engine/Source/Programs/DotNETCommon/DotNETUtilities/StringUtils.cs
--- a/Engine/Source/Programs/DotNETCommon/DotNETUtilities/StringUtils.cs
+++ b/Engine/Source/Programs/DotNETCommon/DotNETUtilities/StringUtils.cs
@@ -45,5 +45,42 @@ namespace Tools.DotNETCommon
 
 			return true;
 		}
+
+		/// <summary>
+		/// Formats a list of strings in the style "1, 2, 3 and 4"
+		/// </summary>
+		/// <param name="Arguments">List of strings to format</param>
+		/// <returns>Formatted list of strings</returns>
+		public static string FormatList(string[] Arguments)
+		{
+			StringBuilder Result = new StringBuilder();
+			if(Arguments.Length > 0)
+			{
+				Result.Append(Arguments[0]);
+				for(int Idx = 1; Idx < Arguments.Length; Idx++)
+				{
+					if(Idx == Arguments.Length - 1)
+					{
+						Result.Append(" and ");
+					}
+					else
+					{
+						Result.Append(", ");
+					}
+					Result.Append(Arguments[Idx]);
+				}
+			}
+			return Result.ToString();
+		}
+
+		/// <summary>
+		/// Formats a list of strings in the style "1, 2, 3 and 4"
+		/// </summary>
+		/// <param name="Arguments">List of strings to format</param>
+		/// <returns>Formatted list of strings</returns>
+		public static string FormatList(IEnumerable<string> Arguments)
+		{
+			return FormatList(Arguments.ToArray());
+		}
 	}
 }
diff --git a/Engine/Source/Programs/UnrealBuildTool/System/CommandLine.cs b/Engine/Source/Programs/UnrealBuildTool/System/CommandLine.cs
--- a/Engine/Source/Programs/UnrealBuildTool/System/CommandLine.cs
+++ b/Engine/Source/Programs/UnrealBuildTool/System/CommandLine.cs
@@ -152,7 +152,7 @@ namespace UnrealBuildTool
 				}
 				else
 				{
-					throw new BuildException("Missing {0} arguments", Utils.FormatList(MissingFieldToParameter.Values.Select(x => x.Prefix.Replace("=", "=..."))));
+					throw new BuildException("Missing {0} arguments", StringUtils.FormatList(MissingFieldToParameter.Values.Select(x => x.Prefix.Replace("=", "=..."))));
 				}
 			}
 		}
diff --git a/Engine/Source/Programs/UnrealBuildTool/System/SingleInstanceMutex.cs b/Engine/Source/Programs/UnrealBuildTool/System/SingleInstanceMutex.cs
--- a/Engine/Source/Programs/UnrealBuildTool/System/SingleInstanceMutex.cs
+++ b/Engine/Source/Programs/UnrealBuildTool/System/SingleInstanceMutex.cs
@@ -5,6 +5,7 @@ using System.Reflection;
 using System.Text;
 using System.Threading;
 using System.Threading.Tasks;
+using Tools.DotNETCommon;
 
 namespace UnrealBuildTool
 {
diff --git a/Engine/Source/Programs/UnrealBuildTool/System/Utils.cs b/Engine/Source/Programs/UnrealBuildTool/System/Utils.cs
--- a/Engine/Source/Programs/UnrealBuildTool/System/Utils.cs
+++ b/Engine/Source/Programs/UnrealBuildTool/System/Utils.cs
@@ -1,4 +1,4 @@
-﻿// Copyright 1998-2018 Epic Games, Inc. All Rights Reserved.
+// Copyright 1998-2018 Epic Games, Inc. All Rights Reserved.
 
 using System;
 using System.Collections.Generic;
@@ -1067,42 +1067,5 @@ namespace UnrealBuildTool
 			}
 			return true;
 		}
-
-		/// <summary>
-		/// Formats a list of strings in the style "1, 2, 3 and 4"
-		/// </summary>
-		/// <param name="Arguments">List of strings to format</param>
-		/// <returns>Formatted list of strings</returns>
-		public static string FormatList(string[] Arguments)
-		{
-			StringBuilder Result = new StringBuilder();
-			if(Arguments.Length > 0)
-			{
-				Result.Append(Arguments[0]);
-				for(int Idx = 1; Idx < Arguments.Length; Idx++)
-				{
-					if(Idx == Arguments.Length - 1)
-					{
-						Result.Append(" and ");
-					}
-					else
-					{
-						Result.Append(", ");
-					}
-					Result.Append(Arguments[Idx]);
-				}
-			}
-			return Result.ToString();
-		}
-
-		/// <summary>
-		/// Formats a list of strings in the style "1, 2, 3 and 4"
-		/// </summary>
-		/// <param name="Arguments">List of strings to format</param>
-		/// <returns>Formatted list of strings</returns>
-		public static string FormatList(IEnumerable<string> Arguments)
-		{
-			return FormatList(Arguments.ToArray());
-		}
 	}
 }
diff --git a/Engine/Source/Programs/UnrealBuildTool/UnrealBuildTool.csproj b/Engine/Source/Programs/UnrealBuildTool/UnrealBuildTool.csproj
--- a/Engine/Source/Programs/UnrealBuildTool/UnrealBuildTool.csproj
+++ b/Engine/Source/Programs/UnrealBuildTool/UnrealBuildTool.csproj
@@ -180,7 +180,6 @@
     <Compile Include="ProjectFiles\Xcode\XcodeProjectFileGenerator.cs" />
     <Compile Include="System\Asn.cs" />
     <Compile Include="System\BuildVersion.cs" />
-    <Compile Include="System\CommandLineArguments.cs" />
     <Compile Include="System\CppCompileEnvironment.cs" />
     <Compile Include="System\CppIncludeParser.cs" />
     <Compile Include="System\DataDrivenPlatformInfo.cs" />
@@ -198,7 +197,6 @@
     <Compile Include="Configuration\UEBuildModuleCPP.cs" />
     <Compile Include="System\PlatformExports.cs" />
     <Compile Include="System\CommandLine.cs" />
-    <Compile Include="System\CommandLineAttribute.cs" />
     <Compile Include="System\ConfigCache.cs" />
     <Compile Include="System\ConfigFile.cs" />
     <Compile Include="System\ConfigFileAttribute.cs" />
