Remove the unused -Receipt= argument to force an additional receipt location for UBT, few other minor bits of cleanup.

#rnx

diff --git a/Engine/Source/Programs/UnrealBuildTool/Configuration/TargetDescriptor.cs b/Engine/Source/Programs/UnrealBuildTool/Configuration/TargetDescriptor.cs
--- a/Engine/Source/Programs/UnrealBuildTool/Configuration/TargetDescriptor.cs
+++ b/Engine/Source/Programs/UnrealBuildTool/Configuration/TargetDescriptor.cs
@@ -21,7 +21,6 @@ namespace UnrealBuildTool
 		public string Architecture;
 		public List<OnlyModule> OnlyModules;
 		public FileReference ForeignPlugin;
-		public string ForceReceiptFileName;
 		public string[] AdditionalArguments;
 
 		/// <summary>
@@ -60,7 +59,6 @@ namespace UnrealBuildTool
 			string Architecture = null;
 			List<OnlyModule> OnlyModules = new List<OnlyModule>();
 			FileReference ForeignPlugin = null;
-			string ForceReceiptFileName = null;
 			List<string> AdditionalArguments = new List<string>();
 
 			// Settings for creating/using static libraries for the engine
@@ -130,10 +128,6 @@ namespace UnrealBuildTool
 						}
 						ForeignPlugin = new FileReference(Value);
 					}
-					else if(ParseArgumentValue(Argument, "-Receipt=", out Value))
-					{
-						ForceReceiptFileName = Value;
-					}
 					else
 					{
 						switch (Arguments[ArgumentIndex].ToUpperInvariant())
@@ -195,7 +189,6 @@ namespace UnrealBuildTool
 				TargetDescriptor Target = new TargetDescriptor(ProjectFile, TargetName, Platform, Configuration, Architecture, AdditionalArguments.ToArray());
 				Target.OnlyModules = OnlyModules;
 				Target.ForeignPlugin = ForeignPlugin;
-				Target.ForceReceiptFileName = ForceReceiptFileName;
 				Targets.Add(Target);
 			}
 
diff --git a/Engine/Source/Programs/UnrealBuildTool/Configuration/UEBuildTarget.cs b/Engine/Source/Programs/UnrealBuildTool/Configuration/UEBuildTarget.cs
--- a/Engine/Source/Programs/UnrealBuildTool/Configuration/UEBuildTarget.cs
+++ b/Engine/Source/Programs/UnrealBuildTool/Configuration/UEBuildTarget.cs
@@ -792,12 +792,6 @@ namespace UnrealBuildTool
 		/// </summary>
 		private KeyValuePair<FileReference, ModuleManifest>[] FileReferenceToModuleManifestPairs;
 
-		/// <summary>
-		/// Force output of the receipt to an additional filename
-		/// </summary>
-		[NonSerialized]
-		private string ForceReceiptFileName;
-
 		/// <summary>
 		/// The name of the .Target.cs file, if the target was created with one
 		/// </summary>
@@ -932,7 +926,6 @@ namespace UnrealBuildTool
 			TargetType = Rules.Type;
 			bPrecompile = InRules.bPrecompile;
 			ForeignPlugin = InDesc.ForeignPlugin;
-			ForceReceiptFileName = InDesc.ForceReceiptFileName;
 
 			// now that we have the platform, we can set the intermediate path to include the platform/architecture name
 			PlatformIntermediateFolder = Path.Combine("Intermediate", "Build", Platform.ToString(), UEBuildPlatform.GetBuildPlatform(Platform).GetFolderNameForArchitecture(Architecture));
@@ -1754,15 +1747,6 @@ namespace UnrealBuildTool
 						Receipt.Write(ReceiptFileName, UnrealBuildTool.EngineDirectory, ProjectDirectory);
 					}
 				}
-				if (ForceReceiptFileName != null)
-				{
-					FileReference ForceReceiptFile = new FileReference(ForceReceiptFileName);
-					if(!UnrealBuildTool.IsFileInstalled(ForceReceiptFile))
-					{
-						DirectoryReference.CreateDirectory(ForceReceiptFile.Directory);
-						Receipt.Write(ForceReceiptFile, UnrealBuildTool.EngineDirectory, ProjectDirectory);
-					}
-				}
 				if(VersionFile != null)
 				{
 					if(!UnrealBuildTool.IsFileInstalled(VersionFile))
@@ -1789,8 +1773,8 @@ namespace UnrealBuildTool
 						if(!FileReference.Exists(FileNameToVersionManifest.Key))
 						{
 							// If the file doesn't already exist, just write it out
-							Directory.CreateDirectory(Path.GetDirectoryName(FileNameToVersionManifest.Key.FullName));
-							FileNameToVersionManifest.Value.Write(FileNameToVersionManifest.Key.FullName);
+							DirectoryReference.CreateDirectory(FileNameToVersionManifest.Key.Directory);
+							FileNameToVersionManifest.Value.Write(FileNameToVersionManifest.Key);
 						}
 						else
 						{
diff --git a/Engine/Source/Programs/UnrealBuildTool/System/ModuleManifest.cs b/Engine/Source/Programs/UnrealBuildTool/System/ModuleManifest.cs
--- a/Engine/Source/Programs/UnrealBuildTool/System/ModuleManifest.cs
+++ b/Engine/Source/Programs/UnrealBuildTool/System/ModuleManifest.cs
@@ -1,4 +1,4 @@
-﻿// Copyright 1998-2018 Epic Games, Inc. All Rights Reserved.
+// Copyright 1998-2018 Epic Games, Inc. All Rights Reserved.
 
 using System;
 using System.Collections.Generic;
@@ -118,10 +118,10 @@ namespace UnrealBuildTool
 		/// Write the receipt to disk.
 		/// </summary>
 		/// <param name="FileName">The file to write to</param>
-		public void Write(string FileName)
+		public void Write(FileReference FileName)
 		{
-			Directory.CreateDirectory(Path.GetDirectoryName(FileName));
-			using(StreamWriter Writer = new StreamWriter(FileName))
+			DirectoryReference.CreateDirectory(FileName.Directory);
+			using(StreamWriter Writer = new StreamWriter(FileName.FullName))
 			{
 				Write(Writer);
 			}
diff --git a/Engine/Source/Programs/UnrealBuildTool/System/Utils.cs b/Engine/Source/Programs/UnrealBuildTool/System/Utils.cs
--- a/Engine/Source/Programs/UnrealBuildTool/System/Utils.cs
+++ b/Engine/Source/Programs/UnrealBuildTool/System/Utils.cs
@@ -142,6 +142,16 @@ namespace UnrealBuildTool
 			return InPath;
 		}
 
+		/// <summary>
+		/// Makes sure path can be used as a command line param (adds quotes if it contains spaces)
+		/// </summary>
+		/// <param name="InPath">Path to convert</param>
+		/// <returns></returns>
+		public static string MakePathSafeToUseWithCommandLine(FileReference InPath)
+		{
+			return MakePathSafeToUseWithCommandLine(InPath.FullName);
+		}
+
 		/// <summary>
 		/// Escapes whitespace in the given command line argument with a backslash. Used on Unix-like platforms for command line arguments in shell commands.
 		/// </summary>
