UBT: Move generation of receipts and manifests, and execution of post-build steps, into the UBT action graph, to allow it to execute correctly when running with makefiles and when exported for other executors.

#rnx

diff --git a/Engine/Source/Programs/AutomationTool/AutomationUtils/UE4Build.cs b/Engine/Source/Programs/AutomationTool/AutomationUtils/UE4Build.cs
--- a/Engine/Source/Programs/AutomationTool/AutomationUtils/UE4Build.cs
+++ b/Engine/Source/Programs/AutomationTool/AutomationUtils/UE4Build.cs
@@ -639,13 +639,6 @@ namespace AutomationTool
 					}
 				}
 			}
-			foreach(XGEItem Item in Actions)
-			{
-				if(Item.Manifest.PostBuildScripts != null)
-				{
-					Utils.ExecuteCustomBuildSteps(Item.Manifest.PostBuildScripts.Select(x => new FileReference(x)).ToArray());
-				}
-			}
 			foreach (var Item in Actions)
 			{
 				XGEFinishBuildWithUBT(Item);
diff --git a/Engine/Source/Programs/DotNETCommon/DotNETUtilities/BinaryFormatterUtils.cs b/Engine/Source/Programs/DotNETCommon/DotNETUtilities/BinaryFormatterUtils.cs
new file mode 100644
--- /dev/null
+++ b/Engine/Source/Programs/DotNETCommon/DotNETUtilities/BinaryFormatterUtils.cs
@@ -0,0 +1,45 @@
+using System;
+using System.Collections.Generic;
+using System.IO;
+using System.Linq;
+using System.Runtime.Serialization.Formatters.Binary;
+using System.Text;
+using System.Threading.Tasks;
+
+namespace Tools.DotNETCommon
+{
+	/// <summary>
+	/// Utility functions for serializing using the BinaryFormatter
+	/// </summary>
+	public static class BinaryFormatterUtils
+	{
+		/// <summary>
+		/// Load an object from a file on disk, using the binary formatter
+		/// </summary>
+		/// <param name="Location">File to read from</param>
+		/// <returns>Instance of the object that was read from disk</returns>
+		public static T Load<T>(FileReference Location)
+		{
+			using (FileStream Stream = new FileStream(Location.FullName, FileMode.Open, FileAccess.Read))
+			{
+				BinaryFormatter Formatter = new BinaryFormatter();
+				return (T)Formatter.Deserialize(Stream);
+			}
+		}
+
+		/// <summary>
+		/// Saves a file to disk, using the binary formatter
+		/// </summary>
+		/// <param name="Location">File to write to</param>
+		/// <param name="Object">Object to serialize</param>
+		public static void Save(FileReference Location, object Object)
+		{
+			DirectoryReference.CreateDirectory(Location.Directory);
+			using (FileStream Stream = new FileStream(Location.FullName, FileMode.Create, FileAccess.Write))
+			{
+				BinaryFormatter Formatter = new BinaryFormatter();
+				Formatter.Serialize(Stream, Object);
+			}
+		}
+	}
+}
diff --git a/Engine/Source/Programs/DotNETCommon/DotNETUtilities/DotNETUtilities.csproj b/Engine/Source/Programs/DotNETCommon/DotNETUtilities/DotNETUtilities.csproj
--- a/Engine/Source/Programs/DotNETCommon/DotNETUtilities/DotNETUtilities.csproj
+++ b/Engine/Source/Programs/DotNETCommon/DotNETUtilities/DotNETUtilities.csproj
@@ -55,6 +55,7 @@
       <Link>Properties\MetaData.cs</Link>
     </Compile>
     <Compile Include="AssemblyUtils.cs" />
+    <Compile Include="BinaryFormatterUtils.cs" />
     <Compile Include="BinaryReaderExtensions.cs" />
     <Compile Include="BinaryWriterExtensions.cs" />
     <Compile Include="CommandLineArguments.cs" />
diff --git a/Engine/Source/Programs/UnrealBuildTool/Configuration/TargetRules.cs b/Engine/Source/Programs/UnrealBuildTool/Configuration/TargetRules.cs
--- a/Engine/Source/Programs/UnrealBuildTool/Configuration/TargetRules.cs
+++ b/Engine/Source/Programs/UnrealBuildTool/Configuration/TargetRules.cs
@@ -950,6 +950,12 @@ namespace UnrealBuildTool
 		/// </summary>
 		public bool bLegacyPublicIncludePaths = true;
 
+		/// <summary>
+		/// Do not allow manifest changes when building this target. Used to cause earlier errors when building multiple targets with a shared build environment.
+		/// </summary>
+		[CommandLine("-NoManifestChanges")]
+		internal bool bNoManifestChanges = false;
+
 		/// <summary>
 		/// The build version string
 		/// </summary>
@@ -2022,6 +2028,11 @@ namespace UnrealBuildTool
 			get { return Inner.bLegacyPublicIncludePaths; }
 		}
 
+		internal bool bNoManifestChanges
+		{
+			get { return Inner.bNoManifestChanges; }
+		}
+
 		public string BuildVersion
 		{
 			get { return Inner.BuildVersion; }
diff --git a/Engine/Source/Programs/UnrealBuildTool/Configuration/UEBuildTarget.cs b/Engine/Source/Programs/UnrealBuildTool/Configuration/UEBuildTarget.cs
--- a/Engine/Source/Programs/UnrealBuildTool/Configuration/UEBuildTarget.cs
+++ b/Engine/Source/Programs/UnrealBuildTool/Configuration/UEBuildTarget.cs
@@ -218,11 +218,6 @@ namespace UnrealBuildTool
 		/// </summary>
 		public readonly List<string> DeployTargetFiles = new List<string>();
 
-		/// <summary>
-		/// 
-		/// </summary>
-		public readonly List<string> PostBuildScripts = new List<string>();
-
 		/// <summary>
 		/// 
 		/// </summary>
@@ -769,15 +764,6 @@ namespace UnrealBuildTool
 		/// </summary>
 		public List<FlatModuleCsDataType> FlatModuleCsData = new List<FlatModuleCsDataType>();
 
-		/// <summary>
-		/// The receipt for this target, which contains a record of this build.
-		/// </summary>
-		public TargetReceipt Receipt
-		{
-			get;
-			private set;
-		}
-
 		/// <summary>
 		/// Filename for the receipt for this target.
 		/// </summary>
@@ -787,11 +773,6 @@ namespace UnrealBuildTool
 			private set;
 		}
 
-		/// <summary>
-		/// Module manifests to be written to each output folder
-		/// </summary>
-		private KeyValuePair<FileReference, ModuleManifest>[] FileReferenceToModuleManifestPairs;
-
 		/// <summary>
 		/// The name of the .Target.cs file, if the target was created with one
 		/// </summary>
@@ -826,15 +807,6 @@ namespace UnrealBuildTool
 			return FlatModuleCsData.Select(Data => Data.BuildCsFilename);
 		}
 
-		/// <summary>
-		/// A list of the module filenames which were used to build this target.
-		/// </summary>
-		/// <returns></returns>
-		public IEnumerable<string> GetAllModuleFolders()
-		{
-			return FlatModuleCsData.SelectMany(Data => Data.UHTHeaderNames);
-		}
-
 		/// <summary>
 		/// Whether this target should be compiled in monolithic mode
 		/// </summary>
@@ -864,9 +836,7 @@ namespace UnrealBuildTool
 			OnlyModules = (List<OnlyModule>)Info.GetValue("om", typeof(List<OnlyModule>));
 			bCompileMonolithic = Info.GetBoolean("cm");
 			FlatModuleCsData = (List<FlatModuleCsDataType>)Info.GetValue("fm", typeof(List<FlatModuleCsDataType>));
-			Receipt = (TargetReceipt)Info.GetValue("re", typeof(TargetReceipt));
 			ReceiptFileName = (FileReference)Info.GetValue("rf", typeof(FileReference));
-			FileReferenceToModuleManifestPairs = (KeyValuePair<FileReference, ModuleManifest>[])Info.GetValue("vm", typeof(KeyValuePair<FileReference, ModuleManifest>[]));
 			TargetRulesFile = (FileReference)Info.GetValue("tc", typeof(FileReference));
 			PreBuildStepScripts = (FileReference[])Info.GetValue("pr", typeof(FileReference[]));
 			PostBuildStepScripts = (FileReference[])Info.GetValue("po", typeof(FileReference[]));
@@ -895,9 +865,7 @@ namespace UnrealBuildTool
 			Info.AddValue("om", OnlyModules);
 			Info.AddValue("cm", bCompileMonolithic);
 			Info.AddValue("fm", FlatModuleCsData);
-			Info.AddValue("re", Receipt);
 			Info.AddValue("rf", ReceiptFileName);
-			Info.AddValue("vm", FileReferenceToModuleManifestPairs);
 			Info.AddValue("tc", TargetRulesFile);
 			Info.AddValue("pr", PreBuildStepScripts);
 			Info.AddValue("po", PostBuildStepScripts);
@@ -1372,19 +1340,6 @@ namespace UnrealBuildTool
 
 			if (!Rules.bDisableLinking)
 			{
-				// Also add the version file if it's been specified
-				if (VersionFile != null)
-				{
-					Manifest.BuildProducts.Add(VersionFile.FullName);
-				}
-
-				// Add all the version manifests to the receipt
-				foreach (FileReference VersionManifestFile in FileReferenceToModuleManifestPairs.Select(x => x.Key))
-				{
-					Manifest.BuildProducts.Add(VersionManifestFile.FullName);
-				}
-
-				UEBuildPlatform BuildPlatform = UEBuildPlatform.GetBuildPlatform(Platform);
 				if (OnlyModules.Count == 0)
 				{
 					Manifest.AddBuildProduct(ReceiptFileName.FullName);
@@ -1394,11 +1349,6 @@ namespace UnrealBuildTool
 				{
 					Manifest.DeployTargetFiles.Add(DeployTargetFile.FullName);
 				}
-
-				if(PostBuildStepScripts != null)
-				{
-					Manifest.PostBuildScripts.AddRange(PostBuildStepScripts.Select(x => x.FullName));
-				}
 			}
 
 			// Remove anything that's not part of the plugin
@@ -1410,12 +1360,46 @@ namespace UnrealBuildTool
 
 			Manifest.BuildProducts.Sort();
 			Manifest.DeployTargetFiles.Sort();
-			Manifest.PostBuildScripts.Sort();
 
 			Log.TraceInformation("Writing manifest to {0}", ManifestPath);
 			Utils.WriteClass<BuildManifest>(Manifest, ManifestPath.FullName, "");
 		}
 
+		/// <summary>
+		/// Prepare all the module manifests for this target
+		/// </summary>
+		/// <returns>Dictionary mapping from filename to module manifest</returns>
+		Dictionary<FileReference, ModuleManifest> PrepareModuleManifests()
+		{
+			Dictionary<FileReference, ModuleManifest> FileNameToModuleManifest = new Dictionary<FileReference, ModuleManifest>();
+			if (!bCompileMonolithic)
+			{
+				// Create the receipts for each folder
+				foreach (UEBuildBinary Binary in Binaries)
+				{
+					if(Binary.Type == UEBuildBinaryType.DynamicLinkLibrary)
+					{
+						DirectoryReference DirectoryName = Binary.OutputFilePath.Directory;
+						bool bIsGameDirectory = !DirectoryName.IsUnderDirectory(UnrealBuildTool.EngineDirectory);
+						FileReference ManifestFileName = FileReference.Combine(DirectoryName, ModuleManifest.GetStandardFileName(AppName, Platform, Configuration, Architecture, bIsGameDirectory));
+
+						ModuleManifest Manifest;
+						if (!FileNameToModuleManifest.TryGetValue(ManifestFileName, out Manifest))
+						{
+							Manifest = new ModuleManifest("");
+							FileNameToModuleManifest.Add(ManifestFileName, Manifest);
+						}
+
+						foreach (UEBuildModuleCPP Module in Binary.Modules.OfType<UEBuildModuleCPP>())
+						{
+							Manifest.ModuleNameToFileName[Module.Name] = Binary.OutputFilePath.GetFileName();
+						}
+					}
+				}
+			}
+			return FileNameToModuleManifest;
+		}
+
 		/// <summary>
 		/// Prepare all the receipts this target (all the .target and .modules files). See the VersionManifest class for an explanation of what these files are.
 		/// </summary>
@@ -1423,14 +1407,8 @@ namespace UnrealBuildTool
 		/// <param name="BuildProducts">Artifacts from the build</param>
 		/// <param name="RuntimeDependencies">Output runtime dependencies</param>
 		/// <param name="HotReload">The hot-reload mode</param>
-		void PrepareReceipts(UEToolChain ToolChain, List<KeyValuePair<FileReference, BuildProductType>> BuildProducts, List<RuntimeDependency> RuntimeDependencies, EHotReload HotReload)
+		TargetReceipt PrepareReceipt(UEToolChain ToolChain, List<KeyValuePair<FileReference, BuildProductType>> BuildProducts, List<RuntimeDependency> RuntimeDependencies, EHotReload HotReload)
 		{
-			// If linking is disabled, don't generate any receipt
-			if(Rules.bDisableLinking)
-			{
-				return;
-			}
-
 			// Read the version file
 			BuildVersion Version;
 			if (!BuildVersion.TryRead(BuildVersion.GetDefaultFileName(), out Version))
@@ -1456,20 +1434,11 @@ namespace UnrealBuildTool
 					{
 						Version = LastVersion;
 					}
-					else
-					{
-						Version.BuildId = "";
-					}
-				}
-				else
-				{
-					// Otherwise generate something randomly.
-					Version.BuildId = Guid.NewGuid().ToString();
 				}
 			}
 
 			// Create the receipt
-			Receipt = new TargetReceipt(TargetName, Platform, Configuration, Version);
+			TargetReceipt Receipt = new TargetReceipt(TargetName, Platform, Configuration, Version);
 
 			// Set the launch executable if there is one
 			if(!Rules.bShouldCompileAsDLL)
@@ -1524,302 +1493,10 @@ namespace UnrealBuildTool
 				}
 			}
 
-			// Also add the version file if it's been specified
-			if (VersionFile != null)
-			{
-				Receipt.BuildProducts.Add(new BuildProduct(VersionFile, BuildProductType.BuildResource));
-			}
-
-			// Prepare all the version manifests
-			Dictionary<FileReference, ModuleManifest> FileNameToModuleManifest = new Dictionary<FileReference, ModuleManifest>();
-			if (!bCompileMonolithic)
-			{
-				// Create the receipts for each folder
-				foreach (UEBuildBinary Binary in Binaries)
-				{
-					if(Binary.Type == UEBuildBinaryType.DynamicLinkLibrary)
-					{
-						DirectoryReference DirectoryName = Binary.OutputFilePath.Directory;
-						bool bIsGameDirectory = !DirectoryName.IsUnderDirectory(UnrealBuildTool.EngineDirectory);
-						FileReference ManifestFileName = FileReference.Combine(DirectoryName, ModuleManifest.GetStandardFileName(AppName, Platform, Configuration, Architecture, bIsGameDirectory));
-
-						ModuleManifest Manifest;
-						if (!FileNameToModuleManifest.TryGetValue(ManifestFileName, out Manifest))
-						{
-							Manifest = new ModuleManifest(Version.BuildId);
-
-							ModuleManifest ExistingManifest;
-							if (ModuleManifest.TryRead(ManifestFileName, out ExistingManifest) && Version.BuildId == ExistingManifest.BuildId)
-							{
-								if (OnlyModules.Count > 0)
-								{
-									// We're just building an existing module; reuse the existing manifest AND build id.
-									Manifest = ExistingManifest;
-								}
-								else if (Version.Changelist != 0)
-								{
-									// We're rebuilding at the same changelist. Keep all the existing binaries.
-									Manifest.ModuleNameToFileName = Manifest.ModuleNameToFileName.Union(ExistingManifest.ModuleNameToFileName).ToDictionary(x => x.Key, x => x.Value);
-								}
-							}
-
-							FileNameToModuleManifest.Add(ManifestFileName, Manifest);
-						}
-
-						foreach (UEBuildModuleCPP Module in Binary.Modules.OfType<UEBuildModuleCPP>())
-						{
-							Manifest.ModuleNameToFileName[Module.Name] = Binary.OutputFilePath.GetFileName();
-						}
-					}
-				}
-			}
-			FileReferenceToModuleManifestPairs = FileNameToModuleManifest.ToArray();
-
-			// Add all the version manifests to the receipt
-			foreach(FileReference VersionManifestFile in FileNameToModuleManifest.Keys)
-			{
-				Receipt.AddBuildProduct(VersionManifestFile, BuildProductType.RequiredResource);
-			}
-
 			// add the SDK used by the tool chain
 			Receipt.AdditionalProperties.Add(new ReceiptProperty("SDK", ToolChain.GetSDKVersion()));
-		}
-
-		/// <summary>
-		/// Try to recycle the build id from existing version manifests in the engine directory rather than generating a new one, if no engine binaries are being modified.
-		/// This allows sharing engine binaries when switching between projects and switching between UE4 and a game-specific project. Note that different targets may require
-		/// additional engine modules to be built, so we don't prohibit files being added or removed.
-		/// </summary>
-		/// <param name="OutputFiles">List of files being modified by this build</param>
-		/// <param name="bNoManifestChanges">Whether manifest changes are allowed. If a manifest has to be changed, an error will be output.</param>
-		/// <returns>True if the existing version manifests will remain valid during this build, false if they are invalidated</returns>
-		public bool TryRecycleVersionManifests(HashSet<FileReference> OutputFiles, bool bNoManifestChanges)
-		{
-			// Make sure we've got a list of version manifests to check against
-			if(FileReferenceToModuleManifestPairs == null)
-			{
-				Log.TraceLog("No file to version manifest mapping; unable to recycle version.");
-				return false;
-			}
-
-			// If there is no version file, don't bother trying to read it
-			if(VersionFile == null)
-			{
-				Log.TraceLog("Target is not using a version file.");
-				return false;
-			}
-
-			// Make sure we've got a file containing the last build id
-			BuildVersion CurrentVersion;
-			if(!BuildVersion.TryRead(VersionFile, out CurrentVersion))
-			{
-				Log.TraceLog("Unable to read version file from {0}.", VersionFile);
-				return false;
-			}
-
-			// Get the project directory. We will ignore any manifests under this directory (ie. anything not under engine/enterprise folders).
-			DirectoryReference ProjectDir = DirectoryReference.FromFile(ProjectFile);
-
-			// Read any the existing module manifests under the engine directory
-			Dictionary<FileReference, ModuleManifest> ExistingFileToManifest = new Dictionary<FileReference, ModuleManifest>();
-			foreach(FileReference ExistingFile in FileReferenceToModuleManifestPairs.Select(x => x.Key))
-			{
-				if(ProjectDir == null || !ExistingFile.IsUnderDirectory(ProjectDir))
-				{
-					ModuleManifest ExistingManifest;
-					if(ModuleManifest.TryRead(ExistingFile, out ExistingManifest))
-					{
-						ExistingFileToManifest.Add(ExistingFile, ExistingManifest);
-					}
-				}
-			}
-
-			// Check if we're modifying any files in an existing valid manifest. If the build id for a manifest doesn't match, we can behave as if it doesn't exist.
-			foreach(KeyValuePair<FileReference, ModuleManifest> ExistingPair in ExistingFileToManifest)
-			{
-				if(ExistingPair.Value.BuildId == CurrentVersion.BuildId)
-				{
-					DirectoryReference ExistingManifestDir = ExistingPair.Key.Directory;
-					foreach(FileReference ExistingFile in ExistingPair.Value.ModuleNameToFileName.Values.Select(x => FileReference.Combine(ExistingManifestDir, x)))
-					{
-						if(OutputFiles.Contains(ExistingFile))
-						{
-							if(bNoManifestChanges)
-							{
-								Log.TraceError("Previous build product would be overwritten: {0}.", ExistingFile);
-							}
-							else
-							{
-								Log.TraceLog("Unable to recycle manifests - modifying {0} invalidates {1}. Using build id {2}.", ExistingFile, ExistingPair.Key, Receipt.Version.BuildId);
-							}
-							return false;
-						}
-					}
-				}
-			}
-
-			// Allow the existing build id to be reused. Update the receipt.
-			Receipt.Version.BuildId = CurrentVersion.BuildId;
 
-			// Merge the existing manifests with the manifests in memory.
-			foreach(KeyValuePair<FileReference, ModuleManifest> NewPair in FileReferenceToModuleManifestPairs)
-			{
-				// Reuse the existing build id
-				ModuleManifest NewManifest = NewPair.Value;
-				NewManifest.BuildId = CurrentVersion.BuildId;
-
-				// Merge in the files from the existing manifest
-				ModuleManifest ExistingManifest;
-				if(ExistingFileToManifest.TryGetValue(NewPair.Key, out ExistingManifest) && ExistingManifest.BuildId == CurrentVersion.BuildId)
-				{
-					foreach(KeyValuePair<string, string> ModulePair in ExistingManifest.ModuleNameToFileName)
-					{
-						if(!NewManifest.ModuleNameToFileName.ContainsKey(ModulePair.Key))
-						{
-							NewManifest.ModuleNameToFileName.Add(ModulePair.Key, ModulePair.Value);
-						}
-					}
-				}
-			}
-
-			// Return success
-			Log.TraceLog("Recycled previous build ID ({0})", CurrentVersion.BuildId);
-			return true;
-		}
-
-		/// <summary>
-		/// Delete all the existing version manifests
-		/// </summary>
-		public void InvalidateVersionManifests()
-		{
-			// Delete all the existing manifests, so we don't try to recycle partial builds in future (the current build may fail after modifying engine files, 
-			// causing bModifyingEngineFiles to be incorrect on the next invocation).
-			if(FileReferenceToModuleManifestPairs != null)
-			{
-				foreach (FileReference VersionManifestFile in FileReferenceToModuleManifestPairs.Select(x => x.Key))
-				{
-					// Make sure the file (and directory) exists before trying to delete it
-					if(FileReference.Exists(VersionManifestFile) && !UnrealBuildTool.IsFileInstalled(VersionManifestFile))
-					{
-						FileReference.Delete(VersionManifestFile);
-					}
-				}
-			}
-		}
-
-		/// <summary>
-		/// Patches the manifests with the new module suffixes from the OnlyModules list.
-		/// </summary>
-		public void PatchModuleManifestsForHotReloadAssembling(List<OnlyModule> OnlyModules)
-		{
-			if (FileReferenceToModuleManifestPairs == null)
-			{
-				return;
-			}
-
-			foreach (KeyValuePair<FileReference, ModuleManifest> FileNameToVersionManifest in FileReferenceToModuleManifestPairs)
-			{
-				foreach (KeyValuePair<string, string> Manifest in FileNameToVersionManifest.Value.ModuleNameToFileName)
-				{
-					string ModuleFilename = Manifest.Value;
-					if (UnrealBuildTool.ReplaceHotReloadFilenameSuffix(ref ModuleFilename, (ModuleName) => UnrealBuildTool.GetReplacementModuleSuffix(OnlyModules, ModuleName)))
-					{
-						FileNameToVersionManifest.Value.ModuleNameToFileName[Manifest.Key] = ModuleFilename;
-						break;
-					}
-				}
-			}
-		}
-
-		/// <summary>
-		/// Writes out the version manifest
-		/// </summary>
-		public void WriteReceipts(bool bNoManifestChanges)
-		{
-			if (Receipt != null)
-			{
-				UEBuildPlatform BuildPlatform = UEBuildPlatform.GetBuildPlatform(Platform);
-				if (OnlyModules == null || OnlyModules.Count == 0)
-				{
-					if(!UnrealBuildTool.IsFileInstalled(ReceiptFileName))
-					{
-						DirectoryReference.CreateDirectory(ReceiptFileName.Directory);
-						Receipt.Write(ReceiptFileName, UnrealBuildTool.EngineDirectory, ProjectDirectory);
-					}
-				}
-				if(VersionFile != null)
-				{
-					if(!UnrealBuildTool.IsFileInstalled(VersionFile))
-					{
-						DirectoryReference.CreateDirectory(VersionFile.Directory);
-
-						StringWriter Writer = new StringWriter();
-						Receipt.Version.Write(Writer);
-
-						string Text = Writer.ToString();
-						if(!FileReference.Exists(VersionFile) || File.ReadAllText(VersionFile.FullName) != Text)
-						{
-							File.WriteAllText(VersionFile.FullName, Text);
-						}
-					}
-				}
-			}
-			if (FileReferenceToModuleManifestPairs != null)
-			{
-				foreach (KeyValuePair<FileReference, ModuleManifest> FileNameToVersionManifest in FileReferenceToModuleManifestPairs)
-				{
-					if(!UnrealBuildTool.IsFileInstalled(FileNameToVersionManifest.Key))
-					{
-						if(!FileReference.Exists(FileNameToVersionManifest.Key))
-						{
-							// If the file doesn't already exist, just write it out
-							DirectoryReference.CreateDirectory(FileNameToVersionManifest.Key.Directory);
-							FileNameToVersionManifest.Value.Write(FileNameToVersionManifest.Key);
-						}
-						else
-						{
-							// Otherwise write it to a buffer first
-							string OutputText;
-							using (StringWriter Writer = new StringWriter())
-							{
-								FileNameToVersionManifest.Value.Write(Writer);
-								OutputText = Writer.ToString();
-							}
-
-							// And only write it to disk if it's been modified
-							string CurrentText = FileReference.ReadAllText(FileNameToVersionManifest.Key);
-							if(CurrentText != OutputText)
-							{
-								if(bNoManifestChanges)
-								{
-									Log.TraceError("Build modifies {0}. This is not permitted. Before:\n    {1}\nAfter:\n    {2}", FileNameToVersionManifest.Key, CurrentText.Replace("\n", "\n    "), OutputText.Replace("\n", "\n    "));
-								}
-								else
-								{
-									FileReference.WriteAllText(FileNameToVersionManifest.Key, OutputText);
-								}
-							}
-						}
-					}
-				}
-			}
-		}
-
-		/// <summary>
-		/// Removes any receipt files in the case of a failed build
-		/// </summary>
-		public void DeleteReceipts()
-		{
-			if (Receipt != null)
-			{
-				if (OnlyModules == null || OnlyModules.Count == 0)
-				{
-					if(!UnrealBuildTool.IsFileInstalled(ReceiptFileName) && FileReference.Exists(ReceiptFileName))
-					{
-						FileReference.Delete(ReceiptFileName);
-					}
-				}
-			}
+			return Receipt;
 		}
 
 		/// <summary>
@@ -2220,23 +1897,6 @@ namespace UnrealBuildTool
 			}
 			BuildProducts.AddRange(RuntimeDependencyTargetFileToSourceFile.Select(x => new KeyValuePair<FileReference, BuildProductType>(x.Key, BuildProductType.RequiredResource)));
 
-			// Also add any explicitly specified build products
-			if(Rules.AdditionalBuildProducts.Count > 0)
-			{
-				Dictionary<string, string> Variables = GetTargetVariables(null);
-				foreach(string AdditionalBuildProduct in Rules.AdditionalBuildProducts)
-				{
-					FileReference BuildProductFile = new FileReference(Utils.ExpandVariables(AdditionalBuildProduct, Variables));
-					BuildProducts.Add(new KeyValuePair<FileReference, BuildProductType>(BuildProductFile, BuildProductType.RequiredResource));
-				}
-			}
-
-			// Create a receipt for the target
-			if (!ProjectFileGenerator.bGenerateProjectFiles)
-			{
-				PrepareReceipts(TargetToolChain, BuildProducts, RuntimeDependencies, HotReload);
-			}
-
 			// Make sure all the checked headers were valid
 			List<string> InvalidIncludeDirectiveMessages = Modules.Values.OfType<UEBuildModuleCPP>().Where(x => x.InvalidIncludeDirectiveMessages != null).SelectMany(x => x.InvalidIncludeDirectiveMessages).ToList();
 			if (InvalidIncludeDirectiveMessages.Count > 0)
@@ -2249,6 +1909,80 @@ namespace UnrealBuildTool
 				return ECompilationResult.Canceled;
 			}
 
+			// Finalize and generate metadata for this target
+			if(!ProjectFileGenerator.bGenerateProjectFiles && !Rules.bDisableLinking)
+			{
+				// Also add any explicitly specified build products
+				if(Rules.AdditionalBuildProducts.Count > 0)
+				{
+					Dictionary<string, string> Variables = GetTargetVariables(null);
+					foreach(string AdditionalBuildProduct in Rules.AdditionalBuildProducts)
+					{
+						FileReference BuildProductFile = new FileReference(Utils.ExpandVariables(AdditionalBuildProduct, Variables));
+						BuildProducts.Add(new KeyValuePair<FileReference, BuildProductType>(BuildProductFile, BuildProductType.RequiredResource));
+					}
+				}
+
+				// Also add the version file as a build product
+				if(VersionFile != null)
+				{
+					BuildProducts.Add(new KeyValuePair<FileReference, BuildProductType>(VersionFile, BuildProductType.RequiredResource));
+				}
+
+				// Prepare the module manifests, and add them to the list of build products
+				Dictionary<FileReference, ModuleManifest> FileNameToModuleManifest = PrepareModuleManifests();
+				BuildProducts.AddRange(FileNameToModuleManifest.Select(x => new KeyValuePair<FileReference, BuildProductType>(x.Key, BuildProductType.RequiredResource)));
+
+				// Prepare the receipt
+				TargetReceipt Receipt = PrepareReceipt(TargetToolChain, BuildProducts, RuntimeDependencies, HotReload);
+
+				// Create an action which to generate the receipts
+				WriteMetadataTargetInfo MetadataTargetInfo = new WriteMetadataTargetInfo(ProjectFile, VersionFile, ReceiptFileName, Receipt, FileNameToModuleManifest);
+				FileReference MetadataTargetFile = FileReference.Combine(ProjectIntermediateDirectory, "Metadata.dat");
+				BinaryFormatterUtils.Save(MetadataTargetFile, MetadataTargetInfo);
+
+				StringBuilder WriteMetadataArguments = new StringBuilder();
+				WriteMetadataArguments.AppendFormat("-Mode=WriteMetadata -Input={0}", Utils.MakePathSafeToUseWithCommandLine(MetadataTargetFile));
+				if(Rules.bNoManifestChanges)
+				{
+					WriteMetadataArguments.Append(" -NoManifestChanges");
+				}
+
+				Action WriteMetadataAction = ActionGraph.AddRecursiveCall(ActionType.WriteMetadata, WriteMetadataArguments.ToString());
+				WriteMetadataAction.WorkingDirectory = UnrealBuildTool.EngineSourceDirectory.FullName;
+				WriteMetadataAction.StatusDescription = "Writing receipt...";
+				WriteMetadataAction.bCanExecuteRemotely = false;
+				WriteMetadataAction.PrerequisiteItems.AddRange(OutputItems);
+				WriteMetadataAction.ProducedItems.Add(FileItem.GetItemByFileReference(ReceiptFileName));
+
+				OutputItems.AddRange(WriteMetadataAction.ProducedItems);
+
+				// Create actions to run the post build steps
+				foreach(FileReference PostBuildStepScript in PostBuildStepScripts)
+				{
+					FileReference OutputFile = new FileReference(PostBuildStepScript.FullName + ".ran");
+
+					Action PostBuildStepAction = ActionGraph.Add(ActionType.PostBuildStep);
+					if(BuildHostPlatform.Current.Platform == UnrealTargetPlatform.Win64)
+					{
+						PostBuildStepAction.CommandPath = "cmd.exe";
+						PostBuildStepAction.CommandArguments = String.Format("/C \"call \"{0}\" && type NUL >\"{1}\"\"", PostBuildStepScript, OutputFile);
+					}
+					else
+					{
+						PostBuildStepAction.CommandPath = "/bin/sh";
+						PostBuildStepAction.CommandArguments = String.Format("\"{0}\" && touch \"{1}\"", PostBuildStepScript, OutputFile);
+					}
+					PostBuildStepAction.WorkingDirectory = UnrealBuildTool.EngineSourceDirectory.FullName;
+					PostBuildStepAction.StatusDescription = String.Format("Executing post build script ({0})", PostBuildStepScript.GetFileName());
+					PostBuildStepAction.bCanExecuteRemotely = false;
+					PostBuildStepAction.PrerequisiteItems.Add(FileItem.GetItemByFileReference(ReceiptFileName));
+					PostBuildStepAction.ProducedItems.Add(FileItem.GetItemByFileReference(OutputFile));
+
+					OutputItems.AddRange(PostBuildStepAction.ProducedItems);
+				}
+			}
+
 			// Build a list of all the files required to build
 			foreach(FileReference DependencyListFileName in Rules.DependencyListFileNames)
 			{
diff --git a/Engine/Source/Programs/UnrealBuildTool/Modes/WriteMetadataMode.cs b/Engine/Source/Programs/UnrealBuildTool/Modes/WriteMetadataMode.cs
new file mode 100644
--- /dev/null
+++ b/Engine/Source/Programs/UnrealBuildTool/Modes/WriteMetadataMode.cs
@@ -0,0 +1,241 @@
+using System;
+using System.Collections.Generic;
+using System.IO;
+using System.Linq;
+using System.Runtime.Serialization.Formatters.Binary;
+using System.Text;
+using System.Threading.Tasks;
+using Tools.DotNETCommon;
+
+namespace UnrealBuildTool
+{
+	/// <summary>
+	/// Parameters for the WriteMetadata mode
+	/// </summary>
+	[Serializable]
+	class WriteMetadataTargetInfo
+	{
+		/// <summary>
+		/// The project file
+		/// </summary>
+		public FileReference ProjectFile;
+
+		/// <summary>
+		/// Output location for the version file
+		/// </summary>
+		public FileReference VersionFile;
+
+		/// <summary>
+		/// Output location for the target file
+		/// </summary>
+		public FileReference ReceiptFile;
+
+		/// <summary>
+		/// The partially constructed receipt data
+		/// </summary>
+		public TargetReceipt Receipt;
+
+		/// <summary>
+		/// Map of module manifest filenames to their location on disk.
+		/// </summary>
+		public Dictionary<FileReference, ModuleManifest> FileToManifest = new Dictionary<FileReference, ModuleManifest>();
+
+		/// <summary>
+		/// Constructor
+		/// </summary>
+		/// <param name="ProjectFile"></param>
+		/// <param name="VersionFile"></param>
+		/// <param name="ReceiptFile"></param>
+		/// <param name="Receipt"></param>
+		/// <param name="FileToManifest"></param>
+		public WriteMetadataTargetInfo(FileReference ProjectFile, FileReference VersionFile, FileReference ReceiptFile, TargetReceipt Receipt, Dictionary<FileReference, ModuleManifest> FileToManifest)//string EngineManifestName, string ProjectManifestName, Dictionary<string, FileReference> ModuleNameToLocation)
+		{
+			this.ProjectFile = ProjectFile;
+			this.VersionFile = VersionFile;
+			this.ReceiptFile = ReceiptFile;
+			this.Receipt = Receipt;
+			this.FileToManifest = FileToManifest;
+		}
+	}
+
+	/// <summary>
+	/// Writes all metadata files at the end of a build (receipts, version files, etc...). This is implemented as a separate mode to allow it to be done as part of the action graph.
+	/// </summary>
+	[ToolMode("WriteMetadata")]
+	class WriteMetadataMode : ToolMode
+	{
+		/// <summary>
+		/// Execute the command
+		/// </summary>
+		/// <param name="Arguments">Command line arguments</param>
+		/// <returns>Exit code</returns>
+		public override int Execute(CommandLineArguments Arguments)
+		{
+			// Read the target info
+			WriteMetadataTargetInfo TargetInfo = BinaryFormatterUtils.Load<WriteMetadataTargetInfo>(Arguments.GetFileReference("-Input="));
+			bool bNoManifestChanges = Arguments.HasOption("-NoManifestChanges");
+			Arguments.CheckAllArgumentsUsed();
+
+			// Check if we need to set a build id
+			TargetReceipt Receipt = TargetInfo.Receipt;
+			if(Receipt.Version.BuildId == null)
+			{
+				// Check if there's an exist version file. If it exists, try to merge in any manifests that are valid (and reuse the existing build id)
+				BuildVersion PreviousVersion;
+				if(TargetInfo.VersionFile != null && BuildVersion.TryRead(TargetInfo.VersionFile, out PreviousVersion))
+				{
+					// Check if we can reuse the existing manifests. This prevents unnecessary builds when switching between projects.
+					Dictionary<FileReference, ModuleManifest> PreviousFileToManifest = new Dictionary<FileReference, ModuleManifest>();
+					if(TryRecyclingManifests(PreviousVersion.BuildId, TargetInfo.FileToManifest.Keys, PreviousFileToManifest))
+					{
+						// Merge files from the existing manifests with the new ones
+						foreach(KeyValuePair<FileReference, ModuleManifest> Pair in PreviousFileToManifest)
+						{
+							ModuleManifest OldManifest = Pair.Value;
+							ModuleManifest NewManifest = TargetInfo.FileToManifest[Pair.Key];
+							foreach(KeyValuePair<string, string> ModulePair in OldManifest.ModuleNameToFileName)
+							{
+								if(!NewManifest.ModuleNameToFileName.ContainsKey(ModulePair.Key))
+								{
+									NewManifest.ModuleNameToFileName.Add(ModulePair.Key, ModulePair.Value);
+								}
+							}
+						}
+
+						// Update the build id to use the current one
+						Receipt.Version.BuildId = PreviousVersion.BuildId;
+					}
+				}
+
+				// If the build id is still not set, generate a new one from a GUID
+				if(Receipt.Version.BuildId == null)
+				{
+					Receipt.Version.BuildId = Guid.NewGuid().ToString();
+				}
+			}
+
+			// Update the build id in all the manifests, and write them out
+			foreach (KeyValuePair<FileReference, ModuleManifest> Pair in TargetInfo.FileToManifest)
+			{
+				FileReference ManifestFile = Pair.Key;
+				if(!UnrealBuildTool.IsFileInstalled(ManifestFile))
+				{
+					ModuleManifest Manifest = Pair.Value;
+					Manifest.BuildId = Receipt.Version.BuildId;
+
+					if(!FileReference.Exists(ManifestFile))
+					{
+						// If the file doesn't already exist, just write it out
+						DirectoryReference.CreateDirectory(ManifestFile.Directory);
+						Manifest.Write(ManifestFile);
+					}
+					else
+					{
+						// Otherwise write it to a buffer first
+						string OutputText;
+						using (StringWriter Writer = new StringWriter())
+						{
+							Manifest.Write(Writer);
+							OutputText = Writer.ToString();
+						}
+
+						// And only write it to disk if it's been modified. Note that if a manifest is out of date, we should have generated a new build id causing the contents to differ.
+						string CurrentText = FileReference.ReadAllText(ManifestFile);
+						if(CurrentText != OutputText)
+						{
+							if(bNoManifestChanges)
+							{
+								Log.TraceError("Build modifies {0}. This is not permitted. Before:\n    {1}\nAfter:\n    {2}", ManifestFile, CurrentText.Replace("\n", "\n    "), OutputText.Replace("\n", "\n    "));
+							}
+							else
+							{
+								FileReference.WriteAllText(ManifestFile, OutputText);
+							}
+						}
+					}
+				}
+			}
+
+			// Write out the version file, if it's changed. Since this file is next to the executable, it may be used by multiple targets, and we should avoid modifying it unless necessary.
+			if(TargetInfo.VersionFile != null && !UnrealBuildTool.IsFileInstalled(TargetInfo.VersionFile))
+			{
+				DirectoryReference.CreateDirectory(TargetInfo.VersionFile.Directory);
+
+				StringWriter Writer = new StringWriter();
+				Receipt.Version.Write(Writer);
+
+				string Text = Writer.ToString();
+				if(!FileReference.Exists(TargetInfo.VersionFile) || File.ReadAllText(TargetInfo.VersionFile.FullName) != Text)
+				{
+					File.WriteAllText(TargetInfo.VersionFile.FullName, Text);
+				}
+			}
+
+			// Write out the receipt
+			if(!UnrealBuildTool.IsFileInstalled(TargetInfo.ReceiptFile))
+			{
+				DirectoryReference.CreateDirectory(TargetInfo.ReceiptFile.Directory);
+				Receipt.Write(TargetInfo.ReceiptFile, UnrealBuildTool.EngineDirectory, DirectoryReference.FromFile(TargetInfo.ProjectFile));
+			}
+
+			return 0;
+		}
+
+		/// <summary>
+		/// Checks whether existing manifests on disk can be merged with new manifests being created, by testing whether any build products they reference have a newer timestamp
+		/// </summary>
+		/// <param name="BuildId">The current build id read from the version file. Only manifests matching this ID will be considered.</param>
+		/// <param name="ManifestFiles">List of new manifest files</param>
+		/// <param name="RecycleFileToManifest">If successful, is populated with a map of filename to existing manifests that can be merged with the new manifests.</param>
+		/// <returns>True if the manifests can be recycled (and fills RecycleFileToManifest)</returns>
+		bool TryRecyclingManifests(string BuildId, IEnumerable<FileReference> ManifestFiles, Dictionary<FileReference, ModuleManifest> RecycleFileToManifest)
+		{
+			bool bCanRecycleManifests = true;
+			foreach(FileReference ManifestFileName in ManifestFiles)
+			{
+				if(ManifestFileName.IsUnderDirectory(UnrealBuildTool.EngineDirectory) && FileReference.Exists(ManifestFileName))
+				{
+					try
+					{
+						ModuleManifest Manifest = ModuleManifest.Read(ManifestFileName);
+						if(Manifest.BuildId == BuildId)
+						{
+							if(IsOutOfDate(ManifestFileName, Manifest))
+							{
+								bCanRecycleManifests = false;
+								break;
+							}
+							RecycleFileToManifest.Add(ManifestFileName, Manifest);
+						}
+					}
+					catch(Exception Ex)
+					{
+						Log.TraceWarning("Unable to read '{0}'; ignoring.", ManifestFileName);
+						Log.TraceLog(ExceptionUtils.FormatExceptionDetails(Ex));
+					}
+				}
+			}
+			return bCanRecycleManifests;
+		}
+
+		/// <summary>
+		/// Checks whether a module manifest on disk is out of date (whether any of the binaries it references are newer than it is)
+		/// </summary>
+		/// <param name="ManifestFileName">Path to the manifest</param>
+		/// <param name="Manifest">The manifest contents</param>
+		/// <returns>True if the manifest is out of date</returns>
+		bool IsOutOfDate(FileReference ManifestFileName, ModuleManifest Manifest)
+		{
+			DateTime ManifestTime = FileReference.GetLastWriteTimeUtc(ManifestFileName);
+			foreach(string FileName in Manifest.ModuleNameToFileName.Values)
+			{
+				FileInfo ModuleInfo = new FileInfo(FileReference.Combine(ManifestFileName.Directory, FileName).FullName);
+				if(!ModuleInfo.Exists || ModuleInfo.LastWriteTimeUtc > ManifestTime)
+				{
+					return true;
+				}
+			}
+			return false;
+		}
+	}
+}
diff --git a/Engine/Source/Programs/UnrealBuildTool/System/ActionGraph.cs b/Engine/Source/Programs/UnrealBuildTool/System/ActionGraph.cs
--- a/Engine/Source/Programs/UnrealBuildTool/System/ActionGraph.cs
+++ b/Engine/Source/Programs/UnrealBuildTool/System/ActionGraph.cs
@@ -24,6 +24,10 @@ namespace UnrealBuildTool
 		GenerateDebugInfo,
 
 		Link,
+
+		WriteMetadata,
+
+		PostBuildStep,
 	}
 
 	/// <summary>
@@ -281,6 +285,14 @@ namespace UnrealBuildTool
 			return NewAction;
 		}
 
+		public Action AddRecursiveCall(ActionType Type, string Arguments)
+		{
+			Action NewAction = Add(Type);
+			NewAction.CommandPath = UnrealBuildTool.GetUBTPath();
+			NewAction.CommandArguments = Arguments;
+			return NewAction;
+		}
+
 		public void FinalizeActionGraph()
 		{
 			// Link producing actions to the items they produce.
@@ -883,7 +895,7 @@ namespace UnrealBuildTool
 					// If the produced file doesn't exist or has zero size, consider it outdated.  The zero size check is to detect cases
 					// where aborting an earlier compile produced invalid zero-sized obj files, but that may cause actions where that's
 					// legitimate output to always be considered outdated.
-					if (ProducedItem.bExists && (ProducedItem.Length > 0 || ProducedItem.IsDirectory))
+					if (ProducedItem.bExists && (RootAction.ActionType != ActionType.Compile || ProducedItem.Length > 0 || ProducedItem.IsDirectory))
 					{
 						// Use the oldest produced item's time as the last execution time.
 						if (ProducedItem.LastWriteTime < LastExecutionTime)
diff --git a/Engine/Source/Programs/UnrealBuildTool/UnrealBuildTool.cs b/Engine/Source/Programs/UnrealBuildTool/UnrealBuildTool.cs
--- a/Engine/Source/Programs/UnrealBuildTool/UnrealBuildTool.cs
+++ b/Engine/Source/Programs/UnrealBuildTool/UnrealBuildTool.cs
@@ -1339,32 +1339,6 @@ namespace UnrealBuildTool
 									CppIncludeThread = new CppIncludeBackgroundThread(TargetToOutdatedPrerequisitesMap, TargetToHeaders);
 								}
 
-								// Check if we're allowed to modify manifests
-								bool bNoManifestChanges = Arguments.Any(x => x.Equals("-NoManifestChanges", StringComparison.OrdinalIgnoreCase));
-
-								// If we're not touching any shared files (ie. anything under Engine), allow the build ids to be recycled between applications.
-								HashSet<FileReference> OutputFiles = new HashSet<FileReference>(ActionsToExecute.SelectMany(x => x.ProducedItems.Select(y => y.Location)));
-								foreach (UEBuildTarget Target in Targets)
-								{
-									if (!Target.TryRecycleVersionManifests(OutputFiles, bNoManifestChanges))
-									{
-										if(bNoManifestChanges)
-										{
-											throw new BuildException("Stopping build due to manifest changes.");
-										}
-										else
-										{
-											Target.InvalidateVersionManifests();
-										}
-									}
-								}
-
-								// Remove the receipts, so we know the target is not valid if the compile fails
-								foreach (UEBuildTarget Target in Targets)
-								{
-									Target.DeleteReceipts();
-								}
-
 								// Execute the actions.
 								string TargetInfoForTelemetry = String.Join("|", Targets.Select(x => String.Format("{0} {1} {2}{3}", x.TargetName, x.Platform, x.Configuration, "")));
 								DateTime ExecutorStartTime = DateTime.UtcNow;
@@ -1384,7 +1358,6 @@ namespace UnrealBuildTool
 								{
 									foreach (UEBuildTarget Target in Targets)
 									{
-										Target.WriteReceipts(bNoManifestChanges);
 										UEBuildPlatform.GetBuildPlatform(Target.Platform).PostBuildSync(Target);
 									}
 									if (ActionsToExecute.Count == 0 && BuildConfiguration.bSkipLinkingWhenNothingToCompile)
@@ -1398,19 +1371,6 @@ namespace UnrealBuildTool
 								}
 							}
 
-							// Execute all the post-build steps
-						if (BuildResult.Succeeded() && !ProjectFileGenerator.bGenerateProjectFiles && !BuildConfiguration.bXGEExport)
-							foreach (UEBuildTarget Target in Targets)
-							{
-								if (!Target.ExecuteCustomPostBuildSteps())
-								{
-									{
-										BuildResult = ECompilationResult.OtherCompilationError;
-										break;
-									}
-								}
-							}
-
 							// Run the deployment steps
 							if (BuildResult.Succeeded()
 								&& String.IsNullOrEmpty(BuildConfiguration.SingleFileToCompile)
@@ -1976,10 +1936,61 @@ namespace UnrealBuildTool
 				}
 			}
 
-			foreach (UEBuildTarget Target in Targets)
+			// Update the action that writes out the module manifests
+			foreach(Action Action in ActionGraph.AllActions)
 			{
-				Target.PatchModuleManifestsForHotReloadAssembling(OnlyModulesLocal);
-				Target.WriteReceipts(false);
+				if(Action.ActionType == ActionType.WriteMetadata)
+				{
+					string Arguments = Action.CommandArguments;
+
+					// Find the argument for the metadata file
+					const string InputArgument = "-Input=";
+
+					int InputIdx = Arguments.IndexOf(InputArgument);
+					if(InputIdx == -1)
+					{
+						throw new Exception("Missing -Input= argument to WriteMetadata command when patching action graph.");
+					}
+
+					int FileNameIdx = InputIdx + InputArgument.Length;
+					if(Arguments[FileNameIdx] == '\"')
+					{
+						FileNameIdx++;
+					}
+
+					int FileNameEndIdx = FileNameIdx;
+					while(FileNameEndIdx < Arguments.Length && (Arguments[FileNameEndIdx] != ' ' || Arguments[FileNameIdx - 1] == '\"') && Arguments[FileNameEndIdx] != '\"')
+					{
+						FileNameEndIdx++;
+					}
+
+					// Read the metadata file
+					FileReference TargetInfoFile = new FileReference(Arguments.Substring(FileNameIdx, FileNameEndIdx - FileNameIdx));
+					if(!FileReference.Exists(TargetInfoFile))
+					{
+						throw new Exception(String.Format("Unable to find metadata file to patch action graph ({0})", TargetInfoFile));
+					}
+
+					// Update the module names
+					WriteMetadataTargetInfo TargetInfo = BinaryFormatterUtils.Load<WriteMetadataTargetInfo>(TargetInfoFile);
+					foreach (KeyValuePair<FileReference, ModuleManifest> FileNameToVersionManifest in TargetInfo.FileToManifest)
+					{
+						foreach (KeyValuePair<string, string> Manifest in FileNameToVersionManifest.Value.ModuleNameToFileName)
+						{
+							string ModuleFilename = Manifest.Value;
+							if (UnrealBuildTool.ReplaceHotReloadFilenameSuffix(ref ModuleFilename, (ModuleName) => UnrealBuildTool.GetReplacementModuleSuffix(OnlyModules, ModuleName)))
+							{
+								FileNameToVersionManifest.Value.ModuleNameToFileName[Manifest.Key] = ModuleFilename;
+								break;
+							}
+						}
+					}
+
+					// Write the hot-reload metadata file and update the argument list
+					FileReference HotReloadTargetInfoFile = FileReference.Combine(TargetInfoFile.Directory, "Metadata-HotReload.dat");
+					BinaryFormatterUtils.Save(HotReloadTargetInfoFile, TargetInfo);
+					Arguments = Arguments.Substring(0, FileNameIdx) + HotReloadTargetInfoFile + Arguments.Substring(FileNameEndIdx);
+				}
 			}
 		}
 	}
diff --git a/Engine/Source/Programs/UnrealBuildTool/UnrealBuildTool.csproj b/Engine/Source/Programs/UnrealBuildTool/UnrealBuildTool.csproj
--- a/Engine/Source/Programs/UnrealBuildTool/UnrealBuildTool.csproj
+++ b/Engine/Source/Programs/UnrealBuildTool/UnrealBuildTool.csproj
@@ -157,6 +157,7 @@
     <Compile Include="Executors\ParallelExecutor.cs" />
     <Compile Include="Executors\SNDBS.cs" />
     <Compile Include="Executors\XGE.cs" />
+    <Compile Include="Modes\WriteMetadataMode.cs" />
     <Compile Include="Platform\IOS\MobileProvisionContents.cs" />
     <Compile Include="Platform\Lumin\UEDeployLumin.cs" />
     <Compile Include="ProjectFiles\CLion\CLionGenerator.cs" />
