UBT: Remove a few unused variables from the deployment context.

#rnx

diff --git a/Engine/Source/Programs/UnrealBuildTool/Configuration/UEBuildDeployTarget.cs b/Engine/Source/Programs/UnrealBuildTool/Configuration/UEBuildDeployTarget.cs
--- a/Engine/Source/Programs/UnrealBuildTool/Configuration/UEBuildDeployTarget.cs
+++ b/Engine/Source/Programs/UnrealBuildTool/Configuration/UEBuildDeployTarget.cs
@@ -25,11 +25,6 @@ namespace UnrealBuildTool
 		/// </summary>
 		public readonly FileReference TargetFile;
 
-		/// <summary>
-		/// The shared app name for this target (eg. UE4Editor)
-		/// </summary>
-		public readonly string AppName;
-
 		/// <summary>
 		/// The name of this target
 		/// </summary>
@@ -70,11 +65,6 @@ namespace UnrealBuildTool
 		/// </summary>
 		public readonly List<FileReference> OutputPaths;
 
-		/// <summary>
-		/// Path to the directory for engine intermediates. May be under the project directory if not being built for the shared build environment.
-		/// </summary>
-		public readonly DirectoryReference EngineIntermediateDirectory;
-
 		/// <summary>
 		/// The project directory, or engine directory for targets without a project file.
 		/// </summary>
@@ -118,13 +108,11 @@ namespace UnrealBuildTool
 		{
 			this.ProjectFile = Target.ProjectFile;
 			this.TargetFile = Target.TargetRulesFile;
-			this.AppName = Target.AppName;
 			this.TargetName = Target.TargetName;
 			this.TargetType = Target.TargetType;
 			this.Platform = Target.Platform;
 			this.Configuration = Target.Configuration;
 			this.OutputPaths = new List<FileReference>(Target.OutputPaths);
-			this.EngineIntermediateDirectory = Target.EngineIntermediateDirectory;
 			this.ProjectDirectory = Target.ProjectDirectory;
 			this.BuildReceiptFileName = Target.ReceiptFileName;
 			this.bUseIncrementalLinking = Target.Rules.bUseIncrementalLinking;
@@ -144,13 +132,11 @@ namespace UnrealBuildTool
 			{
 				ProjectFile = Reader.ReadFileReference();
 				TargetFile = Reader.ReadFileReference();
-				AppName = Reader.ReadString();
 				TargetName = Reader.ReadString();
 				TargetType = (TargetType)Reader.ReadInt32();
 				Platform = (UnrealTargetPlatform)Reader.ReadInt32();
 				Configuration = (UnrealTargetConfiguration)Reader.ReadInt32();
 				OutputPaths = Reader.ReadList(() => Reader.ReadFileReference());
-				EngineIntermediateDirectory = Reader.ReadDirectoryReference();
 				ProjectDirectory = Reader.ReadDirectoryReference();
 				BuildReceiptFileName = Reader.ReadFileReference();
 				bUseIncrementalLinking = Reader.ReadBoolean();
@@ -172,13 +158,11 @@ namespace UnrealBuildTool
 			{
 				Writer.Write(ProjectFile);
 				Writer.Write(TargetFile);
-				Writer.Write(AppName);
 				Writer.Write(TargetName);
 				Writer.Write((Int32)TargetType);
 				Writer.Write((Int32)Platform);
 				Writer.Write((Int32)Configuration);
 				Writer.Write(OutputPaths, Item => Writer.Write(Item));
-				Writer.Write(EngineIntermediateDirectory);
 				Writer.Write(ProjectDirectory);
 				Writer.Write(BuildReceiptFileName);
 				Writer.Write(bUseIncrementalLinking);
