UBT: Expose the location of a target rules file via an internal memory on the TargetRules object.

diff --git a/Engine/Source/Programs/UnrealBuildTool/Configuration/TargetRules.cs b/Engine/Source/Programs/UnrealBuildTool/Configuration/TargetRules.cs
--- a/Engine/Source/Programs/UnrealBuildTool/Configuration/TargetRules.cs
+++ b/Engine/Source/Programs/UnrealBuildTool/Configuration/TargetRules.cs
@@ -139,6 +139,11 @@ namespace UnrealBuildTool
 		/// </summary>
 		public readonly string Name;
 
+		/// <summary>
+		/// File containing this target
+		/// </summary>
+		internal FileReference File;
+
 		/// <summary>
 		/// Platform that this target is being built for.
 		/// </summary>
@@ -1407,6 +1412,11 @@ namespace UnrealBuildTool
 			get { return Inner.Name; }
 		}
 
+		internal FileReference File
+		{
+			get { return Inner.File; }
+		}
+
 		public UnrealTargetPlatform Platform
 		{
 			get { return Inner.Platform; }
diff --git a/Engine/Source/Programs/UnrealBuildTool/Configuration/UEBuildTarget.cs b/Engine/Source/Programs/UnrealBuildTool/Configuration/UEBuildTarget.cs
--- a/Engine/Source/Programs/UnrealBuildTool/Configuration/UEBuildTarget.cs
+++ b/Engine/Source/Programs/UnrealBuildTool/Configuration/UEBuildTarget.cs
@@ -385,8 +385,7 @@ namespace UnrealBuildTool
 
 			RulesAssembly RulesAssembly = RulesCompiler.CreateTargetRulesAssembly(Desc.ProjectFile, Desc.Name, bSkipRulesCompile, bUsePrecompiled, Desc.ForeignPlugin);
 
-			FileReference TargetFileName;
-			TargetRules RulesObject = RulesAssembly.CreateTargetRules(Desc.Name, Desc.Platform, Desc.Configuration, Desc.Architecture, Desc.ProjectFile, Version, Arguments, out TargetFileName);
+			TargetRules RulesObject = RulesAssembly.CreateTargetRules(Desc.Name, Desc.Platform, Desc.Configuration, Desc.Architecture, Desc.ProjectFile, Version, Arguments);
 			if ((ProjectFileGenerator.bGenerateProjectFiles == false) && !GetSupportedPlatforms(RulesObject).Contains(Desc.Platform))
 			{
 				throw new BuildException("{0} does not support the {1} platform.", Desc.Name, Desc.Platform.ToString());
@@ -517,7 +516,7 @@ namespace UnrealBuildTool
 			}
 
 			// Generate a build target from this rules module
-			UEBuildTarget BuildTarget = new UEBuildTarget(Desc, new ReadOnlyTargetRules(RulesObject), RulesAssembly, TargetFileName);
+			UEBuildTarget BuildTarget = new UEBuildTarget(Desc, new ReadOnlyTargetRules(RulesObject), RulesAssembly);
 
 			if (UnrealBuildTool.bPrintPerformanceInfo)
 			{
@@ -880,8 +879,7 @@ namespace UnrealBuildTool
 		/// <param name="InDesc">Target descriptor</param>
 		/// <param name="InRules">The target rules, as created by RulesCompiler.</param>
 		/// <param name="InRulesAssembly">The chain of rules assemblies that this target was created with</param>
-		/// <param name="InTargetCsFilename">The name of the target </param>
-		public UEBuildTarget(TargetDescriptor InDesc, ReadOnlyTargetRules InRules, RulesAssembly InRulesAssembly, FileReference InTargetCsFilename)
+		public UEBuildTarget(TargetDescriptor InDesc, ReadOnlyTargetRules InRules, RulesAssembly InRulesAssembly)
 		{
 			ProjectFile = InDesc.ProjectFile;
 			AppName = InDesc.Name;
@@ -898,8 +896,7 @@ namespace UnrealBuildTool
 			// now that we have the platform, we can set the intermediate path to include the platform/architecture name
 			PlatformIntermediateFolder = Path.Combine("Intermediate", "Build", Platform.ToString(), UEBuildPlatform.GetBuildPlatform(Platform).GetFolderNameForArchitecture(Architecture));
 
-			Debug.Assert(InTargetCsFilename == null || InTargetCsFilename.HasExtension(".Target.cs"));
-			TargetRulesFile = InTargetCsFilename;
+			TargetRulesFile = InRules.File;
 
 			bCompileMonolithic = (Rules.LinkType == TargetLinkType.Monolithic);
 
@@ -928,16 +925,13 @@ namespace UnrealBuildTool
 			{
 				ProjectDirectory = ProjectFile.Directory;
 			}
+			else if (Rules.File.IsUnderDirectory(UnrealBuildTool.EnterpriseDirectory))
+			{
+				ProjectDirectory = UnrealBuildTool.EnterpriseDirectory;
+			}
 			else
 			{
-				if (InTargetCsFilename.IsUnderDirectory(UnrealBuildTool.EnterpriseDirectory))
-				{
-					ProjectDirectory = UnrealBuildTool.EnterpriseDirectory;
-				}
-				else
-				{
-					ProjectDirectory = UnrealBuildTool.EngineDirectory;
-				}
+				ProjectDirectory = UnrealBuildTool.EngineDirectory;
 			}
 
 			// Build the project intermediate directory
diff --git a/Engine/Source/Programs/UnrealBuildTool/System/RulesAssembly.cs b/Engine/Source/Programs/UnrealBuildTool/System/RulesAssembly.cs
--- a/Engine/Source/Programs/UnrealBuildTool/System/RulesAssembly.cs
+++ b/Engine/Source/Programs/UnrealBuildTool/System/RulesAssembly.cs
@@ -440,6 +440,9 @@ namespace UnrealBuildTool
 				throw new BuildException(Ex, "Unable to instantiate instance of '{0}' object type from compiled assembly '{1}'.  Unreal Build Tool creates an instance of your module's 'Rules' object in order to find out about your module's requirements.  The CLR exception details may provide more information:  {2}", TypeName, Path.GetFileNameWithoutExtension(CompiledAssembly.Location), Ex.ToString());
 			}
 
+			// Return the target file name to the caller
+			Rules.File = TargetNameToTargetFile[TargetInfo.Name];
+
 			// Set the default overriddes for the configured target type
 			Rules.SetOverridesForTargetType();
 
@@ -526,23 +529,6 @@ namespace UnrealBuildTool
 			return Rules;
 		}
 
-		/// <summary>
-		/// Creates a target rules object for the specified target name.
-		/// </summary>
-		/// <param name="TargetName">Name of the target</param>
-		/// <param name="Platform">The platform that the target is being built for</param>
-		/// <param name="Configuration">The configuration the target is being built for</param>
-		/// <param name="Architecture">The architecture the target is being built for</param>
-		/// <param name="ProjectFile">The project containing the target being built</param>
-		/// <param name="Arguments">Command line arguments for the target</param>
-		/// <param name="Version">The current build version</param>
-		/// <returns>The build target rules for the specified target</returns>
-		public TargetRules CreateTargetRules(string TargetName, UnrealTargetPlatform Platform, UnrealTargetConfiguration Configuration, string Architecture, FileReference ProjectFile, ReadOnlyBuildVersion Version, string[] Arguments)
-		{
-			FileReference TargetFileName;
-			return CreateTargetRules(TargetName, Platform, Configuration, Architecture, ProjectFile, Version, Arguments, out TargetFileName);
-		}
-
 		/// <summary>
 		/// Creates a target rules object for the specified target name.
 		/// </summary>
@@ -553,9 +539,8 @@ namespace UnrealBuildTool
 		/// <param name="ProjectFile">Path to the project file for this target</param>
 		/// <param name="Version">The current build version</param>
 		/// <param name="Arguments">Command line arguments for this target</param>
-		/// <param name="TargetFileName">The original source file name of the Target.cs file for this target</param>
 		/// <returns>The build target rules for the specified target</returns>
-		public TargetRules CreateTargetRules(string TargetName, UnrealTargetPlatform Platform, UnrealTargetConfiguration Configuration, string Architecture, FileReference ProjectFile, ReadOnlyBuildVersion Version, string[] Arguments, out FileReference TargetFileName)
+		public TargetRules CreateTargetRules(string TargetName, UnrealTargetPlatform Platform, UnrealTargetConfiguration Configuration, string Architecture, FileReference ProjectFile, ReadOnlyBuildVersion Version, string[] Arguments)
 		{
 			bool bFoundTargetName = TargetNameToTargetFile.ContainsKey(TargetName);
 			if (bFoundTargetName == false)
@@ -581,13 +566,10 @@ namespace UnrealBuildTool
 				}
 				else
 				{
-					return Parent.CreateTargetRules(TargetName, Platform, Configuration, Architecture, ProjectFile, Version, Arguments, out TargetFileName);
+					return Parent.CreateTargetRules(TargetName, Platform, Configuration, Architecture, ProjectFile, Version, Arguments);
 				}
 			}
 
-			// Return the target file name to the caller
-			TargetFileName = TargetNameToTargetFile[TargetName];
-
 			// Currently, we expect the user's rules object type name to be the same as the module name + 'Target'
 			string TargetTypeName = TargetName + "Target";
 
