UBT: Use the receipt for deployment rather than duplicating the output paths directly from the target.

diff --git a/Engine/Source/Programs/UnrealBuildTool/Configuration/UEBuildDeployTarget.cs b/Engine/Source/Programs/UnrealBuildTool/Configuration/UEBuildDeployTarget.cs
--- a/Engine/Source/Programs/UnrealBuildTool/Configuration/UEBuildDeployTarget.cs
+++ b/Engine/Source/Programs/UnrealBuildTool/Configuration/UEBuildDeployTarget.cs
@@ -45,26 +45,6 @@ namespace UnrealBuildTool
 		/// </summary>
 		public readonly UnrealTargetConfiguration Configuration;
 
-		/// <summary>
-		/// The output path
-		/// </summary>
-		public FileReference OutputPath
-		{
-			get
-			{
-				if (OutputPaths.Count != 1)
-				{
-					throw new BuildException("Attempted to use UEBuildDeployTarget.OutputPath property, but there are multiple (or no) OutputPaths. You need to handle multiple in the code that called this (size = {0})", OutputPaths.Count);
-				}
-				return OutputPaths[0];
-			}
-		}
-
-		/// <summary>
-		/// The full list of output paths, for platforms that support building multiple binaries simultaneously
-		/// </summary>
-		public readonly List<FileReference> OutputPaths;
-
 		/// <summary>
 		/// The project directory, or engine directory for targets without a project file.
 		/// </summary>
@@ -112,7 +92,6 @@ namespace UnrealBuildTool
 			this.TargetType = Target.TargetType;
 			this.Platform = Target.Platform;
 			this.Configuration = Target.Configuration;
-			this.OutputPaths = new List<FileReference>(Target.OutputPaths);
 			this.ProjectDirectory = Target.ProjectDirectory;
 			this.BuildReceiptFileName = Target.ReceiptFileName;
 			this.bUseIncrementalLinking = Target.Rules.bUseIncrementalLinking;
@@ -136,7 +115,6 @@ namespace UnrealBuildTool
 				TargetType = (TargetType)Reader.ReadInt32();
 				Platform = (UnrealTargetPlatform)Reader.ReadInt32();
 				Configuration = (UnrealTargetConfiguration)Reader.ReadInt32();
-				OutputPaths = Reader.ReadList(() => Reader.ReadFileReference());
 				ProjectDirectory = Reader.ReadDirectoryReference();
 				BuildReceiptFileName = Reader.ReadFileReference();
 				bUseIncrementalLinking = Reader.ReadBoolean();
@@ -162,7 +140,6 @@ namespace UnrealBuildTool
 				Writer.Write((Int32)TargetType);
 				Writer.Write((Int32)Platform);
 				Writer.Write((Int32)Configuration);
-				Writer.Write(OutputPaths, Item => Writer.Write(Item));
 				Writer.Write(ProjectDirectory);
 				Writer.Write(BuildReceiptFileName);
 				Writer.Write(bUseIncrementalLinking);
diff --git a/Engine/Source/Programs/UnrealBuildTool/Platform/Android/UEDeployAndroid.cs b/Engine/Source/Programs/UnrealBuildTool/Platform/Android/UEDeployAndroid.cs
--- a/Engine/Source/Programs/UnrealBuildTool/Platform/Android/UEDeployAndroid.cs
+++ b/Engine/Source/Programs/UnrealBuildTool/Platform/Android/UEDeployAndroid.cs
@@ -3998,16 +3998,17 @@ namespace UnrealBuildTool
 		{
 			AndroidToolChain ToolChain = ((AndroidPlatform)UEBuildPlatform.GetBuildPlatform(InTarget.Platform)).CreateTempToolChainForProject(InTarget.ProjectFile) as AndroidToolChain;
 
-			// we need to strip architecture from any of the output paths
-			string BaseSoName = ToolChain.RemoveArchName(InTarget.OutputPaths[0].FullName);
-
 			// get the receipt
-			UnrealTargetPlatform Platform = InTarget.Platform;
-			UnrealTargetConfiguration Configuration = InTarget.Configuration;
-			string ProjectBaseName = Path.GetFileName(BaseSoName).Replace("-" + Platform, "").Replace("-" + Configuration, "").Replace(".so", "");
-			FileReference ReceiptFilename = TargetReceipt.GetDefaultPath(InTarget.ProjectDirectory, ProjectBaseName, Platform, Configuration, "");
+			FileReference ReceiptFilename = InTarget.BuildReceiptFileName;
 			Log.TraceInformation("Receipt Filename: {0}", ReceiptFilename);
-			SetAndroidPluginData(ToolChain.GetAllArchitectures(), CollectPluginDataPaths(TargetReceipt.Read(ReceiptFilename, UnrealBuildTool.EngineDirectory, InTarget.ProjectDirectory)));
+			TargetReceipt Receipt = TargetReceipt.Read(ReceiptFilename, UnrealBuildTool.EngineDirectory, InTarget.ProjectDirectory);
+			SetAndroidPluginData(ToolChain.GetAllArchitectures(), CollectPluginDataPaths(Receipt));
+
+			// Get the output paths
+			List<FileReference> OutputPaths = Receipt.BuildProducts.Where(x => x.Type == BuildProductType.Executable).Select(x => x.Path).ToList();
+
+			// we need to strip architecture from any of the output paths
+			string BaseSoName = ToolChain.RemoveArchName(OutputPaths[0].FullName);
 
 			// make an apk at the end of compiling, so that we can run without packaging (debugger, cook on the fly, etc)
 			string RelativeEnginePath = UnrealBuildTool.EngineDirectory.MakeRelativeTo(DirectoryReference.GetCurrentDirectory());
@@ -4016,7 +4017,7 @@ namespace UnrealBuildTool
 				bMakeSeparateApks: ShouldMakeSeparateApks(), bIncrementalPackage: true, bDisallowPackagingDataInApk: false, bDisallowExternalFilesDir: true);
 
 			// if we made any non-standard .apk files, the generated debugger settings may be wrong
-			if (ShouldMakeSeparateApks() && (InTarget.OutputPaths.Count > 1 || !InTarget.OutputPaths[0].FullName.Contains("-armv7-es2")))
+			if (ShouldMakeSeparateApks() && (OutputPaths.Count > 1 || !OutputPaths[0].FullName.Contains("-armv7-es2")))
 			{
 				Log.TraceInformation("================================================================================================================================");
 				Log.TraceInformation("Non-default apk(s) have been made: If you are debugging, you will need to manually select one to run in the debugger properties!");
