UBT: Remove unused setting for incremental linking.

#rnx

diff --git a/Engine/Source/Programs/UnrealBuildTool/Configuration/UEBuildDeployTarget.cs b/Engine/Source/Programs/UnrealBuildTool/Configuration/UEBuildDeployTarget.cs
--- a/Engine/Source/Programs/UnrealBuildTool/Configuration/UEBuildDeployTarget.cs
+++ b/Engine/Source/Programs/UnrealBuildTool/Configuration/UEBuildDeployTarget.cs
@@ -55,11 +55,6 @@ namespace UnrealBuildTool
 		/// </summary>
 		public readonly FileReference BuildReceiptFileName;
 
-		/// <summary>
-		/// Whether incremental linking is enabled (disables icon update on Windows)
-		/// </summary>
-		public readonly bool bUseIncrementalLinking;
-
 		/// <summary>
 		/// If true, then a stub IPA will be generated when compiling is done (minimal files needed for a valid IPA)
 		/// </summary>
@@ -94,7 +89,6 @@ namespace UnrealBuildTool
 			this.Configuration = Target.Configuration;
 			this.ProjectDirectory = Target.ProjectDirectory;
 			this.BuildReceiptFileName = Target.ReceiptFileName;
-			this.bUseIncrementalLinking = Target.Rules.bUseIncrementalLinking;
 			this.bCreateStubIPA = Target.Rules.bCreateStubIPA;
 			this.AndroidArchitectures = Target.Rules.AndroidPlatform.Architectures.ToArray();
 			this.AndroidGPUArchitectures = Target.Rules.AndroidPlatform.GPUArchitectures.ToArray();
@@ -117,7 +111,6 @@ namespace UnrealBuildTool
 				Configuration = (UnrealTargetConfiguration)Reader.ReadInt32();
 				ProjectDirectory = Reader.ReadDirectoryReference();
 				BuildReceiptFileName = Reader.ReadFileReference();
-				bUseIncrementalLinking = Reader.ReadBoolean();
 				bCreateStubIPA = Reader.ReadBoolean();
 				AndroidArchitectures = Reader.ReadStringArray();
 				AndroidGPUArchitectures = Reader.ReadStringArray();
@@ -142,7 +135,6 @@ namespace UnrealBuildTool
 				Writer.Write((Int32)Configuration);
 				Writer.Write(ProjectDirectory);
 				Writer.Write(BuildReceiptFileName);
-				Writer.Write(bUseIncrementalLinking);
 				Writer.Write(bCreateStubIPA);
 				Writer.Write(AndroidArchitectures);
 				Writer.Write(AndroidGPUArchitectures);
