UBT: Remove the target file path from the deployment info. The same information can be derived from the project file path.

#rnx

diff --git a/Engine/Source/Programs/UnrealBuildTool/Configuration/UEBuildDeployTarget.cs b/Engine/Source/Programs/UnrealBuildTool/Configuration/UEBuildDeployTarget.cs
--- a/Engine/Source/Programs/UnrealBuildTool/Configuration/UEBuildDeployTarget.cs
+++ b/Engine/Source/Programs/UnrealBuildTool/Configuration/UEBuildDeployTarget.cs
@@ -20,11 +20,6 @@ namespace UnrealBuildTool
 		/// </summary>
 		public readonly FileReference ProjectFile;
 
-		/// <summary>
-		/// Path to the .target.cs
-		/// </summary>
-		public readonly FileReference TargetFile;
-
 		/// <summary>
 		/// The name of this target
 		/// </summary>
@@ -82,7 +77,6 @@ namespace UnrealBuildTool
 		public UEBuildDeployTarget(UEBuildTarget Target)
 		{
 			this.ProjectFile = Target.ProjectFile;
-			this.TargetFile = Target.TargetRulesFile;
 			this.TargetName = Target.TargetName;
 			this.TargetType = Target.TargetType;
 			this.Platform = Target.Platform;
@@ -104,7 +98,6 @@ namespace UnrealBuildTool
 			using (BinaryReader Reader = new BinaryReader(File.Open(Location.FullName, FileMode.Open, FileAccess.Read, FileShare.Read)))
 			{
 				ProjectFile = Reader.ReadFileReference();
-				TargetFile = Reader.ReadFileReference();
 				TargetName = Reader.ReadString();
 				TargetType = (TargetType)Reader.ReadInt32();
 				Platform = (UnrealTargetPlatform)Reader.ReadInt32();
@@ -128,7 +121,6 @@ namespace UnrealBuildTool
 			using (BinaryWriter Writer = new BinaryWriter(File.Open(Location.FullName, FileMode.Create, FileAccess.Write, FileShare.Read)))
 			{
 				Writer.Write(ProjectFile);
-				Writer.Write(TargetFile);
 				Writer.Write(TargetName);
 				Writer.Write((Int32)TargetType);
 				Writer.Write((Int32)Platform);
