UBT: Remove unused Android architecture settings from the deployment context.

#rnx

diff --git a/Engine/Source/Programs/UnrealBuildTool/Configuration/UEBuildDeployTarget.cs b/Engine/Source/Programs/UnrealBuildTool/Configuration/UEBuildDeployTarget.cs
--- a/Engine/Source/Programs/UnrealBuildTool/Configuration/UEBuildDeployTarget.cs
+++ b/Engine/Source/Programs/UnrealBuildTool/Configuration/UEBuildDeployTarget.cs
@@ -55,21 +55,6 @@ namespace UnrealBuildTool
 		/// </summary>
 		public readonly bool bCreateStubIPA;
 
-		/// <summary>
-		/// Which architectures to deploy for on Android
-		/// </summary>
-		public readonly string[] AndroidArchitectures;
-
-		/// <summary>
-		/// Which GPU architectures to deploy for on Android
-		/// </summary>
-		public readonly string[] AndroidGPUArchitectures;
-
-		/// <summary>
-		/// Which GPU architectures to deploy for on Lumin
-		/// </summary>
-		public readonly string[] LuminGPUArchitectures;
-
 		/// <summary>
 		/// Construct the deployment info from a target
 		/// </summary>
@@ -84,9 +69,6 @@ namespace UnrealBuildTool
 			this.ProjectDirectory = Target.ProjectDirectory;
 			this.BuildReceiptFileName = Target.ReceiptFileName;
 			this.bCreateStubIPA = Target.Rules.bCreateStubIPA;
-			this.AndroidArchitectures = Target.Rules.AndroidPlatform.Architectures.ToArray();
-			this.AndroidGPUArchitectures = Target.Rules.AndroidPlatform.GPUArchitectures.ToArray();
-			this.LuminGPUArchitectures = Target.Rules.LuminPlatform.GPUArchitectures.ToArray();
 		}
 
 		/// <summary>
@@ -105,9 +87,6 @@ namespace UnrealBuildTool
 				ProjectDirectory = Reader.ReadDirectoryReference();
 				BuildReceiptFileName = Reader.ReadFileReference();
 				bCreateStubIPA = Reader.ReadBoolean();
-				AndroidArchitectures = Reader.ReadStringArray();
-				AndroidGPUArchitectures = Reader.ReadStringArray();
-				LuminGPUArchitectures = Reader.ReadStringArray();
 			}
 		}
 
@@ -128,9 +107,6 @@ namespace UnrealBuildTool
 				Writer.Write(ProjectDirectory);
 				Writer.Write(BuildReceiptFileName);
 				Writer.Write(bCreateStubIPA);
-				Writer.Write(AndroidArchitectures);
-				Writer.Write(AndroidGPUArchitectures);
-				Writer.Write(LuminGPUArchitectures);
 			}
 		}
 	}
