PR #5237: Fix missing VS version indicator on .sln icon (Contributed by [REDACTED_PERSON])

diff --git a/Engine/Source/Programs/UnrealBuildTool/ProjectFiles/VisualStudio/VCProjectFileGenerator.cs b/Engine/Source/Programs/UnrealBuildTool/ProjectFiles/VisualStudio/VCProjectFileGenerator.cs
--- a/Engine/Source/Programs/UnrealBuildTool/ProjectFiles/VisualStudio/VCProjectFileGenerator.cs
+++ b/Engine/Source/Programs/UnrealBuildTool/ProjectFiles/VisualStudio/VCProjectFileGenerator.cs
@@ -397,9 +397,10 @@ namespace UnrealBuildTool
 			// Setup solution file content
 			StringBuilder VCSolutionFileContent = new StringBuilder();
 
-			// Solution file header
+			// Solution file header. Note that a leading newline is required for file type detection to work correclty in the shell.
 			if (ProjectFileFormat == VCProjectFileFormat.VisualStudio2017)
 			{
+				VCSolutionFileContent.AppendLine();
 				VCSolutionFileContent.AppendLine("Microsoft Visual Studio Solution File, Format Version 12.00");
 				VCSolutionFileContent.AppendLine("# Visual Studio 15");
 				VCSolutionFileContent.AppendLine("VisualStudioVersion = 15.0.25807.0");
@@ -407,6 +408,7 @@ namespace UnrealBuildTool
 			}
 			else if (ProjectFileFormat == VCProjectFileFormat.VisualStudio2015)
 			{
+				VCSolutionFileContent.AppendLine();
 				VCSolutionFileContent.AppendLine("Microsoft Visual Studio Solution File, Format Version 12.00");
 				VCSolutionFileContent.AppendLine("# Visual Studio 14");
 				VCSolutionFileContent.AppendLine("VisualStudioVersion = 14.0.22310.1");
@@ -414,11 +416,13 @@ namespace UnrealBuildTool
 			}
 			else if (ProjectFileFormat == VCProjectFileFormat.VisualStudio2013)
 			{
+				VCSolutionFileContent.AppendLine();
 				VCSolutionFileContent.AppendLine("Microsoft Visual Studio Solution File, Format Version 12.00");
 				VCSolutionFileContent.AppendLine("# Visual Studio 2013");
             }
             else if (ProjectFileFormat == VCProjectFileFormat.VisualStudio2012)
             {
+				VCSolutionFileContent.AppendLine();
                 VCSolutionFileContent.AppendLine("Microsoft Visual Studio Solution File, Format Version 12.00");
 				VCSolutionFileContent.AppendLine("# Visual Studio 2012");
             }
