Remake change to check for installed enterprise files after merge from //UE4/Main.

#rnx

diff --git a/Engine/Source/Programs/UnrealBuildTool/UnrealBuildTool.cs b/Engine/Source/Programs/UnrealBuildTool/UnrealBuildTool.cs
--- a/Engine/Source/Programs/UnrealBuildTool/UnrealBuildTool.cs
+++ b/Engine/Source/Programs/UnrealBuildTool/UnrealBuildTool.cs
@@ -196,6 +196,10 @@ namespace UnrealBuildTool
 			{
 				return true;
 			}
+			if(IsEnterpriseInstalled() && File.IsUnderDirectory(EnterpriseDirectory))
+			{
+				return true;
+			}
 			if(IsProjectInstalled() && File.IsUnderDirectory(InstalledProjectFile.Directory))
 			{
 				return true;
