Print the name of the receipt when compiling, rather than "Writing receipts..."

#rnx

diff --git a/Engine/Source/Programs/UnrealBuildTool/Configuration/UEBuildTarget.cs b/Engine/Source/Programs/UnrealBuildTool/Configuration/UEBuildTarget.cs
--- a/Engine/Source/Programs/UnrealBuildTool/Configuration/UEBuildTarget.cs
+++ b/Engine/Source/Programs/UnrealBuildTool/Configuration/UEBuildTarget.cs
@@ -1913,7 +1913,7 @@ namespace UnrealBuildTool
 
 				Action WriteMetadataAction = ActionGraph.AddRecursiveCall(ActionType.WriteMetadata, WriteMetadataArguments.ToString());
 				WriteMetadataAction.WorkingDirectory = UnrealBuildTool.EngineSourceDirectory.FullName;
-				WriteMetadataAction.StatusDescription = "Writing receipt...";
+				WriteMetadataAction.StatusDescription = ReceiptFileName.GetFileName();
 				WriteMetadataAction.bCanExecuteRemotely = false;
 				WriteMetadataAction.PrerequisiteItems.AddRange(OutputItems);
 				WriteMetadataAction.ProducedItems.Add(FileItem.GetItemByFileReference(ReceiptFileName));
