UBT: Set the appropriate compiler and linker arguments to allow using Live++ without modifications.

* /Gw is always enabled for compilation, which puts symbols in their own section and allows the linker to remove any unreferenced symbols.
* /FUNCTIONPADMIN is enabled in non-shipping builds (via TargetRules.WindowsPlatform.bCreateHotpatchableImage), which allows functions to be detoured.

diff --git a/Engine/Source/Programs/UnrealBuildTool/Platform/Windows/UEBuildWindows.cs b/Engine/Source/Programs/UnrealBuildTool/Platform/Windows/UEBuildWindows.cs
--- a/Engine/Source/Programs/UnrealBuildTool/Platform/Windows/UEBuildWindows.cs
+++ b/Engine/Source/Programs/UnrealBuildTool/Platform/Windows/UEBuildWindows.cs
@@ -187,6 +187,11 @@ namespace UnrealBuildTool
 		/// </summary>
 		public bool bBuildLargeAddressAwareBinary = true;
 
+		/// <summary>
+		/// Create an image that can be hot patched (/FUNCTIONPADMIN)
+		/// </summary>
+		public bool bCreateHotPatchableImage = false;
+
 		/// <summary>
 		/// The Visual C++ environment to use for this target. Only initialized after all the target settings are finalized, in ValidateTarget().
 		/// </summary>
@@ -348,6 +353,11 @@ namespace UnrealBuildTool
 			get { return Inner.bBuildLargeAddressAwareBinary; }
 		}
 
+		public bool bCreateHotpatchableImage
+		{
+			get { return Inner.bCreateHotPatchableImage; }
+		}
+
 		public string GetVisualStudioCompilerVersionName()
 		{
 			return Inner.GetVisualStudioCompilerVersionName();
@@ -442,6 +452,20 @@ namespace UnrealBuildTool
 			return SDK.HasRequiredSDKsInstalled();
 		}
 
+		/// <summary>
+		/// Reset a target's settings to the default
+		/// </summary>
+		/// <param name="Target"></param>
+		public override void ResetTarget(TargetRules Target)
+		{
+			base.ResetTarget(Target);
+
+			if(Target.Configuration != UnrealTargetConfiguration.Shipping)
+			{
+				Target.WindowsPlatform.bCreateHotPatchableImage = true;
+			}
+		}
+
 		/// <summary>
 		/// Validate a target's settings
 		/// </summary>
diff --git a/Engine/Source/Programs/UnrealBuildTool/Platform/Windows/VCToolChain.cs b/Engine/Source/Programs/UnrealBuildTool/Platform/Windows/VCToolChain.cs
--- a/Engine/Source/Programs/UnrealBuildTool/Platform/Windows/VCToolChain.cs
+++ b/Engine/Source/Programs/UnrealBuildTool/Platform/Windows/VCToolChain.cs
@@ -189,6 +189,9 @@ namespace UnrealBuildTool
 			// Compile into an .obj file, and skip linking.
 			Arguments.Add("/c");
 
+			// Put symbols into different sections so the linker can remove them.
+			Arguments.Add("/Gw");
+
 			// Separate functions for linker.
 			Arguments.Add("/Gy");
 
@@ -1257,6 +1260,12 @@ namespace UnrealBuildTool
 				}
 			}
 
+			// Enable function level hot-patching
+			if(!bBuildImportLibraryOnly)
+			{
+				Arguments.Add("/FUNCTIONPADMIN");
+			}
+
 			// For targets that are cross-referenced, we don't want to write a LIB file during the link step as that
 			// file will clobber the import library we went out of our way to generate during an earlier step.  This
 			// file is not needed for our builds, but there is no way to prevent MSVC from generating it when
