Merging //UE4/Dev-Main to Dev-Build (//UE4/Dev-Build)

diff --git a/Engine/Build/BatchFiles/MakeAndInstallSSHKey.bat b/Engine/Build/BatchFiles/MakeAndInstallSSHKey.bat
--- a/Engine/Build/BatchFiles/MakeAndInstallSSHKey.bat
+++ b/Engine/Build/BatchFiles/MakeAndInstallSSHKey.bat
@@ -33,7 +33,7 @@ echo.
 
 pause
 
-%SSH% -p %SSHPORT% "%USER%@%MACHINE%" "if [[ ! -e .ssh ]]; then mkdir .ssh; fi && cd .ssh && if [[ -e authorized_keys ]]; then cp -f authorized_keys authorized_keys_UEBackup; fi && ssh-keygen -t rsa -f RemoteToolChain && mv -f RemoteToolChain.pub RemoteToolChainPublic.key && mv -f RemoteToolChain RemoteToolChainPrivate.key && cat RemoteToolChainPublic.key >> authorized_keys";
+%SSH% -p %SSHPORT% "%USER%@%MACHINE%" "if [[ ! -e .ssh ]]; then mkdir .ssh; fi && cd .ssh && if [[ -e authorized_keys ]]; then cp -f authorized_keys authorized_keys_UEBackup; fi && ssh-keygen -t rsa -m PEM -f RemoteToolChain && mv -f RemoteToolChain.pub RemoteToolChainPublic.key && mv -f RemoteToolChain RemoteToolChainPrivate.key && cat RemoteToolChainPublic.key >> authorized_keys";
 
 echo.
 echo ================================================================================
diff --git a/Engine/Source/Programs/DotNETCommon/DotNETUtilities/StringUtils.cs b/Engine/Source/Programs/DotNETCommon/DotNETUtilities/StringUtils.cs
--- a/Engine/Source/Programs/DotNETCommon/DotNETUtilities/StringUtils.cs
+++ b/Engine/Source/Programs/DotNETCommon/DotNETUtilities/StringUtils.cs
@@ -9,6 +9,22 @@ namespace Tools.DotNETCommon
 {
 	public static class StringUtils
 	{
+		/// <summary>
+		/// Indents a string by a given indent
+		/// </summary>
+		/// <param name="Text">The text to indent</param>
+		/// <param name="Indent">The indent to add to each line</param>
+		/// <returns>The indented string</returns>
+		public static string Indent(string Text, string Indent)
+		{
+			string Result = "";
+			if(Text.Length > 0)
+			{
+				Result = Indent + Text.Replace("\n", "\n" + Indent);
+			}
+			return Result;
+		}
+
 		/// <summary>
 		/// Extension method to allow formatting a string to a stringbuilder and appending a newline
 		/// </summary>
diff --git a/Engine/Source/Programs/UnrealBuildTool/ToolChain/RemoteMac.cs b/Engine/Source/Programs/UnrealBuildTool/ToolChain/RemoteMac.cs
--- a/Engine/Source/Programs/UnrealBuildTool/ToolChain/RemoteMac.cs
+++ b/Engine/Source/Programs/UnrealBuildTool/ToolChain/RemoteMac.cs
@@ -220,6 +220,9 @@ namespace UnrealBuildTool
 				}
 			}
 
+			// Print the path to the private key
+			Log.TraceInformation("[Remote] Using private key at {0}", SshPrivateKey);
+
 			// resolve the rest of the strings
 			RsyncAuthentication = ExpandVariables(RsyncAuthentication);
 			SshAuthentication = ExpandVariables(SshAuthentication);
@@ -250,7 +253,7 @@ namespace UnrealBuildTool
 			StringBuilder Output;
 			if(ExecuteAndCaptureOutput("'echo ~'", out Output) != 0)
 			{
-				throw new BuildException("Unable to determine home directory for remote user");
+				throw new BuildException("Unable to determine home directory for remote user. SSH output:\n{0}", StringUtils.Indent(Output.ToString(), "  "));
 			}
 			RemoteBaseDir = String.Format("{0}/UE4/Builds/{1}", Output.ToString().Trim().TrimEnd('/'), Environment.MachineName);
 			Log.TraceInformation("[Remote] Using base directory '{0}'", RemoteBaseDir);
@@ -291,8 +294,13 @@ namespace UnrealBuildTool
 				FileReference KeyFile = FileReference.Combine(Location, "SSHKeys", ServerName, UserName, "RemoteToolChainPrivate.key");
 				if (FileReference.Exists(KeyFile))
 				{
-					OutPrivateKey = KeyFile;
-					return true;
+					// MacOS Mojave includes a new version of SSH that generates keys that are incompatible with our version of SSH. Make sure the detected keys have the right signature.
+					string Text = FileReference.ReadAllText(KeyFile);
+					if(Text.Contains("---BEGIN RSA PRIVATE KEY---"))
+					{
+						OutPrivateKey = KeyFile;
+						return true;
+					}
 				}
 			}
 
