UBT: Store a relative path to the project file within a target receipt. This simplifies loading and saving receipts by removing the need to pass a project directory.

diff --git a/Engine/Source/Programs/AutomationTool/BuildGraph/Tasks/SanitizeReceiptTask.cs b/Engine/Source/Programs/AutomationTool/BuildGraph/Tasks/SanitizeReceiptTask.cs
--- a/Engine/Source/Programs/AutomationTool/BuildGraph/Tasks/SanitizeReceiptTask.cs
+++ b/Engine/Source/Programs/AutomationTool/BuildGraph/Tasks/SanitizeReceiptTask.cs
@@ -25,12 +25,6 @@ namespace AutomationTool.Tasks
 		/// </summary>
 		[TaskParameter(Optional = true)]
 		public DirectoryReference EngineDir;
-
-		/// <summary>
-		/// Path to the project folder, used to expand $(ProjectDir) properties in receipt files. Defaults to the Engine directory for the current workspace.
-		/// </summary>
-		[TaskParameter(Optional = true)]
-		public DirectoryReference ProjectDir;
 	}
 
 	/// <summary>
@@ -64,9 +58,6 @@ namespace AutomationTool.Tasks
 			// Set the Engine directory
 			DirectoryReference EngineDir = Parameters.EngineDir ?? CommandUtils.EngineDirectory;
 
-			// Set the Project directory
-			DirectoryReference ProjectDir = Parameters.ProjectDir ?? EngineDir;
-
 			// Resolve the input list
 			IEnumerable<FileReference> TargetFiles = ResolveFilespec(CommandUtils.RootDirectory, Parameters.Files, TagNameToFileSet);
 			foreach(FileReference TargetFile in TargetFiles)
@@ -83,7 +74,7 @@ namespace AutomationTool.Tasks
 				{
 					// Read the receipt
 					TargetReceipt Receipt;
-					if (!TargetReceipt.TryRead(TargetFile, EngineDir, ProjectDir, out Receipt))
+					if (!TargetReceipt.TryRead(TargetFile, EngineDir, out Receipt))
 					{
 						CommandUtils.LogWarning("Unable to load file using TagReceipt task ({0})", TargetFile.FullName);
 						continue;
@@ -120,7 +111,7 @@ namespace AutomationTool.Tasks
 					Receipt.RuntimeDependencies = NewRuntimeDependencies;
 				
 					// Save the new receipt
-					Receipt.Write(TargetFile, EngineDir, ProjectDir);
+					Receipt.Write(TargetFile, EngineDir);
 				}
 			}
 		}
diff --git a/Engine/Source/Programs/AutomationTool/BuildGraph/Tasks/StageTask.cs b/Engine/Source/Programs/AutomationTool/BuildGraph/Tasks/StageTask.cs
--- a/Engine/Source/Programs/AutomationTool/BuildGraph/Tasks/StageTask.cs
+++ b/Engine/Source/Programs/AutomationTool/BuildGraph/Tasks/StageTask.cs
@@ -115,7 +115,7 @@ namespace BuildGraph.Tasks
 
 			// Try to load it
 			TargetReceipt Receipt;
-			if(!TargetReceipt.TryRead(ReceiptFileName, SourceEngineDir, SourceProjectDir, out Receipt))
+			if(!TargetReceipt.TryRead(ReceiptFileName, out Receipt))
 			{
 				throw new AutomationException("Couldn't read receipt '{0}'", ReceiptFileName);
 			}
diff --git a/Engine/Source/Programs/AutomationTool/BuildGraph/Tasks/TagReceiptTask.cs b/Engine/Source/Programs/AutomationTool/BuildGraph/Tasks/TagReceiptTask.cs
--- a/Engine/Source/Programs/AutomationTool/BuildGraph/Tasks/TagReceiptTask.cs
+++ b/Engine/Source/Programs/AutomationTool/BuildGraph/Tasks/TagReceiptTask.cs
@@ -30,7 +30,7 @@ namespace AutomationTool.Tasks
 		public DirectoryReference EngineDir;
 
 		/// <summary>
-		/// Path to the project folder, used to expand $(ProjectDir) properties in receipt files. Defaults to the Engine directory for the current workspace.
+		/// Path to the project folder, used to expand $(ProjectDir) properties in receipt files. Defaults to the Engine directory for the current workspace. DEPRECATED.
 		/// </summary>
 		[TaskParameter(Optional = true)]
 		public DirectoryReference ProjectDir;
@@ -113,12 +113,15 @@ namespace AutomationTool.Tasks
 		/// <param name="TagNameToFileSet">Mapping from tag names to the set of files they include</param>
 		public override void Execute(JobContext Job, HashSet<FileReference> BuildProducts, Dictionary<string, HashSet<FileReference>> TagNameToFileSet)
 		{
+			// Output a warning if the project directory is specified
+			if(Parameters.ProjectDir != null)
+			{
+				CommandUtils.LogWarning("The ProjectDir argument to the TagReceipt parameter is deprecated. This path is now determined automatically from the receipt.");
+			}
+
 			// Set the Engine directory
 			DirectoryReference EngineDir = Parameters.EngineDir ?? CommandUtils.EngineDirectory;
 
-			// Set the Project directory
-			DirectoryReference ProjectDir = Parameters.ProjectDir ?? EngineDir;
-
 			// Resolve the input list
 			IEnumerable<FileReference> TargetFiles = ResolveFilespec(CommandUtils.RootDirectory, Parameters.Files, TagNameToFileSet);
 			HashSet<FileReference> Files = new HashSet<FileReference>();
@@ -133,7 +136,7 @@ namespace AutomationTool.Tasks
 
 				// Read the receipt
 				TargetReceipt Receipt;
-				if (!TargetReceipt.TryRead(TargetFile, EngineDir, ProjectDir, out Receipt))
+				if (!TargetReceipt.TryRead(TargetFile, EngineDir, out Receipt))
 				{
 					CommandUtils.LogWarning("Unable to load file using TagReceipt task ({0})", TargetFile.FullName);
 					continue;
diff --git a/Engine/Source/Programs/AutomationTool/Linux/LinuxPlatform.Automation.cs b/Engine/Source/Programs/AutomationTool/Linux/LinuxPlatform.Automation.cs
--- a/Engine/Source/Programs/AutomationTool/Linux/LinuxPlatform.Automation.cs
+++ b/Engine/Source/Programs/AutomationTool/Linux/LinuxPlatform.Automation.cs
@@ -1,4 +1,4 @@
-﻿// Copyright 1998-2018 Epic Games, Inc. All Rights Reserved.
+// Copyright 1998-2018 Epic Games, Inc. All Rights Reserved.
 using System;
 using System.Collections.Generic;
 using System.Linq;
@@ -28,9 +28,7 @@ public abstract class BaseLinuxPlatform : Platform
 			FileReference ReceiptFileName = TargetReceipt.GetDefaultPath(CommandUtils.EngineDirectory, "CrashReportClient", SC.StageTargetPlatform.PlatformType, UnrealTargetConfiguration.Shipping, null);
 			if (FileReference.Exists(ReceiptFileName))
 			{
-				DirectoryReference EngineDir = CommandUtils.EngineDirectory;
-				DirectoryReference ProjectDir = DirectoryReference.FromFile(Params.RawProjectPath);
-				TargetReceipt Receipt = TargetReceipt.Read(ReceiptFileName, EngineDir, ProjectDir);
+				TargetReceipt Receipt = TargetReceipt.Read(ReceiptFileName);
 				SC.StageBuildProductsFromReceipt(Receipt, true, false);
 			}
 		}
diff --git a/Engine/Source/Programs/AutomationTool/Scripts/CopyBuildToStagingDirectory.Automation.cs b/Engine/Source/Programs/AutomationTool/Scripts/CopyBuildToStagingDirectory.Automation.cs
--- a/Engine/Source/Programs/AutomationTool/Scripts/CopyBuildToStagingDirectory.Automation.cs
+++ b/Engine/Source/Programs/AutomationTool/Scripts/CopyBuildToStagingDirectory.Automation.cs
@@ -704,7 +704,7 @@ public partial class Project : CommandUtils
 
 					// Read the receipt for this target
 					TargetReceipt Receipt;
-					if (!TargetReceipt.TryRead(ReceiptFileName, SC.EngineRoot, null, out Receipt))
+					if (!TargetReceipt.TryRead(ReceiptFileName, out Receipt))
 					{
 						throw new AutomationException("Missing or invalid target receipt ({0})", ReceiptFileName);
 					}
@@ -2734,7 +2734,7 @@ public partial class Project : CommandUtils
 
 						// Read the receipt for this target
 						TargetReceipt Receipt;
-						if (!TargetReceipt.TryRead(ReceiptFileName, EngineDir, ProjectDir, out Receipt))
+						if (!TargetReceipt.TryRead(ReceiptFileName, out Receipt))
 						{
 							throw new AutomationException("Missing or invalid target receipt ({0})", ReceiptFileName);
 						}
diff --git a/Engine/Source/Programs/AutomationTool/Scripts/StashTarget.cs b/Engine/Source/Programs/AutomationTool/Scripts/StashTarget.cs
--- a/Engine/Source/Programs/AutomationTool/Scripts/StashTarget.cs
+++ b/Engine/Source/Programs/AutomationTool/Scripts/StashTarget.cs
@@ -35,7 +35,7 @@ namespace AutomationTool
 				throw new AutomationException("Unable to find '{0}'", ReceiptFile);
 			}
 
-			TargetReceipt Receipt = TargetReceipt.Read(ReceiptFile, EngineDirectory, DirectoryReference.FromFile(ProjectFile));
+			TargetReceipt Receipt = TargetReceipt.Read(ReceiptFile);
 
 			// Enumerate all the files we want to move
 			List<FileReference> FilesToMove = new List<FileReference>();
diff --git a/Engine/Source/Programs/AutomationTool/Win/WinPlatform.Automation.cs b/Engine/Source/Programs/AutomationTool/Win/WinPlatform.Automation.cs
--- a/Engine/Source/Programs/AutomationTool/Win/WinPlatform.Automation.cs
+++ b/Engine/Source/Programs/AutomationTool/Win/WinPlatform.Automation.cs
@@ -26,9 +26,7 @@ public abstract class BaseWinPlatform : Platform
 			FileReference ReceiptFileName = TargetReceipt.GetDefaultPath(CommandUtils.EngineDirectory, "CrashReportClient", SC.StageTargetPlatform.PlatformType, UnrealTargetConfiguration.Shipping, null);
 			if(FileReference.Exists(ReceiptFileName))
 			{
-				DirectoryReference EngineDir = CommandUtils.EngineDirectory;
-				DirectoryReference ProjectDir = DirectoryReference.FromFile(Params.RawProjectPath);
-				TargetReceipt Receipt = TargetReceipt.Read(ReceiptFileName, EngineDir, ProjectDir);
+				TargetReceipt Receipt = TargetReceipt.Read(ReceiptFileName);
 				SC.StageBuildProductsFromReceipt(Receipt, true, false);
 			}
 		}
diff --git a/Engine/Source/Programs/UnrealBuildTool/Configuration/UEBuildTarget.cs b/Engine/Source/Programs/UnrealBuildTool/Configuration/UEBuildTarget.cs
--- a/Engine/Source/Programs/UnrealBuildTool/Configuration/UEBuildTarget.cs
+++ b/Engine/Source/Programs/UnrealBuildTool/Configuration/UEBuildTarget.cs
@@ -1401,7 +1401,7 @@ namespace UnrealBuildTool
 			}
 
 			// Create the receipt
-			TargetReceipt Receipt = new TargetReceipt(TargetName, Platform, Configuration, Version);
+			TargetReceipt Receipt = new TargetReceipt(ProjectFile, TargetName, Platform, Configuration, Version);
 
 			// Set the launch executable if there is one
 			if(!Rules.bShouldCompileAsDLL)
diff --git a/Engine/Source/Programs/UnrealBuildTool/Modes/WriteMetadataMode.cs b/Engine/Source/Programs/UnrealBuildTool/Modes/WriteMetadataMode.cs
--- a/Engine/Source/Programs/UnrealBuildTool/Modes/WriteMetadataMode.cs
+++ b/Engine/Source/Programs/UnrealBuildTool/Modes/WriteMetadataMode.cs
@@ -175,7 +175,7 @@ namespace UnrealBuildTool
 			if(!UnrealBuildTool.IsFileInstalled(TargetInfo.ReceiptFile))
 			{
 				DirectoryReference.CreateDirectory(TargetInfo.ReceiptFile.Directory);
-				Receipt.Write(TargetInfo.ReceiptFile, UnrealBuildTool.EngineDirectory, DirectoryReference.FromFile(TargetInfo.ProjectFile));
+				Receipt.Write(TargetInfo.ReceiptFile);
 			}
 
 			return 0;
diff --git a/Engine/Source/Programs/UnrealBuildTool/Platform/Android/UEDeployAndroid.cs b/Engine/Source/Programs/UnrealBuildTool/Platform/Android/UEDeployAndroid.cs
--- a/Engine/Source/Programs/UnrealBuildTool/Platform/Android/UEDeployAndroid.cs
+++ b/Engine/Source/Programs/UnrealBuildTool/Platform/Android/UEDeployAndroid.cs
@@ -4005,7 +4005,7 @@ namespace UnrealBuildTool
 			// get the receipt
 			FileReference ReceiptFilename = InTarget.BuildReceiptFileName;
 			Log.TraceInformation("Receipt Filename: {0}", ReceiptFilename);
-			TargetReceipt Receipt = TargetReceipt.Read(ReceiptFilename, UnrealBuildTool.EngineDirectory, InTarget.ProjectDirectory);
+			TargetReceipt Receipt = TargetReceipt.Read(ReceiptFilename);
 			SetAndroidPluginData(ToolChain.GetAllArchitectures(), CollectPluginDataPaths(Receipt));
 
 			// Get the output paths
diff --git a/Engine/Source/Programs/UnrealBuildTool/Platform/IOS/IOSExports.cs b/Engine/Source/Programs/UnrealBuildTool/Platform/IOS/IOSExports.cs
--- a/Engine/Source/Programs/UnrealBuildTool/Platform/IOS/IOSExports.cs
+++ b/Engine/Source/Programs/UnrealBuildTool/Platform/IOS/IOSExports.cs
@@ -183,7 +183,7 @@ namespace UnrealBuildTool
             // get the receipt
             if (System.IO.File.Exists(BuildRecieptFileName.FullName))
             {
-                TargetReceipt Receipt = TargetReceipt.Read(BuildRecieptFileName, UnrealBuildTool.EngineDirectory, ProjectDirectory);
+                TargetReceipt Receipt = TargetReceipt.Read(BuildRecieptFileName);
                 IEnumerable<ReceiptProperty> Results = Receipt.AdditionalProperties.Where(x => x.Name == "SDK");
 
                 if (Results.Count() > 0)
diff --git a/Engine/Source/Programs/UnrealBuildTool/Platform/IOS/UEDeployIOS.cs b/Engine/Source/Programs/UnrealBuildTool/Platform/IOS/UEDeployIOS.cs
--- a/Engine/Source/Programs/UnrealBuildTool/Platform/IOS/UEDeployIOS.cs
+++ b/Engine/Source/Programs/UnrealBuildTool/Platform/IOS/UEDeployIOS.cs
@@ -767,7 +767,7 @@ namespace UnrealBuildTool
 
 			string RelativeEnginePath = UnrealBuildTool.EngineDirectory.MakeRelativeTo(DirectoryReference.GetCurrentDirectory());
 
-			UPL = new UnrealPluginLanguage(ProjectFile, CollectPluginDataPaths(TargetReceipt.Read(BuildRecieptFileName, UnrealBuildTool.EngineDirectory, new DirectoryReference(ProjectDirectory))), ProjectArches, "", "", UnrealTargetPlatform.IOS);
+			UPL = new UnrealPluginLanguage(ProjectFile, CollectPluginDataPaths(TargetReceipt.Read(BuildRecieptFileName)), ProjectArches, "", "", UnrealTargetPlatform.IOS);
 
 			// Passing in true for distribution is not ideal here but given the way that ios packaging happens and this call chain it seems unavoidable for now, maybe there is a way to correctly pass it in that I can't find?
 			UPL.Init(ProjectArches, true, RelativeEnginePath, BundlePath, ProjectDirectory, Config.ToString());
diff --git a/Engine/Source/Programs/UnrealBuildTool/Platform/Lumin/UEDeployLumin.cs b/Engine/Source/Programs/UnrealBuildTool/Platform/Lumin/UEDeployLumin.cs
--- a/Engine/Source/Programs/UnrealBuildTool/Platform/Lumin/UEDeployLumin.cs
+++ b/Engine/Source/Programs/UnrealBuildTool/Platform/Lumin/UEDeployLumin.cs
@@ -282,7 +282,7 @@ namespace UnrealBuildTool
 				ReceiptFilename = TargetReceipt.GetDefaultPath(UnrealBuildTool.EngineDirectory, "UE4Game", Platform, Configuration, "");
 			}
 			Log.TraceInformation("Receipt Filename: {0}", ReceiptFilename);
-			SetLuminPluginData(MLSDKArches, CollectPluginDataPaths(TargetReceipt.Read(ReceiptFilename, UnrealBuildTool.EngineDirectory, ProjectDirectory)));
+			SetLuminPluginData(MLSDKArches, CollectPluginDataPaths(TargetReceipt.Read(ReceiptFilename)));
 
 			//gather all of the xml
 			UPL.Init(MLSDKArches, true, RelativeEnginePath, UE4BuildPath, RelativeProjectPath, ConfigurationString);
diff --git a/Engine/Source/Programs/UnrealBuildTool/System/ExternalExecution.cs b/Engine/Source/Programs/UnrealBuildTool/System/ExternalExecution.cs
--- a/Engine/Source/Programs/UnrealBuildTool/System/ExternalExecution.cs
+++ b/Engine/Source/Programs/UnrealBuildTool/System/ExternalExecution.cs
@@ -669,8 +669,8 @@ namespace UnrealBuildTool
 		/// </summary>
 		static FileReference GetHeaderToolPath(FileReference ReceiptFile)
 		{
-			TargetReceipt Receipt = TargetReceipt.Read(ReceiptFile, UnrealBuildTool.EngineDirectory, null);
-			return Receipt.BuildProducts[0].Path;
+			TargetReceipt Receipt = TargetReceipt.Read(ReceiptFile);
+			return Receipt.Launch;
 		}
 
 		/// <summary>
@@ -691,7 +691,7 @@ namespace UnrealBuildTool
 				}
 
 				TargetReceipt Receipt;
-				if (!TargetReceipt.TryRead(ReceiptPath, UnrealBuildTool.EngineDirectory, null, out Receipt))
+				if (!TargetReceipt.TryRead(ReceiptPath, out Receipt))
 				{
 					Timestamp = DateTime.MaxValue;
 					return false;
diff --git a/Engine/Source/Programs/UnrealBuildTool/System/TargetReceipt.cs b/Engine/Source/Programs/UnrealBuildTool/System/TargetReceipt.cs
--- a/Engine/Source/Programs/UnrealBuildTool/System/TargetReceipt.cs
+++ b/Engine/Source/Programs/UnrealBuildTool/System/TargetReceipt.cs
@@ -255,6 +255,16 @@ namespace UnrealBuildTool
 	[Serializable]
 	public class TargetReceipt
 	{
+		/// <summary>
+		/// Path to the project file for this target
+		/// </summary>
+		public FileReference ProjectFile;
+
+		/// <summary>
+		/// The project directory
+		/// </summary>
+		public DirectoryReference ProjectDir;
+
 		/// <summary>
 		/// The name of this target
 		/// </summary>
@@ -305,12 +315,15 @@ namespace UnrealBuildTool
 		/// <summary>
 		/// Constructor
 		/// </summary>
+		/// <param name="InProjectFile">Path to the project file for this target</param>
 		/// <param name="InTargetName">The name of the target being compiled</param>
 		/// <param name="InPlatform">Platform for the target being compiled</param>
 		/// <param name="InConfiguration">Configuration of the target being compiled</param>
 		/// <param name="InVersion">Version information for the target</param>
-		public TargetReceipt(string InTargetName, UnrealTargetPlatform InPlatform, UnrealTargetConfiguration InConfiguration, BuildVersion InVersion)
+		public TargetReceipt(FileReference InProjectFile, string InTargetName, UnrealTargetPlatform InPlatform, UnrealTargetConfiguration InConfiguration, BuildVersion InVersion)
 		{
+			ProjectFile = InProjectFile;
+			ProjectDir = DirectoryReference.FromFile(InProjectFile);
 			TargetName = InTargetName;
 			Platform = InPlatform;
 			Configuration = InConfiguration;
@@ -374,7 +387,7 @@ namespace UnrealBuildTool
 		/// <param name="EngineDir">Value of the $(EngineDir) variable.</param>
 		/// <param name="ProjectDir">Value of the $(ProjectDir) variable.</param>
 		/// <returns>Converted path for the file.</returns>
-		public static string InsertPathVariables(FileReference File, DirectoryReference EngineDir, DirectoryReference ProjectDir)
+		static string InsertPathVariables(FileReference File, DirectoryReference EngineDir, DirectoryReference ProjectDir)
 		{
 			if (File.IsUnderDirectory(EngineDir))
 			{
@@ -397,7 +410,7 @@ namespace UnrealBuildTool
 		/// <param name="EngineDir">Value of the $(EngineDir) variable.</param>
 		/// <param name="ProjectDir">Value of the $(ProjectDir) variable.</param>
 		/// <returns>Converted path for the file.</returns>
-		public static FileReference ExpandPathVariables(string Path, DirectoryReference EngineDir, DirectoryReference ProjectDir)
+		static FileReference ExpandPathVariables(string Path, DirectoryReference EngineDir, DirectoryReference ProjectDir)
 		{
 			const string EnginePrefix = "$(EngineDir)";
 			if(Path.StartsWith(EnginePrefix, StringComparison.InvariantCultureIgnoreCase))
@@ -456,13 +469,21 @@ namespace UnrealBuildTool
 			return StagedFileType.NonUFS;
 		}
 
+		/// <summary>
+		/// Read a receipt from disk.
+		/// </summary>
+		/// <param name="Location">Filename to read from</param>
+		public static TargetReceipt Read(FileReference Location)
+		{
+			return Read(Location, UnrealBuildTool.EngineDirectory);
+		}
+
 		/// <summary>
 		/// Read a receipt from disk.
 		/// </summary>
 		/// <param name="Location">Filename to read from</param>
 		/// <param name="EngineDir">Engine directory for expanded variables</param>
-		/// <param name="ProjectDir">Project directory for expanded variables</param>
-		public static TargetReceipt Read(FileReference Location, DirectoryReference EngineDir, DirectoryReference ProjectDir)
+		public static TargetReceipt Read(FileReference Location, DirectoryReference EngineDir)
 		{
 			JsonObject RawObject = JsonObject.Read(Location);
 
@@ -478,8 +499,24 @@ namespace UnrealBuildTool
 				throw new JsonParseException("Invalid 'Version' field");
 			}
 
+			// Read the project path
+			FileReference ProjectFile;
+
+			string RelativeProjectFile;
+			if(RawObject.TryGetStringField("Project", out RelativeProjectFile))
+			{
+				ProjectFile = FileReference.Combine(Location.Directory, RelativeProjectFile);
+			}
+			else
+			{
+				ProjectFile = null;
+			}
+
 			// Create the receipt
-			TargetReceipt Receipt = new TargetReceipt(TargetName, Platform, Configuration, Version);
+			TargetReceipt Receipt = new TargetReceipt(ProjectFile, TargetName, Platform, Configuration, Version);
+
+			// Get the project directory
+			DirectoryReference ProjectDir = Receipt.ProjectDir;
 
 			// Read the launch executable
 			string Launch;
@@ -557,15 +594,25 @@ namespace UnrealBuildTool
 			return Receipt;
 		}
 
+		/// <summary>
+		/// Try to read a receipt from disk, failing gracefully if it can't be read.
+		/// </summary>
+		/// <param name="Location">Filename to read from</param>
+		/// <param name="Receipt">If successful, the receipt that was read</param>
+		/// <returns>True if successful</returns>
+		public static bool TryRead(FileReference Location, out TargetReceipt Receipt)
+		{
+			return TryRead(Location, UnrealBuildTool.EngineDirectory, out Receipt);
+		}
+
 		/// <summary>
 		/// Try to read a receipt from disk, failing gracefully if it can't be read.
 		/// </summary>
 		/// <param name="Location">Filename to read from</param>
 		/// <param name="EngineDir">Engine directory for expanded paths</param>
-		/// <param name="ProjectDir">Project directory for expanded paths</param>
 		/// <param name="Receipt">If successful, the receipt that was read</param>
 		/// <returns>True if successful</returns>
-		public static bool TryRead(FileReference Location, DirectoryReference EngineDir, DirectoryReference ProjectDir, out TargetReceipt Receipt)
+		public static bool TryRead(FileReference Location, DirectoryReference EngineDir, out TargetReceipt Receipt)
 		{
 			if (!FileReference.Exists(Location))
 			{
@@ -575,7 +622,7 @@ namespace UnrealBuildTool
 
 			try
 			{
-				Receipt = Read(Location, EngineDir, ProjectDir);
+				Receipt = Read(Location, EngineDir);
 				return true;
 			}
 			catch (Exception)
@@ -585,13 +632,21 @@ namespace UnrealBuildTool
 			}
 		}
 
+		/// <summary>
+		/// Write the receipt to disk.
+		/// </summary>
+		/// <param name="Location">Output filename</param>
+		public void Write(FileReference Location)
+		{
+			Write(Location, UnrealBuildTool.EngineDirectory);
+		}
+
 		/// <summary>
 		/// Write the receipt to disk.
 		/// </summary>
 		/// <param name="Location">Output filename</param>
 		/// <param name="EngineDir">Engine directory for expanded paths</param>
-		/// <param name="ProjectDir">Project directory for expanded paths</param>
-		public void Write(FileReference Location, DirectoryReference EngineDir, DirectoryReference ProjectDir)
+		public void Write(FileReference Location, DirectoryReference EngineDir)
 		{
 			using (JsonWriter Writer = new JsonWriter(Location.FullName))
 			{
@@ -600,6 +655,11 @@ namespace UnrealBuildTool
 				Writer.WriteValue("Platform", Platform.ToString());
 				Writer.WriteValue("Configuration", Configuration.ToString());
 
+				if(ProjectFile != null)
+				{
+					Writer.WriteValue("Project", ProjectFile.MakeRelativeTo(Location.Directory).Replace(Path.DirectorySeparatorChar, '/'));
+				}
+
 				Writer.WriteObjectStart("Version");
 				Version.WriteProperties(Writer);
 				Writer.WriteObjectEnd();
