UBT: Add the target type to the receipt.

#rnx

diff --git a/Engine/Source/Programs/UnrealBuildTool/Configuration/UEBuildTarget.cs b/Engine/Source/Programs/UnrealBuildTool/Configuration/UEBuildTarget.cs
--- a/Engine/Source/Programs/UnrealBuildTool/Configuration/UEBuildTarget.cs
+++ b/Engine/Source/Programs/UnrealBuildTool/Configuration/UEBuildTarget.cs
@@ -1401,7 +1401,7 @@ namespace UnrealBuildTool
 			}
 
 			// Create the receipt
-			TargetReceipt Receipt = new TargetReceipt(ProjectFile, TargetName, Platform, Configuration, Version);
+			TargetReceipt Receipt = new TargetReceipt(ProjectFile, TargetName, TargetType, Platform, Configuration, Version);
 
 			// Set the launch executable if there is one
 			if(!Rules.bShouldCompileAsDLL)
@@ -1916,6 +1916,7 @@ namespace UnrealBuildTool
 				WriteMetadataAction.StatusDescription = ReceiptFileName.GetFileName();
 				WriteMetadataAction.bCanExecuteRemotely = false;
 				WriteMetadataAction.PrerequisiteItems.AddRange(OutputItems);
+				WriteMetadataAction.PrerequisiteItems.Add(FileItem.GetExistingItemByPath(Assembly.GetExecutingAssembly().Location));
 				WriteMetadataAction.ProducedItems.Add(FileItem.GetItemByFileReference(ReceiptFileName));
 
 				OutputItems.AddRange(WriteMetadataAction.ProducedItems);
diff --git a/Engine/Source/Programs/UnrealBuildTool/System/TargetReceipt.cs b/Engine/Source/Programs/UnrealBuildTool/System/TargetReceipt.cs
--- a/Engine/Source/Programs/UnrealBuildTool/System/TargetReceipt.cs
+++ b/Engine/Source/Programs/UnrealBuildTool/System/TargetReceipt.cs
@@ -280,6 +280,11 @@ namespace UnrealBuildTool
 		/// </summary>
 		public UnrealTargetConfiguration Configuration;
 
+		/// <summary>
+		/// The type of the target
+		/// </summary>
+		public TargetType TargetType;
+
 		/// <summary>
 		/// Version information for this target.
 		/// </summary>
@@ -317,16 +322,18 @@ namespace UnrealBuildTool
 		/// </summary>
 		/// <param name="InProjectFile">Path to the project file for this target</param>
 		/// <param name="InTargetName">The name of the target being compiled</param>
+		/// <param name="InTargetType">The type of target</param>
 		/// <param name="InPlatform">Platform for the target being compiled</param>
 		/// <param name="InConfiguration">Configuration of the target being compiled</param>
 		/// <param name="InVersion">Version information for the target</param>
-		public TargetReceipt(FileReference InProjectFile, string InTargetName, UnrealTargetPlatform InPlatform, UnrealTargetConfiguration InConfiguration, BuildVersion InVersion)
+		public TargetReceipt(FileReference InProjectFile, string InTargetName, TargetType InTargetType, UnrealTargetPlatform InPlatform, UnrealTargetConfiguration InConfiguration, BuildVersion InVersion)
 		{
 			ProjectFile = InProjectFile;
 			ProjectDir = DirectoryReference.FromFile(InProjectFile);
 			TargetName = InTargetName;
 			Platform = InPlatform;
 			Configuration = InConfiguration;
+			TargetType = InTargetType;
 			Version = InVersion;
 		}
 
@@ -489,6 +496,7 @@ namespace UnrealBuildTool
 
 			// Read the initial fields
 			string TargetName = RawObject.GetStringField("TargetName");
+			TargetType TargetType = RawObject.GetEnumField<TargetType>("TargetType");
 			UnrealTargetPlatform Platform = RawObject.GetEnumField<UnrealTargetPlatform>("Platform");
 			UnrealTargetConfiguration Configuration = RawObject.GetEnumField<UnrealTargetConfiguration>("Configuration");
 
@@ -513,7 +521,7 @@ namespace UnrealBuildTool
 			}
 
 			// Create the receipt
-			TargetReceipt Receipt = new TargetReceipt(ProjectFile, TargetName, Platform, Configuration, Version);
+			TargetReceipt Receipt = new TargetReceipt(ProjectFile, TargetName, TargetType, Platform, Configuration, Version);
 
 			// Get the project directory
 			DirectoryReference ProjectDir = Receipt.ProjectDir;
@@ -654,21 +662,22 @@ namespace UnrealBuildTool
 				Writer.WriteValue("TargetName", TargetName);
 				Writer.WriteValue("Platform", Platform.ToString());
 				Writer.WriteValue("Configuration", Configuration.ToString());
+				Writer.WriteValue("TargetType", TargetType.ToString());
 
 				if(ProjectFile != null)
 				{
 					Writer.WriteValue("Project", ProjectFile.MakeRelativeTo(Location.Directory).Replace(Path.DirectorySeparatorChar, '/'));
 				}
 
-				Writer.WriteObjectStart("Version");
-				Version.WriteProperties(Writer);
-				Writer.WriteObjectEnd();
-
 				if(Launch != null)
 				{
 					Writer.WriteValue("Launch", InsertPathVariables(Launch, EngineDir, ProjectDir));
 				}
 
+				Writer.WriteObjectStart("Version");
+				Version.WriteProperties(Writer);
+				Writer.WriteObjectEnd();
+
 				Writer.WriteArrayStart("BuildProducts");
 				foreach (BuildProduct BuildProduct in BuildProducts)
 				{
