UBT: Move standalone deployment functionality into a separate mode, and out of the regular build loop.

diff --git a/Engine/Source/Programs/AutomationTool/AutomationUtils/UE4Build.cs b/Engine/Source/Programs/AutomationTool/AutomationUtils/UE4Build.cs
--- a/Engine/Source/Programs/AutomationTool/AutomationUtils/UE4Build.cs
+++ b/Engine/Source/Programs/AutomationTool/AutomationUtils/UE4Build.cs
@@ -153,7 +153,7 @@ namespace AutomationTool
 			// run the deployment steps, if necessary
 			foreach(string DeployTargetFile in Item.Manifest.DeployTargetFiles)
 			{
-				CommandUtils.RunUBT(CommandUtils.CmdEnv, UBTExecutable, String.Format("-deploy=\"{0}\" -nomutex", DeployTargetFile));
+				CommandUtils.RunUBT(CommandUtils.CmdEnv, UBTExecutable, String.Format("-Mode=Deploy -TargetInfo=\"{0}\"", DeployTargetFile));
 			}
 
 			foreach (string ManifestItem in Item.Manifest.BuildProducts)
diff --git a/Engine/Source/Programs/UnrealBuildTool/Configuration/BuildConfiguration.cs b/Engine/Source/Programs/UnrealBuildTool/Configuration/BuildConfiguration.cs
--- a/Engine/Source/Programs/UnrealBuildTool/Configuration/BuildConfiguration.cs
+++ b/Engine/Source/Programs/UnrealBuildTool/Configuration/BuildConfiguration.cs
@@ -163,12 +163,6 @@ namespace UnrealBuildTool
 		[CommandLine("-Clean")]
 		public bool bCleanProject = false;
 
-		/// <summary>
-		/// If we are just running the deployment step, specifies the path to the given deployment settings
-		/// </summary>
-		[CommandLine("-Deploy")]
-		public FileReference DeployTargetFile = null;
-
 		/// <summary>
 		/// If true, force header regeneration. Intended for the build machine
 		/// </summary>
diff --git a/Engine/Source/Programs/UnrealBuildTool/Modes/DeployMode.cs b/Engine/Source/Programs/UnrealBuildTool/Modes/DeployMode.cs
new file mode 100644
--- /dev/null
+++ b/Engine/Source/Programs/UnrealBuildTool/Modes/DeployMode.cs
@@ -0,0 +1,48 @@
+using System;
+using System.Collections.Generic;
+using System.Linq;
+using System.Text;
+using System.Threading.Tasks;
+using Tools.DotNETCommon;
+
+namespace UnrealBuildTool
+{
+	/// <summary>
+	/// Invokes the deployment handler for a target.
+	/// </summary>
+	[ToolMode("Deploy")]
+	class DeployMode : ToolMode
+	{
+		/// <summary>
+		/// If we are just running the deployment step, specifies the path to the given deployment settings
+		/// </summary>
+		[CommandLine("-TargetInfo", Required=true)]
+		public FileReference TargetInfoFile = null;
+
+		/// <summary>
+		/// Execute the tool mode
+		/// </summary>
+		/// <param name="Arguments">Command line arguments</param>
+		/// <returns>Exit code</returns>
+		public override int Execute(CommandLineArguments Arguments)
+		{
+			// Apply the arguments
+			Arguments.ApplyTo(this);
+			Arguments.CheckAllArgumentsUsed();
+
+			// Change the working directory to be the Engine/Source folder. We are likely running from Engine/Binaries/DotNET
+			// This is critical to be done early so any code that relies on the current directory being Engine/Source will work.
+			DirectoryReference.SetCurrentDirectory(UnrealBuildTool.EngineSourceDirectory);
+
+			// Find and register all tool chains, build platforms, etc. that are present
+			UnrealBuildTool.RegisterAllUBTClasses(false);
+
+			// Execute the deploy
+			UEBuildDeployTarget DeployTarget = new UEBuildDeployTarget(TargetInfoFile);
+			Log.WriteLine(LogEventType.Console, "Deploying {0} {1} {2}...", DeployTarget.TargetName, DeployTarget.Platform, DeployTarget.Configuration);
+			UEBuildPlatform.GetBuildPlatform(DeployTarget.Platform).Deploy(DeployTarget);
+
+			return (int)ECompilationResult.Succeeded;
+		}
+	}
+}
diff --git a/Engine/Source/Programs/UnrealBuildTool/UnrealBuildTool.cs b/Engine/Source/Programs/UnrealBuildTool/UnrealBuildTool.cs
--- a/Engine/Source/Programs/UnrealBuildTool/UnrealBuildTool.cs
+++ b/Engine/Source/Programs/UnrealBuildTool/UnrealBuildTool.cs
@@ -584,21 +584,12 @@ namespace UnrealBuildTool
 						JunkDeleter.DeleteJunk();
 					}
 
-					if (BuildConfiguration.DeployTargetFile != null)
+					// Build our project
+					if (Result == ECompilationResult.Succeeded)
 					{
-						UEBuildDeployTarget DeployTarget = new UEBuildDeployTarget(BuildConfiguration.DeployTargetFile);
-						Log.WriteLine(LogEventType.Console, "Deploying {0} {1} {2}...", DeployTarget.TargetName, DeployTarget.Platform, DeployTarget.Configuration);
-						UEBuildPlatform.GetBuildPlatform(DeployTarget.Platform).Deploy(DeployTarget);
-						Result = ECompilationResult.Succeeded;
-					}
-					else
-					{
-						// Build our project
-						if (Result == ECompilationResult.Succeeded)
-						{
-							Result = RunUBT(BuildConfiguration, Arguments, ProjectFile, true);
-						}
+						Result = RunUBT(BuildConfiguration, Arguments, ProjectFile, true);
 					}
+
 					// Print some performance info
 					double BuildDuration = (DateTime.UtcNow - StartTime).TotalSeconds;
 					if (UnrealBuildTool.bPrintPerformanceInfo)
diff --git a/Engine/Source/Programs/UnrealBuildTool/UnrealBuildTool.csproj b/Engine/Source/Programs/UnrealBuildTool/UnrealBuildTool.csproj
--- a/Engine/Source/Programs/UnrealBuildTool/UnrealBuildTool.csproj
+++ b/Engine/Source/Programs/UnrealBuildTool/UnrealBuildTool.csproj
@@ -140,6 +140,7 @@
     </When>
   </Choose>
   <ItemGroup>
+    <Compile Include="Modes\DeployMode.cs" />
     <Compile Include="Modes\GenerateProjectFilesMode.cs" />
     <Compile Include="Modes\SetupPlatformsMode.cs" />
     <Compile Include="Modes\ValidatePlatformsMode.cs" />
