UBT: Remove the unused RequiresDeployAfterCompile() method from UEBuildPlatform.

diff --git a/Engine/Source/Programs/UnrealBuildTool/Configuration/UEBuildPlatform.cs b/Engine/Source/Programs/UnrealBuildTool/Configuration/UEBuildPlatform.cs
--- a/Engine/Source/Programs/UnrealBuildTool/Configuration/UEBuildPlatform.cs
+++ b/Engine/Source/Programs/UnrealBuildTool/Configuration/UEBuildPlatform.cs
@@ -638,15 +638,6 @@ namespace UnrealBuildTool
 			return TempList;
 		}
 
-		/// <summary>
-		/// Whether the build platform requires deployment prep
-		/// </summary>
-		/// <returns></returns>
-		public virtual bool RequiresDeployPrepAfterCompile()
-		{
-			return false;
-		}
-
 		/// <summary>
 		/// Return all valid configurations for this platform
 		/// Typically, this is always Debug, Development, and Shipping - but Test is a likely future addition for some platforms
diff --git a/Engine/Source/Programs/UnrealBuildTool/Platform/Android/UEBuildAndroid.cs b/Engine/Source/Programs/UnrealBuildTool/Platform/Android/UEBuildAndroid.cs
--- a/Engine/Source/Programs/UnrealBuildTool/Platform/Android/UEBuildAndroid.cs
+++ b/Engine/Source/Programs/UnrealBuildTool/Platform/Android/UEBuildAndroid.cs
@@ -222,11 +222,6 @@ namespace UnrealBuildTool
 			return true;
 		}
 
-		public override bool RequiresDeployPrepAfterCompile()
-		{
-			return true;
-		}
-
 		/// <summary>
 		/// Modify the rules for a newly created module, where the target is a different host platform.
 		/// This is not required - but allows for hiding details of a particular platform.
