Remove hacks for configuring stub creation on IOS. The -CreateStub argument is now only specified when building remotely, and when building the engine as part of an installed engine build.

#rnx

diff --git a/Engine/Build/InstalledEngineBuild.xml b/Engine/Build/InstalledEngineBuild.xml
--- a/Engine/Build/InstalledEngineBuild.xml
+++ b/Engine/Build/InstalledEngineBuild.xml
@@ -384,7 +384,7 @@
 		
 			<ForEach Name="Configuration" Values="$(GameConfigurations)">
 				<Compile Target="UE4Game" Platform="IOS" Configuration="$(Configuration)" Tag="#UE4Game IOS" Arguments="-precompile -allmodules -nolink $(TargetDebugInfoArg)"/>
-				<Compile Target="UE4Game" Platform="IOS" Configuration="$(Configuration)" Tag="#UE4Game IOS" Arguments="-precompile $(TargetDebugInfoArg)" Clean="false"/>
+				<Compile Target="UE4Game" Platform="IOS" Configuration="$(Configuration)" Tag="#UE4Game IOS" Arguments="-precompile -createstub $(TargetDebugInfoArg)" Clean="false"/>
 			</ForEach>
 
 			<Tag Files="#UE4Game IOS" Filter="$(IOSStripFilter)" With="#UE4Game IOS Unstripped"/>
diff --git a/Engine/Source/Programs/AutomationTool/AutomationUtils/UBTUtils.cs b/Engine/Source/Programs/AutomationTool/AutomationUtils/UBTUtils.cs
--- a/Engine/Source/Programs/AutomationTool/AutomationUtils/UBTUtils.cs
+++ b/Engine/Source/Programs/AutomationTool/AutomationUtils/UBTUtils.cs
@@ -34,10 +34,6 @@ namespace AutomationTool
 			{
 				CommandLine += " -2015";
 			}
-			if (!IsBuildMachine && UnrealBuildTool.BuildHostPlatform.Current.Platform == UnrealBuildTool.UnrealTargetPlatform.Mac)
-			{
-				CommandLine += " -nocreatestub";
-			}
 			CommandLine += " -NoHotReload";
 			if (bJunkDeleted || GlobalCommandLine.IgnoreJunk)
 			{
diff --git a/Engine/Source/Programs/UnrealBuildTool/Configuration/TargetRules.cs b/Engine/Source/Programs/UnrealBuildTool/Configuration/TargetRules.cs
--- a/Engine/Source/Programs/UnrealBuildTool/Configuration/TargetRules.cs
+++ b/Engine/Source/Programs/UnrealBuildTool/Configuration/TargetRules.cs
@@ -880,8 +880,7 @@ namespace UnrealBuildTool
 		/// If true, then a stub IPA will be generated when compiling is done (minimal files needed for a valid IPA).
 		/// </summary>
 		[CommandLine("-CreateStub", Value = "true")]
-		[CommandLine("-NoCreateStub", Value = "false")]
-		public bool bCreateStubIPA = true;
+		public bool bCreateStubIPA = false;
 
 		/// <summary>
 		/// If true, then a stub IPA will be generated when compiling is done (minimal files needed for a valid IPA).
@@ -1177,12 +1176,6 @@ namespace UnrealBuildTool
 			// If we've got a changelist set, set that we're making a formal build
 			bFormalBuild = (Version.Changelist != 0 && Version.IsPromotedBuild);
 
-			// @todo remove this hacky build system stuff
-			if (bCreateStubIPA && (String.IsNullOrEmpty(Environment.GetEnvironmentVariable("uebp_LOCAL_ROOT")) && BuildHostPlatform.Current.Platform == UnrealTargetPlatform.Mac))
-			{
-				bCreateStubIPA = false;
-			}
-
 			// Set the default build version
 			if(String.IsNullOrEmpty(BuildVersion))
 			{
