UBT: Remove the unused bCopyAppBundleBackToDevice setting.

#rnx

diff --git a/Engine/Source/Programs/UnrealBuildTool/Configuration/TargetRules.cs b/Engine/Source/Programs/UnrealBuildTool/Configuration/TargetRules.cs
--- a/Engine/Source/Programs/UnrealBuildTool/Configuration/TargetRules.cs
+++ b/Engine/Source/Programs/UnrealBuildTool/Configuration/TargetRules.cs
@@ -882,12 +882,6 @@ namespace UnrealBuildTool
 		[CommandLine("-CreateStub", Value = "true")]
 		public bool bCreateStubIPA = false;
 
-		/// <summary>
-		/// If true, then a stub IPA will be generated when compiling is done (minimal files needed for a valid IPA).
-		/// </summary>
-		[CommandLine("-CopyAppBundleBackToDevice")]
-		public bool bCopyAppBundleBackToDevice = false;
-
 		/// <summary>
 		/// When enabled, allows XGE to compile pre-compiled header files on remote machines.  Otherwise, PCHs are always generated locally.
 		/// </summary>
@@ -2028,11 +2022,6 @@ namespace UnrealBuildTool
 			get { return Inner.bCreateStubIPA; }
 		}
 
-		public bool bCopyAppBundleBackToDevice
-		{
-			get { return Inner.bCopyAppBundleBackToDevice; }
-		}
-
 		public bool bAllowRemotelyCompiledPCHs
 		{
 			get { return Inner.bAllowRemotelyCompiledPCHs; }
