UBT: Change deployment on IOS to only read the receipt once (step towards removing the deployment target context).

#rnx

diff --git a/Engine/Source/Programs/UnrealBuildTool/Platform/IOS/IOSExports.cs b/Engine/Source/Programs/UnrealBuildTool/Platform/IOS/IOSExports.cs
--- a/Engine/Source/Programs/UnrealBuildTool/Platform/IOS/IOSExports.cs
+++ b/Engine/Source/Programs/UnrealBuildTool/Platform/IOS/IOSExports.cs
@@ -79,7 +79,8 @@ namespace UnrealBuildTool
 		/// <returns></returns>
 		public static bool PrepForUATPackageOrDeploy(UnrealTargetConfiguration Config, FileReference ProjectFile, string InProjectName, DirectoryReference InProjectDirectory, string InExecutablePath, DirectoryReference InEngineDir, bool bForDistribution, string CookFlavor, bool bIsDataDeploy, bool bCreateStubIPA, FileReference BuildReceiptFileName)
 		{
-			return new UEDeployIOS().PrepForUATPackageOrDeploy(Config, ProjectFile, InProjectName, InProjectDirectory.FullName, InExecutablePath, InEngineDir.FullName, bForDistribution, CookFlavor, bIsDataDeploy, bCreateStubIPA, BuildReceiptFileName);
+			TargetReceipt Receipt = TargetReceipt.Read(BuildReceiptFileName);
+			return new UEDeployIOS().PrepForUATPackageOrDeploy(Config, ProjectFile, InProjectName, InProjectDirectory.FullName, InExecutablePath, InEngineDir.FullName, bForDistribution, CookFlavor, bIsDataDeploy, bCreateStubIPA, Receipt);
 		}
 
 		/// <summary>
@@ -100,7 +101,8 @@ namespace UnrealBuildTool
 		/// <returns></returns>
 		public static bool GeneratePList(FileReference ProjectFile, UnrealTargetConfiguration Config, DirectoryReference ProjectDirectory, bool bIsUE4Game, string GameName, string ProjectName, DirectoryReference InEngineDir, DirectoryReference AppDirectory, FileReference BuildReceiptFileName, out bool bSupportsPortrait, out bool bSupportsLandscape, out bool bSkipIcons)
 		{
-			return new UEDeployIOS().GeneratePList(ProjectFile, Config, ProjectDirectory.FullName, bIsUE4Game, GameName, ProjectName, InEngineDir.FullName, AppDirectory.FullName, BuildReceiptFileName, out bSupportsPortrait, out bSupportsLandscape, out bSkipIcons);
+			TargetReceipt Receipt = TargetReceipt.Read(BuildReceiptFileName);
+			return new UEDeployIOS().GeneratePList(ProjectFile, Config, ProjectDirectory.FullName, bIsUE4Game, GameName, ProjectName, InEngineDir.FullName, AppDirectory.FullName, Receipt, out bSupportsPortrait, out bSupportsLandscape, out bSkipIcons);
 		}
 
 		/// <summary>
@@ -178,12 +180,11 @@ namespace UnrealBuildTool
         /// <summary>
         /// 
         /// </summary>
-        public static bool SupportsIconCatalog(DirectoryReference ProjectDirectory, FileReference BuildRecieptFileName)
+        public static bool SupportsIconCatalog(DirectoryReference ProjectDirectory, TargetReceipt Receipt)
         {
             // get the receipt
-            if (System.IO.File.Exists(BuildRecieptFileName.FullName))
+            if (Receipt != null)
             {
-                TargetReceipt Receipt = TargetReceipt.Read(BuildRecieptFileName);
                 IEnumerable<ReceiptProperty> Results = Receipt.AdditionalProperties.Where(x => x.Name == "SDK");
 
                 if (Results.Count() > 0)
diff --git a/Engine/Source/Programs/UnrealBuildTool/Platform/IOS/UEDeployIOS.cs b/Engine/Source/Programs/UnrealBuildTool/Platform/IOS/UEDeployIOS.cs
--- a/Engine/Source/Programs/UnrealBuildTool/Platform/IOS/UEDeployIOS.cs
+++ b/Engine/Source/Programs/UnrealBuildTool/Platform/IOS/UEDeployIOS.cs
@@ -186,7 +186,7 @@ namespace UnrealBuildTool
 			return result;
 		}
 
-		public static bool GenerateIOSPList(FileReference ProjectFile, UnrealTargetConfiguration Config, string ProjectDirectory, bool bIsUE4Game, string GameName, string ProjectName, string InEngineDir, string AppDirectory, FileReference BuildRecieptFileName, out bool bSupportsPortrait, out bool bSupportsLandscape, out bool bSkipIcons, UEDeployIOS InThis = null)
+		public static bool GenerateIOSPList(FileReference ProjectFile, UnrealTargetConfiguration Config, string ProjectDirectory, bool bIsUE4Game, string GameName, string ProjectName, string InEngineDir, string AppDirectory, TargetReceipt Receipt, out bool bSupportsPortrait, out bool bSupportsLandscape, out bool bSkipIcons, UEDeployIOS InThis = null)
 		{
 			// generate the Info.plist for future use
 			string BuildDirectory = ProjectDirectory + "/Build/IOS";
@@ -457,7 +457,7 @@ namespace UnrealBuildTool
 			}
 			Text.AppendLine("\t</array>");
 
-			if (IOSExports.SupportsIconCatalog(new DirectoryReference(ProjectDirectory), BuildRecieptFileName))
+			if (IOSExports.SupportsIconCatalog(new DirectoryReference(ProjectDirectory), Receipt))
 			{
 				bSkipIcons = true;
 				Text.AppendLine("\t<key>CFBundleIcons</key>");
@@ -740,7 +740,7 @@ namespace UnrealBuildTool
 			return bSkipDefaultPNGs;
 		}
 
-		public virtual bool GeneratePList(FileReference ProjectFile, UnrealTargetConfiguration Config, string ProjectDirectory, bool bIsUE4Game, string GameName, string ProjectName, string InEngineDir, string AppDirectory, FileReference BuildRecieptFileName, out bool bSupportsPortrait, out bool bSupportsLandscape, out bool bSkipIcons)
+		public virtual bool GeneratePList(FileReference ProjectFile, UnrealTargetConfiguration Config, string ProjectDirectory, bool bIsUE4Game, string GameName, string ProjectName, string InEngineDir, string AppDirectory, TargetReceipt Receipt, out bool bSupportsPortrait, out bool bSupportsLandscape, out bool bSkipIcons)
 		{
 			// remember name with -IOS-Shipping, etc
 			string ExeName = GameName;
@@ -767,12 +767,12 @@ namespace UnrealBuildTool
 
 			string RelativeEnginePath = UnrealBuildTool.EngineDirectory.MakeRelativeTo(DirectoryReference.GetCurrentDirectory());
 
-			UPL = new UnrealPluginLanguage(ProjectFile, CollectPluginDataPaths(TargetReceipt.Read(BuildRecieptFileName)), ProjectArches, "", "", UnrealTargetPlatform.IOS);
+			UPL = new UnrealPluginLanguage(ProjectFile, CollectPluginDataPaths(Receipt), ProjectArches, "", "", UnrealTargetPlatform.IOS);
 
 			// Passing in true for distribution is not ideal here but given the way that ios packaging happens and this call chain it seems unavoidable for now, maybe there is a way to correctly pass it in that I can't find?
 			UPL.Init(ProjectArches, true, RelativeEnginePath, BundlePath, ProjectDirectory, Config.ToString());
 
-			return GenerateIOSPList(ProjectFile, Config, ProjectDirectory, bIsUE4Game, GameName, ProjectName, InEngineDir, AppDirectory, BuildRecieptFileName, out bSupportsPortrait, out bSupportsLandscape, out bSkipIcons, this);
+			return GenerateIOSPList(ProjectFile, Config, ProjectDirectory, bIsUE4Game, GameName, ProjectName, InEngineDir, AppDirectory, Receipt, out bSupportsPortrait, out bSupportsLandscape, out bSkipIcons, this);
 		}
 
 		protected virtual void CopyCloudResources(string InEngineDir, string AppDirectory)
@@ -873,7 +873,7 @@ namespace UnrealBuildTool
 			}
 		}
 
-		public bool PrepForUATPackageOrDeploy(UnrealTargetConfiguration Config, FileReference ProjectFile, string InProjectName, string InProjectDirectory, string InExecutablePath, string InEngineDir, bool bForDistribution, string CookFlavor, bool bIsDataDeploy, bool bCreateStubIPA, FileReference BuildReceiptFileName)
+		public bool PrepForUATPackageOrDeploy(UnrealTargetConfiguration Config, FileReference ProjectFile, string InProjectName, string InProjectDirectory, string InExecutablePath, string InEngineDir, bool bForDistribution, string CookFlavor, bool bIsDataDeploy, bool bCreateStubIPA, TargetReceipt Receipt)
 		{
 			if (BuildHostPlatform.Current.Platform != UnrealTargetPlatform.Mac)
 			{
@@ -1036,7 +1036,7 @@ namespace UnrealBuildTool
 			bool bSupportsPortrait = true;
 			bool bSupportsLandscape = false;
 			bool bSkipIcons = false;
-			bool bSkipDefaultPNGs = GeneratePList(ProjectFile, Config, InProjectDirectory, bIsUE4Game, GameExeName, InProjectName, InEngineDir, AppDirectory, BuildReceiptFileName, out bSupportsPortrait, out bSupportsLandscape, out bSkipIcons);
+			bool bSkipDefaultPNGs = GeneratePList(ProjectFile, Config, InProjectDirectory, bIsUE4Game, GameExeName, InProjectName, InEngineDir, AppDirectory, Receipt, out bSupportsPortrait, out bSupportsLandscape, out bSkipIcons);
 
 			// ensure the destination is writable
 			if (File.Exists(AppDirectory + "/" + GameName))
@@ -1078,6 +1078,8 @@ namespace UnrealBuildTool
 
 		public override bool PrepTargetForDeployment(UEBuildDeployTarget InTarget)
 		{
+			TargetReceipt Receipt = TargetReceipt.Read(InTarget.BuildReceiptFileName);
+
 			string SubDir = GetTargetPlatformName();
 
 			string GameName = InTarget.TargetName;
@@ -1099,13 +1101,13 @@ namespace UnrealBuildTool
 
 			if (BuildHostPlatform.Current.Platform == UnrealTargetPlatform.Mac && Environment.GetEnvironmentVariable("UBT_NO_POST_DEPLOY") != "true")
 			{
-				return PrepForUATPackageOrDeploy(InTarget.Configuration, InTarget.ProjectFile, GameName, ProjectDirectory, BuildPath + "/" + DecoratedGameName, "../../Engine", false, "", false, InTarget.bCreateStubIPA, InTarget.BuildReceiptFileName);
+				return PrepForUATPackageOrDeploy(InTarget.Configuration, InTarget.ProjectFile, GameName, ProjectDirectory, BuildPath + "/" + DecoratedGameName, "../../Engine", false, "", false, InTarget.bCreateStubIPA, Receipt);
 			}
 			else
 			{
 				// @todo tvos merge: This used to copy the bundle back - where did that code go? It needs to be fixed up for TVOS directories
 				bool bSupportPortrait, bSupportLandscape, bSkipIcons;
-				GeneratePList(InTarget.ProjectFile, InTarget.Configuration, ProjectDirectory, bIsUE4Game, GameName, (InTarget.ProjectFile == null) ? "" : Path.GetFileNameWithoutExtension(InTarget.ProjectFile.FullName), "../../Engine", "", InTarget.BuildReceiptFileName, out bSupportPortrait, out bSupportLandscape, out bSkipIcons);
+				GeneratePList(InTarget.ProjectFile, InTarget.Configuration, ProjectDirectory, bIsUE4Game, GameName, (InTarget.ProjectFile == null) ? "" : Path.GetFileNameWithoutExtension(InTarget.ProjectFile.FullName), "../../Engine", "", Receipt, out bSupportPortrait, out bSupportLandscape, out bSkipIcons);
 			}
 			return true;
 		}
diff --git a/Engine/Source/Programs/UnrealBuildTool/Platform/TVOS/TVOSExports.cs b/Engine/Source/Programs/UnrealBuildTool/Platform/TVOS/TVOSExports.cs
--- a/Engine/Source/Programs/UnrealBuildTool/Platform/TVOS/TVOSExports.cs
+++ b/Engine/Source/Programs/UnrealBuildTool/Platform/TVOS/TVOSExports.cs
@@ -78,7 +78,8 @@ namespace UnrealBuildTool
 		/// <returns></returns>
 		public static bool PrepForUATPackageOrDeploy(UnrealTargetConfiguration Config, FileReference ProjectFile, string InProjectName, DirectoryReference InProjectDirectory, string InExecutablePath, DirectoryReference InEngineDir, bool bForDistribution, string CookFlavor, bool bIsDataDeploy, bool bCreateStubIPA, FileReference BuildReceiptFileName)
 		{
-			return new UEDeployTVOS().PrepForUATPackageOrDeploy(Config, ProjectFile, InProjectName, InProjectDirectory.FullName, InExecutablePath, InEngineDir.FullName, bForDistribution, CookFlavor, bIsDataDeploy, bCreateStubIPA, BuildReceiptFileName);
+			TargetReceipt Receipt = TargetReceipt.Read(BuildReceiptFileName);
+			return new UEDeployTVOS().PrepForUATPackageOrDeploy(Config, ProjectFile, InProjectName, InProjectDirectory.FullName, InExecutablePath, InEngineDir.FullName, bForDistribution, CookFlavor, bIsDataDeploy, bCreateStubIPA, Receipt);
 		}
 
 		/// <summary>
@@ -99,7 +100,8 @@ namespace UnrealBuildTool
 		/// <returns></returns>
 		public static bool GeneratePList(FileReference ProjectFile, UnrealTargetConfiguration Config, DirectoryReference ProjectDirectory, bool bIsUE4Game, string GameName, string ProjectName, DirectoryReference InEngineDir, DirectoryReference AppDirectory, FileReference BuildReceiptFileName, out bool bSupportsPortrait, out bool bSupportsLandscape, out bool bSkipIcons)
 		{
-			return new UEDeployTVOS().GeneratePList(ProjectFile, Config, ProjectDirectory.FullName, bIsUE4Game, GameName, ProjectName, InEngineDir.FullName, AppDirectory.FullName, BuildReceiptFileName, out bSupportsPortrait, out bSupportsLandscape, out bSkipIcons);
+			TargetReceipt Receipt = TargetReceipt.Read(BuildReceiptFileName);
+			return new UEDeployTVOS().GeneratePList(ProjectFile, Config, ProjectDirectory.FullName, bIsUE4Game, GameName, ProjectName, InEngineDir.FullName, AppDirectory.FullName, Receipt, out bSupportsPortrait, out bSupportsLandscape, out bSkipIcons);
 		}
 	}
 }
diff --git a/Engine/Source/Programs/UnrealBuildTool/Platform/TVOS/UEDeployTVOS.cs b/Engine/Source/Programs/UnrealBuildTool/Platform/TVOS/UEDeployTVOS.cs
--- a/Engine/Source/Programs/UnrealBuildTool/Platform/TVOS/UEDeployTVOS.cs
+++ b/Engine/Source/Programs/UnrealBuildTool/Platform/TVOS/UEDeployTVOS.cs
@@ -331,7 +331,7 @@ namespace UnrealBuildTool
 			return bSkipDefaultPNGs;
 		}
 
-		public override bool GeneratePList(FileReference ProjectFile, UnrealTargetConfiguration Config, string ProjectDirectory, bool bIsUE4Game, string GameName, string ProjectName, string InEngineDir, string AppDirectory, FileReference BuildRecieptFileName, out bool bSupportsPortrait, out bool bSupportsLandscape, out bool bSkipIcons)
+		public override bool GeneratePList(FileReference ProjectFile, UnrealTargetConfiguration Config, string ProjectDirectory, bool bIsUE4Game, string GameName, string ProjectName, string InEngineDir, string AppDirectory, TargetReceipt Receipt, out bool bSupportsPortrait, out bool bSupportsLandscape, out bool bSkipIcons)
 		{
 			bSupportsLandscape = bSupportsPortrait = true;
             bSkipIcons = true;
diff --git a/Engine/Source/Programs/UnrealBuildTool/ProjectFiles/Xcode/XcodeProject.cs b/Engine/Source/Programs/UnrealBuildTool/ProjectFiles/Xcode/XcodeProject.cs
--- a/Engine/Source/Programs/UnrealBuildTool/ProjectFiles/Xcode/XcodeProject.cs
+++ b/Engine/Source/Programs/UnrealBuildTool/ProjectFiles/Xcode/XcodeProject.cs
@@ -879,7 +879,9 @@ namespace UnrealBuildTool
 							}
 							Directory.CreateDirectory(Path.GetDirectoryName(IOSInfoPlistPath));
                             bool bSupportPortrait, bSupportLandscape, bSkipIcons;
-							UEDeployIOS.GenerateIOSPList(ProjectFile, Config.BuildConfig, ProjectPath.FullName, bIsUE4Game, GameName, Config.BuildTarget, EngineDir.FullName, ProjectPath + "/Binaries/IOS/Payload", ReceiptFilename, out bSupportPortrait, out bSupportLandscape, out bSkipIcons);
+							TargetReceipt Receipt;
+							TargetReceipt.TryRead(ReceiptFilename, out Receipt);
+							UEDeployIOS.GenerateIOSPList(ProjectFile, Config.BuildConfig, ProjectPath.FullName, bIsUE4Game, GameName, Config.BuildTarget, EngineDir.FullName, ProjectPath + "/Binaries/IOS/Payload", Receipt, out bSupportPortrait, out bSupportLandscape, out bSkipIcons);
 						}
 						if (bCreateTVOSInfoPlist)
 						{
