Fix CIS.

#rnx

diff --git a/Engine/Source/Programs/UnrealBuildTool/ProjectFiles/ProjectFileGenerator.cs b/Engine/Source/Programs/UnrealBuildTool/ProjectFiles/ProjectFileGenerator.cs
--- a/Engine/Source/Programs/UnrealBuildTool/ProjectFiles/ProjectFileGenerator.cs
+++ b/Engine/Source/Programs/UnrealBuildTool/ProjectFiles/ProjectFileGenerator.cs
@@ -953,7 +953,7 @@ namespace UnrealBuildTool
 				}
 
 				// write out any additional debug information for the solution (such as UnrealVS configuration)
-				WriteDebugSolutionFiles(IntermediateProjectFilesPath);
+				WriteDebugSolutionFiles(PlatformProjectGenerators, IntermediateProjectFilesPath);
 
 				// Generate IntelliSense data if we need to.  This involves having UBT simulate the action compilation of
 				// the targets so that we can extra the compiler defines, include paths, etc.
@@ -2525,8 +2525,9 @@ namespace UnrealBuildTool
 		/// <summary>
 		/// Writes any additional solution-wide debug files (e.g. UnrealVS hints)
 		/// </summary>
+		/// <param name="PlatformProjectGenerators">The platform project generators</param>
 		/// <param name="IntermediateProjectFilesPath">Intermediate project files folder</param>
-		protected virtual void WriteDebugSolutionFiles( DirectoryReference IntermediateProjectFilesPath )
+		protected virtual void WriteDebugSolutionFiles( PlatformProjectGeneratorCollection PlatformProjectGenerators, DirectoryReference IntermediateProjectFilesPath )
 		{
 		}
 
diff --git a/Engine/Source/Programs/UnrealBuildTool/ProjectFiles/VisualStudio/VCProjectFileGenerator.cs b/Engine/Source/Programs/UnrealBuildTool/ProjectFiles/VisualStudio/VCProjectFileGenerator.cs
--- a/Engine/Source/Programs/UnrealBuildTool/ProjectFiles/VisualStudio/VCProjectFileGenerator.cs
+++ b/Engine/Source/Programs/UnrealBuildTool/ProjectFiles/VisualStudio/VCProjectFileGenerator.cs
@@ -819,13 +819,13 @@ namespace UnrealBuildTool
 			return bSuccess;
 		}
 
-		protected override void WriteDebugSolutionFiles( DirectoryReference IntermediateProjectFilesPath )
+		protected override void WriteDebugSolutionFiles( PlatformProjectGeneratorCollection PlatformProjectGenerators, DirectoryReference IntermediateProjectFilesPath )
 		{
 			//build and collect UnrealVS configuration
 			StringBuilder UnrealVSContent = new StringBuilder();
 			foreach (UnrealTargetPlatform SupportedPlatform in SupportedPlatforms)
 			{
-				UEPlatformProjectGenerator ProjGenerator = UEPlatformProjectGenerator.GetPlatformProjectGenerator(SupportedPlatform, true);
+				PlatformProjectGenerator ProjGenerator = PlatformProjectGenerators.GetPlatformProjectGenerator(SupportedPlatform, true);
 				if (ProjGenerator != null)
 				{
 					ProjGenerator.GetUnrealVSConfigurationEntries(UnrealVSContent);
