UBT: Always allow using SNDBS on Windows, if it's available.

#rnx

diff --git a/Engine/Source/Programs/UnrealBuildTool/Platform/Windows/UEBuildWindows.cs b/Engine/Source/Programs/UnrealBuildTool/Platform/Windows/UEBuildWindows.cs
--- a/Engine/Source/Programs/UnrealBuildTool/Platform/Windows/UEBuildWindows.cs
+++ b/Engine/Source/Programs/UnrealBuildTool/Platform/Windows/UEBuildWindows.cs
@@ -1302,18 +1302,7 @@ namespace UnrealBuildTool
 		/// </summary>
 		public override bool CanUseSNDBS()
 		{
-			// Check that SN-DBS is available
-			string SCERootPath = Environment.GetEnvironmentVariable("SCE_ROOT_DIR");
-			if (!String.IsNullOrEmpty(SCERootPath))
-			{
-				string SNDBSPath = Path.Combine(SCERootPath, "common", "sn-dbs", "bin", "dbsbuild.exe");
-				bool bIsSNDBSAvailable = File.Exists(SNDBSPath);
-				return bIsSNDBSAvailable;
-			}
-			else
-			{
-				return false;
-			}
+			return true;
 		}
 
 		/// <summary>
