UBT: Fix build products being added to a receipt that may not exist in an installed build.

diff --git a/Engine/Source/Programs/UnrealBuildTool/Configuration/UEBuildTarget.cs b/Engine/Source/Programs/UnrealBuildTool/Configuration/UEBuildTarget.cs
--- a/Engine/Source/Programs/UnrealBuildTool/Configuration/UEBuildTarget.cs
+++ b/Engine/Source/Programs/UnrealBuildTool/Configuration/UEBuildTarget.cs
@@ -1865,6 +1865,12 @@ namespace UnrealBuildTool
 			}
 			BuildProducts.AddRange(RuntimeDependencyTargetFileToSourceFile.Select(x => new KeyValuePair<FileReference, BuildProductType>(x.Key, BuildProductType.RequiredResource)));
 
+			// Remove any installed build products that don't exist. They may be part of an optional install.
+			if(UnrealBuildTool.IsEngineInstalled())
+			{
+				BuildProducts.RemoveAll(x => UnrealBuildTool.IsFileInstalled(x.Key) && !FileReference.Exists(x.Key));
+			}
+
 			// Make sure all the checked headers were valid
 			List<string> InvalidIncludeDirectiveMessages = Modules.Values.OfType<UEBuildModuleCPP>().Where(x => x.InvalidIncludeDirectiveMessages != null).SelectMany(x => x.InvalidIncludeDirectiveMessages).ToList();
 			if (InvalidIncludeDirectiveMessages.Count > 0)
