UBT: Remove the bCreateStubIPA argument from the cross-platform deployment parameter block. It's only needed in IOSToolChain.PostBuildSync(), right before we create the stub.

#rnx

diff --git a/Engine/Source/Programs/UnrealBuildTool/Configuration/UEBuildDeployTarget.cs b/Engine/Source/Programs/UnrealBuildTool/Configuration/UEBuildDeployTarget.cs
--- a/Engine/Source/Programs/UnrealBuildTool/Configuration/UEBuildDeployTarget.cs
+++ b/Engine/Source/Programs/UnrealBuildTool/Configuration/UEBuildDeployTarget.cs
@@ -50,11 +50,6 @@ namespace UnrealBuildTool
 		/// </summary>
 		public readonly FileReference BuildReceiptFileName;
 
-		/// <summary>
-		/// If true, then a stub IPA will be generated when compiling is done (minimal files needed for a valid IPA)
-		/// </summary>
-		public readonly bool bCreateStubIPA;
-
 		/// <summary>
 		/// Construct the deployment info from a target
 		/// </summary>
@@ -68,7 +63,6 @@ namespace UnrealBuildTool
 			this.Configuration = Target.Configuration;
 			this.ProjectDirectory = Target.ProjectDirectory;
 			this.BuildReceiptFileName = Target.ReceiptFileName;
-			this.bCreateStubIPA = Target.Rules.bCreateStubIPA;
 		}
 
 		/// <summary>
@@ -86,7 +80,6 @@ namespace UnrealBuildTool
 				Configuration = (UnrealTargetConfiguration)Reader.ReadInt32();
 				ProjectDirectory = Reader.ReadDirectoryReference();
 				BuildReceiptFileName = Reader.ReadFileReference();
-				bCreateStubIPA = Reader.ReadBoolean();
 			}
 		}
 
@@ -106,7 +99,6 @@ namespace UnrealBuildTool
 				Writer.Write((Int32)Configuration);
 				Writer.Write(ProjectDirectory);
 				Writer.Write(BuildReceiptFileName);
-				Writer.Write(bCreateStubIPA);
 			}
 		}
 	}
diff --git a/Engine/Source/Programs/UnrealBuildTool/Platform/IOS/IOSToolChain.cs b/Engine/Source/Programs/UnrealBuildTool/Platform/IOS/IOSToolChain.cs
--- a/Engine/Source/Programs/UnrealBuildTool/Platform/IOS/IOSToolChain.cs
+++ b/Engine/Source/Programs/UnrealBuildTool/Platform/IOS/IOSToolChain.cs
@@ -1495,7 +1495,7 @@ namespace UnrealBuildTool
 
             // ensure the plist, entitlements, and provision files are properly copied
             UEDeployIOS DeployHandler = (Target.Platform == UnrealTargetPlatform.IOS ? new UEDeployIOS() : new UEDeployTVOS());
-            DeployHandler.PrepTargetForDeployment(new UEBuildDeployTarget(Target));
+            DeployHandler.PrepTargetForDeployment(new UEBuildDeployTarget(Target), Target.Rules.bCreateStubIPA);
 
 			// copy the executable
 			if (!File.Exists(FinalRemoteExecutablePath))
@@ -1532,7 +1532,7 @@ namespace UnrealBuildTool
 
 				// ensure the plist, entitlements, and provision files are properly copied
 				DeployHandler = (Target.Platform == UnrealTargetPlatform.IOS ? new UEDeployIOS() : new UEDeployTVOS());
-				DeployHandler.PrepTargetForDeployment(new UEBuildDeployTarget(Target));
+				DeployHandler.PrepTargetForDeployment(new UEBuildDeployTarget(Target), true);
 
 				FileReference SignProjectScript = FileReference.Combine(Target.ProjectIntermediateDirectory, "SignProject.sh");
 				using(StreamWriter Writer = new StreamWriter(SignProjectScript.FullName))
diff --git a/Engine/Source/Programs/UnrealBuildTool/Platform/IOS/UEDeployIOS.cs b/Engine/Source/Programs/UnrealBuildTool/Platform/IOS/UEDeployIOS.cs
--- a/Engine/Source/Programs/UnrealBuildTool/Platform/IOS/UEDeployIOS.cs
+++ b/Engine/Source/Programs/UnrealBuildTool/Platform/IOS/UEDeployIOS.cs
@@ -1077,6 +1077,11 @@ namespace UnrealBuildTool
 		}
 
 		public override bool PrepTargetForDeployment(UEBuildDeployTarget InTarget)
+		{
+			return PrepTargetForDeployment(InTarget, false);
+		}
+
+		public bool PrepTargetForDeployment(UEBuildDeployTarget InTarget, bool bCreateStubIPA)
 		{
 			TargetReceipt Receipt = TargetReceipt.Read(InTarget.BuildReceiptFileName);
 
@@ -1101,7 +1106,7 @@ namespace UnrealBuildTool
 
 			if (BuildHostPlatform.Current.Platform == UnrealTargetPlatform.Mac && Environment.GetEnvironmentVariable("UBT_NO_POST_DEPLOY") != "true")
 			{
-				return PrepForUATPackageOrDeploy(InTarget.Configuration, InTarget.ProjectFile, GameName, ProjectDirectory, BuildPath + "/" + DecoratedGameName, "../../Engine", false, "", false, InTarget.bCreateStubIPA, Receipt);
+				return PrepForUATPackageOrDeploy(InTarget.Configuration, InTarget.ProjectFile, GameName, ProjectDirectory, BuildPath + "/" + DecoratedGameName, "../../Engine", false, "", false, bCreateStubIPA, Receipt);
 			}
 			else
 			{
