UBT: Move the bStripSymbolsOnIOS and bCreateStubIPA settings into the IOS settings object.

diff --git a/Engine/Source/Programs/UnrealBuildTool/Configuration/TargetRules.cs b/Engine/Source/Programs/UnrealBuildTool/Configuration/TargetRules.cs
--- a/Engine/Source/Programs/UnrealBuildTool/Configuration/TargetRules.cs
+++ b/Engine/Source/Programs/UnrealBuildTool/Configuration/TargetRules.cs
@@ -808,8 +808,22 @@ namespace UnrealBuildTool
 		/// <summary>
 		/// Whether to strip iOS symbols or not (implied by bGeneratedSYMFile).
 		/// </summary>
-		[XmlConfigFile(Category = "BuildConfiguration")]
-		public bool bStripSymbolsOnIOS = false;
+		[Obsolete("bStripSymbolsOnIOS has been deprecated. Use IOSPlatform.bStripSymbols instead.")]
+		public bool bStripSymbolsOnIOS
+		{
+			get { return IOSPlatform.bStripSymbols; }
+			set { IOSPlatform.bStripSymbols = value; }
+		}
+
+		/// <summary>
+		/// If true, then a stub IPA will be generated when compiling is done (minimal files needed for a valid IPA).
+		/// </summary>
+		[Obsolete("bCreateStubIPA has been deprecated. Use IOSPlatform.bCreateStubIPA instead.")]
+		public bool bCreateStubIPA
+		{
+			get { return IOSPlatform.bCreateStubIPA; }
+			set { IOSPlatform.bCreateStubIPA = value; }
+		}
 
 		/// <summary>
 		/// If true, then enable memory profiling in the build (defines USE_MALLOC_PROFILER=1 and forces bOmitFramePointers=false).
@@ -876,12 +890,6 @@ namespace UnrealBuildTool
 		[CommandLine("-SkipDeploy", Value = "false")]
 		public bool bDeployAfterCompile = false;
 
-		/// <summary>
-		/// If true, then a stub IPA will be generated when compiling is done (minimal files needed for a valid IPA).
-		/// </summary>
-		[CommandLine("-CreateStub", Value = "true")]
-		public bool bCreateStubIPA = false;
-
 		/// <summary>
 		/// When enabled, allows XGE to compile pre-compiled header files on remote machines.  Otherwise, PCHs are always generated locally.
 		/// </summary>
@@ -1967,9 +1975,10 @@ namespace UnrealBuildTool
 			get { return Inner.bOmitFramePointers; }
 		}
 
+		[Obsolete("bStripSymbolsOnIOS has been deprecated. Use IOSPlatform.bStripSymbols instead.")]
 		public bool bStripSymbolsOnIOS
 		{
-			get { return Inner.bStripSymbolsOnIOS; }
+			get { return IOSPlatform.bStripSymbols; }
 		}
 
 		public bool bUseMallocProfiler
@@ -2022,9 +2031,10 @@ namespace UnrealBuildTool
 			get { return Inner.bDeployAfterCompile; }
 		}
 
+		[Obsolete("bCreateStubIPA has been deprecated. Use IOSPlatform.bCreateStubIPA instead.")]
 		public bool bCreateStubIPA
 		{
-			get { return Inner.bCreateStubIPA; }
+			get { return IOSPlatform.bCreateStubIPA; }
 		}
 
 		public bool bAllowRemotelyCompiledPCHs
diff --git a/Engine/Source/Programs/UnrealBuildTool/Platform/IOS/IOSToolChain.cs b/Engine/Source/Programs/UnrealBuildTool/Platform/IOS/IOSToolChain.cs
--- a/Engine/Source/Programs/UnrealBuildTool/Platform/IOS/IOSToolChain.cs
+++ b/Engine/Source/Programs/UnrealBuildTool/Platform/IOS/IOSToolChain.cs
@@ -137,7 +137,7 @@ namespace UnrealBuildTool
 
 		public override void ModifyBuildProducts(ReadOnlyTargetRules Target, UEBuildBinary Binary, List<string> Libraries, List<UEBuildBundleResource> BundleResources, Dictionary<FileReference, BuildProductType> BuildProducts)
 		{
-			if (Target.bCreateStubIPA && Binary.Type != UEBuildBinaryType.StaticLibrary)
+			if (Target.IOSPlatform.bCreateStubIPA && Binary.Type != UEBuildBinaryType.StaticLibrary)
 			{
 				FileReference StubFile = FileReference.Combine(Binary.OutputFilePath.Directory, Binary.OutputFilePath.GetFileNameWithoutExtension() + ".stub");
 				BuildProducts.Add(StubFile, BuildProductType.Package);
@@ -1473,7 +1473,7 @@ namespace UnrealBuildTool
 			string FinalRemoteExecutablePath = String.Format("{0}/Payload/{1}.app/{1}", RemoteShadowDirectoryMac, AppName);
 
 			// strip the debug info from the executable if needed
-			if (Target.Rules.bStripSymbolsOnIOS || (Target.Configuration == UnrealTargetConfiguration.Shipping))
+			if (Target.Rules.IOSPlatform.bStripSymbols || (Target.Configuration == UnrealTargetConfiguration.Shipping))
 			{
 				Process StripProcess = new Process();
 				StripProcess.StartInfo.WorkingDirectory = RemoteShadowDirectoryMac;
@@ -1495,7 +1495,7 @@ namespace UnrealBuildTool
 
             // ensure the plist, entitlements, and provision files are properly copied
             UEDeployIOS DeployHandler = (Target.Platform == UnrealTargetPlatform.IOS ? new UEDeployIOS() : new UEDeployTVOS());
-            DeployHandler.PrepTargetForDeployment(new UEBuildDeployTarget(Target), Target.Rules.bCreateStubIPA);
+            DeployHandler.PrepTargetForDeployment(new UEBuildDeployTarget(Target), Target.Rules.IOSPlatform.bCreateStubIPA);
 
 			// copy the executable
 			if (!File.Exists(FinalRemoteExecutablePath))
@@ -1509,7 +1509,7 @@ namespace UnrealBuildTool
 				GenerateCrashlyticsData(RemoteShadowDirectoryMac, Path.GetFileName(Target.OutputPath.FullName), Target.ProjectDirectory.FullName, AppName);
 			}
 
-			if (Target.Rules.bCreateStubIPA)
+			if (Target.Rules.IOSPlatform.bCreateStubIPA)
 			{
 				// generate the dummy project so signing works
 				DirectoryReference XcodeWorkspaceDir;
diff --git a/Engine/Source/Programs/UnrealBuildTool/Platform/IOS/UEBuildIOS.cs b/Engine/Source/Programs/UnrealBuildTool/Platform/IOS/UEBuildIOS.cs
--- a/Engine/Source/Programs/UnrealBuildTool/Platform/IOS/UEBuildIOS.cs
+++ b/Engine/Source/Programs/UnrealBuildTool/Platform/IOS/UEBuildIOS.cs
@@ -17,6 +17,18 @@ namespace UnrealBuildTool
 	/// </summary>
 	public class IOSTargetRules
 	{
+		/// <summary>
+		/// Whether to strip iOS symbols or not (implied by bGeneratedSYMFile).
+		/// </summary>
+		[XmlConfigFile(Category = "BuildConfiguration")]
+		public bool bStripSymbols = false;
+
+		/// <summary>
+		/// If true, then a stub IPA will be generated when compiling is done (minimal files needed for a valid IPA).
+		/// </summary>
+		[CommandLine("-CreateStub", Value = "true")]
+		public bool bCreateStubIPA = false;
+
 		/// <summary>
 		/// Don't generate crashlytics data
 		/// </summary>
@@ -68,6 +80,16 @@ namespace UnrealBuildTool
 #if !__MonoCS__
 #pragma warning disable CS1591
 #endif
+		public bool bStripSymbols
+		{
+			get { return Inner.bStripSymbols; }
+		}
+			
+		public bool bCreateStubIPA
+		{
+			get { return Inner.bCreateStubIPA; }
+		}
+
 		public bool bSkipCrashlytics
 		{
 			get { return Inner.bSkipCrashlytics; }
