UBT: Change deployment to use the target receipt rather than a custom binary parameter block.

diff --git a/Engine/Source/Programs/AutomationTool/AutomationUtils/UE4Build.cs b/Engine/Source/Programs/AutomationTool/AutomationUtils/UE4Build.cs
--- a/Engine/Source/Programs/AutomationTool/AutomationUtils/UE4Build.cs
+++ b/Engine/Source/Programs/AutomationTool/AutomationUtils/UE4Build.cs
@@ -153,7 +153,7 @@ namespace AutomationTool
 			// run the deployment steps, if necessary
 			foreach(string DeployTargetFile in Item.Manifest.DeployTargetFiles)
 			{
-				CommandUtils.RunUBT(CommandUtils.CmdEnv, UBTExecutable, String.Format("-Mode=Deploy -TargetInfo=\"{0}\"", DeployTargetFile));
+				CommandUtils.RunUBT(CommandUtils.CmdEnv, UBTExecutable, String.Format("-Mode=Deploy -Receipt=\"{0}\"", DeployTargetFile));
 			}
 
 			foreach (string ManifestItem in Item.Manifest.BuildProducts)
diff --git a/Engine/Source/Programs/UnrealBuildTool/Configuration/UEBuildDeploy.cs b/Engine/Source/Programs/UnrealBuildTool/Configuration/UEBuildDeploy.cs
--- a/Engine/Source/Programs/UnrealBuildTool/Configuration/UEBuildDeploy.cs
+++ b/Engine/Source/Programs/UnrealBuildTool/Configuration/UEBuildDeploy.cs
@@ -17,9 +17,9 @@ namespace UnrealBuildTool
 		/// <summary>
 		/// Prepare the target for deployment
 		/// </summary>
-		/// <param name="InTarget"> The target for deployment</param>
-		/// <returns>bool   true if successful, false if not</returns>
-		public virtual bool PrepTargetForDeployment(UEBuildDeployTarget InTarget)
+		/// <param name="Receipt">Receipt for the target being deployed</param>
+		/// <returns>True if successful, false if not</returns>
+		public virtual bool PrepTargetForDeployment(TargetReceipt Receipt)
 		{
 			return true;
 		}
diff --git a/Engine/Source/Programs/UnrealBuildTool/Configuration/UEBuildDeployTarget.cs b/Engine/Source/Programs/UnrealBuildTool/Configuration/UEBuildDeployTarget.cs
deleted file mode 100644
diff --git a/Engine/Source/Programs/UnrealBuildTool/Configuration/UEBuildPlatform.cs b/Engine/Source/Programs/UnrealBuildTool/Configuration/UEBuildPlatform.cs
--- a/Engine/Source/Programs/UnrealBuildTool/Configuration/UEBuildPlatform.cs
+++ b/Engine/Source/Programs/UnrealBuildTool/Configuration/UEBuildPlatform.cs
@@ -840,7 +840,7 @@ namespace UnrealBuildTool
 		/// <summary>
 		/// Deploys the given target
 		/// </summary>
-		/// <param name="Target">Information about the target being deployed</param>
-		public abstract void Deploy(UEBuildDeployTarget Target);
+		/// <param name="Receipt">Receipt for the target being deployed</param>
+		public abstract void Deploy(TargetReceipt Receipt);
 	}
 }
diff --git a/Engine/Source/Programs/UnrealBuildTool/Configuration/UEBuildTarget.cs b/Engine/Source/Programs/UnrealBuildTool/Configuration/UEBuildTarget.cs
--- a/Engine/Source/Programs/UnrealBuildTool/Configuration/UEBuildTarget.cs
+++ b/Engine/Source/Programs/UnrealBuildTool/Configuration/UEBuildTarget.cs
@@ -762,9 +762,9 @@ namespace UnrealBuildTool
 		FileReference[] PostBuildStepScripts;
 
 		/// <summary>
-		/// File containing information needed to deploy this target
+		/// Whether to deploy this target after compilation
 		/// </summary>
-		public FileReference DeployTargetFile;
+		public bool bDeployAfterCompile;
 
 		/// <summary>
 		/// Directories which are scanned for source files. The UBT makefile uses this to check for files being added/removed.
@@ -813,7 +813,7 @@ namespace UnrealBuildTool
 			TargetRulesFile = (FileReference)Info.GetValue("tc", typeof(FileReference));
 			PreBuildStepScripts = (FileReference[])Info.GetValue("pr", typeof(FileReference[]));
 			PostBuildStepScripts = (FileReference[])Info.GetValue("po", typeof(FileReference[]));
-			DeployTargetFile = (FileReference)Info.GetValue("dt", typeof(FileReference));
+			bDeployAfterCompile = Info.GetBoolean("dt");
 			bHasProjectScriptPlugin = Info.GetBoolean("sp");
 			SourceDirectories = (HashSet<DirectoryReference>)Info.GetValue("sd", typeof(HashSet<DirectoryReference>));
 		}
@@ -842,7 +842,7 @@ namespace UnrealBuildTool
 			Info.AddValue("tc", TargetRulesFile);
 			Info.AddValue("pr", PreBuildStepScripts);
 			Info.AddValue("po", PostBuildStepScripts);
-			Info.AddValue("dt", DeployTargetFile);
+			Info.AddValue("dt", bDeployAfterCompile);
 			Info.AddValue("sp", bHasProjectScriptPlugin);
 			Info.AddValue("sd", SourceDirectories);
 		}
@@ -866,6 +866,7 @@ namespace UnrealBuildTool
 			TargetType = Rules.Type;
 			bPrecompile = InRules.bPrecompile;
 			ForeignPlugin = InDesc.ForeignPlugin;
+			bDeployAfterCompile = InRules.bDeployAfterCompile;
 
 			// now that we have the platform, we can set the intermediate path to include the platform/architecture name
 			PlatformIntermediateFolder = Path.Combine("Intermediate", "Build", Platform.ToString(), UEBuildPlatform.GetBuildPlatform(Platform).GetFolderNameForArchitecture(Architecture));
@@ -1313,9 +1314,9 @@ namespace UnrealBuildTool
 					Manifest.AddBuildProduct(ReceiptFileName.FullName);
 				}
 
-				if (DeployTargetFile != null)
+				if (bDeployAfterCompile)
 				{
-					Manifest.DeployTargetFiles.Add(DeployTargetFile.FullName);
+					Manifest.DeployTargetFiles.Add(ReceiptFileName.FullName);
 				}
 			}
 
@@ -1742,14 +1743,6 @@ namespace UnrealBuildTool
 			// On Mac and Linux we have actions that should be executed after all the binaries are created
 			TargetToolChain.SetupBundleDependencies(Binaries, TargetName);
 
-			// Write out the deployment context, if necessary
-			if(Rules.bDeployAfterCompile && !Rules.bDisableLinking)
-			{
-				UEBuildDeployTarget DeployTarget = new UEBuildDeployTarget(this);
-				DeployTargetFile = FileReference.Combine(ProjectIntermediateDirectory, "Deploy.dat");
-				DeployTarget.Write(DeployTargetFile);
-			}
-
 			if (!ProjectFileGenerator.bGenerateProjectFiles)
 			{
 				HashSet<UEBuildModuleCPP> ModulesToGenerateHeadersFor = GatherDependencyModules(OriginalBinaries.ToList());
diff --git a/Engine/Source/Programs/UnrealBuildTool/Modes/DeployMode.cs b/Engine/Source/Programs/UnrealBuildTool/Modes/DeployMode.cs
--- a/Engine/Source/Programs/UnrealBuildTool/Modes/DeployMode.cs
+++ b/Engine/Source/Programs/UnrealBuildTool/Modes/DeployMode.cs
@@ -16,8 +16,8 @@ namespace UnrealBuildTool
 		/// <summary>
 		/// If we are just running the deployment step, specifies the path to the given deployment settings
 		/// </summary>
-		[CommandLine("-TargetInfo", Required=true)]
-		public FileReference TargetInfoFile = null;
+		[CommandLine("-Receipt", Required=true)]
+		public FileReference ReceiptFile = null;
 
 		/// <summary>
 		/// Execute the tool mode
@@ -38,9 +38,9 @@ namespace UnrealBuildTool
 			UnrealBuildTool.RegisterAllUBTClasses(false);
 
 			// Execute the deploy
-			UEBuildDeployTarget DeployTarget = new UEBuildDeployTarget(TargetInfoFile);
-			Log.WriteLine(LogEventType.Console, "Deploying {0} {1} {2}...", DeployTarget.TargetName, DeployTarget.Platform, DeployTarget.Configuration);
-			UEBuildPlatform.GetBuildPlatform(DeployTarget.Platform).Deploy(DeployTarget);
+			TargetReceipt Receipt = TargetReceipt.Read(ReceiptFile);
+			Log.WriteLine(LogEventType.Console, "Deploying {0} {1} {2}...", Receipt.TargetName, Receipt.Platform, Receipt.Configuration);
+			UEBuildPlatform.GetBuildPlatform(Receipt.Platform).Deploy(Receipt);
 
 			return (int)ECompilationResult.Succeeded;
 		}
diff --git a/Engine/Source/Programs/UnrealBuildTool/Platform/Android/UEBuildAndroid.cs b/Engine/Source/Programs/UnrealBuildTool/Platform/Android/UEBuildAndroid.cs
--- a/Engine/Source/Programs/UnrealBuildTool/Platform/Android/UEBuildAndroid.cs
+++ b/Engine/Source/Programs/UnrealBuildTool/Platform/Android/UEBuildAndroid.cs
@@ -465,11 +465,11 @@ namespace UnrealBuildTool
 		/// <summary>
 		/// Deploys the given target
 		/// </summary>
-		/// <param name="Target">Information about the target being deployed</param>
-		public override void Deploy(UEBuildDeployTarget Target)
+		/// <param name="Receipt">Receipt for the target being deployed</param>
+		public override void Deploy(TargetReceipt Receipt)
 		{
 			// do not package data if building via UBT
-			new UEDeployAndroid(Target.ProjectFile, false).PrepTargetForDeployment(Target);
+			new UEDeployAndroid(Receipt.ProjectFile, false).PrepTargetForDeployment(Receipt);
 		}
 	}
 
diff --git a/Engine/Source/Programs/UnrealBuildTool/Platform/Android/UEDeployAndroid.cs b/Engine/Source/Programs/UnrealBuildTool/Platform/Android/UEDeployAndroid.cs
--- a/Engine/Source/Programs/UnrealBuildTool/Platform/Android/UEDeployAndroid.cs
+++ b/Engine/Source/Programs/UnrealBuildTool/Platform/Android/UEDeployAndroid.cs
@@ -3998,14 +3998,11 @@ namespace UnrealBuildTool
 			return PluginExtras;
 		}
 
-		public override bool PrepTargetForDeployment(UEBuildDeployTarget InTarget)
+		public override bool PrepTargetForDeployment(TargetReceipt Receipt)
 		{
-			AndroidToolChain ToolChain = ((AndroidPlatform)UEBuildPlatform.GetBuildPlatform(InTarget.Platform)).CreateTempToolChainForProject(InTarget.ProjectFile) as AndroidToolChain;
+			AndroidToolChain ToolChain = ((AndroidPlatform)UEBuildPlatform.GetBuildPlatform(Receipt.Platform)).CreateTempToolChainForProject(Receipt.ProjectFile) as AndroidToolChain;
 
 			// get the receipt
-			FileReference ReceiptFilename = InTarget.BuildReceiptFileName;
-			Log.TraceInformation("Receipt Filename: {0}", ReceiptFilename);
-			TargetReceipt Receipt = TargetReceipt.Read(ReceiptFilename);
 			SetAndroidPluginData(ToolChain.GetAllArchitectures(), CollectPluginDataPaths(Receipt));
 
 			// Get the output paths
@@ -4016,8 +4013,9 @@ namespace UnrealBuildTool
 
 			// make an apk at the end of compiling, so that we can run without packaging (debugger, cook on the fly, etc)
 			string RelativeEnginePath = UnrealBuildTool.EngineDirectory.MakeRelativeTo(DirectoryReference.GetCurrentDirectory());
-			string TargetName = (InTarget.ProjectFile == null ? InTarget.TargetName : InTarget.ProjectFile.GetFileNameWithoutAnyExtensions());
-			MakeApk(ToolChain, TargetName, InTarget.TargetType, InTarget.ProjectDirectory.FullName, BaseSoName, RelativeEnginePath, bForDistribution: false, CookFlavor: "",
+			string TargetName = (Receipt.ProjectFile == null ? Receipt.TargetName : Receipt.ProjectFile.GetFileNameWithoutAnyExtensions());
+			DirectoryReference ProjectDirectory = DirectoryReference.FromFile(Receipt.ProjectFile) ?? UnrealBuildTool.EngineDirectory;
+			MakeApk(ToolChain, TargetName, Receipt.TargetType, ProjectDirectory.FullName, BaseSoName, RelativeEnginePath, bForDistribution: false, CookFlavor: "",
 				bMakeSeparateApks: ShouldMakeSeparateApks(), bIncrementalPackage: true, bDisallowPackagingDataInApk: false, bDisallowExternalFilesDir: true);
 
 			// if we made any non-standard .apk files, the generated debugger settings may be wrong
diff --git a/Engine/Source/Programs/UnrealBuildTool/Platform/HTML5/UEBuildHTML5.cs b/Engine/Source/Programs/UnrealBuildTool/Platform/HTML5/UEBuildHTML5.cs
--- a/Engine/Source/Programs/UnrealBuildTool/Platform/HTML5/UEBuildHTML5.cs
+++ b/Engine/Source/Programs/UnrealBuildTool/Platform/HTML5/UEBuildHTML5.cs
@@ -231,8 +231,8 @@ namespace UnrealBuildTool
 		/// <summary>
 		/// Deploys the given target
 		/// </summary>
-		/// <param name="Target">Information about the target being deployed</param>
-		public override void Deploy(UEBuildDeployTarget Target)
+		/// <param name="Receipt">Receipt for the target being deployed</param>
+		public override void Deploy(TargetReceipt Receipt)
 		{
 		}
 	}
diff --git a/Engine/Source/Programs/UnrealBuildTool/Platform/IOS/IOSToolChain.cs b/Engine/Source/Programs/UnrealBuildTool/Platform/IOS/IOSToolChain.cs
--- a/Engine/Source/Programs/UnrealBuildTool/Platform/IOS/IOSToolChain.cs
+++ b/Engine/Source/Programs/UnrealBuildTool/Platform/IOS/IOSToolChain.cs
@@ -1465,6 +1465,8 @@ namespace UnrealBuildTool
 				return;
 			}
 
+			TargetReceipt Receipt = TargetReceipt.Read(Target.ReceiptFileName);
+
 			IOSProjectSettings ProjectSettings = ((IOSPlatform)UEBuildPlatform.GetBuildPlatform(Target.Platform)).ReadProjectSettings(Target.ProjectFile);
 
 			string AppName = Target.TargetType == TargetType.Game ? Target.TargetName : Target.AppName;
@@ -1495,7 +1497,7 @@ namespace UnrealBuildTool
 
             // ensure the plist, entitlements, and provision files are properly copied
             UEDeployIOS DeployHandler = (Target.Platform == UnrealTargetPlatform.IOS ? new UEDeployIOS() : new UEDeployTVOS());
-            DeployHandler.PrepTargetForDeployment(new UEBuildDeployTarget(Target), Target.Rules.IOSPlatform.bCreateStubIPA);
+            DeployHandler.PrepTargetForDeployment(Receipt, Target.Rules.IOSPlatform.bCreateStubIPA);
 
 			// copy the executable
 			if (!File.Exists(FinalRemoteExecutablePath))
@@ -1532,7 +1534,7 @@ namespace UnrealBuildTool
 
 				// ensure the plist, entitlements, and provision files are properly copied
 				DeployHandler = (Target.Platform == UnrealTargetPlatform.IOS ? new UEDeployIOS() : new UEDeployTVOS());
-				DeployHandler.PrepTargetForDeployment(new UEBuildDeployTarget(Target), true);
+				DeployHandler.PrepTargetForDeployment(Receipt, true);
 
 				FileReference SignProjectScript = FileReference.Combine(Target.ProjectIntermediateDirectory, "SignProject.sh");
 				using(StreamWriter Writer = new StreamWriter(SignProjectScript.FullName))
diff --git a/Engine/Source/Programs/UnrealBuildTool/Platform/IOS/UEBuildIOS.cs b/Engine/Source/Programs/UnrealBuildTool/Platform/IOS/UEBuildIOS.cs
--- a/Engine/Source/Programs/UnrealBuildTool/Platform/IOS/UEBuildIOS.cs
+++ b/Engine/Source/Programs/UnrealBuildTool/Platform/IOS/UEBuildIOS.cs
@@ -1003,10 +1003,10 @@ namespace UnrealBuildTool
 		/// <summary>
 		/// Deploys the given target
 		/// </summary>
-		/// <param name="Target">Information about the target being deployed</param>
-		public override void Deploy(UEBuildDeployTarget Target)
+		/// <param name="Receipt">Receipt for the target being deployed</param>
+		public override void Deploy(TargetReceipt Receipt)
 		{
-			new UEDeployIOS().PrepTargetForDeployment(Target);
+			new UEDeployIOS().PrepTargetForDeployment(Receipt);
 		}
 	}
 
diff --git a/Engine/Source/Programs/UnrealBuildTool/Platform/IOS/UEDeployIOS.cs b/Engine/Source/Programs/UnrealBuildTool/Platform/IOS/UEDeployIOS.cs
--- a/Engine/Source/Programs/UnrealBuildTool/Platform/IOS/UEDeployIOS.cs
+++ b/Engine/Source/Programs/UnrealBuildTool/Platform/IOS/UEDeployIOS.cs
@@ -1076,43 +1076,41 @@ namespace UnrealBuildTool
 			return true;
 		}
 
-		public override bool PrepTargetForDeployment(UEBuildDeployTarget InTarget)
+		public override bool PrepTargetForDeployment(TargetReceipt Receipt)
 		{
-			return PrepTargetForDeployment(InTarget, false);
+			return PrepTargetForDeployment(Receipt, false);
 		}
 
-		public bool PrepTargetForDeployment(UEBuildDeployTarget InTarget, bool bCreateStubIPA)
+		public bool PrepTargetForDeployment(TargetReceipt Receipt, bool bCreateStubIPA)
 		{
-			TargetReceipt Receipt = TargetReceipt.Read(InTarget.BuildReceiptFileName);
-
 			string SubDir = GetTargetPlatformName();
 
-			string GameName = InTarget.TargetName;
-			string BuildPath = (GameName == "UE4Game" ? "../../Engine" : InTarget.ProjectDirectory.FullName) + "/Binaries/" + SubDir;
-			string ProjectDirectory = InTarget.ProjectDirectory.FullName;
+			string GameName = Receipt.TargetName;
+			string ProjectDirectory = (DirectoryReference.FromFile(Receipt.ProjectFile) ?? UnrealBuildTool.EngineDirectory).FullName;
+			string BuildPath = (GameName == "UE4Game" ? "../../Engine" : ProjectDirectory) + "/Binaries/" + SubDir;
 			bool bIsUE4Game = GameName.Contains("UE4Game");
 
-			Console.WriteLine("1 GameName: {0}, ProjectName: {1}", GameName, (InTarget.ProjectFile == null) ? "" : Path.GetFileNameWithoutExtension(InTarget.ProjectFile.FullName));
+			Console.WriteLine("1 GameName: {0}, ProjectName: {1}", GameName, (Receipt.ProjectFile == null) ? "" : Path.GetFileNameWithoutExtension(Receipt.ProjectFile.FullName));
 
 			string DecoratedGameName;
-			if (InTarget.Configuration == UnrealTargetConfiguration.Development)
+			if (Receipt.Configuration == UnrealTargetConfiguration.Development)
 			{
 				DecoratedGameName = GameName;
 			}
 			else
 			{
-				DecoratedGameName = String.Format("{0}-{1}-{2}", GameName, InTarget.Platform.ToString(), InTarget.Configuration.ToString());
+				DecoratedGameName = String.Format("{0}-{1}-{2}", GameName, Receipt.Platform.ToString(), Receipt.Configuration.ToString());
 			}
 
 			if (BuildHostPlatform.Current.Platform == UnrealTargetPlatform.Mac && Environment.GetEnvironmentVariable("UBT_NO_POST_DEPLOY") != "true")
 			{
-				return PrepForUATPackageOrDeploy(InTarget.Configuration, InTarget.ProjectFile, GameName, ProjectDirectory, BuildPath + "/" + DecoratedGameName, "../../Engine", false, "", false, bCreateStubIPA, Receipt);
+				return PrepForUATPackageOrDeploy(Receipt.Configuration, Receipt.ProjectFile, GameName, ProjectDirectory, BuildPath + "/" + DecoratedGameName, "../../Engine", false, "", false, bCreateStubIPA, Receipt);
 			}
 			else
 			{
 				// @todo tvos merge: This used to copy the bundle back - where did that code go? It needs to be fixed up for TVOS directories
 				bool bSupportPortrait, bSupportLandscape, bSkipIcons;
-				GeneratePList(InTarget.ProjectFile, InTarget.Configuration, ProjectDirectory, bIsUE4Game, GameName, (InTarget.ProjectFile == null) ? "" : Path.GetFileNameWithoutExtension(InTarget.ProjectFile.FullName), "../../Engine", "", Receipt, out bSupportPortrait, out bSupportLandscape, out bSkipIcons);
+				GeneratePList(Receipt.ProjectFile, Receipt.Configuration, ProjectDirectory, bIsUE4Game, GameName, (Receipt.ProjectFile == null) ? "" : Path.GetFileNameWithoutExtension(Receipt.ProjectFile.FullName), "../../Engine", "", Receipt, out bSupportPortrait, out bSupportLandscape, out bSkipIcons);
 			}
 			return true;
 		}
diff --git a/Engine/Source/Programs/UnrealBuildTool/Platform/Linux/UEBuildLinux.cs b/Engine/Source/Programs/UnrealBuildTool/Platform/Linux/UEBuildLinux.cs
--- a/Engine/Source/Programs/UnrealBuildTool/Platform/Linux/UEBuildLinux.cs
+++ b/Engine/Source/Programs/UnrealBuildTool/Platform/Linux/UEBuildLinux.cs
@@ -537,8 +537,8 @@ namespace UnrealBuildTool
 		/// <summary>
 		/// Deploys the given target
 		/// </summary>
-		/// <param name="Target">Information about the target being deployed</param>
-		public override void Deploy(UEBuildDeployTarget Target)
+		/// <param name="Receipt">Receipt for the target being deployed</param>
+		public override void Deploy(TargetReceipt Receipt)
 		{
 		}
 	}
diff --git a/Engine/Source/Programs/UnrealBuildTool/Platform/Lumin/UEBuildLumin.cs b/Engine/Source/Programs/UnrealBuildTool/Platform/Lumin/UEBuildLumin.cs
--- a/Engine/Source/Programs/UnrealBuildTool/Platform/Lumin/UEBuildLumin.cs
+++ b/Engine/Source/Programs/UnrealBuildTool/Platform/Lumin/UEBuildLumin.cs
@@ -255,10 +255,10 @@ namespace UnrealBuildTool
 		/// <summary>
 		/// Deploys the given target
 		/// </summary>
-		/// <param name="Target">Information about the target being deployed</param>
-		public override void Deploy(UEBuildDeployTarget Target)
+		/// <param name="Receipt">Receipt for the target being deployed</param>
+		public override void Deploy(TargetReceipt Receipt)
 		{
-			new UEDeployLumin(Target.ProjectFile).PrepTargetForDeployment(Target);
+			new UEDeployLumin(Receipt.ProjectFile).PrepTargetForDeployment(Receipt);
 		}
 
 		public override void ValidateTarget(TargetRules Target)
diff --git a/Engine/Source/Programs/UnrealBuildTool/Platform/Lumin/UEDeployLumin.cs b/Engine/Source/Programs/UnrealBuildTool/Platform/Lumin/UEDeployLumin.cs
--- a/Engine/Source/Programs/UnrealBuildTool/Platform/Lumin/UEDeployLumin.cs
+++ b/Engine/Source/Programs/UnrealBuildTool/Platform/Lumin/UEDeployLumin.cs
@@ -381,7 +381,7 @@ namespace UnrealBuildTool
 			return true;
 		}
 
-		public override bool PrepTargetForDeployment(UEBuildDeployTarget InTarget)
+		public override bool PrepTargetForDeployment(TargetReceipt Receipt)
 		{
 			// @todo Lumin: Need to create a MabuFile with no data files - including the executable!!
 			return true;
diff --git a/Engine/Source/Programs/UnrealBuildTool/Platform/Mac/UEBuildMac.cs b/Engine/Source/Programs/UnrealBuildTool/Platform/Mac/UEBuildMac.cs
--- a/Engine/Source/Programs/UnrealBuildTool/Platform/Mac/UEBuildMac.cs
+++ b/Engine/Source/Programs/UnrealBuildTool/Platform/Mac/UEBuildMac.cs
@@ -348,10 +348,10 @@ namespace UnrealBuildTool
 		/// <summary>
 		/// Deploys the given target
 		/// </summary>
-		/// <param name="Target">Information about the target being deployed</param>
-		public override void Deploy(UEBuildDeployTarget Target)
+		/// <param name="Receipt">Receipt for the target being deployed</param>
+		public override void Deploy(TargetReceipt Receipt)
 		{
-			new UEDeployMac().PrepTargetForDeployment(Target);
+			new UEDeployMac().PrepTargetForDeployment(Receipt);
 		}
 	}
 
diff --git a/Engine/Source/Programs/UnrealBuildTool/Platform/Mac/UEDeployMac.cs b/Engine/Source/Programs/UnrealBuildTool/Platform/Mac/UEDeployMac.cs
--- a/Engine/Source/Programs/UnrealBuildTool/Platform/Mac/UEDeployMac.cs
+++ b/Engine/Source/Programs/UnrealBuildTool/Platform/Mac/UEDeployMac.cs
@@ -14,7 +14,7 @@ namespace UnrealBuildTool
 {
 	class UEDeployMac : UEBuildDeploy
 	{
-		public override bool PrepTargetForDeployment(UEBuildDeployTarget InTarget)
+		public override bool PrepTargetForDeployment(TargetReceipt Receipt)
 		{
 			Log.TraceInformation("Deploying now!");
 			return true;
diff --git a/Engine/Source/Programs/UnrealBuildTool/Platform/TVOS/UEBuildTVOS.cs b/Engine/Source/Programs/UnrealBuildTool/Platform/TVOS/UEBuildTVOS.cs
--- a/Engine/Source/Programs/UnrealBuildTool/Platform/TVOS/UEBuildTVOS.cs
+++ b/Engine/Source/Programs/UnrealBuildTool/Platform/TVOS/UEBuildTVOS.cs
@@ -130,9 +130,9 @@ namespace UnrealBuildTool
 			return new TVOSToolChain(Target.ProjectFile, ProjectSettings);
 		}
 
-		public override void Deploy(UEBuildDeployTarget Target)
+		public override void Deploy(TargetReceipt Receipt)
 		{
-			new UEDeployTVOS().PrepTargetForDeployment(Target);
+			new UEDeployTVOS().PrepTargetForDeployment(Receipt);
 		}
 	}
 
diff --git a/Engine/Source/Programs/UnrealBuildTool/Platform/Windows/UEBuildWindows.cs b/Engine/Source/Programs/UnrealBuildTool/Platform/Windows/UEBuildWindows.cs
--- a/Engine/Source/Programs/UnrealBuildTool/Platform/Windows/UEBuildWindows.cs
+++ b/Engine/Source/Programs/UnrealBuildTool/Platform/Windows/UEBuildWindows.cs
@@ -1772,8 +1772,8 @@ namespace UnrealBuildTool
 		/// <summary>
 		/// Deploys the given target
 		/// </summary>
-		/// <param name="Target">Information about the target being deployed</param>
-		public override void Deploy(UEBuildDeployTarget Target)
+		/// <param name="Receipt">Receipt for the target being deployed</param>
+		public override void Deploy(TargetReceipt Receipt)
 		{
 		}
 	}
diff --git a/Engine/Source/Programs/UnrealBuildTool/UnrealBuildTool.cs b/Engine/Source/Programs/UnrealBuildTool/UnrealBuildTool.cs
--- a/Engine/Source/Programs/UnrealBuildTool/UnrealBuildTool.cs
+++ b/Engine/Source/Programs/UnrealBuildTool/UnrealBuildTool.cs
@@ -1376,11 +1376,11 @@ namespace UnrealBuildTool
 							{
 								foreach (UEBuildTarget Target in Targets)
 								{
-									if (Target.DeployTargetFile != null && Target.OnlyModules.Count == 0)
+									if (Target.bDeployAfterCompile && Target.OnlyModules.Count == 0)
 									{
 										Log.WriteLine(LogEventType.Console, "Deploying {0} {1} {2}...", Target.TargetName, Target.Platform, Target.Configuration);
-										UEBuildDeployTarget DeployTarget = new UEBuildDeployTarget(Target.DeployTargetFile);
-										UEBuildPlatform.GetBuildPlatform(Target.Platform).Deploy(DeployTarget);
+										TargetReceipt Receipt = TargetReceipt.Read(Target.ReceiptFileName);
+										UEBuildPlatform.GetBuildPlatform(Target.Platform).Deploy(Receipt);
 									}
 								}
 							}
diff --git a/Engine/Source/Programs/UnrealBuildTool/UnrealBuildTool.csproj b/Engine/Source/Programs/UnrealBuildTool/UnrealBuildTool.csproj
--- a/Engine/Source/Programs/UnrealBuildTool/UnrealBuildTool.csproj
+++ b/Engine/Source/Programs/UnrealBuildTool/UnrealBuildTool.csproj
@@ -148,7 +148,6 @@
     <Compile Include="Configuration\ModuleRules.cs" />
     <Compile Include="Configuration\PlatformTargetRules.cs" />
     <Compile Include="Configuration\TargetDescriptor.cs" />
-    <Compile Include="Configuration\UEBuildDeployTarget.cs" />
     <Compile Include="Configuration\UEBuildModuleExternal.cs" />
     <Compile Include="Configuration\UEBuildPlatformFactory.cs" />
     <Compile Include="Configuration\UEBuildPlatformSDK.cs" />
