UBT: Remove the unused CppCompileEnvironment.LocalShadowDirectory variable.

#rnx

diff --git a/Engine/Source/Programs/UnrealBuildTool/Configuration/UEBuildTarget.cs b/Engine/Source/Programs/UnrealBuildTool/Configuration/UEBuildTarget.cs
--- a/Engine/Source/Programs/UnrealBuildTool/Configuration/UEBuildTarget.cs
+++ b/Engine/Source/Programs/UnrealBuildTool/Configuration/UEBuildTarget.cs
@@ -3418,9 +3418,6 @@ namespace UnrealBuildTool
                 }
             }
 
-			// By default, shadow source files for this target in the root OutputDirectory
-			GlobalCompileEnvironment.LocalShadowDirectory = LinkIntermediateDirectory;
-
 			GlobalCompileEnvironment.Definitions.Add("UNICODE");
 			GlobalCompileEnvironment.Definitions.Add("_UNICODE");
 			GlobalCompileEnvironment.Definitions.Add("__UNREAL__");
diff --git a/Engine/Source/Programs/UnrealBuildTool/System/CppCompileEnvironment.cs b/Engine/Source/Programs/UnrealBuildTool/System/CppCompileEnvironment.cs
--- a/Engine/Source/Programs/UnrealBuildTool/System/CppCompileEnvironment.cs
+++ b/Engine/Source/Programs/UnrealBuildTool/System/CppCompileEnvironment.cs
@@ -90,11 +90,6 @@ namespace UnrealBuildTool
 		/// </summary>
 		public readonly string Architecture;
 
-		/// <summary>
-		/// The directory to shadow source files in for syncing to remote compile servers
-		/// </summary>
-		public DirectoryReference LocalShadowDirectory = null;
-
 		/// <summary>
 		/// The name of the header file which is precompiled.
 		/// </summary>
@@ -336,7 +331,6 @@ namespace UnrealBuildTool
 			Platform = Other.Platform;
 			Configuration = Other.Configuration;
 			Architecture = Other.Architecture;
-			LocalShadowDirectory = Other.LocalShadowDirectory;
 			PrecompiledHeaderIncludeFilename = Other.PrecompiledHeaderIncludeFilename;
 			PrecompiledHeaderAction = Other.PrecompiledHeaderAction;
 			bUseRTTI = Other.bUseRTTI;
