Fix HTML5 build failure due to attempt to deploy when -nolink is specified.

#rnx

diff --git a/Engine/Source/Programs/UnrealBuildTool/Configuration/UEBuildTarget.cs b/Engine/Source/Programs/UnrealBuildTool/Configuration/UEBuildTarget.cs
--- a/Engine/Source/Programs/UnrealBuildTool/Configuration/UEBuildTarget.cs
+++ b/Engine/Source/Programs/UnrealBuildTool/Configuration/UEBuildTarget.cs
@@ -866,7 +866,7 @@ namespace UnrealBuildTool
 			TargetType = Rules.Type;
 			bPrecompile = InRules.bPrecompile;
 			ForeignPlugin = InDesc.ForeignPlugin;
-			bDeployAfterCompile = InRules.bDeployAfterCompile;
+			bDeployAfterCompile = InRules.bDeployAfterCompile && !InRules.bDisableLinking;
 
 			// now that we have the platform, we can set the intermediate path to include the platform/architecture name
 			PlatformIntermediateFolder = Path.Combine("Intermediate", "Build", Platform.ToString(), UEBuildPlatform.GetBuildPlatform(Platform).GetFolderNameForArchitecture(Architecture));
