UBT: Do symbol stripping for IOS as part of the action graph, rather than in PostBuildSync. This ensures that it runs even if UBT makefiles are enabled.

#rnx

diff --git a/Engine/Source/Programs/UnrealBuildTool/Platform/IOS/IOSToolChain.cs b/Engine/Source/Programs/UnrealBuildTool/Platform/IOS/IOSToolChain.cs
--- a/Engine/Source/Programs/UnrealBuildTool/Platform/IOS/IOSToolChain.cs
+++ b/Engine/Source/Programs/UnrealBuildTool/Platform/IOS/IOSToolChain.cs
@@ -72,11 +72,13 @@ namespace UnrealBuildTool
 
 	class IOSToolChain : AppleToolChain
 	{
+		public readonly ReadOnlyTargetRules Target;
 		protected IOSProjectSettings ProjectSettings;
 
-		public IOSToolChain(FileReference InProjectFile, IOSProjectSettings InProjectSettings)
-			: this(CppPlatform.IOS, InProjectFile, InProjectSettings, () => new IOSToolChainSettings())
+		public IOSToolChain(ReadOnlyTargetRules Target, IOSProjectSettings InProjectSettings)
+			: this(CppPlatform.IOS, (Target == null)? null : Target.ProjectFile, InProjectSettings, () => new IOSToolChainSettings())
 		{
+			this.Target = Target;
 		}
 
 		protected IOSToolChain(CppPlatform TargetPlatform, FileReference InProjectFile, IOSProjectSettings InProjectSettings, Func<IOSToolChainSettings> InCreateSettings)
@@ -1324,6 +1326,24 @@ namespace UnrealBuildTool
                 }
             }
 
+			// strip the debug info from the executable if needed. creates a dummy output file for the action graph to track that it's out of date.
+			if (Target.IOSPlatform.bStripSymbols || (Target.Configuration == UnrealTargetConfiguration.Shipping))
+			{
+				FileItem StripCompleteFile = FileItem.GetItemByFileReference(FileReference.Combine(BinaryLinkEnvironment.IntermediateDirectory, Executable.Location.GetFileName() + ".stripped"));
+
+				Action StripAction = ActionGraph.Add(ActionType.CreateAppBundle);
+				StripAction.WorkingDirectory = GetMacDevSrcRoot();
+				StripAction.CommandPath = "sh";
+				StripAction.CommandArguments = String.Format("-c '\"{0}strip\" \"{1}\" && touch \"{2}\"'", Settings.Value.ToolchainDir, Executable.Location, StripCompleteFile);
+				StripAction.PrerequisiteItems.Add(Executable);
+				StripAction.PrerequisiteItems.AddRange(OutputFiles);
+				StripAction.ProducedItems.Add(StripCompleteFile);
+				StripAction.StatusDescription = String.Format("Stripping symbols from {0}", Executable.AbsolutePath);
+				StripAction.bCanExecuteRemotely = false;
+
+				OutputFiles.Add(StripCompleteFile);
+			}
+
 			if(ShouldCompileAssetCatalog())
 			{
 				// generate the asset catalog
@@ -1474,27 +1494,6 @@ namespace UnrealBuildTool
 			string RemoteShadowDirectoryMac = Path.GetDirectoryName(Target.OutputPath.FullName);
 			string FinalRemoteExecutablePath = String.Format("{0}/Payload/{1}.app/{1}", RemoteShadowDirectoryMac, AppName);
 
-			// strip the debug info from the executable if needed
-			if (Target.Rules.IOSPlatform.bStripSymbols || (Target.Configuration == UnrealTargetConfiguration.Shipping))
-			{
-				Process StripProcess = new Process();
-				StripProcess.StartInfo.WorkingDirectory = RemoteShadowDirectoryMac;
-				StripProcess.StartInfo.FileName = new IOSToolChainSettings().ToolchainDir + "strip";
-				StripProcess.StartInfo.Arguments = "\"" + Target.OutputPath + "\"";
-
-				ProcessOutput Output = new ProcessOutput();
-				StripProcess.OutputDataReceived += new DataReceivedEventHandler(Output.OutputReceivedDataEventHandler);
-				StripProcess.ErrorDataReceived += new DataReceivedEventHandler(Output.OutputReceivedDataEventHandler);
-
-				Output.OutputReceivedDataEventHandlerEncounteredError = false;
-				Output.OutputReceivedDataEventHandlerEncounteredErrorMessage = "";
-				Utils.RunLocalProcess(StripProcess);
-				if (Output.OutputReceivedDataEventHandlerEncounteredError)
-				{
-					throw new Exception(Output.OutputReceivedDataEventHandlerEncounteredErrorMessage);
-				}
-			}
-
             // ensure the plist, entitlements, and provision files are properly copied
             UEDeployIOS DeployHandler = (Target.Platform == UnrealTargetPlatform.IOS ? new UEDeployIOS() : new UEDeployTVOS());
             DeployHandler.PrepTargetForDeployment(Receipt, Target.Rules.IOSPlatform.bCreateStubIPA);
diff --git a/Engine/Source/Programs/UnrealBuildTool/Platform/IOS/UEBuildIOS.cs b/Engine/Source/Programs/UnrealBuildTool/Platform/IOS/UEBuildIOS.cs
--- a/Engine/Source/Programs/UnrealBuildTool/Platform/IOS/UEBuildIOS.cs
+++ b/Engine/Source/Programs/UnrealBuildTool/Platform/IOS/UEBuildIOS.cs
@@ -997,7 +997,7 @@ namespace UnrealBuildTool
 		public override UEToolChain CreateToolChain(CppPlatform CppPlatform, ReadOnlyTargetRules Target)
 		{
 			IOSProjectSettings ProjectSettings = ReadProjectSettings(Target.ProjectFile);
-			return new IOSToolChain(Target.ProjectFile, ProjectSettings);
+			return new IOSToolChain(Target, ProjectSettings);
 		}
 
 		/// <summary>
