UBT: Pass the UPL instance as a parameter to methods that need it, rather than storing it on the deployment context.

diff --git a/Engine/Source/Programs/UnrealBuildTool/Platform/IOS/UEDeployIOS.cs b/Engine/Source/Programs/UnrealBuildTool/Platform/IOS/UEDeployIOS.cs
--- a/Engine/Source/Programs/UnrealBuildTool/Platform/IOS/UEDeployIOS.cs
+++ b/Engine/Source/Programs/UnrealBuildTool/Platform/IOS/UEDeployIOS.cs
@@ -186,7 +186,7 @@ namespace UnrealBuildTool
 			return result;
 		}
 
-		public static bool GenerateIOSPList(FileReference ProjectFile, UnrealTargetConfiguration Config, string ProjectDirectory, bool bIsUE4Game, string GameName, string ProjectName, string InEngineDir, string AppDirectory, TargetReceipt Receipt, out bool bSupportsPortrait, out bool bSupportsLandscape, out bool bSkipIcons, UEDeployIOS InThis = null)
+		public static bool GenerateIOSPList(FileReference ProjectFile, UnrealTargetConfiguration Config, string ProjectDirectory, bool bIsUE4Game, string GameName, string ProjectName, string InEngineDir, string AppDirectory, TargetReceipt Receipt, UnrealPluginLanguage UPL, out bool bSupportsPortrait, out bool bSupportsLandscape, out bool bSkipIcons)
 		{
 			// generate the Info.plist for future use
 			string BuildDirectory = ProjectDirectory + "/Build/IOS";
@@ -280,18 +280,17 @@ namespace UnrealBuildTool
 			string RequiredCaps = "";
 
 			IOSProjectSettings ProjectSettings = ((IOSPlatform)UEBuildPlatform.GetBuildPlatform(UnrealTargetPlatform.IOS)).ReadProjectSettings(ProjectFile);
-			if (InThis != null)
+
+			List<string> Arches = ((Config == UnrealTargetConfiguration.Shipping) ? ProjectSettings.ShippingArchitectures : ProjectSettings.NonShippingArchitectures).ToList();
+			if (Arches.Count > 1)
 			{
-				List<string> Arches = ((Config == UnrealTargetConfiguration.Shipping) ? ProjectSettings.ShippingArchitectures : ProjectSettings.NonShippingArchitectures).ToList();
-				if (Arches.Count > 1)
-				{
-					RequiredCaps += "\t\t<string>armv7</string>\n";
-				}
-				else
-				{
-					RequiredCaps += "\t\t<string>" + Arches[0] + "</string>\n";
-				}
+				RequiredCaps += "\t\t<string>armv7</string>\n";
 			}
+			else
+			{
+				RequiredCaps += "\t\t<string>" + Arches[0] + "</string>\n";
+			}
+
 			ConfigHierarchy GameIni = ConfigCache.ReadHierarchy(ConfigHierarchyType.Game, DirRef, UnrealTargetPlatform.IOS);
 
 			Ini.GetBool("/Script/IOSRuntimeSettings.IOSRuntimeSettings", "bSupportsOpenGLES2", out bSupported);
@@ -705,7 +704,7 @@ namespace UnrealBuildTool
 				Directory.CreateDirectory(IntermediateDirectory);
 			}
 
-			if (InThis != null && InThis.UPL != null)
+			if (UPL != null)
 			{
 				// Allow UPL to modify the plist here
 				XDocument XDoc;
@@ -719,7 +718,7 @@ namespace UnrealBuildTool
 				}
 
 				XDoc.DocumentType.InternalSubset = "";
-				InThis.UPL.ProcessPluginNode("None", "iosPListUpdates", "", ref XDoc);
+				UPL.ProcessPluginNode("None", "iosPListUpdates", "", ref XDoc);
 				string result = XDoc.Declaration.ToString() + "\n" + XDoc.ToString().Replace("<!DOCTYPE plist PUBLIC \"-//Apple//DTD PLIST 1.0//EN\" \"http://www.apple.com/DTDs/PropertyList-1.0.dtd\"[]>", "<!DOCTYPE plist PUBLIC \"-//Apple//DTD PLIST 1.0//EN\" \"http://www.apple.com/DTDs/PropertyList-1.0.dtd\">");
 				File.WriteAllText(PListFile, result);
 
@@ -767,12 +766,12 @@ namespace UnrealBuildTool
 
 			string RelativeEnginePath = UnrealBuildTool.EngineDirectory.MakeRelativeTo(DirectoryReference.GetCurrentDirectory());
 
-			UPL = new UnrealPluginLanguage(ProjectFile, CollectPluginDataPaths(Receipt), ProjectArches, "", "", UnrealTargetPlatform.IOS);
+			UnrealPluginLanguage UPL = new UnrealPluginLanguage(ProjectFile, CollectPluginDataPaths(Receipt), ProjectArches, "", "", UnrealTargetPlatform.IOS);
 
 			// Passing in true for distribution is not ideal here but given the way that ios packaging happens and this call chain it seems unavoidable for now, maybe there is a way to correctly pass it in that I can't find?
 			UPL.Init(ProjectArches, true, RelativeEnginePath, BundlePath, ProjectDirectory, Config.ToString(), BuildVersion.ReadDefault());
 
-			return GenerateIOSPList(ProjectFile, Config, ProjectDirectory, bIsUE4Game, GameName, ProjectName, InEngineDir, AppDirectory, Receipt, out bSupportsPortrait, out bSupportsLandscape, out bSkipIcons, this);
+			return GenerateIOSPList(ProjectFile, Config, ProjectDirectory, bIsUE4Game, GameName, ProjectName, InEngineDir, AppDirectory, Receipt, UPL, out bSupportsPortrait, out bSupportsLandscape, out bSkipIcons);
 		}
 
 		protected virtual void CopyCloudResources(string InEngineDir, string AppDirectory)
diff --git a/Engine/Source/Programs/UnrealBuildTool/Platform/TVOS/UEDeployTVOS.cs b/Engine/Source/Programs/UnrealBuildTool/Platform/TVOS/UEDeployTVOS.cs
--- a/Engine/Source/Programs/UnrealBuildTool/Platform/TVOS/UEDeployTVOS.cs
+++ b/Engine/Source/Programs/UnrealBuildTool/Platform/TVOS/UEDeployTVOS.cs
@@ -24,7 +24,7 @@ namespace UnrealBuildTool
 			return "TVOS";
 		}
 
-		public static bool GenerateTVOSPList(string ProjectDirectory, bool bIsUE4Game, string GameName, string ProjectName, string InEngineDir, string AppDirectory, UEDeployTVOS InThis = null)
+		public static bool GenerateTVOSPList(string ProjectDirectory, bool bIsUE4Game, string GameName, string ProjectName, string InEngineDir, string AppDirectory, UnrealPluginLanguage UPL)
 		{
 			// @todo tvos: THIS!
 
@@ -296,7 +296,7 @@ namespace UnrealBuildTool
 				Directory.CreateDirectory(IntermediateDirectory);
 			}
 
-			if(InThis != null && InThis.UPL != null)
+			if(UPL != null)
 			{
 				// Allow UPL to modify the plist here
 				XDocument XDoc;
@@ -310,7 +310,7 @@ namespace UnrealBuildTool
 				}
 
 				XDoc.DocumentType.InternalSubset = "";
-				InThis.UPL.ProcessPluginNode("None", "iosPListUpdates", "", ref XDoc);
+				UPL.ProcessPluginNode("None", "iosPListUpdates", "", ref XDoc);
                 string result = XDoc.Declaration.ToString() + "\n" + XDoc.ToString().Replace("<!DOCTYPE plist PUBLIC \"-//Apple//DTD PLIST 1.0//EN\" \"http://www.apple.com/DTDs/PropertyList-1.0.dtd\"[]>", "<!DOCTYPE plist PUBLIC \"-//Apple//DTD PLIST 1.0//EN\" \"http://www.apple.com/DTDs/PropertyList-1.0.dtd\">");
                 File.WriteAllText(PListFile, result);
 			}
@@ -335,7 +335,7 @@ namespace UnrealBuildTool
 		{
 			bSupportsLandscape = bSupportsPortrait = true;
             bSkipIcons = true;
-			return GenerateTVOSPList(ProjectDirectory, bIsUE4Game, GameName, ProjectName, InEngineDir, AppDirectory, this);
+			return GenerateTVOSPList(ProjectDirectory, bIsUE4Game, GameName, ProjectName, InEngineDir, AppDirectory, null);
 		}
 
         protected override void CopyGraphicsResources(bool bSkipDefaultPNGs, bool bSkipIcons, string InEngineDir, string AppDirectory, string BuildDirectory, string IntermediateDir, bool bSupportsPortrait, bool bSupportsLandscape)
diff --git a/Engine/Source/Programs/UnrealBuildTool/ProjectFiles/Xcode/XcodeProject.cs b/Engine/Source/Programs/UnrealBuildTool/ProjectFiles/Xcode/XcodeProject.cs
--- a/Engine/Source/Programs/UnrealBuildTool/ProjectFiles/Xcode/XcodeProject.cs
+++ b/Engine/Source/Programs/UnrealBuildTool/ProjectFiles/Xcode/XcodeProject.cs
@@ -881,12 +881,12 @@ namespace UnrealBuildTool
                             bool bSupportPortrait, bSupportLandscape, bSkipIcons;
 							TargetReceipt Receipt;
 							TargetReceipt.TryRead(ReceiptFilename, out Receipt);
-							UEDeployIOS.GenerateIOSPList(ProjectFile, Config.BuildConfig, ProjectPath.FullName, bIsUE4Game, GameName, Config.BuildTarget, EngineDir.FullName, ProjectPath + "/Binaries/IOS/Payload", Receipt, out bSupportPortrait, out bSupportLandscape, out bSkipIcons);
+							UEDeployIOS.GenerateIOSPList(ProjectFile, Config.BuildConfig, ProjectPath.FullName, bIsUE4Game, GameName, Config.BuildTarget, EngineDir.FullName, ProjectPath + "/Binaries/IOS/Payload", Receipt, null, out bSupportPortrait, out bSupportLandscape, out bSkipIcons);
 						}
 						if (bCreateTVOSInfoPlist)
 						{
 							Directory.CreateDirectory(Path.GetDirectoryName(TVOSInfoPlistPath));
-							UEDeployTVOS.GenerateTVOSPList(ProjectPath.FullName, bIsUE4Game, GameName, Config.BuildTarget, EngineDir.FullName, ProjectPath + "/Binaries/TVOS/Payload");
+							UEDeployTVOS.GenerateTVOSPList(ProjectPath.FullName, bIsUE4Game, GameName, Config.BuildTarget, EngineDir.FullName, ProjectPath + "/Binaries/TVOS/Payload", null);
 						}
 					}
 				}
