UBT: Remove the SupportsIconCatalog function from the exported function class. It's not needed in UAT.

#rnx

diff --git a/Engine/Source/Programs/UnrealBuildTool/Platform/IOS/IOSExports.cs b/Engine/Source/Programs/UnrealBuildTool/Platform/IOS/IOSExports.cs
--- a/Engine/Source/Programs/UnrealBuildTool/Platform/IOS/IOSExports.cs
+++ b/Engine/Source/Programs/UnrealBuildTool/Platform/IOS/IOSExports.cs
@@ -176,30 +176,5 @@ namespace UnrealBuildTool
 				}
 			}
 		}
-
-        /// <summary>
-        /// 
-        /// </summary>
-        public static bool SupportsIconCatalog(DirectoryReference ProjectDirectory, TargetReceipt Receipt)
-        {
-            // get the receipt
-            if (Receipt != null)
-            {
-                IEnumerable<ReceiptProperty> Results = Receipt.AdditionalProperties.Where(x => x.Name == "SDK");
-
-                if (Results.Count() > 0)
-                {
-                    if (Single.Parse(Results.ElementAt(0).Value) >= 11.0f)
-                    {
-                        return true;
-                    }
-                    else
-                    {
-                        return false;
-                    }
-                }
-            }
-            return false;
-        }
     }
 }
diff --git a/Engine/Source/Programs/UnrealBuildTool/Platform/IOS/UEDeployIOS.cs b/Engine/Source/Programs/UnrealBuildTool/Platform/IOS/UEDeployIOS.cs
--- a/Engine/Source/Programs/UnrealBuildTool/Platform/IOS/UEDeployIOS.cs
+++ b/Engine/Source/Programs/UnrealBuildTool/Platform/IOS/UEDeployIOS.cs
@@ -456,7 +456,7 @@ namespace UnrealBuildTool
 			}
 			Text.AppendLine("\t</array>");
 
-			if (IOSExports.SupportsIconCatalog(new DirectoryReference(ProjectDirectory), Receipt))
+			if (SupportsIconCatalog(Receipt))
 			{
 				bSkipIcons = true;
 				Text.AppendLine("\t<key>CFBundleIcons</key>");
@@ -739,6 +739,32 @@ namespace UnrealBuildTool
 			return bSkipDefaultPNGs;
 		}
 
+
+        /// <summary>
+        /// 
+        /// </summary>
+        static bool SupportsIconCatalog(TargetReceipt Receipt)
+        {
+            // get the receipt
+            if (Receipt != null)
+            {
+                IEnumerable<ReceiptProperty> Results = Receipt.AdditionalProperties.Where(x => x.Name == "SDK");
+
+                if (Results.Count() > 0)
+                {
+                    if (Single.Parse(Results.ElementAt(0).Value) >= 11.0f)
+                    {
+                        return true;
+                    }
+                    else
+                    {
+                        return false;
+                    }
+                }
+            }
+            return false;
+        }
+
 		public virtual bool GeneratePList(FileReference ProjectFile, UnrealTargetConfiguration Config, string ProjectDirectory, bool bIsUE4Game, string GameName, string ProjectName, string InEngineDir, string AppDirectory, TargetReceipt Receipt, out bool bSupportsPortrait, out bool bSupportsLandscape, out bool bSkipIcons)
 		{
 			// remember name with -IOS-Shipping, etc
