UBT: Remove unnecessary variables.

#rnx

diff --git a/Engine/Source/Programs/UnrealBuildTool/Platform/IOS/IOSToolChain.cs b/Engine/Source/Programs/UnrealBuildTool/Platform/IOS/IOSToolChain.cs
--- a/Engine/Source/Programs/UnrealBuildTool/Platform/IOS/IOSToolChain.cs
+++ b/Engine/Source/Programs/UnrealBuildTool/Platform/IOS/IOSToolChain.cs
@@ -123,10 +123,6 @@ namespace UnrealBuildTool
 		/// </summary>
 		private const string IOSArchiver = "libtool";
 
-		public static List<FileReference> BuiltBinaries = new List<FileReference>();
-
-		private bool bUseMallocProfiler;
-
 		/// <summary>
 		/// Additional frameworks stored locally so we have access without LinkEnvironment
 		/// </summary>
@@ -189,7 +185,7 @@ namespace UnrealBuildTool
 					BuildProducts.Add(AssetFile, BuildProductType.RequiredResource);
 				}
 			}
-            if ((ProjectSettings.bGeneratedSYMFile == true || ProjectSettings.bGeneratedSYMBundle == true) && (ProjectSettings.bGenerateCrashReportSymbols || bUseMallocProfiler) && Binary.Type == UEBuildBinaryType.Executable)
+            if ((ProjectSettings.bGeneratedSYMFile == true || ProjectSettings.bGeneratedSYMBundle == true) && (ProjectSettings.bGenerateCrashReportSymbols || Target.bUseMallocProfiler) && Binary.Type == UEBuildBinaryType.Executable)
             {
                 FileReference DebugFile = FileReference.Combine(Binary.OutputFilePath.Directory, Binary.OutputFilePath.GetFileNameWithoutExtension() + ".udebugsymbols");
                 BuildProducts.Add(DebugFile, BuildProductType.SymbolFile);
@@ -206,12 +202,6 @@ namespace UnrealBuildTool
 			return OutputType == BuildProductType.Executable;
 		}
 
-		public override void SetUpGlobalEnvironment(ReadOnlyTargetRules Target)
-		{
-			base.SetUpGlobalEnvironment(Target);
-			bUseMallocProfiler = Target.bUseMallocProfiler;
-		}
-
 		string GetCompileArguments_Global(CppCompileEnvironment CompileEnvironment)
 		{
             string Result = "";
@@ -1320,7 +1310,7 @@ namespace UnrealBuildTool
 			if (ProjectSettings.bGeneratedSYMFile == true || ProjectSettings.bGeneratedSYMBundle == true || BinaryLinkEnvironment.bUsePDBFiles == true)
             {
                 OutputFiles.Add(GenerateDebugInfo(Executable, ActionGraph));
-                if (ProjectSettings.bGenerateCrashReportSymbols || bUseMallocProfiler)
+                if (ProjectSettings.bGenerateCrashReportSymbols || Target.bUseMallocProfiler)
                 {
                     OutputFiles.Add(GeneratePseudoPDB(Executable, ActionGraph));
                 }
