UBT: Extract IOS frameworks as part of the action graph. This removes the need to extract them for incremental builds and allows dependency checking to work when the zip file changes. Also remove the UEBuildPlatform.PreBuildSync() callback, which is no longer needed.

diff --git a/Engine/Source/Programs/UnrealBuildTool/Configuration/UEBuildFramework.cs b/Engine/Source/Programs/UnrealBuildTool/Configuration/UEBuildFramework.cs
--- a/Engine/Source/Programs/UnrealBuildTool/Configuration/UEBuildFramework.cs
+++ b/Engine/Source/Programs/UnrealBuildTool/Configuration/UEBuildFramework.cs
@@ -5,53 +5,64 @@ using System.Collections.Generic;
 using System.Linq;
 using System.Text;
 using System.Threading.Tasks;
+using Tools.DotNETCommon;
 
 namespace UnrealBuildTool
 {
 	/// <summary>
-	/// 
+	/// Represents a Mac/IOS framework
 	/// </summary>
 	class UEBuildFramework
 	{
 		/// <summary>
-		/// 
+		/// The name of this framework
 		/// </summary>
-		/// <param name="Name"></param>
-		/// <param name="ZipPath"></param>
-		/// <param name="CopyBundledAssets"></param>
-		public UEBuildFramework(string Name, string ZipPath = null, string CopyBundledAssets = null)
-		{
-			this.FrameworkName = Name;
-			this.FrameworkZipPath = ZipPath;
-			this.CopyBundledAssets = CopyBundledAssets;
-		}
+		public readonly string Name;
 
 		/// <summary>
-		/// 
+		/// Path to a zip file containing the framework. May be null.
 		/// </summary>
-		/// <param name="Framework"></param>
-		public UEBuildFramework(ModuleRules.Framework Framework)
-		{
-			FrameworkName = Framework.Name;
-			FrameworkZipPath = Framework.ZipPath;
-			CopyBundledAssets = Framework.CopyBundledAssets;
-		}
+		public readonly FileReference ZipFile;
 
-		internal UEBuildModule OwningModule = null;
+		/// <summary>
+		/// Path to the extracted framework directory.
+		/// </summary>
+		public readonly DirectoryReference OutputDirectory;
 
 		/// <summary>
 		/// 
 		/// </summary>
-		public string FrameworkName = null;
+		public readonly string CopyBundledAssets;
 
 		/// <summary>
-		/// 
+		/// File created after the framework has been extracted. Used to add dependencies into the action graph.
 		/// </summary>
-		public string FrameworkZipPath = null;
+		public FileItem ExtractedTokenFile;
 
 		/// <summary>
-		/// 
+		/// Constructor
 		/// </summary>
-		public string CopyBundledAssets = null;
+		/// <param name="Name">The framework name</param>
+		/// <param name="CopyBundledAssets"></param>
+		public UEBuildFramework(string Name, string CopyBundledAssets = null)
+		{
+			this.Name = Name;
+			this.CopyBundledAssets = CopyBundledAssets;
+		}
+
+		/// <summary>
+		/// Constructor
+		/// </summary>
+		/// <param name="Name">The framework name</param>
+		/// <param name="ZipFile">Path to the zip file for this framework</param>
+		/// <param name="OutputDirectory">Path for the extracted zip file</param>
+		/// <param name="CopyBundledAssets"></param>
+		public UEBuildFramework(string Name, FileReference ZipFile, DirectoryReference OutputDirectory, string CopyBundledAssets)
+		{
+			this.Name = Name;
+			this.ZipFile = ZipFile;
+			this.OutputDirectory = OutputDirectory;
+			this.CopyBundledAssets = CopyBundledAssets;
+		}
 	}
 }
diff --git a/Engine/Source/Programs/UnrealBuildTool/Configuration/UEBuildModule.cs b/Engine/Source/Programs/UnrealBuildTool/Configuration/UEBuildModule.cs
--- a/Engine/Source/Programs/UnrealBuildTool/Configuration/UEBuildModule.cs
+++ b/Engine/Source/Programs/UnrealBuildTool/Configuration/UEBuildModule.cs
@@ -173,7 +173,25 @@ namespace UnrealBuildTool
 			PublicAdditionalLibraries = HashSetFromOptionalEnumerableStringParameter(Rules.PublicAdditionalLibraries);
 			PublicFrameworks = HashSetFromOptionalEnumerableStringParameter(Rules.PublicFrameworks);
 			PublicWeakFrameworks = HashSetFromOptionalEnumerableStringParameter(Rules.PublicWeakFrameworks);
-			PublicAdditionalFrameworks = Rules.PublicAdditionalFrameworks == null ? new HashSet<UEBuildFramework>() : new HashSet<UEBuildFramework>(Rules.PublicAdditionalFrameworks.Select(x => new UEBuildFramework(x)));
+
+			PublicAdditionalFrameworks = new HashSet<UEBuildFramework>();
+			if(Rules.PublicAdditionalFrameworks != null)
+			{
+				foreach(ModuleRules.Framework FrameworkRules in Rules.PublicAdditionalFrameworks)
+				{
+					UEBuildFramework Framework;
+					if(String.IsNullOrEmpty(FrameworkRules.ZipPath))
+					{
+						Framework = new UEBuildFramework(FrameworkRules.Name, FrameworkRules.CopyBundledAssets);
+					}
+					else
+					{
+						Framework = new UEBuildFramework(FrameworkRules.Name, FileReference.Combine(ModuleDirectory, FrameworkRules.ZipPath), DirectoryReference.Combine(UnrealBuildTool.EngineDirectory, "Intermediate", "UnzippedFrameworks", Name, Path.GetFileNameWithoutExtension(FrameworkRules.ZipPath)), FrameworkRules.CopyBundledAssets);
+					}
+					PublicAdditionalFrameworks.Add(Framework);
+				}
+			}
+
 			PublicAdditionalShadowFiles = HashSetFromOptionalEnumerableStringParameter(Rules.PublicAdditionalShadowFiles);
 			PublicAdditionalBundleResources = Rules.AdditionalBundleResources == null ? new HashSet<UEBuildBundleResource>() : new HashSet<UEBuildBundleResource>(Rules.AdditionalBundleResources.Select(x => new UEBuildBundleResource(x)));
 			PublicDelayLoadDLLs = HashSetFromOptionalEnumerableStringParameter(Rules.PublicDelayLoadDLLs);
@@ -487,12 +505,6 @@ namespace UnrealBuildTool
 
 			// Add the additional frameworks so that the compiler can know about their #include paths
 			AdditionalFrameworks.AddRange(PublicAdditionalFrameworks);
-
-			// Remember the module so we can refer to it when needed
-			foreach (UEBuildFramework Framework in PublicAdditionalFrameworks)
-			{
-				Framework.OwningModule = this;
-			}
 		}
 
 		/// <summary>
@@ -679,11 +691,6 @@ namespace UnrealBuildTool
 				Frameworks.AddRange(PublicFrameworks);
 				WeakFrameworks.AddRange(PublicWeakFrameworks);
 				AdditionalBundleResources.AddRange(PublicAdditionalBundleResources);
-				// Remember the module so we can refer to it when needed
-				foreach (UEBuildFramework Framework in PublicAdditionalFrameworks)
-				{
-					Framework.OwningModule = this;
-				}
 				AdditionalFrameworks.AddRange(PublicAdditionalFrameworks);
 				AdditionalShadowFiles.AddRange(PublicAdditionalShadowFiles);
 				DelayLoadDLLs.AddRange(PublicDelayLoadDLLs);
diff --git a/Engine/Source/Programs/UnrealBuildTool/Configuration/UEBuildPlatform.cs b/Engine/Source/Programs/UnrealBuildTool/Configuration/UEBuildPlatform.cs
--- a/Engine/Source/Programs/UnrealBuildTool/Configuration/UEBuildPlatform.cs
+++ b/Engine/Source/Programs/UnrealBuildTool/Configuration/UEBuildPlatform.cs
@@ -295,10 +295,6 @@ namespace UnrealBuildTool
 			return false;
 		}
 
-		public virtual void PreBuildSync()
-		{
-		}
-
 		public virtual void PostBuildSync(UEBuildTarget Target)
 		{
 		}
diff --git a/Engine/Source/Programs/UnrealBuildTool/Platform/IOS/IOSToolChain.cs b/Engine/Source/Programs/UnrealBuildTool/Platform/IOS/IOSToolChain.cs
--- a/Engine/Source/Programs/UnrealBuildTool/Platform/IOS/IOSToolChain.cs
+++ b/Engine/Source/Programs/UnrealBuildTool/Platform/IOS/IOSToolChain.cs
@@ -332,9 +332,9 @@ namespace UnrealBuildTool
 			// Add additional frameworks so that their headers can be found
 			foreach (UEBuildFramework Framework in CompileEnvironment.AdditionalFrameworks)
 			{
-				if (Framework.OwningModule != null && Framework.FrameworkZipPath != null && Framework.FrameworkZipPath != "")
+				if(Framework.OutputDirectory != null)
 				{
-					Result += " -F\"" + GetExtractedFrameworkDir(Framework) + "\"";
+					Result += String.Format(" -F\"{0}\"", Framework.OutputDirectory);
 				}
 			}
 
@@ -409,26 +409,6 @@ namespace UnrealBuildTool
 			return Result;
 		}
 
-		static FileReference GetFrameworkZip(UEBuildFramework Framework)
-		{
-			if (Framework.OwningModule == null)
-			{
-				throw new BuildException("GetLocalFrameworkZipPath: No owning module for framework {0}", Framework.FrameworkName);
-			}
-
-			return FileReference.Combine(Framework.OwningModule.ModuleDirectory, Framework.FrameworkZipPath);
-		}
-
-		static DirectoryReference GetExtractedFrameworkDir(UEBuildFramework Framework)
-		{
-			if (Framework.OwningModule == null)
-			{
-				throw new BuildException("GetRemoteIntermediateFrameworkZipPath: No owning module for framework {0}", Framework.FrameworkName);
-			}
-
-			return DirectoryReference.Combine(UnrealBuildTool.EngineDirectory, "Intermediate", "UnzippedFrameworks", Framework.OwningModule.Name, Framework.FrameworkZipPath.Replace(".zip", ""));
-		}
-
 		static void CleanIntermediateDirectory(string Path)
 		{
 			string ResultsText;
@@ -548,13 +528,13 @@ namespace UnrealBuildTool
 			}
 			foreach (UEBuildFramework Framework in LinkEnvironment.AdditionalFrameworks)
 			{
-				if (Framework.OwningModule != null && Framework.FrameworkZipPath != null && Framework.FrameworkZipPath != "")
+				if(Framework.OutputDirectory != null)
 				{
 					// If this framework has a zip specified, we'll need to setup the path as well
-					Result += " -F\"" + GetExtractedFrameworkDir(Framework) + "\"";
+					Result += String.Format(" -F\"{0}\"", Framework.OutputDirectory);
 				}
 
-				Result += " -framework " + Framework.FrameworkName;
+				Result += " -framework " + Framework.Name;
 			}
 			foreach (string Framework in LinkEnvironment.WeakFrameworks)
 			{
@@ -722,6 +702,15 @@ namespace UnrealBuildTool
 				CompileAction.bCanExecuteRemotely = false;
 				CompileAction.bShouldOutputStatusDescription = true;
 
+				foreach(UEBuildFramework Framework in CompileEnvironment.AdditionalFrameworks)
+				{
+					if(Framework.ZipFile != null)
+					{
+						FileItem ExtractedTokenFile = ExtractFramework(Framework, ActionGraph);
+						CompileAction.PrerequisiteItems.Add(ExtractedTokenFile);
+					}
+				}
+
 				AddFrameworksToTrack(CompileEnvironment.AdditionalFrameworks, CompileAction);
 			}
 			return Result;
@@ -731,7 +720,7 @@ namespace UnrealBuildTool
 		{
 			foreach (UEBuildFramework Framework in Frameworks)
 			{
-				if (Framework.OwningModule == null || Framework.FrameworkZipPath == null || Framework.FrameworkZipPath == "")
+				if (Framework.ZipFile == null)
 				{
 					continue;	// Only care about frameworks that have a zip specified
 				}
@@ -790,6 +779,14 @@ namespace UnrealBuildTool
 			}
 
 			// Handle additional framework assets that might need to be shadowed
+			foreach(UEBuildFramework Framework in LinkEnvironment.AdditionalFrameworks)
+			{
+				if(Framework.ZipFile != null)
+				{
+					FileItem ExtractedTokenFile = ExtractFramework(Framework, ActionGraph);
+					LinkAction.PrerequisiteItems.Add(ExtractedTokenFile);
+				}
+			}
 			AddFrameworksToTrack(LinkEnvironment.AdditionalFrameworks, LinkAction);
 
 			// Add the output file as a production of the link action.
@@ -1124,32 +1121,39 @@ namespace UnrealBuildTool
 			}
 		}
 
-		public static void PreBuildSync()
+		FileItem ExtractFramework(UEBuildFramework Framework, ActionGraph ActionGraph)
 		{
-			// Unzip any third party frameworks that are stored as zips
-			foreach (UEBuildFramework Framework in RememberedAdditionalFrameworks)
+			if(Framework.ZipFile == null)
 			{
-				FileReference ZipSrcPath = GetFrameworkZip(Framework);
-				if(!FileReference.Exists(ZipSrcPath))
-				{
-					throw new BuildException("Unable to find framework '{0}'", ZipSrcPath);
-				}
-
-				DirectoryReference FrameworkDstPath = GetExtractedFrameworkDir(Framework);
-				Log.TraceInformation("Unzipping: {0} -> {1}", ZipSrcPath, FrameworkDstPath);
+				throw new BuildException("Unable to extract framework '{0}' - no zip file specified", Framework.Name);
+			}
+			if(Framework.ExtractedTokenFile == null)
+			{
+				FileItem InputFile = FileItem.GetItemByFileReference(Framework.ZipFile);
+				Framework.ExtractedTokenFile = FileItem.GetItemByFileReference(new FileReference(Framework.OutputDirectory.FullName + ".extracted"));
 
-				CleanIntermediateDirectory(FrameworkDstPath.FullName);
+				StringBuilder ExtractScript = new StringBuilder();
+				ExtractScript.AppendLine("#!/bin/sh");
+				ExtractScript.AppendLine("set -e");
+				// ExtractScript.AppendLine("set -x"); // For debugging
+				ExtractScript.AppendLine(String.Format("[ -d {0} ] && rm -rf {0}", Utils.EscapeShellArgument(Framework.OutputDirectory.FullName)));
+				ExtractScript.AppendLine(String.Format("unzip -q -o {0} -d {1}", Utils.EscapeShellArgument(Framework.ZipFile.FullName), Utils.EscapeShellArgument(Framework.OutputDirectory.ParentDirectory.FullName))); // Zip contains folder with the same name, hence ParentDirectory
+				ExtractScript.AppendLine(String.Format("touch {0}", Utils.EscapeShellArgument(Framework.ExtractedTokenFile.AbsolutePath)));
 
-				// Assume that there is another directory inside the zip with the same name as the zip
-				DirectoryReference ZipDstPath = FrameworkDstPath.ParentDirectory;
+				FileItem ExtractScriptFileItem = FileItem.CreateIntermediateTextFile(new FileReference(Framework.OutputDirectory.FullName + ".sh"), ExtractScript.ToString());
 
-				// If we're on the mac, just unzip using the shell
-				string ResultsText;
-				if(RunExecutableAndWait("unzip", String.Format("-o \"{0}\" -d \"{1}\"", ZipSrcPath, ZipDstPath), out ResultsText) != 0)
-				{
-					throw new BuildException("Unable to extract {0}:\n{1}", ZipSrcPath, ResultsText);
-				}
+				Action UnzipAction = ActionGraph.Add(ActionType.BuildProject);
+				UnzipAction.CommandPath = "/bin/sh";
+				UnzipAction.CommandArguments = Utils.EscapeShellArgument(ExtractScriptFileItem.AbsolutePath);
+				UnzipAction.WorkingDirectory = UnrealBuildTool.EngineDirectory.FullName;
+				UnzipAction.PrerequisiteItems.Add(InputFile);
+				UnzipAction.PrerequisiteItems.Add(ExtractScriptFileItem);
+				UnzipAction.ProducedItems.Add(Framework.ExtractedTokenFile);
+				UnzipAction.DeleteItems.Add(Framework.ExtractedTokenFile);
+				UnzipAction.StatusDescription = String.Format("Unzipping: {0} -> {1}", Framework.ZipFile, Framework.OutputDirectory);
+				UnzipAction.bCanExecuteRemotely = false;
 			}
+			return Framework.ExtractedTokenFile;
         }
 
         public static DirectoryReference GenerateAssetCatalog(FileReference ProjectFile, CppPlatform Platform, ref bool bUserImagesExist)
@@ -1684,22 +1688,20 @@ namespace UnrealBuildTool
 
 				foreach (UEBuildFramework Framework in RememberedAdditionalFrameworks)
 				{
-					if (Framework.OwningModule == null || Framework.CopyBundledAssets == null || Framework.CopyBundledAssets == "")
+					if (!String.IsNullOrEmpty(Framework.CopyBundledAssets) && Framework.OutputDirectory != null)
 					{
-						continue;		// Only care if we need to copy bundle assets
-					}
-
-					string UnpackedZipPath = GetExtractedFrameworkDir(Framework).FullName;
+						string UnpackedZipPath = Framework.OutputDirectory.FullName;
 
-					// For now, this is hard coded, but we need to loop over all modules, and copy bundled assets that need it
-					string LocalSource = UnpackedZipPath + "/" + Framework.CopyBundledAssets;
-					string BundleName = Framework.CopyBundledAssets.Substring(Framework.CopyBundledAssets.LastIndexOf('/') + 1);
-					string LocalDest = LocalFrameworkAssets + "/" + BundleName;
+						// For now, this is hard coded, but we need to loop over all modules, and copy bundled assets that need it
+						string LocalSource = UnpackedZipPath + "/" + Framework.CopyBundledAssets;
+						string BundleName = Framework.CopyBundledAssets.Substring(Framework.CopyBundledAssets.LastIndexOf('/') + 1);
+						string LocalDest = LocalFrameworkAssets + "/" + BundleName;
 
-					Log.TraceInformation("Copying bundled asset... LocalSource: {0}, LocalDest: {1}", LocalSource, LocalDest);
+						Log.TraceInformation("Copying bundled asset... LocalSource: {0}, LocalDest: {1}", LocalSource, LocalDest);
 
-					string ResultsText;
-                    RunExecutableAndWait("cp", String.Format("-R -L \"{0}\" \"{1}\"", LocalSource, LocalDest), out ResultsText);
+						string ResultsText;
+						RunExecutableAndWait("cp", String.Format("-R -L \"{0}\" \"{1}\"", LocalSource, LocalDest), out ResultsText);
+					}
                 }
             }
 		}
diff --git a/Engine/Source/Programs/UnrealBuildTool/Platform/IOS/UEBuildIOS.cs b/Engine/Source/Programs/UnrealBuildTool/Platform/IOS/UEBuildIOS.cs
--- a/Engine/Source/Programs/UnrealBuildTool/Platform/IOS/UEBuildIOS.cs
+++ b/Engine/Source/Programs/UnrealBuildTool/Platform/IOS/UEBuildIOS.cs
@@ -770,11 +770,6 @@ namespace UnrealBuildTool
 			return true;
 		}
 
-		public override void PreBuildSync()
-		{
-			IOSToolChain.PreBuildSync();
-		}
-
 		public override void PostBuildSync(UEBuildTarget Target)
 		{
 			IOSToolChain.PostBuildSync(Target);
diff --git a/Engine/Source/Programs/UnrealBuildTool/Platform/Mac/MacToolChain.cs b/Engine/Source/Programs/UnrealBuildTool/Platform/Mac/MacToolChain.cs
--- a/Engine/Source/Programs/UnrealBuildTool/Platform/Mac/MacToolChain.cs
+++ b/Engine/Source/Programs/UnrealBuildTool/Platform/Mac/MacToolChain.cs
@@ -737,9 +737,9 @@ namespace UnrealBuildTool
 			}
 			foreach (UEBuildFramework Framework in LinkEnvironment.AdditionalFrameworks)
 			{
-				if (!AllFrameworks.ContainsKey(Framework.FrameworkName))
+				if (!AllFrameworks.ContainsKey(Framework.Name))
 				{
-					AllFrameworks.Add(Framework.FrameworkName, false);
+					AllFrameworks.Add(Framework.Name, false);
 				}
 			}
 			foreach (string Framework in LinkEnvironment.WeakFrameworks)
diff --git a/Engine/Source/Programs/UnrealBuildTool/UnrealBuildTool.cs b/Engine/Source/Programs/UnrealBuildTool/UnrealBuildTool.cs
--- a/Engine/Source/Programs/UnrealBuildTool/UnrealBuildTool.cs
+++ b/Engine/Source/Programs/UnrealBuildTool/UnrealBuildTool.cs
@@ -1310,11 +1310,6 @@ namespace UnrealBuildTool
 								// Make sure any old DLL files from in-engine recompiles aren't lying around.  Must be called after the action graph is finalized.
 								ActionGraph.DeleteStaleHotReloadDLLs();
 
-								foreach (UEBuildTarget Target in Targets)
-								{
-									UEBuildPlatform.GetBuildPlatform(Target.Platform).PreBuildSync();
-								}
-
 								// Plan the actions to execute for the build.
 								Dictionary<UEBuildTarget, List<FileItem>> TargetToOutdatedPrerequisitesMap;
 								List<Action> ActionsToExecute = ActionGraph.GetActionsToExecute(BuildConfiguration, UBTMakefile.PrerequisiteActions, Targets, TargetToHeaders, bIsAssemblingBuild, bNeedsFullCPPIncludeRescan, out TargetToOutdatedPrerequisitesMap);
