UBT: Remove unused RecursivelyAddPrecompiledModules method.

#rnx

diff --git a/Engine/Source/Programs/UnrealBuildTool/Configuration/UEBuildModule.cs b/Engine/Source/Programs/UnrealBuildTool/Configuration/UEBuildModule.cs
--- a/Engine/Source/Programs/UnrealBuildTool/Configuration/UEBuildModule.cs
+++ b/Engine/Source/Programs/UnrealBuildTool/Configuration/UEBuildModule.cs
@@ -762,14 +762,6 @@ namespace UnrealBuildTool
 		{
 		}
 
-		/// <summary>
-		/// Gets all of the modules precompiled along with this module
-		/// </summary>
-		/// <param name="Modules">Set of all the precompiled modules</param>
-		public virtual void RecursivelyAddPrecompiledModules(List<UEBuildModule> Modules)
-		{
-		}
-
 		public delegate UEBuildModule CreateModuleDelegate(string Name, string ReferenceChain);
 
 		/// <summary>
diff --git a/Engine/Source/Programs/UnrealBuildTool/Configuration/UEBuildModuleCPP.cs b/Engine/Source/Programs/UnrealBuildTool/Configuration/UEBuildModuleCPP.cs
--- a/Engine/Source/Programs/UnrealBuildTool/Configuration/UEBuildModuleCPP.cs
+++ b/Engine/Source/Programs/UnrealBuildTool/Configuration/UEBuildModuleCPP.cs
@@ -1183,34 +1183,5 @@ namespace UnrealBuildTool
 				}
 			}
 		}
-
-		public override void RecursivelyAddPrecompiledModules(List<UEBuildModule> Modules)
-		{
-			if (!Modules.Contains(this))
-			{
-				Modules.Add(this);
-
-				// Get the dependent modules
-				List<UEBuildModule> DependentModules = new List<UEBuildModule>();
-				if (PrivateDependencyModules != null)
-				{
-					DependentModules.AddRange(PrivateDependencyModules);
-				}
-				if (PublicDependencyModules != null)
-				{
-					DependentModules.AddRange(PublicDependencyModules);
-				}
-				if (DynamicallyLoadedModules != null)
-				{
-					DependentModules.AddRange(DynamicallyLoadedModules);
-				}
-
-				// Find modules for each of them, and add their dependencies too
-				foreach (UEBuildModule DependentModule in DependentModules)
-				{
-					DependentModule.RecursivelyAddPrecompiledModules(Modules);
-				}
-			}
-		}
 	}
 }
