UBT: Allow adaptive non-unity builds to use the existing PCH if the compile settings match.

diff --git a/Engine/Source/Programs/UnrealBuildTool/Configuration/UEBuildModuleCPP.cs b/Engine/Source/Programs/UnrealBuildTool/Configuration/UEBuildModuleCPP.cs
--- a/Engine/Source/Programs/UnrealBuildTool/Configuration/UEBuildModuleCPP.cs
+++ b/Engine/Source/Programs/UnrealBuildTool/Configuration/UEBuildModuleCPP.cs
@@ -773,9 +773,13 @@ namespace UnrealBuildTool
 				{
 					AdaptiveOutput = CompileAdaptiveNonUnityFilesWithoutPCH(ToolChain, AdaptiveUnityEnvironment, AdaptiveFiles, IntermediateDirectory, Name, ActionGraph);
 				}
+				else if(AdaptiveUnityEnvironment.bOptimizeCode != CompileEnvironment.bOptimizeCode || AdaptiveUnityEnvironment.bSupportEditAndContinue != CompileEnvironment.bSupportEditAndContinue)
+				{
+					AdaptiveOutput = CompileAdaptiveNonUnityFiles(ToolChain, AdaptiveUnityEnvironment, AdaptiveFiles, IntermediateDirectory, Name, ActionGraph);
+				}
 				else
 				{
-					AdaptiveOutput = CompileAdaptiveNonUnityFilesWithPCH(ToolChain, AdaptiveUnityEnvironment, AdaptiveFiles, IntermediateDirectory, Name, ActionGraph);
+					AdaptiveOutput = CompileAdaptiveNonUnityFiles(ToolChain, CompileEnvironment, AdaptiveFiles, IntermediateDirectory, Name, ActionGraph);
 				}
 
 				// Merge output
@@ -786,7 +790,7 @@ namespace UnrealBuildTool
 			return OutputFiles;
 		}
 
-		static CPPOutput CompileAdaptiveNonUnityFilesWithPCH(UEToolChain ToolChain, CppCompileEnvironment CompileEnvironment, List<FileItem> Files, DirectoryReference IntermediateDirectory, string ModuleName, ActionGraph ActionGraph)
+		static CPPOutput CompileAdaptiveNonUnityFiles(UEToolChain ToolChain, CppCompileEnvironment CompileEnvironment, List<FileItem> Files, DirectoryReference IntermediateDirectory, string ModuleName, ActionGraph ActionGraph)
 		{
 			// Compile the files
 			return ToolChain.CompileCPPFiles(CompileEnvironment, Files, IntermediateDirectory, ModuleName, ActionGraph);
