UBT: Deprecate bCompileLeanAndMeanUE. This setting is vague and non-descriptive, and causes several other settings to be overriden. It is forced on by the current target type and individual target platforms, preventing targets from modifying those values. Anything explicitly using bCompileLeanAndMeanUE now uses bBuildDeveloperTools instead, and the UE_BUILD_MINIMAL macro at runtime has been removed.

The values that it previously forced values to are now configured as follows:

* bBuildEditor is a read-only property that returns true if and only if the target type is Editor.
* bCompileSimplygon and bCompileSimplygonSSF now a property which will default to true for Windows editor, if Simplygon is available. It can also be explicitly enabled or disabled.
* bCompileSpeedTree is now a property which will default to true for editor targets, and can also be explicitly enabled or disabled.

diff --git a/Engine/Source/Developer/TargetPlatform/TargetPlatform.Build.cs b/Engine/Source/Developer/TargetPlatform/TargetPlatform.Build.cs
--- a/Engine/Source/Developer/TargetPlatform/TargetPlatform.Build.cs
+++ b/Engine/Source/Developer/TargetPlatform/TargetPlatform.Build.cs
@@ -24,10 +24,7 @@ public class TargetPlatform : ModuleRules
 
 			if ((Target.Platform == UnrealTargetPlatform.Win64) || (Target.Platform == UnrealTargetPlatform.Win32))
 			{
-				if (Target.bCompileLeanAndMeanUE == false)
-				{
-					DynamicallyLoadedModuleNames.Add("TextureFormatIntelISPCTexComp");
-				}
+				DynamicallyLoadedModuleNames.Add("TextureFormatIntelISPCTexComp");
 			}
 
 			if (Target.Platform == UnrealTargetPlatform.Win32 ||
@@ -38,12 +35,9 @@ public class TargetPlatform : ModuleRules
 				DynamicallyLoadedModuleNames.Add("ShaderFormatD3D");
 				DynamicallyLoadedModuleNames.Add("MetalShaderFormat");
 
-				if (Target.bCompileLeanAndMeanUE == false)
-				{
-					DynamicallyLoadedModuleNames.Add("TextureFormatDXT");
-					DynamicallyLoadedModuleNames.Add("TextureFormatPVR");
-					DynamicallyLoadedModuleNames.Add("TextureFormatASTC");
-				}
+				DynamicallyLoadedModuleNames.Add("TextureFormatDXT");
+				DynamicallyLoadedModuleNames.Add("TextureFormatPVR");
+				DynamicallyLoadedModuleNames.Add("TextureFormatASTC");
 
 				DynamicallyLoadedModuleNames.Add("TextureFormatUncompressed");
 
@@ -68,12 +62,9 @@ public class TargetPlatform : ModuleRules
 			}
 			else if (Target.Platform == UnrealTargetPlatform.Mac)
 			{
-				if (Target.bCompileLeanAndMeanUE == false)
-				{
-					DynamicallyLoadedModuleNames.Add("TextureFormatDXT");
-					DynamicallyLoadedModuleNames.Add("TextureFormatPVR");
-					DynamicallyLoadedModuleNames.Add("TextureFormatASTC");
-				}
+				DynamicallyLoadedModuleNames.Add("TextureFormatDXT");
+				DynamicallyLoadedModuleNames.Add("TextureFormatPVR");
+				DynamicallyLoadedModuleNames.Add("TextureFormatASTC");
 
 				DynamicallyLoadedModuleNames.Add("TextureFormatUncompressed");
 
@@ -94,12 +85,9 @@ public class TargetPlatform : ModuleRules
 			}
 			else if (Target.Platform == UnrealTargetPlatform.Linux)
 			{
-				if (Target.bCompileLeanAndMeanUE == false)
-				{
-					DynamicallyLoadedModuleNames.Add("TextureFormatDXT");
-					DynamicallyLoadedModuleNames.Add("TextureFormatPVR");
-					DynamicallyLoadedModuleNames.Add("TextureFormatASTC");
-				}
+				DynamicallyLoadedModuleNames.Add("TextureFormatDXT");
+				DynamicallyLoadedModuleNames.Add("TextureFormatPVR");
+				DynamicallyLoadedModuleNames.Add("TextureFormatASTC");
 
 				DynamicallyLoadedModuleNames.Add("TextureFormatUncompressed");
 
diff --git a/Engine/Source/Developer/TextureCompressor/TextureCompressor.Build.cs b/Engine/Source/Developer/TextureCompressor/TextureCompressor.Build.cs
--- a/Engine/Source/Developer/TextureCompressor/TextureCompressor.Build.cs
+++ b/Engine/Source/Developer/TextureCompressor/TextureCompressor.Build.cs
@@ -16,7 +16,7 @@ public class TextureCompressor : ModuleRules
 			}
 			);
 
-		if (Target.bCompileLeanAndMeanUE == false)
+		if (Target.bBuildDeveloperTools)
 		{
 			AddEngineThirdPartyPrivateStaticDependencies(Target, "nvTextureTools");
 		}
diff --git a/Engine/Source/Editor/GameProjectGeneration/Private/GameProjectUtils.cpp b/Engine/Source/Editor/GameProjectGeneration/Private/GameProjectUtils.cpp
--- a/Engine/Source/Editor/GameProjectGeneration/Private/GameProjectUtils.cpp
+++ b/Engine/Source/Editor/GameProjectGeneration/Private/GameProjectUtils.cpp
@@ -3570,7 +3570,7 @@ bool GameProjectUtils::HasDefaultBuildSettings(const FName InPlatformInfoName)
 	// first check default build settings for all platforms
 	TArray<FString> BoolKeys, IntKeys, StringKeys, BuildKeys;
 	BuildKeys.Add(TEXT("bCompileApex")); BuildKeys.Add(TEXT("bCompileICU"));
-	BuildKeys.Add(TEXT("bCompileSimplygon")); BuildKeys.Add(TEXT("bCompileSimplygonSSF")); BuildKeys.Add(TEXT("bCompileLeanAndMeanUE"));
+	BuildKeys.Add(TEXT("bCompileSimplygon")); BuildKeys.Add(TEXT("bCompileSimplygonSSF"));
 	BuildKeys.Add(TEXT("bIncludeADO"));	BuildKeys.Add(TEXT("bCompileRecast")); BuildKeys.Add(TEXT("bCompileSpeedTree"));
 	BuildKeys.Add(TEXT("bCompileWithPluginSupport")); BuildKeys.Add(TEXT("bCompilePhysXVehicle")); BuildKeys.Add(TEXT("bCompileFreeType"));
 	BuildKeys.Add(TEXT("bCompileForSize"));	BuildKeys.Add(TEXT("bCompileCEF3"));
diff --git a/Engine/Source/Editor/LandscapeEditor/LandscapeEditor.Build.cs b/Engine/Source/Editor/LandscapeEditor/LandscapeEditor.Build.cs
--- a/Engine/Source/Editor/LandscapeEditor/LandscapeEditor.Build.cs
+++ b/Engine/Source/Editor/LandscapeEditor/LandscapeEditor.Build.cs
@@ -66,7 +66,7 @@ public class LandscapeEditor : ModuleRules
 		}
 
 		// KissFFT is used by the smooth tool.
-		if (Target.bCompileLeanAndMeanUE == false &&
+		if (Target.bBuildDeveloperTools &&
 			(Target.Platform == UnrealTargetPlatform.Win64 || Target.Platform == UnrealTargetPlatform.Win32 || Target.Platform == UnrealTargetPlatform.Mac || Target.Platform == UnrealTargetPlatform.Linux))
 		{
 			AddEngineThirdPartyPrivateStaticDependencies(Target, "Kiss_FFT");
diff --git a/Engine/Source/Programs/BlankProgram/BlankProgram.Target.cs b/Engine/Source/Programs/BlankProgram/BlankProgram.Target.cs
--- a/Engine/Source/Programs/BlankProgram/BlankProgram.Target.cs
+++ b/Engine/Source/Programs/BlankProgram/BlankProgram.Target.cs
@@ -13,14 +13,12 @@ public class BlankProgramTarget : TargetRules
 		LaunchModuleName = "BlankProgram";
 
 		// Lean and mean
-		bCompileLeanAndMeanUE = true;
+		bBuildDeveloperTools = false;
 
 		// Never use malloc profiling in Unreal Header Tool.  We set this because often UHT is compiled right before the engine
 		// automatically by Unreal Build Tool, but if bUseMallocProfiler is defined, UHT can operate incorrectly.
 		bUseMallocProfiler = false;
 
-		// No editor needed
-		bBuildEditor = false;
 		// Editor-only data, however, is needed
 		bBuildWithEditorOnlyData = true;
 
diff --git a/Engine/Source/Programs/BuildPatchTool/BuildPatchTool.Target.cs b/Engine/Source/Programs/BuildPatchTool/BuildPatchTool.Target.cs
--- a/Engine/Source/Programs/BuildPatchTool/BuildPatchTool.Target.cs
+++ b/Engine/Source/Programs/BuildPatchTool/BuildPatchTool.Target.cs
@@ -15,7 +15,6 @@ public class BuildPatchToolTarget : TargetRules
 		bOutputPubliclyDistributable = true;
 		UndecoratedConfiguration = UnrealTargetConfiguration.Shipping;
 
-		bBuildEditor = false;
 		bCompileAgainstEngine = false;
 		bCompileAgainstCoreUObject = false;
 		if(Target.Configuration == UnrealTargetConfiguration.Shipping)
@@ -23,7 +22,7 @@ public class BuildPatchToolTarget : TargetRules
 			// AutomationController is referenced by BuildPatchTool.build.cs, and references a ton of editor things that it probably shouldn't. Need ApplicationCore for this.
 			bCompileAgainstApplicationCore = false;
 		}
-		bCompileLeanAndMeanUE = true;
+		bBuildDeveloperTools = false;
 		bUseLoggingInShipping = true;
 		bUseChecksInShipping = true;
 		bIsBuildingConsoleApplication = true;
diff --git a/Engine/Source/Programs/CrashReportClient/CrashReportClient.Target.cs b/Engine/Source/Programs/CrashReportClient/CrashReportClient.Target.cs
--- a/Engine/Source/Programs/CrashReportClient/CrashReportClient.Target.cs
+++ b/Engine/Source/Programs/CrashReportClient/CrashReportClient.Target.cs
@@ -22,10 +22,7 @@ public class CrashReportClientTarget : TargetRules
 
         bOutputPubliclyDistributable = true;
 
-		bCompileLeanAndMeanUE = true;
-
-		// Don't need editor
-		bBuildEditor = false;
+		bBuildDeveloperTools = false;
 
 		// CrashReportClient doesn't ever compile with the engine linked in
 		bCompileAgainstEngine = false;
diff --git a/Engine/Source/Programs/IOS/UnrealLaunchDaemon/UnrealLaunchDaemon.Target.cs b/Engine/Source/Programs/IOS/UnrealLaunchDaemon/UnrealLaunchDaemon.Target.cs
--- a/Engine/Source/Programs/IOS/UnrealLaunchDaemon/UnrealLaunchDaemon.Target.cs
+++ b/Engine/Source/Programs/IOS/UnrealLaunchDaemon/UnrealLaunchDaemon.Target.cs
@@ -17,8 +17,7 @@ public class UnrealLaunchDaemonTarget : TargetRules
 
 		LaunchModuleName = "UnrealLaunchDaemon";
 
-		bCompileLeanAndMeanUE = true;
-		bBuildEditor = false;
+		bBuildDeveloperTools = false;
 		bCompileAgainstEngine = false;
 
 		bHasExports = false;
diff --git a/Engine/Source/Programs/Mac/DsymExporter/DsymExporter.Target.cs b/Engine/Source/Programs/Mac/DsymExporter/DsymExporter.Target.cs
--- a/Engine/Source/Programs/Mac/DsymExporter/DsymExporter.Target.cs
+++ b/Engine/Source/Programs/Mac/DsymExporter/DsymExporter.Target.cs
@@ -12,10 +12,7 @@ public class DsymExporterTarget : TargetRules
 		LinkType = TargetLinkType.Monolithic;
 		LaunchModuleName = "DsymExporter";
 
-		bCompileLeanAndMeanUE = true;
-
-		// Don't need editor
-		bBuildEditor = false;
+		bBuildDeveloperTools = false;
 
 		// DsymExporter doesn't ever compile with the engine linked in
 		bCompileAgainstEngine = false;
diff --git a/Engine/Source/Programs/Mac/UE4EditorServices/UE4EditorServices.Target.cs b/Engine/Source/Programs/Mac/UE4EditorServices/UE4EditorServices.Target.cs
--- a/Engine/Source/Programs/Mac/UE4EditorServices/UE4EditorServices.Target.cs
+++ b/Engine/Source/Programs/Mac/UE4EditorServices/UE4EditorServices.Target.cs
@@ -12,12 +12,10 @@ public class UE4EditorServicesTarget : TargetRules
 		LinkType = TargetLinkType.Monolithic;
 		LaunchModuleName = "UE4EditorServices";
 
-		bCompileLeanAndMeanUE = true;
+		bBuildDeveloperTools = false;
 		bUseMallocProfiler = false;
 
-		// No editor needed
 		bCompileICU = false;
-		bBuildEditor = false;
 		// Editor-only data, however, is needed
 		bBuildWithEditorOnlyData = false;
 
diff --git a/Engine/Source/Programs/Mac/UnrealAtoS/UnrealAtoS.Target.cs b/Engine/Source/Programs/Mac/UnrealAtoS/UnrealAtoS.Target.cs
--- a/Engine/Source/Programs/Mac/UnrealAtoS/UnrealAtoS.Target.cs
+++ b/Engine/Source/Programs/Mac/UnrealAtoS/UnrealAtoS.Target.cs
@@ -12,10 +12,7 @@ public class UnrealAtoSTarget : TargetRules
 		LinkType = TargetLinkType.Monolithic;
 		LaunchModuleName = "UnrealAtoS";
 
-		bCompileLeanAndMeanUE = true;
-
-		// Don't need editor
-		bBuildEditor = false;
+		bBuildDeveloperTools = false;
 
 		// DsymExporter doesn't ever compile with the engine linked in
 		bCompileAgainstEngine = false;
diff --git a/Engine/Source/Programs/MayaLiveLinkPlugin/MayaLiveLinkPlugin2016.Target.cs b/Engine/Source/Programs/MayaLiveLinkPlugin/MayaLiveLinkPlugin2016.Target.cs
--- a/Engine/Source/Programs/MayaLiveLinkPlugin/MayaLiveLinkPlugin2016.Target.cs
+++ b/Engine/Source/Programs/MayaLiveLinkPlugin/MayaLiveLinkPlugin2016.Target.cs
@@ -1,4 +1,4 @@
-﻿// Copyright 1998-2018 Epic Games, Inc. All Rights Reserved.
+// Copyright 1998-2018 Epic Games, Inc. All Rights Reserved.
 
 using UnrealBuildTool;
 using System.Collections.Generic;
@@ -13,9 +13,8 @@ public abstract class MayaLiveLinkPluginTargetBase : TargetRules
 		LinkType = TargetLinkType.Monolithic;
 
 		// We only need minimal use of the engine for this plugin
-		bCompileLeanAndMeanUE = true;
+		bBuildDeveloperTools = false;
 		bUseMallocProfiler = false;
-		bBuildEditor = false;
 		bBuildWithEditorOnlyData = true;
 		bCompileAgainstEngine = false;
 		bCompileAgainstCoreUObject = true;
diff --git a/Engine/Source/Programs/MinidumpDiagnostics/MinidumpDiagnostics.Target.cs b/Engine/Source/Programs/MinidumpDiagnostics/MinidumpDiagnostics.Target.cs
--- a/Engine/Source/Programs/MinidumpDiagnostics/MinidumpDiagnostics.Target.cs
+++ b/Engine/Source/Programs/MinidumpDiagnostics/MinidumpDiagnostics.Target.cs
@@ -12,12 +12,9 @@ public class MinidumpDiagnosticsTarget : TargetRules
 		LinkType = TargetLinkType.Monolithic;
 		LaunchModuleName = "MinidumpDiagnostics";
 
-		bCompileLeanAndMeanUE = true;
+		bBuildDeveloperTools = false;
 		bCompileICU = false;
 
-		// Don't need editor
-		bBuildEditor = false;
-
 		// MinidumpDiagnostics doesn't ever compile with the engine linked in
 		bCompileAgainstEngine = false;
 
diff --git a/Engine/Source/Programs/ShaderCompileWorker/ShaderCompileWorker.Target.cs b/Engine/Source/Programs/ShaderCompileWorker/ShaderCompileWorker.Target.cs
--- a/Engine/Source/Programs/ShaderCompileWorker/ShaderCompileWorker.Target.cs
+++ b/Engine/Source/Programs/ShaderCompileWorker/ShaderCompileWorker.Target.cs
@@ -30,7 +30,7 @@ public class ShaderCompileWorkerTarget : TargetRules
 		// Turn off various third party features we don't need
 
 		// Currently we force Lean and Mean mode
-		bCompileLeanAndMeanUE = true;
+		bBuildDeveloperTools = false;
 
 		// ShaderCompileWorker isn't localized, so doesn't need ICU
 		bCompileICU = false;
diff --git a/Engine/Source/Programs/SlateViewer/SlateViewer.Target.cs b/Engine/Source/Programs/SlateViewer/SlateViewer.Target.cs
--- a/Engine/Source/Programs/SlateViewer/SlateViewer.Target.cs
+++ b/Engine/Source/Programs/SlateViewer/SlateViewer.Target.cs
@@ -15,10 +15,7 @@ public class SlateViewerTarget : TargetRules
 		LaunchModuleName = "SlateViewer";
 		ExtraModuleNames.Add("EditorStyle");
 
-		bCompileLeanAndMeanUE = true;
-
-		// Don't need editor
-		bBuildEditor = false;
+		bBuildDeveloperTools = false;
 
 		// SlateViewer doesn't ever compile with the engine linked in
 		bCompileAgainstEngine = false;
diff --git a/Engine/Source/Programs/SymbolDebugger/SymbolDebugger.Target.cs b/Engine/Source/Programs/SymbolDebugger/SymbolDebugger.Target.cs
--- a/Engine/Source/Programs/SymbolDebugger/SymbolDebugger.Target.cs
+++ b/Engine/Source/Programs/SymbolDebugger/SymbolDebugger.Target.cs
@@ -16,10 +16,7 @@ public class SymbolDebuggerTarget : TargetRules
 
 		GlobalDefinitions.Add("WITH_DATABASE_SUPPORT=1");
 
-		bCompileLeanAndMeanUE = true;
-
-		// Don't need editor
-		bBuildEditor = false;
+		bBuildDeveloperTools = false;
 
 		// SymbolDebugger doesn't ever compile with the engine linked in
 		bCompileAgainstEngine = false;
diff --git a/Engine/Source/Programs/TestPAL/TestPAL.Target.cs b/Engine/Source/Programs/TestPAL/TestPAL.Target.cs
--- a/Engine/Source/Programs/TestPAL/TestPAL.Target.cs
+++ b/Engine/Source/Programs/TestPAL/TestPAL.Target.cs
@@ -1,4 +1,4 @@
-﻿// Copyright 1998-2018 Epic Games, Inc. All Rights Reserved.
+// Copyright 1998-2018 Epic Games, Inc. All Rights Reserved.
 
 using UnrealBuildTool;
 using System.Collections.Generic;
@@ -13,11 +13,7 @@ public class TestPALTarget : TargetRules
 		LaunchModuleName = "TestPAL";
 
 		// Lean and mean
-		bCompileLeanAndMeanUE = true;
-
-		// No editor or editor-only data is needed
-		bBuildEditor = false;
-		//bBuildWithEditorOnlyData = false;
+		bBuildDeveloperTools = false;
 
 		// Compile out references from Core to the rest of the engine
 		bCompileAgainstEngine = false;
diff --git a/Engine/Source/Programs/UnrealBuildTool/Configuration/TargetRules.cs b/Engine/Source/Programs/UnrealBuildTool/Configuration/TargetRules.cs
--- a/Engine/Source/Programs/UnrealBuildTool/Configuration/TargetRules.cs
+++ b/Engine/Source/Programs/UnrealBuildTool/Configuration/TargetRules.cs
@@ -6,6 +6,7 @@ using System.Collections.Generic;
 using System.IO;
 using System.Linq;
 using System.Text;
+using System.Threading;
 using System.Threading.Tasks;
 using Tools.DotNETCommon;
 
@@ -93,7 +94,7 @@ namespace UnrealBuildTool
 	/// <summary>
 	/// Attribute used to mark fields which much match between targets in the shared build environment
 	/// </summary>
-	[AttributeUsage(AttributeTargets.Field, AllowMultiple = false)]
+	[AttributeUsage(AttributeTargets.Field | AttributeTargets.Property, AllowMultiple = false)]
 	class RequiresUniqueBuildEnvironmentAttribute : Attribute
 	{
 	}
@@ -347,10 +348,12 @@ namespace UnrealBuildTool
 		/// <summary>
 		/// Whether to compile the editor or not. Only desktop platforms (Windows or Mac) will use this, other platforms force this to false.
 		/// </summary>
-		[RequiresUniqueBuildEnvironment]
-		[CommandLine("-NoEditor", Value = "false")]
-		public bool bBuildEditor = true;
-
+		public bool bBuildEditor
+		{
+			get { return (Type == TargetType.Editor); }
+			set { Log.TraceWarning("Setting {0}.bBuildEditor is deprecated. Set {0}.Type instead.", GetType().Name); }
+		}
+			
 		/// <summary>
 		/// Whether to compile code related to building assets. Consoles generally cannot build assets. Desktop platforms generally can.
 		/// </summary>
@@ -364,11 +367,20 @@ namespace UnrealBuildTool
 		[CommandLine("-NoEditorOnlyData", Value = "false")]
 		public bool bBuildWithEditorOnlyData = true;
 
+		/// <summary>
+		/// Manually specified value for bBuildDeveloperTools.
+		/// </summary>
+		bool? bBuildDeveloperToolsOverride;
+
 		/// <summary>
 		/// Whether to compile the developer tools.
 		/// </summary>
 		[RequiresUniqueBuildEnvironment]
-		public bool? bBuildDeveloperTools;
+		public bool bBuildDeveloperTools
+		{
+			set { bBuildDeveloperToolsOverride = value; }
+			get { return bBuildDeveloperToolsOverride ?? (bCompileAgainstEngine && (Type == TargetType.Editor || Type == TargetType.Program)); }
+		}
 
 		/// <summary>
 		/// Whether to force compiling the target platform modules, even if they wouldn't normally be built.
@@ -381,26 +393,57 @@ namespace UnrealBuildTool
 		public bool bForceBuildShaderFormats = false;
 
 		/// <summary>
-		/// Whether we should compile in support for Simplygon or not.
+		/// Cached value for whether Simplygon is available
+		/// </summary>
+		static Lazy<bool> HasSimplygon = new Lazy<bool>(() => FileReference.Exists(FileReference.Combine(UnrealBuildTool.EngineDirectory, "Source", "ThirdParty", "NotForLicensees", "Simplygon", "Simplygon-latest", "Inc", "SimplygonSDK.h")), LazyThreadSafetyMode.PublicationOnly);
+
+		/// <summary>
+		/// Optional override for whether to use Simplygon. Allows the default value to be computed dynamically if unspecified.
 		/// </summary>
-		[RequiresUniqueBuildEnvironment]
 		[CommandLine("-WithSimplygon")]
 		[ConfigFile(ConfigHierarchyType.Engine, "/Script/BuildSettings.BuildSettings", "bCompileSimplygon")]
-		public bool bCompileSimplygon = true;
+		bool? bOverrideCompileSimplygon;
+		
+		/// <summary>
+		/// Whether we should compile in support for Simplygon or not.
+		/// </summary>
+		[RequiresUniqueBuildEnvironment]
+		public bool bCompileSimplygon
+		{
+			set { bOverrideCompileSimplygon = value; }
+			get { return bOverrideCompileSimplygon ?? (Type == TargetType.Editor && (Platform == UnrealTargetPlatform.Win32 || Platform == UnrealTargetPlatform.Win64) && HasSimplygon.Value); }
+		}
+		
+		/// <summary>
+		/// Cached value for whether SimplygonSSF is available
+		/// </summary>
+		static Lazy<bool> HasSimplygonSSF = new Lazy<bool>(() => FileReference.Exists(FileReference.Combine(UnrealBuildTool.EngineDirectory, "Source", "ThirdParty", "NotForLicensees", "SSF", "Public", "ssf.h")), LazyThreadSafetyMode.PublicationOnly);
+
+        /// <summary>
+        /// Manually specified value for bCompileSimplygonSSF. Allows the default to be determined based on the target type.
+        /// </summary>
+		[ConfigFile(ConfigHierarchyType.Engine, "/Script/BuildSettings.BuildSettings", "bCompileSimplygonSSF")]
+        bool? bOverrideCompileSimplygonSSF;
 
         /// <summary>
         /// Whether we should compile in support for Simplygon's SSF library or not.
         /// </summary>
 		[RequiresUniqueBuildEnvironment]
-		[ConfigFile(ConfigHierarchyType.Engine, "/Script/BuildSettings.BuildSettings", "bCompileSimplygonSSF")]
-        public bool bCompileSimplygonSSF = true;
+		public bool bCompileSimplygonSSF
+		{
+			set { bOverrideCompileSimplygonSSF = value; }
+			get { return bOverrideCompileSimplygonSSF ?? (Type == TargetType.Editor && (Platform == UnrealTargetPlatform.Win32 || Platform == UnrealTargetPlatform.Win64) && HasSimplygonSSF.Value); }
+		}
 
         /// <summary>
 		/// Whether to compile lean and mean version of UE.
 		/// </summary>
-		[RequiresUniqueBuildEnvironment]
-		[ConfigFile(ConfigHierarchyType.Engine, "/Script/BuildSettings.BuildSettings", "bCompileLeanAndMeanUE")]
-		public bool bCompileLeanAndMeanUE = false;
+		[Obsolete("bCompileLeanAndMeanUE is deprecated. Set bBuildDeveloperTools to the opposite value instead.")]
+		public bool bCompileLeanAndMeanUE
+		{
+			get { return !bBuildDeveloperTools; }
+			set { bBuildDeveloperTools = !value; }
+		}
 
         /// <summary>
 		/// Whether to utilize cache freed OS allocs with MallocBinned
@@ -444,9 +487,18 @@ namespace UnrealBuildTool
 		/// <summary>
 		/// Whether to compile SpeedTree support.
 		/// </summary>
-		[RequiresUniqueBuildEnvironment]
 		[ConfigFile(ConfigHierarchyType.Engine, "/Script/BuildSettings.BuildSettings", "bCompileSpeedTree")]
-		public bool bCompileSpeedTree = true;
+		bool? bOverrideCompileSpeedTree;
+
+		/// <summary>
+		/// Whether we should compile in support for Simplygon or not.
+		/// </summary>
+		[RequiresUniqueBuildEnvironment]
+		public bool bCompileSpeedTree
+		{
+			set { bOverrideCompileSpeedTree = value; }
+			get { return bOverrideCompileSpeedTree ?? Type == TargetType.Editor; }
+		}
 
 		/// <summary>
 		/// Enable exceptions for all modules.
@@ -1162,24 +1214,6 @@ namespace UnrealBuildTool
 				UEBuildPlatform.GetBuildPlatform(Platform).ResetTarget(this);
 			}
 
-			// Check that the appropriate headers exist to enable Simplygon
-			if(bCompileSimplygon)
-			{
-				FileReference HeaderFile = FileReference.Combine(UnrealBuildTool.EngineDirectory, "Source", "ThirdParty", "NotForLicensees", "Simplygon", "Simplygon-latest", "Inc", "SimplygonSDK.h");
-				if(!FileReference.Exists(HeaderFile))
-				{
-					bCompileSimplygon = false;
-				}
-			}
-			if(bCompileSimplygonSSF)
-			{
-				FileReference HeaderFile = FileReference.Combine(UnrealBuildTool.EngineDirectory, "Source", "ThirdParty", "NotForLicensees", "SSF", "Public", "ssf.h");
-				if(!FileReference.Exists(HeaderFile))
-				{
-					bCompileSimplygonSSF = false;
-				}
-			}
-
 			// If we've got a changelist set, set that we're making a formal build
 			bFormalBuild = (Version.Changelist != 0 && Version.IsPromotedBuild);
 
@@ -1234,10 +1268,7 @@ namespace UnrealBuildTool
 		{
 			if(Type == global::UnrealBuildTool.TargetType.Game)
 			{
-				bCompileLeanAndMeanUE = true;
-
 				// Do not include the editor
-				bBuildEditor = false;
 				bBuildWithEditorOnlyData = false;
 
 				// Require cooked data
@@ -1254,10 +1285,7 @@ namespace UnrealBuildTool
 			}
 			else if(Type == global::UnrealBuildTool.TargetType.Client)
 			{
-				bCompileLeanAndMeanUE = true;
-
 				// Do not include the editor
-				bBuildEditor = false;
 				bBuildWithEditorOnlyData = false;
 
 				// Require cooked data
@@ -1277,10 +1305,7 @@ namespace UnrealBuildTool
 			}
 			else if(Type == global::UnrealBuildTool.TargetType.Editor)
 			{
-				bCompileLeanAndMeanUE = false;
-
 				// Do not include the editor
-				bBuildEditor = true;
 				bBuildWithEditorOnlyData = true;
 
 				// Require cooked data
@@ -1303,10 +1328,7 @@ namespace UnrealBuildTool
 			}
 			else if(Type == global::UnrealBuildTool.TargetType.Server)
 			{
-				bCompileLeanAndMeanUE = true;
-
 				// Do not include the editor
-				bBuildEditor = false;
 				bBuildWithEditorOnlyData = false;
 
 				// Require cooked data
@@ -1642,8 +1664,7 @@ namespace UnrealBuildTool
 
 		public bool bBuildDeveloperTools
 		{
-			// appropriate default will be set by this point
-			get { return Inner.bBuildDeveloperTools.Value; }
+			get { return Inner.bBuildDeveloperTools; }
 		}
 
 		public bool bForceBuildTargetPlatforms
@@ -1666,6 +1687,7 @@ namespace UnrealBuildTool
 			get { return Inner.bCompileSimplygonSSF; }
 		}
 
+		[Obsolete("bCompileLeanAndMeanUE is deprecated. Use bBuildDeveloperTools instead.")]
 		public bool bCompileLeanAndMeanUE
 		{
 			get { return Inner.bCompileLeanAndMeanUE; }
diff --git a/Engine/Source/Programs/UnrealBuildTool/Configuration/UEBuildPlatform.cs b/Engine/Source/Programs/UnrealBuildTool/Configuration/UEBuildPlatform.cs
--- a/Engine/Source/Programs/UnrealBuildTool/Configuration/UEBuildPlatform.cs
+++ b/Engine/Source/Programs/UnrealBuildTool/Configuration/UEBuildPlatform.cs
@@ -709,7 +709,7 @@ namespace UnrealBuildTool
 		{
 			string[] BoolKeys = new string[] {
 				"bCompileApex", "bCompileICU", "bCompileSimplygon", "bCompileSimplygonSSF",
-				"bCompileLeanAndMeanUE", "bIncludeADO", "bCompileRecast", "bCompileSpeedTree",
+				"bIncludeADO", "bCompileRecast", "bCompileSpeedTree",
 				"bCompileWithPluginSupport", "bCompilePhysXVehicle", "bCompileFreeType",
 				"bCompileForSize", "bCompileCEF3"
 			};
diff --git a/Engine/Source/Programs/UnrealBuildTool/Configuration/UEBuildTarget.cs b/Engine/Source/Programs/UnrealBuildTool/Configuration/UEBuildTarget.cs
--- a/Engine/Source/Programs/UnrealBuildTool/Configuration/UEBuildTarget.cs
+++ b/Engine/Source/Programs/UnrealBuildTool/Configuration/UEBuildTarget.cs
@@ -545,32 +545,47 @@ namespace UnrealBuildTool
 				{
 					if(Field.GetCustomAttribute<RequiresUniqueBuildEnvironmentAttribute>() != null)
 					{
-						// Get the values for the current target and for the base target
 						object ThisValue = Field.GetValue(ThisObjects[Idx]);
 						object BaseValue = Field.GetValue(BaseObjects[Idx]);
-
-						// Check if the fields match, treating lists of strings (eg. definitions) differently to value types.
-						bool bFieldsMatch;
-						if(ThisValue == null || BaseValue == null)
-						{
-							bFieldsMatch = (ThisValue == BaseValue);
-						}
-						else if(typeof(IEnumerable<string>).IsAssignableFrom(Field.FieldType))
-						{
-							bFieldsMatch = Enumerable.SequenceEqual((IEnumerable<string>)ThisValue, (IEnumerable<string>)BaseValue);
-						}
-						else
-						{
-							bFieldsMatch = ThisValue.Equals(BaseValue);
-						}
-
-						// Throw an exception if they don't match
-						if(!bFieldsMatch)
-						{
-							throw new BuildException("{0} modifies the value of {1}. This is not allowed, as {0} has build products in common with {2}.\nRemove the modified setting or change {0} to use a unique build environment by setting 'BuildEnvironment = TargetBuildEnvironment.Unique;' in the {3} constructor.", ThisTargetName, Field.Name, BaseTargetName, ThisRules.GetType().Name);
-						}
+						CheckValuesMatch(ThisRules.GetType(), ThisTargetName, BaseTargetName, Field.Name, Field.FieldType, ThisValue, BaseValue);
 					}
 				}
+				foreach(PropertyInfo Property in ObjectType.GetProperties())
+				{
+					if(Property.GetCustomAttribute<RequiresUniqueBuildEnvironmentAttribute>() != null)
+					{
+						object ThisValue = Property.GetValue(ThisObjects[Idx]);
+						object BaseValue = Property.GetValue(BaseObjects[Idx]);
+						CheckValuesMatch(ThisRules.GetType(), ThisTargetName, BaseTargetName, Property.Name, Property.PropertyType, ThisValue, BaseValue);
+					}
+				}
+			}
+		}
+
+		/// <summary>
+		/// Check that two values match between a base and derived rules type
+		/// </summary>
+		static void CheckValuesMatch(Type RulesType, string ThisTargetName, string BaseTargetName, string FieldName, Type ValueType, object ThisValue, object BaseValue)
+		{
+			// Check if the fields match, treating lists of strings (eg. definitions) differently to value types.
+			bool bFieldsMatch;
+			if(ThisValue == null || BaseValue == null)
+			{
+				bFieldsMatch = (ThisValue == BaseValue);
+			}
+			else if(typeof(IEnumerable<string>).IsAssignableFrom(ValueType))
+			{
+				bFieldsMatch = Enumerable.SequenceEqual((IEnumerable<string>)ThisValue, (IEnumerable<string>)BaseValue);
+			}
+			else
+			{
+				bFieldsMatch = ThisValue.Equals(BaseValue);
+			}
+
+			// Throw an exception if they don't match
+			if(!bFieldsMatch)
+			{
+				throw new BuildException("{0} modifies the value of {1}. This is not allowed, as {0} has build products in common with {2}.\nRemove the modified setting or change {0} to use a unique build environment by setting 'BuildEnvironment = TargetBuildEnvironment.Unique;' in the {3} constructor.", ThisTargetName, FieldName, BaseTargetName, RulesType.Name);
 			}
 		}
 
@@ -3509,16 +3524,6 @@ namespace UnrealBuildTool
 				GlobalCompileEnvironment.Definitions.Add("USE_CHECKS_IN_SHIPPING=0");
 			}
 
-			// Propagate whether we want a lean and mean build to the C++ code.
-			if (Rules.bCompileLeanAndMeanUE)
-			{
-				GlobalCompileEnvironment.Definitions.Add("UE_BUILD_MINIMAL=1");
-			}
-			else
-			{
-				GlobalCompileEnvironment.Definitions.Add("UE_BUILD_MINIMAL=0");
-			}
-
 			// bBuildEditor has now been set appropriately for all platforms, so this is here to make sure the #define 
 			if (Rules.bBuildEditor)
 			{
diff --git a/Engine/Source/Programs/UnrealBuildTool/Platform/Android/UEBuildAndroid.cs b/Engine/Source/Programs/UnrealBuildTool/Platform/Android/UEBuildAndroid.cs
--- a/Engine/Source/Programs/UnrealBuildTool/Platform/Android/UEBuildAndroid.cs
+++ b/Engine/Source/Programs/UnrealBuildTool/Platform/Android/UEBuildAndroid.cs
@@ -95,20 +95,17 @@ namespace UnrealBuildTool
 			ValidateTarget(Target);
 
 			Target.bDeployAfterCompile = true;
+
+			Target.bCompileSimplygon = false;
+			Target.bCompileSimplygonSSF = false;
 		}
 
 		public override void ValidateTarget(TargetRules Target)
 		{
-			Target.bCompileLeanAndMeanUE = true;
 			Target.bCompilePhysX = true;
 			Target.bCompileAPEX = false;
 			Target.bCompileNvCloth = false;
 
-			Target.bBuildEditor = false;
-			Target.bBuildDeveloperTools = false;
-			Target.bCompileSimplygon = false;
-			Target.bCompileSimplygonSSF = false;
-
 			Target.bCompileRecast = true;
 		}
 
diff --git a/Engine/Source/Programs/UnrealBuildTool/Platform/HTML5/UEBuildHTML5.cs b/Engine/Source/Programs/UnrealBuildTool/Platform/HTML5/UEBuildHTML5.cs
--- a/Engine/Source/Programs/UnrealBuildTool/Platform/HTML5/UEBuildHTML5.cs
+++ b/Engine/Source/Programs/UnrealBuildTool/Platform/HTML5/UEBuildHTML5.cs
@@ -42,16 +42,16 @@ namespace UnrealBuildTool
 		public override void ResetTarget(TargetRules Target)
 		{
 			ValidateTarget(Target);
+
+			Target.bCompileSimplygon = false;
+			Target.bCompileSimplygonSSF = false;
 		}
 
 		public override void ValidateTarget(TargetRules Target)
 		{
-			Target.bCompileLeanAndMeanUE = true;
 			Target.bCompileAPEX = false;
 			Target.bCompileNvCloth = false;
 			Target.bCompilePhysX = true;
-			Target.bCompileSimplygon = false;
-			Target.bCompileSimplygonSSF = false;
 			Target.bCompileForSize = true;			// {true:[all:-Oz], false:[developer:-O2, shipping:-O3]}  WARNING: need emscripten version >= 1.37.13
 			Target.bUsePCHFiles = false;
 			Target.bDeployAfterCompile = true;
diff --git a/Engine/Source/Programs/UnrealBuildTool/Platform/IOS/UEBuildIOS.cs b/Engine/Source/Programs/UnrealBuildTool/Platform/IOS/UEBuildIOS.cs
--- a/Engine/Source/Programs/UnrealBuildTool/Platform/IOS/UEBuildIOS.cs
+++ b/Engine/Source/Programs/UnrealBuildTool/Platform/IOS/UEBuildIOS.cs
@@ -636,13 +636,10 @@ namespace UnrealBuildTool
 				Target.bCompilePhysX = false;
 			}
 
-			Target.bBuildEditor = false;
-			Target.bBuildDeveloperTools = false;
 			Target.bCompileAPEX = false;
             Target.bCompileNvCloth = false;
             Target.bCompileSimplygon = false;
             Target.bCompileSimplygonSSF = false;
-			Target.bBuildDeveloperTools = false;
 
 			Target.bDeployAfterCompile = true;
 		}
diff --git a/Engine/Source/Programs/UnrealBuildTool/Platform/Linux/UEBuildLinux.cs b/Engine/Source/Programs/UnrealBuildTool/Platform/Linux/UEBuildLinux.cs
--- a/Engine/Source/Programs/UnrealBuildTool/Platform/Linux/UEBuildLinux.cs
+++ b/Engine/Source/Programs/UnrealBuildTool/Platform/Linux/UEBuildLinux.cs
@@ -229,12 +229,13 @@ namespace UnrealBuildTool
 		public override void ResetTarget(TargetRules Target)
 		{
 			ValidateTarget(Target);
+
+			Target.bCompileSimplygon = false;
+			Target.bCompileSimplygonSSF = false;
 		}
 
 		public override void ValidateTarget(TargetRules Target)
 		{
-			Target.bCompileSimplygon = false;
-			Target.bCompileSimplygonSSF = false;
 			// depends on arch, APEX cannot be as of November'16 compiled for AArch32/64
 			Target.bCompileAPEX = Target.Architecture.StartsWith("x86_64");
 			Target.bCompileNvCloth = Target.Architecture.StartsWith("x86_64");
@@ -262,13 +263,6 @@ namespace UnrealBuildTool
 			{
 				Target.bCompileICU = false;
 			}
-
-			if (ProjectFileGenerator.bGenerateProjectFiles)
-			{
-				// When generating project files we need intellisense generator to include info from all modules,
-				// including editor-only third party libs
-				Target.bCompileLeanAndMeanUE = false;
-			}
 		}
 
 		/// <summary>
diff --git a/Engine/Source/Programs/UnrealBuildTool/Platform/Mac/UEBuildMac.cs b/Engine/Source/Programs/UnrealBuildTool/Platform/Mac/UEBuildMac.cs
--- a/Engine/Source/Programs/UnrealBuildTool/Platform/Mac/UEBuildMac.cs
+++ b/Engine/Source/Programs/UnrealBuildTool/Platform/Mac/UEBuildMac.cs
@@ -153,12 +153,6 @@ namespace UnrealBuildTool
 				Target.GlobalDefinitions.Add("FORCE_ANSI_ALLOCATOR=1");
 			}
 
-			if (ProjectFileGenerator.bGenerateProjectFiles)
-			{
-				// When generating project files we need intellisense generator to include info from all modules, including editor-only third party libs
-				Target.bCompileLeanAndMeanUE = false;
-			}
-
 			Target.bUsePDBFiles = !Target.bDisableDebugInfo && Target.Configuration != UnrealTargetConfiguration.Debug && Platform == UnrealTargetPlatform.Mac && Target.MacPlatform.bGenerateDsymFile;
 
 			// we always deploy - the build machines need to be able to copy the files back, which needs the full bundle
diff --git a/Engine/Source/Programs/UnrealBuildTool/System/ConfigCache.cs b/Engine/Source/Programs/UnrealBuildTool/System/ConfigCache.cs
--- a/Engine/Source/Programs/UnrealBuildTool/System/ConfigCache.cs
+++ b/Engine/Source/Programs/UnrealBuildTool/System/ConfigCache.cs
@@ -1,4 +1,4 @@
-﻿// Copyright 1998-2018 Epic Games, Inc. All Rights Reserved.
+// Copyright 1998-2018 Epic Games, Inc. All Rights Reserved.
 
 using System;
 using System.Collections.Generic;
@@ -232,7 +232,11 @@ namespace UnrealBuildTool
 				if (!TypeToConfigFields.TryGetValue(TargetObjectType, out Fields))
 				{
 					Fields = new List<ConfigField>();
-					foreach (FieldInfo FieldInfo in TargetObjectType.GetFields(BindingFlags.Instance | BindingFlags.GetField | BindingFlags.GetProperty | BindingFlags.Public | BindingFlags.NonPublic))
+					if(TargetObjectType.BaseType != null)
+					{
+						Fields.AddRange(FindConfigFieldsForType(TargetObjectType.BaseType));
+					}
+					foreach (FieldInfo FieldInfo in TargetObjectType.GetFields(BindingFlags.Instance | BindingFlags.GetField | BindingFlags.GetProperty | BindingFlags.Public | BindingFlags.NonPublic | BindingFlags.DeclaredOnly))
 					{
 						IEnumerable<ConfigFileAttribute> Attributes = FieldInfo.GetCustomAttributes<ConfigFileAttribute>();
 						foreach (ConfigFileAttribute Attribute in Attributes)
@@ -406,6 +410,10 @@ namespace UnrealBuildTool
 					return false;
 				}
 			}
+			else if(FieldType.GetGenericTypeDefinition() == typeof(Nullable<>))
+			{
+				return TryParseValue(Text, FieldType.GetGenericArguments()[0], out Value);
+			}
 			else
 			{
 				throw new Exception("Unsupported type for [ConfigFile] attribute");
diff --git a/Engine/Source/Programs/UnrealBuildTool/System/RulesAssembly.cs b/Engine/Source/Programs/UnrealBuildTool/System/RulesAssembly.cs
--- a/Engine/Source/Programs/UnrealBuildTool/System/RulesAssembly.cs
+++ b/Engine/Source/Programs/UnrealBuildTool/System/RulesAssembly.cs
@@ -478,22 +478,6 @@ namespace UnrealBuildTool
 				}
 			}
 
-			// Lean and mean means no Editor and other frills.
-			if (Rules.bCompileLeanAndMeanUE)
-			{
-				Rules.bBuildEditor = false;
-				Rules.bBuildDeveloperTools = Rules.bBuildDeveloperTools ?? false;
-				Rules.bCompileSimplygon = false;
-				Rules.bCompileSimplygonSSF = false;
-				Rules.bCompileSpeedTree = false;
-			}
-
-			// if the bBuildDeveloperTools switch hasn't been defined, default it to the bCompileAgainstEngine switch.
-			if (!Rules.bBuildDeveloperTools.HasValue)
-			{
-				Rules.bBuildDeveloperTools = Rules.bCompileAgainstEngine;
-			}
-
 			// Automatically include CoreUObject
 			if (Rules.bCompileAgainstEngine)
 			{
diff --git a/Engine/Source/Programs/UnrealCEFSubProcess/UnrealCEFSubProcess.Target.cs b/Engine/Source/Programs/UnrealCEFSubProcess/UnrealCEFSubProcess.Target.cs
--- a/Engine/Source/Programs/UnrealCEFSubProcess/UnrealCEFSubProcess.Target.cs
+++ b/Engine/Source/Programs/UnrealCEFSubProcess/UnrealCEFSubProcess.Target.cs
@@ -20,7 +20,7 @@ public class UnrealCEFSubProcessTarget : TargetRules
 		// Turn off various third party features we don't need
 
 		// Currently we force Lean and Mean mode
-		bCompileLeanAndMeanUE = true;
+		bBuildDeveloperTools = false;
 
 		// Currently this app is not linking against the engine, so we'll compile out references from Core to the rest of the engine
 		bCompileAgainstEngine = false;
diff --git a/Engine/Source/Programs/UnrealFileServer/UnrealFileServer.Target.cs b/Engine/Source/Programs/UnrealFileServer/UnrealFileServer.Target.cs
--- a/Engine/Source/Programs/UnrealFileServer/UnrealFileServer.Target.cs
+++ b/Engine/Source/Programs/UnrealFileServer/UnrealFileServer.Target.cs
@@ -1,4 +1,4 @@
-﻿// Copyright 1998-2018 Epic Games, Inc. All Rights Reserved.
+// Copyright 1998-2018 Epic Games, Inc. All Rights Reserved.
 
 using UnrealBuildTool;
 using System.Collections.Generic;
@@ -11,15 +11,10 @@ public class UnrealFileServerTarget : TargetRules
 		LinkType = TargetLinkType.Modular;
 		LaunchModuleName = "UnrealFileServer";
 
-		// Lean and mean
-		bCompileLeanAndMeanUE = true;
-
 		// Never use malloc profiling in Unreal Header Tool.  We set this because often UHT is compiled right before the engine
 		// automatically by Unreal Build Tool, but if bUseMallocProfiler is defined, UHT can operate incorrectly.
 		bUseMallocProfiler = false;
 
-		// No editor needed
-		bBuildEditor = false;
 		// Editor-only data, however, is needed
 		bBuildWithEditorOnlyData = true;
 
diff --git a/Engine/Source/Programs/UnrealFrontend/UnrealFrontend.Target.cs b/Engine/Source/Programs/UnrealFrontend/UnrealFrontend.Target.cs
--- a/Engine/Source/Programs/UnrealFrontend/UnrealFrontend.Target.cs
+++ b/Engine/Source/Programs/UnrealFrontend/UnrealFrontend.Target.cs
@@ -12,7 +12,6 @@ public class UnrealFrontendTarget : TargetRules
 		AdditionalPlugins.Add("UdpMessaging");
         LaunchModuleName = "UnrealFrontend";
 
-		bBuildEditor = false;
 		bCompileAgainstEngine = false;
 		bCompileAgainstCoreUObject = true;
 		bForceBuildTargetPlatforms = true;
diff --git a/Engine/Source/Programs/UnrealHeaderTool/UnrealHeaderTool.Target.cs b/Engine/Source/Programs/UnrealHeaderTool/UnrealHeaderTool.Target.cs
--- a/Engine/Source/Programs/UnrealHeaderTool/UnrealHeaderTool.Target.cs
+++ b/Engine/Source/Programs/UnrealHeaderTool/UnrealHeaderTool.Target.cs
@@ -11,17 +11,12 @@ public class UnrealHeaderToolTarget : TargetRules
 		LinkType = TargetLinkType.Modular;
 		LaunchModuleName = "UnrealHeaderTool";
 
-		// Lean and mean
-		bCompileLeanAndMeanUE = true;
-
 		// Never use malloc profiling in Unreal Header Tool.  We set this because often UHT is compiled right before the engine
 		// automatically by Unreal Build Tool, but if bUseMallocProfiler is defined, UHT can operate incorrectly.
 		bUseMallocProfiler = false;
 
-        // No editor needed
         bCompileICU = false;
-        bBuildEditor = false;
-		// Editor-only data, however, is needed
+        // Editor-only data, however, is needed
 		bBuildWithEditorOnlyData = true;
 
 		// Currently this app is not linking against the engine, so we'll compile out references from Core to the rest of the engine
diff --git a/Engine/Source/Programs/UnrealLightmass/UnrealLightmass.Target.cs b/Engine/Source/Programs/UnrealLightmass/UnrealLightmass.Target.cs
--- a/Engine/Source/Programs/UnrealLightmass/UnrealLightmass.Target.cs
+++ b/Engine/Source/Programs/UnrealLightmass/UnrealLightmass.Target.cs
@@ -14,14 +14,12 @@ public class UnrealLightmassTarget : TargetRules
 		LaunchModuleName = "UnrealLightmass";
 
 		// Lean and mean
-		bCompileLeanAndMeanUE = true;
+		bBuildDeveloperTools = false;
 
 		// Never use malloc profiling in Unreal Header Tool.  We set this because often UHT is compiled right before the engine
 		// automatically by Unreal Build Tool, but if bUseMallocProfiler is defined, UHT can operate incorrectly.
 		bUseMallocProfiler = false;
 
-		// No editor needed
-		bBuildEditor = false;
 		// Editor-only data, however, is needed
 		bBuildWithEditorOnlyData = true;
 
diff --git a/Engine/Source/Programs/UnrealPak/UnrealPak.Target.cs b/Engine/Source/Programs/UnrealPak/UnrealPak.Target.cs
--- a/Engine/Source/Programs/UnrealPak/UnrealPak.Target.cs
+++ b/Engine/Source/Programs/UnrealPak/UnrealPak.Target.cs
@@ -11,12 +11,10 @@ public class UnrealPakTarget : TargetRules
 		LinkType = TargetLinkType.Modular;
 		LaunchModuleName = "UnrealPak";
 
-		bCompileLeanAndMeanUE = true;
+		bBuildDeveloperTools = false;
 		bUseMallocProfiler = false;
 		bCompileWithPluginSupport = true;
 
-		// No editor needed
-		bBuildEditor = false;
 		// Editor-only data, however, is needed
 		bBuildWithEditorOnlyData = true;
 
diff --git a/Engine/Source/Programs/UnrealVersionSelector/UnrealVersionSelector.Target.cs b/Engine/Source/Programs/UnrealVersionSelector/UnrealVersionSelector.Target.cs
--- a/Engine/Source/Programs/UnrealVersionSelector/UnrealVersionSelector.Target.cs
+++ b/Engine/Source/Programs/UnrealVersionSelector/UnrealVersionSelector.Target.cs
@@ -11,7 +11,7 @@ public class UnrealVersionSelectorTarget : TargetRules
 		LinkType = TargetLinkType.Monolithic;
 		LaunchModuleName = "UnrealVersionSelector";
 		
-		bCompileLeanAndMeanUE = true;
+		bBuildDeveloperTools = false;
 		bUseMallocProfiler = false;
 
 		bool bUsingSlate = (Target.Platform == UnrealTargetPlatform.Linux);
@@ -21,9 +21,7 @@ public class UnrealVersionSelectorTarget : TargetRules
 			ExtraModuleNames.Add("EditorStyle");
 		}
 
-		// No editor needed
 		bCompileICU = false;
-		bBuildEditor = false;
 		// Editor-only data, however, is needed
 		bBuildWithEditorOnlyData = bUsingSlate;
 
diff --git a/Engine/Source/Programs/UnrealWatchdog/UnrealWatchdog.Target.cs b/Engine/Source/Programs/UnrealWatchdog/UnrealWatchdog.Target.cs
--- a/Engine/Source/Programs/UnrealWatchdog/UnrealWatchdog.Target.cs
+++ b/Engine/Source/Programs/UnrealWatchdog/UnrealWatchdog.Target.cs
@@ -1,4 +1,4 @@
-﻿// Copyright 1998-2018 Epic Games, Inc. All Rights Reserved.
+// Copyright 1998-2018 Epic Games, Inc. All Rights Reserved.
 
 using UnrealBuildTool;
 using System.Collections.Generic;
@@ -15,14 +15,13 @@ public class UnrealWatchdogTarget : TargetRules
 		LaunchModuleName = "UnrealWatchdog";
 
 		// Lean and mean
-		bCompileLeanAndMeanUE = true;
+		bBuildDeveloperTools = false;
 		bUseMallocProfiler = false;
 
         bUseLoggingInShipping = true;
 
         // No editor needed
-        bBuildEditor = false;
-		bBuildWithEditorOnlyData = false;
+        bBuildWithEditorOnlyData = false;
 
 		// Currently this app is not linking against the engine, so we'll compile out references from Core to the rest of the engine
 		bCompileAgainstEngine = false;
diff --git a/Engine/Source/Programs/Windows/BootstrapPackagedGame/BootstrapPackagedGame.Target.cs b/Engine/Source/Programs/Windows/BootstrapPackagedGame/BootstrapPackagedGame.Target.cs
--- a/Engine/Source/Programs/Windows/BootstrapPackagedGame/BootstrapPackagedGame.Target.cs
+++ b/Engine/Source/Programs/Windows/BootstrapPackagedGame/BootstrapPackagedGame.Target.cs
@@ -1,4 +1,4 @@
-﻿// Copyright 1998-2018 Epic Games, Inc. All Rights Reserved.
+// Copyright 1998-2018 Epic Games, Inc. All Rights Reserved.
 
 using UnrealBuildTool;
 using System.Collections.Generic;
@@ -19,9 +19,8 @@ public class BootstrapPackagedGameTarget : TargetRules
 		bUseMallocProfiler = false;
 
 		// Disable all parts of the editor.
-		bCompileLeanAndMeanUE = true;
+		bBuildDeveloperTools = false;
 		bCompileICU = false;
-		bBuildEditor = false;
 		bBuildWithEditorOnlyData = false;
 		bCompileAgainstEngine = false;
 		bCompileAgainstCoreUObject = false;
diff --git a/Engine/Source/Runtime/AVIWriter/AVIWriter.Build.cs b/Engine/Source/Runtime/AVIWriter/AVIWriter.Build.cs
--- a/Engine/Source/Runtime/AVIWriter/AVIWriter.Build.cs
+++ b/Engine/Source/Runtime/AVIWriter/AVIWriter.Build.cs
@@ -22,7 +22,7 @@ namespace UnrealBuildTool.Rules
 			if ((Target.Platform == UnrealTargetPlatform.Win64) ||
 				(Target.Platform == UnrealTargetPlatform.Win32))
 			{
-				if (Target.bCompileLeanAndMeanUE == false)
+				if (Target.bBuildDeveloperTools)
 				{
 					AddEngineThirdPartyPrivateStaticDependencies(Target, "DirectShow");
 				}
diff --git a/Engine/Source/Runtime/AVIWriter/Private/AVIWriter.cpp b/Engine/Source/Runtime/AVIWriter/Private/AVIWriter.cpp
--- a/Engine/Source/Runtime/AVIWriter/Private/AVIWriter.cpp
+++ b/Engine/Source/Runtime/AVIWriter/Private/AVIWriter.cpp
@@ -18,7 +18,7 @@ class FAVIWriterModule : public IModuleInterface
 
 IMPLEMENT_MODULE(FAVIWriterModule, AVIWriter);
 
-#if PLATFORM_WINDOWS && !UE_BUILD_MINIMAL
+#if PLATFORM_WINDOWS && WITH_UNREAL_DEVELOPER_TOOLS
 
 #include "Windows/WindowsHWrapper.h"
 #include "Windows/AllowWindowsPlatformTypes.h"
@@ -425,7 +425,7 @@ private:
 	IBaseFilter* EncodingFilter;
 };
 
-#elif PLATFORM_MAC && !UE_BUILD_MINIMAL
+#elif PLATFORM_MAC && WITH_UNREAL_DEVELOPER_TOOLS
 
 #import <AVFoundation/AVFoundation.h>
 
@@ -840,9 +840,9 @@ void FAVIWriter::Update(double FrameTimeSeconds, TArray<FColor> FrameData)
 
 FAVIWriter* FAVIWriter::CreateInstance(const FAVIWriterOptions& InOptions)
 {
-#if PLATFORM_WINDOWS && !UE_BUILD_MINIMAL
+#if PLATFORM_WINDOWS && WITH_UNREAL_DEVELOPER_TOOLS
 	return new FAVIWriterWin(InOptions);
-#elif PLATFORM_MAC && !UE_BUILD_MINIMAL
+#elif PLATFORM_MAC && WITH_UNREAL_DEVELOPER_TOOLS
 	return new FAVIWriterMac(InOptions);
 #else
 	return nullptr;
diff --git a/Engine/Source/Runtime/AVIWriter/Private/CapturePin.cpp b/Engine/Source/Runtime/AVIWriter/Private/CapturePin.cpp
--- a/Engine/Source/Runtime/AVIWriter/Private/CapturePin.cpp
+++ b/Engine/Source/Runtime/AVIWriter/Private/CapturePin.cpp
@@ -7,7 +7,7 @@
 
 DEFINE_LOG_CATEGORY(LogMovieCapture);
 
-#if PLATFORM_WINDOWS && !UE_BUILD_MINIMAL
+#if PLATFORM_WINDOWS && WITH_UNREAL_DEVELOPER_TOOLS
 
 
 FCapturePin::FCapturePin(HRESULT *phr, CSource *pFilter, const FAVIWriter& InWriter)
diff --git a/Engine/Source/Runtime/AVIWriter/Private/CapturePin.h b/Engine/Source/Runtime/AVIWriter/Private/CapturePin.h
--- a/Engine/Source/Runtime/AVIWriter/Private/CapturePin.h
+++ b/Engine/Source/Runtime/AVIWriter/Private/CapturePin.h
@@ -11,7 +11,7 @@
 class FAVIWriter;
 struct FCapturedFrame;
 
-#if PLATFORM_WINDOWS && !UE_BUILD_MINIMAL
+#if PLATFORM_WINDOWS && WITH_UNREAL_DEVELOPER_TOOLS
 
 #include "Windows/WindowsHWrapper.h"
 #include "Windows/AllowWindowsPlatformTypes.h"
diff --git a/Engine/Source/Runtime/AVIWriter/Private/CaptureSource.cpp b/Engine/Source/Runtime/AVIWriter/Private/CaptureSource.cpp
--- a/Engine/Source/Runtime/AVIWriter/Private/CaptureSource.cpp
+++ b/Engine/Source/Runtime/AVIWriter/Private/CaptureSource.cpp
@@ -9,7 +9,7 @@
 #include "HAL/PlatformProcess.h"
 #include "HAL/Event.h"
 
-#if PLATFORM_WINDOWS && !UE_BUILD_MINIMAL
+#if PLATFORM_WINDOWS && WITH_UNREAL_DEVELOPER_TOOLS
 
 FCaptureSource::FCaptureSource(const FAVIWriter& Writer)
            : CSource(NAME("ViewportCaptureFilter"), nullptr, CLSID_ViewportCaptureSource)
diff --git a/Engine/Source/Runtime/AVIWriter/Private/CaptureSource.h b/Engine/Source/Runtime/AVIWriter/Private/CaptureSource.h
--- a/Engine/Source/Runtime/AVIWriter/Private/CaptureSource.h
+++ b/Engine/Source/Runtime/AVIWriter/Private/CaptureSource.h
@@ -11,7 +11,7 @@
 
 class FAVIWriter;
 
-#if PLATFORM_WINDOWS && !UE_BUILD_MINIMAL
+#if PLATFORM_WINDOWS && WITH_UNREAL_DEVELOPER_TOOLS
 
 typedef TCHAR* PTCHAR;
 
diff --git a/Engine/Source/Runtime/Core/Public/Misc/Build.h b/Engine/Source/Runtime/Core/Public/Misc/Build.h
--- a/Engine/Source/Runtime/Core/Public/Misc/Build.h
+++ b/Engine/Source/Runtime/Core/Public/Misc/Build.h
@@ -117,14 +117,6 @@
 	#define WITH_AUTOMATION_WORKER !(UE_BUILD_SHIPPING || HACK_HEADER_GENERATOR)
 #endif
 
-/**
- * Whether we want the slimmest possible build of UE4 or not. Don't modify directly but rather change UEBuildConfiguration.cs in UBT.
- */
-#ifndef UE_BUILD_MINIMAL
-	#define UE_BUILD_MINIMAL 0 // for auto-complete
-	#error UBT should always define UE_BUILD_MINIMAL to be 0 or 1
-#endif
-
 /**
 * Whether we want a monolithic build (no DLLs); must be defined by UBT
 */
@@ -220,14 +212,14 @@
 #if UE_BUILD_DEBUG
 	#define DO_GUARD_SLOW									1
 	#define DO_CHECK										1
-	#define STATS											((!UE_BUILD_MINIMAL || !WITH_EDITORONLY_DATA || USE_STATS_WITHOUT_ENGINE || USE_MALLOC_PROFILER || FORCE_USE_STATS) && !ENABLE_STATNAMEDEVENTS)
+	#define STATS											((WITH_UNREAL_DEVELOPER_TOOLS || !WITH_EDITORONLY_DATA || USE_STATS_WITHOUT_ENGINE || USE_MALLOC_PROFILER || FORCE_USE_STATS) && !ENABLE_STATNAMEDEVENTS)
 	#define ALLOW_DEBUG_FILES								1
 	#define ALLOW_CONSOLE									1
 	#define NO_LOGGING										0
 #elif UE_BUILD_DEVELOPMENT
 	#define DO_GUARD_SLOW									0
 	#define DO_CHECK										1
-	#define STATS											((!UE_BUILD_MINIMAL || !WITH_EDITORONLY_DATA || USE_STATS_WITHOUT_ENGINE || USE_MALLOC_PROFILER || FORCE_USE_STATS) && !ENABLE_STATNAMEDEVENTS)
+	#define STATS											((WITH_UNREAL_DEVELOPER_TOOLS || !WITH_EDITORONLY_DATA || USE_STATS_WITHOUT_ENGINE || USE_MALLOC_PROFILER || FORCE_USE_STATS) && !ENABLE_STATNAMEDEVENTS)
 	#define ALLOW_DEBUG_FILES								1
 	#define ALLOW_CONSOLE									1
 	#define NO_LOGGING										0
diff --git a/Engine/Source/Runtime/DatabaseSupport/Public/Database.h b/Engine/Source/Runtime/DatabaseSupport/Public/Database.h
--- a/Engine/Source/Runtime/DatabaseSupport/Public/Database.h
+++ b/Engine/Source/Runtime/DatabaseSupport/Public/Database.h
@@ -8,7 +8,7 @@
  * Whether to compile in support for database connectivity and SQL execution.
  */
 #ifndef WITH_DATABASE_SUPPORT
-	#define WITH_DATABASE_SUPPORT (!UE_BUILD_MINIMAL && !UE_BUILD_SHIPPING)
+	#define WITH_DATABASE_SUPPORT (WITH_UNREAL_DEVELOPER_TOOLS && !UE_BUILD_SHIPPING)
 #endif
 
 // Only use ADO on windows, if support is enabled and not for shipping games.
diff --git a/Engine/Source/Runtime/Engine/Private/ModelLight.cpp b/Engine/Source/Runtime/Engine/Private/ModelLight.cpp
--- a/Engine/Source/Runtime/Engine/Private/ModelLight.cpp
+++ b/Engine/Source/Runtime/Engine/Private/ModelLight.cpp
@@ -35,7 +35,7 @@ TMap<UModelComponent*, TIndirectArray<FModelElement> > UModelComponent::TempBSPE
 #define SHADOWMAP_TEXTURE_WIDTH		512
 #define SHADOWMAP_TEXTURE_HEIGHT	512
 
-#if (defined(_MSC_VER) || PLATFORM_MAC || PLATFORM_UNIX) && WITH_EDITOR && !UE_BUILD_MINIMAL
+#if (defined(_MSC_VER) || PLATFORM_MAC || PLATFORM_UNIX) && WITH_EDITOR
 	/** Whether to allow cropping of unmapped borders in lightmaps and shadowmaps. Controlled by BaseLightmass.ini setting. */
 	extern ENGINE_API bool GAllowLightmapCropping;
 #endif
