UBT: Move the PostBuild step on IOS (uploading crashlytics data, creating stub, creating the app bundle) into the action graph. Properties from the receipt required by UPL are now passed in via the link environment, and are serialized to a temporary file along with other parameters for the target, then passed to a special UBT mode via a child process in the action graph.

Also re-enabled makefiles on IOS, now that there is no external state, and removed the PostBuildSync callback from the platform layer.

diff --git a/Engine/Source/Programs/UnrealBuildTool/Configuration/UEBuildModule.cs b/Engine/Source/Programs/UnrealBuildTool/Configuration/UEBuildModule.cs
--- a/Engine/Source/Programs/UnrealBuildTool/Configuration/UEBuildModule.cs
+++ b/Engine/Source/Programs/UnrealBuildTool/Configuration/UEBuildModule.cs
@@ -725,6 +725,9 @@ namespace UnrealBuildTool
 				DependencyModule.SetupPublicLinkEnvironment(SourceBinary, LinkEnvironment.LibraryPaths, LinkEnvironment.AdditionalLibraries, LinkEnvironment.RuntimeLibraryPaths, LinkEnvironment.Frameworks, LinkEnvironment.WeakFrameworks,
 					LinkEnvironment.AdditionalFrameworks, LinkEnvironment.AdditionalShadowFiles, LinkEnvironment.AdditionalBundleResources, LinkEnvironment.DelayLoadDLLs, BinaryDependencies, VisitedModules, ExeDir);
 			}
+
+			// Add all the additional properties
+			LinkEnvironment.AdditionalProperties.AddRange(Rules.AdditionalPropertiesForReceipt.Inner);
 		}
 
 		/// <summary>
diff --git a/Engine/Source/Programs/UnrealBuildTool/Platform/IOS/IOSPostBuildSyncMode.cs b/Engine/Source/Programs/UnrealBuildTool/Platform/IOS/IOSPostBuildSyncMode.cs
new file mode 100644
--- /dev/null
+++ b/Engine/Source/Programs/UnrealBuildTool/Platform/IOS/IOSPostBuildSyncMode.cs
@@ -0,0 +1,82 @@
+using System;
+using System.Collections.Generic;
+using System.Linq;
+using System.Text;
+using System.Threading.Tasks;
+using Tools.DotNETCommon;
+
+namespace UnrealBuildTool
+{
+	[Serializable]
+	class IOSPostBuildSyncTarget
+	{
+		public UnrealTargetPlatform Platform;
+		public UnrealTargetConfiguration Configuration;
+		public FileReference ProjectFile;
+		public string TargetName;
+		public TargetType TargetType;
+		public FileReference OutputPath;
+		public List<string> UPLScripts;
+		public VersionNumber SdkVersion;
+		public bool bCreateStubIPA;
+		public bool bSkipCrashlytics;
+		public DirectoryReference ProjectDirectory;
+		public DirectoryReference ProjectIntermediateDirectory;
+		public string ImportProvision;
+		public string ImportCertificate;
+		public string ImportCertificatePassword;
+		public Dictionary<string, DirectoryReference> FrameworkNameToSourceDir;
+
+		public IOSPostBuildSyncTarget(ReadOnlyTargetRules Target, FileReference OutputPath, DirectoryReference ProjectIntermediateDirectory, List<string> UPLScripts, VersionNumber SdkVersion, Dictionary<string, DirectoryReference> FrameworkNameToSourceDir)
+		{
+			this.Platform = Target.Platform;
+			this.Configuration = Target.Configuration;
+			this.ProjectFile = Target.ProjectFile;
+			this.TargetName = Target.Name;
+			this.TargetType = Target.Type;
+			this.OutputPath = OutputPath;
+			this.UPLScripts = UPLScripts;
+			this.SdkVersion = SdkVersion;
+			this.bCreateStubIPA = Target.IOSPlatform.bCreateStubIPA;
+			this.bSkipCrashlytics = Target.IOSPlatform.bSkipCrashlytics;
+			this.ProjectDirectory = DirectoryReference.FromFile(Target.ProjectFile) ?? UnrealBuildTool.EngineDirectory;
+			this.ProjectIntermediateDirectory = ProjectIntermediateDirectory;
+			this.ImportProvision = Target.IOSPlatform.ImportProvision;
+			this.ImportCertificate = Target.IOSPlatform.ImportCertificate;
+			this.ImportCertificatePassword = Target.IOSPlatform.ImportCertificatePassword;
+			this.FrameworkNameToSourceDir = FrameworkNameToSourceDir;
+		}
+	}
+
+	[ToolMode("IOSPostBuildSync")]
+	class IOSPostBuildSyncMode : ToolMode
+	{
+		[CommandLine("-Input=", Required = true)]
+		public FileReference InputFile = null;
+
+		[CommandLine("-XmlConfigCache=")]
+		public FileReference XmlConfigCache = null;
+
+		public override int Execute(CommandLineArguments Arguments)
+		{
+			Arguments.ApplyTo(this);
+			Arguments.CheckAllArgumentsUsed();
+
+			// Read the XML config
+			XmlConfig.ReadConfigFiles(XmlConfigCache);
+
+			// Change the working directory to be the Engine/Source folder. We are likely running from Engine/Binaries/DotNET
+			// This is critical to be done early so any code that relies on the current directory being Engine/Source will work.
+			DirectoryReference.SetCurrentDirectory(UnrealBuildTool.EngineSourceDirectory);
+
+			// Find and register all tool chains, build platforms, etc. that are present
+			UnrealBuildTool.RegisterAllUBTClasses(false);
+
+			// Run the PostBuildSync command
+			IOSPostBuildSyncTarget Target = BinaryFormatterUtils.Load<IOSPostBuildSyncTarget>(InputFile);
+			IOSToolChain.PostBuildSync(Target);
+
+			return 0;
+		}
+	}
+}
diff --git a/Engine/Source/Programs/UnrealBuildTool/Platform/IOS/IOSToolChain.cs b/Engine/Source/Programs/UnrealBuildTool/Platform/IOS/IOSToolChain.cs
--- a/Engine/Source/Programs/UnrealBuildTool/Platform/IOS/IOSToolChain.cs
+++ b/Engine/Source/Programs/UnrealBuildTool/Platform/IOS/IOSToolChain.cs
@@ -123,11 +123,6 @@ namespace UnrealBuildTool
 		/// </summary>
 		private const string IOSArchiver = "libtool";
 
-		/// <summary>
-		/// Additional frameworks stored locally so we have access without LinkEnvironment
-		/// </summary>
-		public static List<UEBuildFramework> RememberedAdditionalFrameworks = new List<UEBuildFramework>();
-
         public override string GetSDKVersion()
         {
             return Settings.Value.IOSSDKVersionFloat.ToString();
@@ -710,32 +705,10 @@ namespace UnrealBuildTool
 						CompileAction.PrerequisiteItems.Add(ExtractedTokenFile);
 					}
 				}
-
-				AddFrameworksToTrack(CompileEnvironment.AdditionalFrameworks, CompileAction);
 			}
 			return Result;
 		}
 
-		static void AddFrameworksToTrack(List<UEBuildFramework> Frameworks, Action DependentAction)
-		{
-			foreach (UEBuildFramework Framework in Frameworks)
-			{
-				if (Framework.ZipFile == null)
-				{
-					continue;	// Only care about frameworks that have a zip specified
-				}
-
-				// If we've already remembered this framework, skip
-				if (RememberedAdditionalFrameworks.Contains(Framework))
-				{
-					continue;
-				}
-
-				// Remember any files we need to unzip
-				RememberedAdditionalFrameworks.Add(Framework);
-			}
-		}
-
 		public override FileItem LinkFiles(LinkEnvironment LinkEnvironment, bool bBuildImportLibraryOnly, ActionGraph ActionGraph)
 		{
 			string LinkerPath = Settings.Value.ToolchainDir +
@@ -787,7 +760,6 @@ namespace UnrealBuildTool
 					LinkAction.PrerequisiteItems.Add(ExtractedTokenFile);
 				}
 			}
-			AddFrameworksToTrack(LinkEnvironment.AdditionalFrameworks, LinkAction);
 
 			// Add the output file as a production of the link action.
 			FileItem OutputFile = FileItem.GetItemByFileReference(LinkEnvironment.OutputFilePath);
@@ -1128,7 +1100,7 @@ namespace UnrealBuildTool
 				throw new BuildException("Unable to extract framework '{0}' - no zip file specified", Framework.Name);
 			}
 			if(Framework.ExtractedTokenFile == null)
-			{
+				{
 				FileItem InputFile = FileItem.GetItemByFileReference(Framework.ZipFile);
 				Framework.ExtractedTokenFile = FileItem.GetItemByFileReference(new FileReference(Framework.OutputDirectory.FullName + ".extracted"));
 
@@ -1303,7 +1275,7 @@ namespace UnrealBuildTool
 
         public override ICollection<FileItem> PostBuild(FileItem Executable, LinkEnvironment BinaryLinkEnvironment, ActionGraph ActionGraph)
         {
-            ICollection<FileItem> OutputFiles = base.PostBuild(Executable, BinaryLinkEnvironment, ActionGraph);
+            List<FileItem> OutputFiles = new List<FileItem>(base.PostBuild(Executable, BinaryLinkEnvironment, ActionGraph));
 
             if (BinaryLinkEnvironment.bIsBuildingLibrary)
             {
@@ -1362,6 +1334,44 @@ namespace UnrealBuildTool
 				OutputFiles.Add(AssetCatalogFile);
 			}
 
+			// Generate the app bundle
+			if(!Target.bDisableLinking)
+			{
+				Log.TraceInformation("Adding PostBuildSync action");
+			
+				List<string> UPLScripts = UEDeployIOS.CollectPluginDataPaths(BinaryLinkEnvironment.AdditionalProperties);
+				VersionNumber SdkVersion = VersionNumber.Parse(Settings.Value.IOSSDKVersion);
+
+				Dictionary<string, DirectoryReference> FrameworkNameToSourceDir = new Dictionary<string, DirectoryReference>();
+				foreach (UEBuildFramework Framework in BinaryLinkEnvironment.AdditionalFrameworks)
+				{
+					if (Framework.OutputDirectory != null && !String.IsNullOrEmpty(Framework.CopyBundledAssets))
+					{
+						// For now, this is hard coded, but we need to loop over all modules, and copy bundled assets that need it
+						DirectoryReference LocalSource = DirectoryReference.Combine(Framework.OutputDirectory, Framework.CopyBundledAssets);
+						string BundleName = Framework.CopyBundledAssets.Substring(Framework.CopyBundledAssets.LastIndexOf('/') + 1);
+						FrameworkNameToSourceDir[BundleName] = LocalSource;
+					}
+				}
+
+				IOSPostBuildSyncTarget PostBuildSyncTarget = new IOSPostBuildSyncTarget(Target, Executable.Location, BinaryLinkEnvironment.IntermediateDirectory, UPLScripts, SdkVersion, FrameworkNameToSourceDir);
+				FileReference PostBuildSyncFile = FileReference.Combine(BinaryLinkEnvironment.IntermediateDirectory, "PostBuildSync.dat");
+				BinaryFormatterUtils.Save(PostBuildSyncFile, PostBuildSyncTarget);
+
+				string PostBuildSyncArguments = String.Format("-Mode=IOSPostBuildSync -Input=\"{0}\" -XmlConfigCache=\"{1}\"", PostBuildSyncFile, XmlConfig.CacheFile);
+
+				Action PostBuildSyncAction = ActionGraph.AddRecursiveCall(ActionType.CreateAppBundle, PostBuildSyncArguments);
+				PostBuildSyncAction.WorkingDirectory = UnrealBuildTool.EngineSourceDirectory.FullName;
+				PostBuildSyncAction.PrerequisiteItems.Add(Executable);
+				PostBuildSyncAction.PrerequisiteItems.AddRange(OutputFiles);
+				PostBuildSyncAction.ProducedItems.Add(FileItem.GetItemByFileReference(GetStagedExecutablePath(Executable.Location, Target.Name)));
+				PostBuildSyncAction.DeleteItems.AddRange(PostBuildSyncAction.ProducedItems);
+				PostBuildSyncAction.StatusDescription = "Executing PostBuildSync";
+				PostBuildSyncAction.bCanExecuteRemotely = false;
+
+				OutputFiles.AddRange(PostBuildSyncAction.ProducedItems);
+			}
+
 			return OutputFiles;
         }
 
@@ -1467,44 +1477,34 @@ namespace UnrealBuildTool
 			}
 		}
 
-        public static void PostBuildSync(UEBuildTarget Target)
+		public static FileReference GetStagedExecutablePath(FileReference Executable, string TargetName)
 		{
-			if (Target.Rules == null)
-			{
-				Log.TraceWarning("Unable to PostBuildSync, Target has no Rules object");
-				return;
-			}
-			if(Target.Rules.bDisableLinking)
-			{
-				return;
-			}
-
-			TargetReceipt Receipt = TargetReceipt.Read(Target.ReceiptFileName);
+			return FileReference.Combine(Executable.Directory, "Payload", TargetName + ".app", TargetName);
+		}
 
+        public static void PostBuildSync(IOSPostBuildSyncTarget Target)
+		{
 			IOSProjectSettings ProjectSettings = ((IOSPlatform)UEBuildPlatform.GetBuildPlatform(Target.Platform)).ReadProjectSettings(Target.ProjectFile);
 
-			string AppName = Target.TargetType == TargetType.Game ? Target.TargetName : Target.AppName;
+			string AppName = Target.TargetName;
 
-			string RemoteShadowDirectoryMac = Path.GetDirectoryName(Target.OutputPath.FullName);
-			string FinalRemoteExecutablePath = String.Format("{0}/Payload/{1}.app/{1}", RemoteShadowDirectoryMac, AppName);
+			string RemoteShadowDirectoryMac = Target.OutputPath.Directory.FullName;
+			FileReference StagedExecutablePath = GetStagedExecutablePath(Target.OutputPath, Target.TargetName);
 
             // ensure the plist, entitlements, and provision files are properly copied
             UEDeployIOS DeployHandler = (Target.Platform == UnrealTargetPlatform.IOS ? new UEDeployIOS() : new UEDeployTVOS());
-            DeployHandler.PrepTargetForDeployment(Receipt, Target.Rules.IOSPlatform.bCreateStubIPA);
+            DeployHandler.PrepTargetForDeployment(Target.ProjectFile, Target.TargetName, Target.Platform, Target.Configuration, Target.UPLScripts, Target.SdkVersion, Target.bCreateStubIPA);
 
 			// copy the executable
-			if (!File.Exists(FinalRemoteExecutablePath))
-			{
-				Directory.CreateDirectory(String.Format("{0}/Payload/{1}.app", RemoteShadowDirectoryMac, AppName));
-			}
-			File.Copy(Target.OutputPath.FullName, FinalRemoteExecutablePath, true);
+			DirectoryReference.CreateDirectory(StagedExecutablePath.Directory);
+			FileReference.Copy(Target.OutputPath, StagedExecutablePath, true);
 
-			if (!Target.Rules.IOSPlatform.bSkipCrashlytics)
+			if (!Target.bSkipCrashlytics)
 			{
 				GenerateCrashlyticsData(RemoteShadowDirectoryMac, Path.GetFileName(Target.OutputPath.FullName), Target.ProjectDirectory.FullName, AppName);
 			}
 
-			if (Target.Rules.IOSPlatform.bCreateStubIPA)
+			if (Target.bCreateStubIPA)
 			{
 				// generate the dummy project so signing works
 				DirectoryReference XcodeWorkspaceDir;
@@ -1527,7 +1527,7 @@ namespace UnrealBuildTool
 
 				// ensure the plist, entitlements, and provision files are properly copied
 				DeployHandler = (Target.Platform == UnrealTargetPlatform.IOS ? new UEDeployIOS() : new UEDeployTVOS());
-				DeployHandler.PrepTargetForDeployment(Receipt, true);
+				DeployHandler.PrepTargetForDeployment(Target.ProjectFile, Target.TargetName, Target.Platform, Target.Configuration, Target.UPLScripts, Target.SdkVersion, true);
 
 				FileReference SignProjectScript = FileReference.Combine(Target.ProjectIntermediateDirectory, "SignProject.sh");
 				using(StreamWriter Writer = new StreamWriter(SignProjectScript.FullName))
@@ -1538,9 +1538,9 @@ namespace UnrealBuildTool
 					Writer.WriteLine("set -x");
 
 					// Copy the mobile provision into the system store
-					if(Target.Rules.IOSPlatform.ImportProvision != null)
+					if(Target.ImportProvision != null)
 					{
-						Writer.WriteLine("cp -f {0} ~/Library/MobileDevice/Provisioning\\ Profiles/", Utils.EscapeShellArgument(Target.Rules.IOSPlatform.ImportProvision));
+						Writer.WriteLine("cp -f {0} ~/Library/MobileDevice/Provisioning\\ Profiles/", Utils.EscapeShellArgument(Target.ImportProvision));
 					}
 
 					// Path to the temporary keychain. When -ImportCertificate is specified, we will temporarily add this to the list of keychains to search, and remove it later.
@@ -1548,7 +1548,7 @@ namespace UnrealBuildTool
 
 					// Get the signing certificate to use
 					string SigningCertificate;
-					if(Target.Rules.IOSPlatform.ImportCertificate == null)
+					if(Target.ImportCertificate == null)
 					{
 						// Take it from the standard settings
 						IOSProvisioningData ProvisioningData = ((IOSPlatform)UEBuildPlatform.GetBuildPlatform(Target.Platform)).ReadProvisioningData(Target.ProjectFile);
@@ -1566,11 +1566,11 @@ namespace UnrealBuildTool
 						X509Certificate2 Certificate;
 						try
 						{
-							Certificate = new X509Certificate2(Target.Rules.IOSPlatform.ImportCertificate, Target.Rules.IOSPlatform.ImportCertificatePassword ?? "");
+							Certificate = new X509Certificate2(Target.ImportCertificate, Target.ImportCertificatePassword ?? "");
 						}
 						catch(Exception Ex)
 						{
-							throw new BuildException(Ex, "Unable to read certificate '{0}': {1}", Target.Rules.IOSPlatform.ImportCertificate, Ex.Message);
+							throw new BuildException(Ex, "Unable to read certificate '{0}': {1}", Target.ImportCertificate, Ex.Message);
 						}
 						SigningCertificate = Certificate.GetNameInfo(X509NameType.SimpleName, false);
 
@@ -1584,7 +1584,7 @@ namespace UnrealBuildTool
 						Writer.WriteLine("security list-keychains");
 						Writer.WriteLine("security set-keychain-settings -t 3600 -l  \"{0}\"", TempKeychain);
 						Writer.WriteLine("security -v unlock-keychain -p \"A\" \"{0}\"", TempKeychain);
-						Writer.WriteLine("security import {0} -P {1} -k \"{2}\" -T /usr/bin/codesign -T /usr/bin/security -t agg", Utils.EscapeShellArgument(Target.Rules.IOSPlatform.ImportCertificate), Utils.EscapeShellArgument(Target.Rules.IOSPlatform.ImportCertificatePassword), TempKeychain);
+						Writer.WriteLine("security import {0} -P {1} -k \"{2}\" -T /usr/bin/codesign -T /usr/bin/security -t agg", Utils.EscapeShellArgument(Target.ImportCertificate), Utils.EscapeShellArgument(Target.ImportCertificatePassword), TempKeychain);
 						Writer.WriteLine("security set-key-partition-list -S apple-tool:,apple:,codesign: -s -k \"A\" -D '{0}' -t private {1}", SigningCertificate, TempKeychain);
 
 						// Set parameters to make sure it uses the correct identity and keychain
@@ -1596,7 +1596,7 @@ namespace UnrealBuildTool
 					FileReference MobileProvisionFile;
 					string MobileProvisionUUID;
 					string TeamUUID;
-					if(Target.Rules.IOSPlatform.ImportProvision == null)
+					if(Target.ImportProvision == null)
 					{
 						IOSProvisioningData ProvisioningData = ((IOSPlatform)UEBuildPlatform.GetBuildPlatform(Target.Platform)).ReadProvisioningData(ProjectSettings);
 						MobileProvisionFile = ProvisioningData.MobileProvisionFile;
@@ -1605,7 +1605,7 @@ namespace UnrealBuildTool
 					}
 					else
 					{
-						MobileProvisionFile = new FileReference(Target.Rules.IOSPlatform.ImportProvision);
+						MobileProvisionFile = new FileReference(Target.ImportProvision);
 
 						MobileProvisionContents MobileProvision = MobileProvisionContents.Read(MobileProvisionFile);
 						MobileProvisionUUID = MobileProvision.GetUniqueId();
@@ -1613,9 +1613,9 @@ namespace UnrealBuildTool
 					}
 
 					string ConfigName = Target.Configuration.ToString();
-					if (Target.Rules.Type != TargetType.Game && Target.Rules.Type != TargetType.Program)
+					if (Target.TargetType != TargetType.Game && Target.TargetType != TargetType.Program)
 					{
-						ConfigName += " " + Target.Rules.Type.ToString();
+						ConfigName += " " + Target.TargetType.ToString();
 					}
 
 					string SchemeName;
@@ -1686,22 +1686,17 @@ namespace UnrealBuildTool
 				// Clean the local dest directory if it exists
 				CleanIntermediateDirectory(LocalFrameworkAssets);
 
-				foreach (UEBuildFramework Framework in RememberedAdditionalFrameworks)
+				foreach (KeyValuePair<string, DirectoryReference> Pair in Target.FrameworkNameToSourceDir)
 				{
-					if (!String.IsNullOrEmpty(Framework.CopyBundledAssets) && Framework.OutputDirectory != null)
-					{
-						string UnpackedZipPath = Framework.OutputDirectory.FullName;
+					string UnpackedZipPath = Pair.Value.FullName;
 
-						// For now, this is hard coded, but we need to loop over all modules, and copy bundled assets that need it
-						string LocalSource = UnpackedZipPath + "/" + Framework.CopyBundledAssets;
-						string BundleName = Framework.CopyBundledAssets.Substring(Framework.CopyBundledAssets.LastIndexOf('/') + 1);
-						string LocalDest = LocalFrameworkAssets + "/" + BundleName;
+					// For now, this is hard coded, but we need to loop over all modules, and copy bundled assets that need it
+					string LocalDest = LocalFrameworkAssets + "/" + Pair.Key;
 
-						Log.TraceInformation("Copying bundled asset... LocalSource: {0}, LocalDest: {1}", LocalSource, LocalDest);
+					Log.TraceInformation("Copying bundled asset... LocalSource: {0}, LocalDest: {1}", Pair.Value, LocalDest);
 
-						string ResultsText;
-						RunExecutableAndWait("cp", String.Format("-R -L \"{0}\" \"{1}\"", LocalSource, LocalDest), out ResultsText);
-					}
+					string ResultsText;
+                    RunExecutableAndWait("cp", String.Format("-R -L \"{0}\" \"{1}\"", Pair.Value, LocalDest), out ResultsText);
                 }
             }
 		}
diff --git a/Engine/Source/Programs/UnrealBuildTool/Platform/IOS/UEBuildIOS.cs b/Engine/Source/Programs/UnrealBuildTool/Platform/IOS/UEBuildIOS.cs
--- a/Engine/Source/Programs/UnrealBuildTool/Platform/IOS/UEBuildIOS.cs
+++ b/Engine/Source/Programs/UnrealBuildTool/Platform/IOS/UEBuildIOS.cs
@@ -767,11 +767,6 @@ namespace UnrealBuildTool
 			return true;
 		}
 
-		public override void PostBuildSync(UEBuildTarget Target)
-		{
-			IOSToolChain.PostBuildSync(Target);
-		}
-
 		public bool HasCustomIcons(DirectoryReference ProjectDirectoryName)
 		{
 			string IconDir = Path.Combine(ProjectDirectoryName.FullName, "Build", "IOS", "Resources", "Graphics");
diff --git a/Engine/Source/Programs/UnrealBuildTool/Platform/IOS/UEDeployIOS.cs b/Engine/Source/Programs/UnrealBuildTool/Platform/IOS/UEDeployIOS.cs
--- a/Engine/Source/Programs/UnrealBuildTool/Platform/IOS/UEDeployIOS.cs
+++ b/Engine/Source/Programs/UnrealBuildTool/Platform/IOS/UEDeployIOS.cs
@@ -186,7 +186,7 @@ namespace UnrealBuildTool
 			return result;
 		}
 
-		public static bool GenerateIOSPList(FileReference ProjectFile, UnrealTargetConfiguration Config, string ProjectDirectory, bool bIsUE4Game, string GameName, string ProjectName, string InEngineDir, string AppDirectory, TargetReceipt Receipt, UnrealPluginLanguage UPL, out bool bSupportsPortrait, out bool bSupportsLandscape, out bool bSkipIcons)
+		public static bool GenerateIOSPList(FileReference ProjectFile, UnrealTargetConfiguration Config, string ProjectDirectory, bool bIsUE4Game, string GameName, string ProjectName, string InEngineDir, string AppDirectory, VersionNumber SdkVersion, UnrealPluginLanguage UPL, out bool bSupportsPortrait, out bool bSupportsLandscape, out bool bSkipIcons)
 		{
 			// generate the Info.plist for future use
 			string BuildDirectory = ProjectDirectory + "/Build/IOS";
@@ -456,7 +456,7 @@ namespace UnrealBuildTool
 			}
 			Text.AppendLine("\t</array>");
 
-			if (SupportsIconCatalog(Receipt))
+			if (SupportsIconCatalog(SdkVersion))
 			{
 				bSkipIcons = true;
 				Text.AppendLine("\t<key>CFBundleIcons</key>");
@@ -739,33 +739,33 @@ namespace UnrealBuildTool
 			return bSkipDefaultPNGs;
 		}
 
+        static bool SupportsIconCatalog(VersionNumber SdkVersion)
+        {
+			return SdkVersion != null && SdkVersion >= new VersionNumber(11);
+		}
 
-        /// <summary>
-        /// 
-        /// </summary>
-        static bool SupportsIconCatalog(TargetReceipt Receipt)
+        public static VersionNumber GetSdkVersion(TargetReceipt Receipt)
         {
-            // get the receipt
+			VersionNumber SdkVersion = null;
             if (Receipt != null)
             {
-                IEnumerable<ReceiptProperty> Results = Receipt.AdditionalProperties.Where(x => x.Name == "SDK");
-
-                if (Results.Count() > 0)
-                {
-                    if (Single.Parse(Results.ElementAt(0).Value) >= 11.0f)
-                    {
-                        return true;
-                    }
-                    else
-                    {
-                        return false;
-                    }
-                }
-            }
-            return false;
+                ReceiptProperty SdkVersionProperty = Receipt.AdditionalProperties.FirstOrDefault(x => x.Name == "SDK");
+				if(SdkVersionProperty != null)
+				{
+					VersionNumber.TryParse(SdkVersionProperty.Value, out SdkVersion);
+				}
+			}
+			return SdkVersion;
         }
 
-		public virtual bool GeneratePList(FileReference ProjectFile, UnrealTargetConfiguration Config, string ProjectDirectory, bool bIsUE4Game, string GameName, string ProjectName, string InEngineDir, string AppDirectory, TargetReceipt Receipt, out bool bSupportsPortrait, out bool bSupportsLandscape, out bool bSkipIcons)
+		public bool GeneratePList(FileReference ProjectFile, UnrealTargetConfiguration Config, string ProjectDirectory, bool bIsUE4Game, string GameName, string ProjectName, string InEngineDir, string AppDirectory, TargetReceipt Receipt, out bool bSupportsPortrait, out bool bSupportsLandscape, out bool bSkipIcons)
+		{
+			List<string> UPLScripts = CollectPluginDataPaths(Receipt.AdditionalProperties);
+			VersionNumber SdkVersion = GetSdkVersion(Receipt);
+			return GeneratePList(ProjectFile, Config, ProjectDirectory, bIsUE4Game, GameName, ProjectName, InEngineDir, AppDirectory, UPLScripts, SdkVersion, out bSupportsPortrait, out bSupportsLandscape, out bSkipIcons);
+		}
+
+		public virtual bool GeneratePList(FileReference ProjectFile, UnrealTargetConfiguration Config, string ProjectDirectory, bool bIsUE4Game, string GameName, string ProjectName, string InEngineDir, string AppDirectory, List<string> UPLScripts, VersionNumber SdkVersion, out bool bSupportsPortrait, out bool bSupportsLandscape, out bool bSkipIcons)
 		{
 			// remember name with -IOS-Shipping, etc
 			string ExeName = GameName;
@@ -792,12 +792,12 @@ namespace UnrealBuildTool
 
 			string RelativeEnginePath = UnrealBuildTool.EngineDirectory.MakeRelativeTo(DirectoryReference.GetCurrentDirectory());
 
-			UnrealPluginLanguage UPL = new UnrealPluginLanguage(ProjectFile, CollectPluginDataPaths(Receipt), ProjectArches, "", "", UnrealTargetPlatform.IOS);
+			UnrealPluginLanguage UPL = new UnrealPluginLanguage(ProjectFile, UPLScripts, ProjectArches, "", "", UnrealTargetPlatform.IOS);
 
 			// Passing in true for distribution is not ideal here but given the way that ios packaging happens and this call chain it seems unavoidable for now, maybe there is a way to correctly pass it in that I can't find?
 			UPL.Init(ProjectArches, true, RelativeEnginePath, BundlePath, ProjectDirectory, Config.ToString(), BuildVersion.ReadDefault());
 
-			return GenerateIOSPList(ProjectFile, Config, ProjectDirectory, bIsUE4Game, GameName, ProjectName, InEngineDir, AppDirectory, Receipt, UPL, out bSupportsPortrait, out bSupportsLandscape, out bSkipIcons);
+			return GenerateIOSPList(ProjectFile, Config, ProjectDirectory, bIsUE4Game, GameName, ProjectName, InEngineDir, AppDirectory, SdkVersion, UPL, out bSupportsPortrait, out bSupportsLandscape, out bSkipIcons);
 		}
 
 		protected virtual void CopyCloudResources(string InEngineDir, string AppDirectory)
@@ -899,6 +899,13 @@ namespace UnrealBuildTool
 		}
 
 		public bool PrepForUATPackageOrDeploy(UnrealTargetConfiguration Config, FileReference ProjectFile, string InProjectName, string InProjectDirectory, string InExecutablePath, string InEngineDir, bool bForDistribution, string CookFlavor, bool bIsDataDeploy, bool bCreateStubIPA, TargetReceipt Receipt)
+		{
+			List<string> UPLScripts = CollectPluginDataPaths(Receipt.AdditionalProperties);
+			VersionNumber SdkVersion = GetSdkVersion(Receipt);
+			return PrepForUATPackageOrDeploy(Config, ProjectFile, InProjectName, InProjectDirectory, InExecutablePath, InEngineDir, bForDistribution, CookFlavor, bIsDataDeploy, bCreateStubIPA, UPLScripts, SdkVersion);
+		}
+
+		public bool PrepForUATPackageOrDeploy(UnrealTargetConfiguration Config, FileReference ProjectFile, string InProjectName, string InProjectDirectory, string InExecutablePath, string InEngineDir, bool bForDistribution, string CookFlavor, bool bIsDataDeploy, bool bCreateStubIPA, List<string> UPLScripts, VersionNumber SdkVersion)
 		{
 			if (BuildHostPlatform.Current.Platform != UnrealTargetPlatform.Mac)
 			{
@@ -1061,7 +1068,7 @@ namespace UnrealBuildTool
 			bool bSupportsPortrait = true;
 			bool bSupportsLandscape = false;
 			bool bSkipIcons = false;
-			bool bSkipDefaultPNGs = GeneratePList(ProjectFile, Config, InProjectDirectory, bIsUE4Game, GameExeName, InProjectName, InEngineDir, AppDirectory, Receipt, out bSupportsPortrait, out bSupportsLandscape, out bSkipIcons);
+			bool bSkipDefaultPNGs = GeneratePList(ProjectFile, Config, InProjectDirectory, bIsUE4Game, GameExeName, InProjectName, InEngineDir, AppDirectory, UPLScripts, SdkVersion, out bSupportsPortrait, out bSupportsLandscape, out bSkipIcons);
 
 			// ensure the destination is writable
 			if (File.Exists(AppDirectory + "/" + GameName))
@@ -1103,47 +1110,49 @@ namespace UnrealBuildTool
 
 		public override bool PrepTargetForDeployment(TargetReceipt Receipt)
 		{
-			return PrepTargetForDeployment(Receipt, false);
+			List<string> UPLScripts = CollectPluginDataPaths(Receipt.AdditionalProperties);
+			VersionNumber SdkVersion = GetSdkVersion(Receipt);
+			return PrepTargetForDeployment(Receipt.ProjectFile, Receipt.TargetName, Receipt.Platform, Receipt.Configuration, UPLScripts, SdkVersion, false);
 		}
 
-		public bool PrepTargetForDeployment(TargetReceipt Receipt, bool bCreateStubIPA)
+		public bool PrepTargetForDeployment(FileReference ProjectFile, string TargetName, UnrealTargetPlatform Platform, UnrealTargetConfiguration Configuration, List<string> UPLScripts, VersionNumber SdkVersion, bool bCreateStubIPA)
 		{
 			string SubDir = GetTargetPlatformName();
 
-			string GameName = Receipt.TargetName;
-			string ProjectDirectory = (DirectoryReference.FromFile(Receipt.ProjectFile) ?? UnrealBuildTool.EngineDirectory).FullName;
+			string GameName = TargetName;
+			string ProjectDirectory = (DirectoryReference.FromFile(ProjectFile) ?? UnrealBuildTool.EngineDirectory).FullName;
 			string BuildPath = (GameName == "UE4Game" ? "../../Engine" : ProjectDirectory) + "/Binaries/" + SubDir;
 			bool bIsUE4Game = GameName.Contains("UE4Game");
 
-			Console.WriteLine("1 GameName: {0}, ProjectName: {1}", GameName, (Receipt.ProjectFile == null) ? "" : Path.GetFileNameWithoutExtension(Receipt.ProjectFile.FullName));
+			Console.WriteLine("1 GameName: {0}, ProjectName: {1}", GameName, (ProjectFile == null) ? "" : Path.GetFileNameWithoutExtension(ProjectFile.FullName));
 
 			string DecoratedGameName;
-			if (Receipt.Configuration == UnrealTargetConfiguration.Development)
+			if (Configuration == UnrealTargetConfiguration.Development)
 			{
 				DecoratedGameName = GameName;
 			}
 			else
 			{
-				DecoratedGameName = String.Format("{0}-{1}-{2}", GameName, Receipt.Platform.ToString(), Receipt.Configuration.ToString());
+				DecoratedGameName = String.Format("{0}-{1}-{2}", GameName, Platform.ToString(), Configuration.ToString());
 			}
 
 			if (BuildHostPlatform.Current.Platform == UnrealTargetPlatform.Mac && Environment.GetEnvironmentVariable("UBT_NO_POST_DEPLOY") != "true")
 			{
-				return PrepForUATPackageOrDeploy(Receipt.Configuration, Receipt.ProjectFile, GameName, ProjectDirectory, BuildPath + "/" + DecoratedGameName, "../../Engine", false, "", false, bCreateStubIPA, Receipt);
+				return PrepForUATPackageOrDeploy(Configuration, ProjectFile, GameName, ProjectDirectory, BuildPath + "/" + DecoratedGameName, "../../Engine", false, "", false, bCreateStubIPA, UPLScripts, SdkVersion);
 			}
 			else
 			{
 				// @todo tvos merge: This used to copy the bundle back - where did that code go? It needs to be fixed up for TVOS directories
 				bool bSupportPortrait, bSupportLandscape, bSkipIcons;
-				GeneratePList(Receipt.ProjectFile, Receipt.Configuration, ProjectDirectory, bIsUE4Game, GameName, (Receipt.ProjectFile == null) ? "" : Path.GetFileNameWithoutExtension(Receipt.ProjectFile.FullName), "../../Engine", "", Receipt, out bSupportPortrait, out bSupportLandscape, out bSkipIcons);
+				GeneratePList(ProjectFile, Configuration, ProjectDirectory, bIsUE4Game, GameName, (ProjectFile == null) ? "" : Path.GetFileNameWithoutExtension(ProjectFile.FullName), "../../Engine", "", UPLScripts, SdkVersion, out bSupportPortrait, out bSupportLandscape, out bSkipIcons);
 			}
 			return true;
 		}
 
-		private List<string> CollectPluginDataPaths(TargetReceipt Receipt)
+		public static List<string> CollectPluginDataPaths(List<ReceiptProperty> ReceiptProperties)
 		{
 			List<string> PluginExtras = new List<string>();
-			if (Receipt == null)
+			if (ReceiptProperties == null)
 			{
 				Log.TraceInformation("Receipt is NULL");
 				//Log.TraceInformation("Receipt is NULL");
@@ -1151,7 +1160,7 @@ namespace UnrealBuildTool
 			}
 
 			// collect plugin extra data paths from target receipt
-			IEnumerable<ReceiptProperty> Results = Receipt.AdditionalProperties.Where(x => x.Name == "IOSPlugin");
+			IEnumerable<ReceiptProperty> Results = ReceiptProperties.Where(x => x.Name == "IOSPlugin");
 			foreach(ReceiptProperty Property in Results)
 			{
 				// Keep only unique paths
diff --git a/Engine/Source/Programs/UnrealBuildTool/Platform/TVOS/UEDeployTVOS.cs b/Engine/Source/Programs/UnrealBuildTool/Platform/TVOS/UEDeployTVOS.cs
--- a/Engine/Source/Programs/UnrealBuildTool/Platform/TVOS/UEDeployTVOS.cs
+++ b/Engine/Source/Programs/UnrealBuildTool/Platform/TVOS/UEDeployTVOS.cs
@@ -331,7 +331,7 @@ namespace UnrealBuildTool
 			return bSkipDefaultPNGs;
 		}
 
-		public override bool GeneratePList(FileReference ProjectFile, UnrealTargetConfiguration Config, string ProjectDirectory, bool bIsUE4Game, string GameName, string ProjectName, string InEngineDir, string AppDirectory, TargetReceipt Receipt, out bool bSupportsPortrait, out bool bSupportsLandscape, out bool bSkipIcons)
+		public override bool GeneratePList(FileReference ProjectFile, UnrealTargetConfiguration Config, string ProjectDirectory, bool bIsUE4Game, string GameName, string ProjectName, string InEngineDir, string AppDirectory, List<string> UPLScripts, VersionNumber SdkVersion, out bool bSupportsPortrait, out bool bSupportsLandscape, out bool bSkipIcons)
 		{
 			bSupportsLandscape = bSupportsPortrait = true;
             bSkipIcons = true;
diff --git a/Engine/Source/Programs/UnrealBuildTool/ProjectFiles/Xcode/XcodeProject.cs b/Engine/Source/Programs/UnrealBuildTool/ProjectFiles/Xcode/XcodeProject.cs
--- a/Engine/Source/Programs/UnrealBuildTool/ProjectFiles/Xcode/XcodeProject.cs
+++ b/Engine/Source/Programs/UnrealBuildTool/ProjectFiles/Xcode/XcodeProject.cs
@@ -881,7 +881,8 @@ namespace UnrealBuildTool
                             bool bSupportPortrait, bSupportLandscape, bSkipIcons;
 							TargetReceipt Receipt;
 							TargetReceipt.TryRead(ReceiptFilename, out Receipt);
-							UEDeployIOS.GenerateIOSPList(ProjectFile, Config.BuildConfig, ProjectPath.FullName, bIsUE4Game, GameName, Config.BuildTarget, EngineDir.FullName, ProjectPath + "/Binaries/IOS/Payload", Receipt, null, out bSupportPortrait, out bSupportLandscape, out bSkipIcons);
+							VersionNumber SdkVersion = UEDeployIOS.GetSdkVersion(Receipt);
+							UEDeployIOS.GenerateIOSPList(ProjectFile, Config.BuildConfig, ProjectPath.FullName, bIsUE4Game, GameName, Config.BuildTarget, EngineDir.FullName, ProjectPath + "/Binaries/IOS/Payload", SdkVersion, null, out bSupportPortrait, out bSupportLandscape, out bSkipIcons);
 						}
 						if (bCreateTVOSInfoPlist)
 						{
diff --git a/Engine/Source/Programs/UnrealBuildTool/System/LinkEnvironment.cs b/Engine/Source/Programs/UnrealBuildTool/System/LinkEnvironment.cs
--- a/Engine/Source/Programs/UnrealBuildTool/System/LinkEnvironment.cs
+++ b/Engine/Source/Programs/UnrealBuildTool/System/LinkEnvironment.cs
@@ -292,6 +292,11 @@ namespace UnrealBuildTool
 		/// </summary>
 		public string ModuleDefinitionFile;
 
+		/// <summary>
+		/// All the additional properties from the modules linked into this binary
+		/// </summary>
+		public List<ReceiptProperty> AdditionalProperties = new List<ReceiptProperty>();
+
 		/// <summary>
 		/// Default constructor.
 		/// </summary>
@@ -357,6 +362,7 @@ namespace UnrealBuildTool
 			DefaultResourceFiles.AddRange(Other.DefaultResourceFiles);
 			CommonResourceFiles.AddRange(Other.CommonResourceFiles);
 			ModuleDefinitionFile = Other.ModuleDefinitionFile;
+			AdditionalProperties.AddRange(Other.AdditionalProperties);
         }
 	}
 }
diff --git a/Engine/Source/Programs/UnrealBuildTool/System/VersionNumber.cs b/Engine/Source/Programs/UnrealBuildTool/System/VersionNumber.cs
--- a/Engine/Source/Programs/UnrealBuildTool/System/VersionNumber.cs
+++ b/Engine/Source/Programs/UnrealBuildTool/System/VersionNumber.cs
@@ -11,6 +11,7 @@ namespace UnrealBuildTool
 	/// <summary>
 	/// Stores a numeric version consisting of any number of components.
 	/// </summary>
+	[Serializable]
 	class VersionNumber : IComparable<VersionNumber>
 	{
 		/// <summary>
diff --git a/Engine/Source/Programs/UnrealBuildTool/UnrealBuildTool.cs b/Engine/Source/Programs/UnrealBuildTool/UnrealBuildTool.cs
--- a/Engine/Source/Programs/UnrealBuildTool/UnrealBuildTool.cs
+++ b/Engine/Source/Programs/UnrealBuildTool/UnrealBuildTool.cs
@@ -1346,10 +1346,6 @@ namespace UnrealBuildTool
 								// if the build succeeded, write the receipts and do any needed syncing
 								if (bSuccess)
 								{
-									foreach (UEBuildTarget Target in Targets)
-									{
-										UEBuildPlatform.GetBuildPlatform(Target.Platform).PostBuildSync(Target);
-									}
 									if (ActionsToExecute.Count == 0 && BuildConfiguration.bSkipLinkingWhenNothingToCompile)
 									{
 										BuildResult = ECompilationResult.UpToDate;
diff --git a/Engine/Source/Programs/UnrealBuildTool/UnrealBuildTool.csproj b/Engine/Source/Programs/UnrealBuildTool/UnrealBuildTool.csproj
--- a/Engine/Source/Programs/UnrealBuildTool/UnrealBuildTool.csproj
+++ b/Engine/Source/Programs/UnrealBuildTool/UnrealBuildTool.csproj
@@ -158,6 +158,7 @@
     <Compile Include="Executors\SNDBS.cs" />
     <Compile Include="Executors\XGE.cs" />
     <Compile Include="Modes\WriteMetadataMode.cs" />
+    <Compile Include="Platform\IOS\IOSPostBuildSyncMode.cs" />
     <Compile Include="Platform\IOS\MobileProvisionContents.cs" />
     <Compile Include="Platform\Lumin\UEDeployLumin.cs" />
     <Compile Include="ProjectFiles\CLion\CLionGenerator.cs" />
