UBT: Remove some initialization code which is no longer used.

#rnx

diff --git a/Engine/Source/Programs/UnrealBuildTool/UnrealBuildTool.cs b/Engine/Source/Programs/UnrealBuildTool/UnrealBuildTool.cs
--- a/Engine/Source/Programs/UnrealBuildTool/UnrealBuildTool.cs
+++ b/Engine/Source/Programs/UnrealBuildTool/UnrealBuildTool.cs
@@ -309,7 +309,7 @@ namespace UnrealBuildTool
 			return false;
 		}
 
-		private static bool ParseRocketCommandlineArg(string InArg, ref string OutGameName, ref FileReference ProjectFile)
+		private static bool ParseRocketCommandlineArg(string InArg, ref FileReference ProjectFile)
 		{
 			string LowercaseArg = InArg.ToLowerInvariant();
 
@@ -334,7 +334,6 @@ namespace UnrealBuildTool
 			if (ProjectArg != null && ProjectArg.Length > 0)
 			{
 				string ProjectPath = Path.GetDirectoryName(ProjectArg);
-				OutGameName = Path.GetFileNameWithoutExtension(ProjectArg);
 
 				if (Path.IsPathRooted(ProjectPath) == false)
 				{
@@ -463,8 +462,7 @@ namespace UnrealBuildTool
 					FileReference ProjectFile = null;
 					foreach (string Arg in Arguments)
 					{
-						string TempGameName = null;
-						if (ParseRocketCommandlineArg(Arg, ref TempGameName, ref ProjectFile) == true)
+						if (ParseRocketCommandlineArg(Arg, ref ProjectFile) == true)
 						{
 							// This is to allow relative paths for the project file
 							Log.TraceVerbose("UBT Running for Rocket: " + ProjectFile);
@@ -501,39 +499,15 @@ namespace UnrealBuildTool
 					Log.TraceVerbose("UnrealBuildTool (DEBUG OUTPUT MODE)");
 					Log.TraceVerbose("Command-line: {0}", String.Join(" ", Arguments));
 
-					string GameName = null;
 					bool bSpecificModulesOnly = false;
-
-					// We need to be able to identify if one of the arguments is the platform...
-					UnrealTargetPlatform CheckPlatform = UnrealTargetPlatform.Unknown;
-					foreach (string Argument in Arguments)
-					{
-						UnrealTargetPlatform ParsedPlatform = UEBuildPlatform.ConvertStringToPlatform(Argument);
-						if (ParsedPlatform != UnrealTargetPlatform.Unknown)
-						{
-							CheckPlatform = ParsedPlatform;
-						}
-					}
-
-					// @todo ubtmake: remove this when building with RPCUtility works
-					// @todo tvos merge: Check the change to this line, not clear why. Is TVOS needed here?
-					if (BuildHostPlatform.Current.Platform != UnrealTargetPlatform.Mac && (CheckPlatform == UnrealTargetPlatform.Mac || CheckPlatform == UnrealTargetPlatform.IOS || CheckPlatform == UnrealTargetPlatform.TVOS))
-					{
-						BuildConfiguration.bUseUBTMakefiles = false;
-					}
-
 					foreach (string Arg in Arguments)
 					{
 						string LowercaseArg = Arg.ToLowerInvariant();
-						if (ProjectFile == null && ParseRocketCommandlineArg(Arg, ref GameName, ref ProjectFile))
+						if (ProjectFile == null && ParseRocketCommandlineArg(Arg, ref ProjectFile))
 						{
 							// Already handled at startup. Calling now just to properly set the game name
 							continue;
 						}
-						else if (LowercaseArg == "development" || LowercaseArg == "debug" || LowercaseArg == "shipping" || LowercaseArg == "test" || LowercaseArg == "debuggame")
-						{
-							//ConfigName = Arg;
-						}
 						else if (LowercaseArg.StartsWith("-modulewithsuffix="))
 						{
 							bSpecificModulesOnly = true;
@@ -543,30 +517,10 @@ namespace UnrealBuildTool
 						{
 							BuildConfiguration.bIgnoreJunk = true;
 						}
-						else if (LowercaseArg == "-define")
-						{
-							// Skip -define
-						}
-						else if (CheckPlatform.ToString().ToLowerInvariant() == LowercaseArg)
-						{
-							// It's the platform set...
-							//PlatformName = Arg;
-						}
 						else if (LowercaseArg == "-vsdebugandroid")
 						{
 							AndroidProjectGenerator.VSDebugCommandLineOptionPresent = true;
 						}
-						else
-						{
-							// This arg may be a game name. Check for the existence of a game folder with this name.
-							// "Engine" is not a valid game name.
-							if (LowercaseArg != "engine" && Arg.IndexOfAny(Path.GetInvalidPathChars()) == -1 && Arg.IndexOfAny(new char[] { ':', Path.DirectorySeparatorChar, Path.AltDirectorySeparatorChar }) == -1 &&
-								DirectoryReference.Exists(DirectoryReference.Combine(RootDirectory, Arg, "Config")))
-							{
-								GameName = Arg;
-								Log.TraceVerbose("CommandLine: Found game name '{0}'", GameName);
-							}
-						}
 					}
 
 					// Find and register all tool chains, build platforms, etc. that are present
