UBT: Cleaning up a few more initialization bits.

#rnx

diff --git a/Engine/Source/Programs/UnrealBuildTool/Configuration/BuildConfiguration.cs b/Engine/Source/Programs/UnrealBuildTool/Configuration/BuildConfiguration.cs
--- a/Engine/Source/Programs/UnrealBuildTool/Configuration/BuildConfiguration.cs
+++ b/Engine/Source/Programs/UnrealBuildTool/Configuration/BuildConfiguration.cs
@@ -149,6 +149,7 @@ namespace UnrealBuildTool
 		/// Whether to skip checking for files identified by the junk manifest
 		/// </summary>
 		[XmlConfigFile]
+		[CommandLine("-IgnoreJunk")]
 		public bool bIgnoreJunk = false;
 
 		/// <summary>
diff --git a/Engine/Source/Programs/UnrealBuildTool/UnrealBuildTool.cs b/Engine/Source/Programs/UnrealBuildTool/UnrealBuildTool.cs
--- a/Engine/Source/Programs/UnrealBuildTool/UnrealBuildTool.cs
+++ b/Engine/Source/Programs/UnrealBuildTool/UnrealBuildTool.cs
@@ -292,7 +292,7 @@ namespace UnrealBuildTool
 
 			for(int Idx = 0; Idx < Arguments.Count; Idx++)
 			{
-				if(Arguments[Idx][0] != '-' && Arguments[Idx].EndsWith(".uproject"))
+				if(Arguments[Idx][0] != '-' && Arguments[Idx].EndsWith(".uproject", StringComparison.OrdinalIgnoreCase))
 				{
 					Arguments.MarkAsUsed(Idx);
 					ProjectFile = new FileReference(Arguments[Idx]);
@@ -303,63 +303,13 @@ namespace UnrealBuildTool
 			if(IsProjectInstalled())
 			{
 				ProjectFile = InstalledProjectFile;
+				return true;
 			}
 
 			ProjectFile = null;
 			return false;
 		}
 
-		private static bool ParseRocketCommandlineArg(string InArg, ref FileReference ProjectFile)
-		{
-			string LowercaseArg = InArg.ToLowerInvariant();
-
-			string ProjectArg = null;
-			if (LowercaseArg.StartsWith("-project="))
-			{
-				ProjectArg = InArg.Substring(9).Trim(new Char[] { ' ', '"', '\'' });
-			}
-			else if (LowercaseArg.EndsWith(".uproject"))
-			{
-				ProjectArg = InArg;
-			}
-			else if (LowercaseArg.StartsWith("-remoteini="))
-			{
-				RemoteIniPath = InArg.Substring(11);
-			}
-			else
-			{
-				return false;
-			}
-
-			if (ProjectArg != null && ProjectArg.Length > 0)
-			{
-				string ProjectPath = Path.GetDirectoryName(ProjectArg);
-
-				if (Path.IsPathRooted(ProjectPath) == false)
-				{
-					if (Directory.Exists(ProjectPath) == false)
-					{
-						// If it is *not* rooted, then we potentially received a path that is 
-						// relative to Engine/Binaries/[PLATFORM] rather than Engine/Source
-						if (ProjectPath.StartsWith("../") || ProjectPath.StartsWith("..\\"))
-						{
-							string AlternativeProjectpath = ProjectPath.Substring(3);
-							if (Directory.Exists(AlternativeProjectpath))
-							{
-								ProjectPath = AlternativeProjectpath;
-								ProjectArg = ProjectArg.Substring(3);
-								Debug.Assert(ProjectArg.Length > 0);
-							}
-						}
-					}
-				}
-
-				ProjectFile = new FileReference(ProjectArg);
-			}
-
-			return true;
-		}
-
 		/// <summary>
 		/// UBT startup order is fairly fragile, and relies on globals that may or may not be safe to use yet.
 		/// This function is for super early startup stuff that should not access Configuration classes (anything loaded by XmlConfg).
@@ -395,19 +345,8 @@ namespace UnrealBuildTool
 					// This is critical to be done early so any code that relies on the current directory being Engine/Source will work.
 					DirectoryReference.SetCurrentDirectory(EngineSourceDirectory);
 
-					// Parse the argument for overriding the XML configuration file location (for remote builds)
-					FileReference XmlConfigCache = null;
-					foreach (string Argument in Arguments)
-					{
-						const string Prefix = "-XmlConfigCache=";
-						if(Argument.StartsWith(Prefix, StringComparison.OrdinalIgnoreCase))
-						{
-							XmlConfigCache = new FileReference(Argument.Substring(Prefix.Length));
-							break;
-						}
-					}
-
 					// Read the XML configuration files
+					FileReference XmlConfigCache = CmdLine.GetFileReferenceOrDefault("-XmlConfigCache=", null);
 					XmlConfig.ReadConfigFiles(XmlConfigCache);
 
 					// Create the build configuration object, and read the settings
@@ -444,6 +383,9 @@ namespace UnrealBuildTool
 						}
 					}
 
+					// Parse the remote INI setting
+					CmdLine.TryGetValue("-RemoteIni=", out RemoteIniPath);
+
 					// Create the log file, and flush the startup listener to it
 					if (!String.IsNullOrEmpty(BuildConfiguration.LogFileName) && !Log.HasFileWriter())
 					{
@@ -458,31 +400,10 @@ namespace UnrealBuildTool
 					}
 					Trace.Listeners.Remove(StartupListener);
 
-					// Parse rocket-specific arguments.
-					FileReference ProjectFile = null;
-					foreach (string Arg in Arguments)
-					{
-						if (ParseRocketCommandlineArg(Arg, ref ProjectFile) == true)
-						{
-							// This is to allow relative paths for the project file
-							Log.TraceVerbose("UBT Running for Rocket: " + ProjectFile);
-							break;
-						}
-					}
-
-					// Read the project file from the installed project text file
-					if (ProjectFile == null)
-					{
-						FileReference InstalledProjectFile = FileReference.Combine(RootDirectory, "Engine", "Build", "InstalledProjectBuild.txt");
-						if (FileReference.Exists(InstalledProjectFile))
-						{
-							ProjectFile = FileReference.Combine(UnrealBuildTool.RootDirectory, File.ReadAllText(InstalledProjectFile.FullName).Trim());
-						}
-					}
-
 					// Build the list of game projects that we know about. When building from the editor (for hot-reload) or for projects from installed builds, we require the 
 					// project file to be passed in. Otherwise we scan for projects in directories named in UE4Games.uprojectdirs.
-					if (ProjectFile != null)
+					FileReference ProjectFile;
+					if (TryParseProjectFileArgument(CmdLine, out ProjectFile))
 					{
 						UProjectInfo.AddProject(ProjectFile);
 					}
@@ -496,27 +417,15 @@ namespace UnrealBuildTool
 
 					DateTime BasicInitStartTime = DateTime.UtcNow;
 
-					Log.TraceVerbose("UnrealBuildTool (DEBUG OUTPUT MODE)");
-					Log.TraceVerbose("Command-line: {0}", String.Join(" ", Arguments));
-
 					bool bSpecificModulesOnly = false;
 					foreach (string Arg in Arguments)
 					{
 						string LowercaseArg = Arg.ToLowerInvariant();
-						if (ProjectFile == null && ParseRocketCommandlineArg(Arg, ref ProjectFile))
-						{
-							// Already handled at startup. Calling now just to properly set the game name
-							continue;
-						}
-						else if (LowercaseArg.StartsWith("-modulewithsuffix="))
+						if (LowercaseArg.StartsWith("-modulewithsuffix="))
 						{
 							bSpecificModulesOnly = true;
 							continue;
 						}
-						else if (LowercaseArg == "-ignorejunk")
-						{
-							BuildConfiguration.bIgnoreJunk = true;
-						}
 						else if (LowercaseArg == "-vsdebugandroid")
 						{
 							AndroidProjectGenerator.VSDebugCommandLineOptionPresent = true;
