Add missing copyright notice.

#rnx

diff --git a/Engine/Source/Programs/UnrealBuildTool/Platform/IOS/IOSPostBuildSyncMode.cs b/Engine/Source/Programs/UnrealBuildTool/Platform/IOS/IOSPostBuildSyncMode.cs
--- a/Engine/Source/Programs/UnrealBuildTool/Platform/IOS/IOSPostBuildSyncMode.cs
+++ b/Engine/Source/Programs/UnrealBuildTool/Platform/IOS/IOSPostBuildSyncMode.cs
@@ -1,3 +1,5 @@
+// Copyright 1998-2018 Epic Games, Inc. All Rights Reserved.
+
 using System;
 using System.Collections.Generic;
 using System.Linq;
