UBT: Pass specific arguments to ExecuteHeaderToolIfNecessary, rather than the whole UEBuildTarget object. Effort to reduce the number of things that may not be available when using a makefile.

#rnx

diff --git a/Engine/Source/Programs/UnrealBuildTool/Configuration/UEBuildTarget.cs b/Engine/Source/Programs/UnrealBuildTool/Configuration/UEBuildTarget.cs
--- a/Engine/Source/Programs/UnrealBuildTool/Configuration/UEBuildTarget.cs
+++ b/Engine/Source/Programs/UnrealBuildTool/Configuration/UEBuildTarget.cs
@@ -1782,10 +1782,11 @@ namespace UnrealBuildTool
 				// files that are injected as top level prerequisites.  If UHT only emitted included header files, we wouldn't need to run it during the Gather phase at all.
 				if (UObjectModules.Count > 0)
 				{
-					// Execute the header tool
 					FileReference ModuleInfoFileName = FileReference.Combine(ProjectIntermediateDirectory, GetTargetName() + ".uhtmanifest");
-					ECompilationResult UHTResult = ECompilationResult.OtherCompilationError;
-					if (!ExternalExecution.ExecuteHeaderToolIfNecessary(BuildConfiguration, this, GlobalCompileEnvironment, UObjectModules, ModuleInfoFileName, ref UHTResult, HotReload, true, bIsAssemblingBuild))
+
+					// Execute the header tool
+					ECompilationResult UHTResult = ExternalExecution.ExecuteHeaderToolIfNecessary(BuildConfiguration, ProjectFile, TargetName, TargetType, bHasProjectScriptPlugin, GlobalCompileEnvironment, UObjectModules, ModuleInfoFileName, HotReload, true, bIsAssemblingBuild);
+					if(!UHTResult.Succeeded())
 					{
 						Log.TraceInformation(String.Format("Error: UnrealHeaderTool failed for target '{0}' (platform: {1}, module info: {2}, exit code: {3} ({4})).", GetTargetName(), Platform.ToString(), ModuleInfoFileName, UHTResult.ToString(), (int)UHTResult));
 						return UHTResult;
diff --git a/Engine/Source/Programs/UnrealBuildTool/System/ExternalExecution.cs b/Engine/Source/Programs/UnrealBuildTool/System/ExternalExecution.cs
--- a/Engine/Source/Programs/UnrealBuildTool/System/ExternalExecution.cs
+++ b/Engine/Source/Programs/UnrealBuildTool/System/ExternalExecution.cs
@@ -313,11 +313,11 @@ namespace UnrealBuildTool
 		public string ExternalDependenciesFile; // File to contain additional dependencies that the generated code depends on
 		public List<Module> Modules;
 
-		public UHTManifest(UEBuildTarget Target, string InRootLocalPath, string InExternalDependenciesFile, List<Module> InModules)
+		public UHTManifest(string InTargetName, TargetType InTargetType, string InRootLocalPath, string InExternalDependenciesFile, List<Module> InModules)
 		{
-			IsGameTarget = (Target.TargetType != TargetType.Program);
+			IsGameTarget = (InTargetType != TargetType.Program);
 			RootLocalPath = InRootLocalPath;
-			TargetName = Target.GetTargetName();
+			TargetName = InTargetName;
 			ExternalDependenciesFile = InExternalDependenciesFile;
 			Modules = InModules;
 		}
@@ -1087,7 +1087,7 @@ namespace UnrealBuildTool
 		/// Builds and runs the header tool and touches the header directories.
 		/// Performs any early outs if headers need no changes, given the UObject modules, tool path, game name, and configuration
 		/// </summary>
-		public static bool ExecuteHeaderToolIfNecessary(BuildConfiguration BuildConfiguration, UEBuildTarget Target, CppCompileEnvironment GlobalCompileEnvironment, List<UHTModuleInfo> UObjectModules, FileReference ModuleInfoFileName, ref ECompilationResult UHTResult, EHotReload HotReload, bool bIsGatheringBuild, bool bIsAssemblingBuild)
+		public static ECompilationResult ExecuteHeaderToolIfNecessary(BuildConfiguration BuildConfiguration, FileReference ProjectFile, string TargetName, TargetType TargetType, bool bHasProjectScriptPlugin, CppCompileEnvironment GlobalCompileEnvironment, List<UHTModuleInfo> UObjectModules, FileReference ModuleInfoFileName, EHotReload HotReload, bool bIsGatheringBuild, bool bIsAssemblingBuild)
 		{
 			if (ProgressWriter.bWriteMarkup)
 			{
@@ -1096,14 +1096,14 @@ namespace UnrealBuildTool
 			using (ProgressWriter Progress = new ProgressWriter("Generating code...", false))
 			{
 				// We never want to try to execute the header tool when we're already trying to build it!
-				bool bIsBuildingUHT = Target.GetTargetName().Equals("UnrealHeaderTool", StringComparison.InvariantCultureIgnoreCase);
+				bool bIsBuildingUHT = TargetName.Equals("UnrealHeaderTool", StringComparison.InvariantCultureIgnoreCase);
 
 				string RootLocalPath = UnrealBuildTool.RootDirectory.FullName;
 
 				UnrealTargetConfiguration UHTConfig = BuildConfiguration.bForceDebugUnrealHeaderTool ? UnrealTargetConfiguration.Debug : UnrealTargetConfiguration.Development;
 
 				// Figure out the receipt path
-				FileReference HeaderToolReceipt = GetHeaderToolReceiptFile(Target.ProjectFile, UHTConfig, Target.bHasProjectScriptPlugin);
+				FileReference HeaderToolReceipt = GetHeaderToolReceiptFile(ProjectFile, UHTConfig, bHasProjectScriptPlugin);
 
 				// check if UHT is out of date
 				DateTime HeaderToolTimestamp = DateTime.MaxValue;
@@ -1127,12 +1127,12 @@ namespace UnrealBuildTool
 				{
 					Modules.Add(new UHTManifest.Module(UObjectModule));
 				}
-				UHTManifest Manifest = new UHTManifest(Target, RootLocalPath, ExternalDependenciesFile.FullName, Modules);
+				UHTManifest Manifest = new UHTManifest(TargetName, TargetType, RootLocalPath, ExternalDependenciesFile.FullName, Modules);
 
 				if (!bIsBuildingUHT && bUHTNeedsToRun)
 				{
 					// Always build UnrealHeaderTool if header regeneration is required, unless we're running within an installed ecosystem or hot-reloading
-					if ((!UnrealBuildTool.IsEngineInstalled() || Target.bHasProjectScriptPlugin) &&
+					if ((!UnrealBuildTool.IsEngineInstalled() || bHasProjectScriptPlugin) &&
 						!BuildConfiguration.bDoNotBuildUHT &&
 						HotReload != EHotReload.FromIDE &&
 						!(bHaveHeaderTool && !bIsGatheringBuild && bIsAssemblingBuild))	// If running in "assembler only" mode, we assume UHT is already up to date for much faster iteration!
@@ -1170,9 +1170,9 @@ namespace UnrealBuildTool
 						UBTArguments.Append(" -ignorejunk");
 
 						// Add UHT plugins to UBT command line as external plugins
-						if(Target.bHasProjectScriptPlugin && Target.ProjectFile != null)
+						if(bHasProjectScriptPlugin && ProjectFile != null)
 						{
-							UBTArguments.AppendFormat(" -project=\"{0}\"", Target.ProjectFile);
+							UBTArguments.AppendFormat(" -project=\"{0}\"", ProjectFile);
 						}
 
 						// Add any global override for the compiler
@@ -1193,13 +1193,13 @@ namespace UnrealBuildTool
 
 						if (RunExternalDotNETExecutable(UnrealBuildTool.GetUBTPath(), UBTArguments.ToString()) != 0)
 						{
-							return false;
+							return ECompilationResult.OtherCompilationError;
 						}
 					}
 
 					Progress.Write(1, 3);
 
-					string ActualTargetName = String.IsNullOrEmpty(Target.GetTargetName()) ? "UE4" : Target.GetTargetName();
+					string ActualTargetName = String.IsNullOrEmpty(TargetName) ? "UE4" : TargetName;
 					Log.TraceInformation("Parsing headers for {0}", ActualTargetName);
 
 					FileReference HeaderToolPath = GetHeaderToolPath(HeaderToolReceipt);
@@ -1212,7 +1212,7 @@ namespace UnrealBuildTool
 					Directory.CreateDirectory(ModuleInfoFileName.Directory.FullName);
 					System.IO.File.WriteAllText(ModuleInfoFileName.FullName, fastJSON.JSON.Instance.ToJSON(Manifest, new fastJSON.JSONParameters { UseExtensions = false }));
 
-					string CmdLine = (Target.ProjectFile != null) ? "\"" + Target.ProjectFile.FullName + "\"" : Target.GetTargetName();
+					string CmdLine = (ProjectFile != null) ? "\"" + ProjectFile.FullName + "\"" : TargetName;
 					CmdLine += " \"" + ModuleInfoFileName + "\" -LogCmds=\"loginit warning, logexit warning, logdatabase error\" -Unattended -WarningsAsErrors";
 					if (UnrealBuildTool.IsEngineInstalled())
 					{
@@ -1224,16 +1224,11 @@ namespace UnrealBuildTool
 						CmdLine += " -FailIfGeneratedCodeChanges";
 					}
 
-					if (Target.Rules != null && !Target.Rules.bCompileAgainstEngine)
-					{
-						CmdLine += " -NoEnginePlugins";
-					}
-
 					Log.TraceInformation("  Running UnrealHeaderTool {0}", CmdLine);
 
 					Stopwatch s = new Stopwatch();
 					s.Start();
-					UHTResult = (ECompilationResult)RunExternalNativeExecutable(ExternalExecution.GetHeaderToolPath(HeaderToolReceipt), CmdLine);
+					ECompilationResult UHTResult = (ECompilationResult)RunExternalNativeExecutable(ExternalExecution.GetHeaderToolPath(HeaderToolReceipt), CmdLine);
 					s.Stop();
 
 					if (UHTResult != ECompilationResult.Succeeded)
@@ -1254,7 +1249,7 @@ namespace UnrealBuildTool
 						{
 							Log.TraceInformation(String.Format("UnrealHeaderTool failed with exit code 0x{0:X} - check that UE4 prerequisites are installed.", (int)UHTResult));
 						}
-						return false;
+						return UHTResult;
 					}
 
 					Log.TraceInformation("Reflection code generated for {0} in {1} seconds", ActualTargetName, s.Elapsed.TotalSeconds);
@@ -1286,7 +1281,7 @@ namespace UnrealBuildTool
 			{
 				Log.WriteLine(LogEventType.Console, "@progress pop");
 			}
-			return true;
+			return ECompilationResult.Succeeded;
 		}
 	}
 }
diff --git a/Engine/Source/Programs/UnrealBuildTool/UnrealBuildTool.cs b/Engine/Source/Programs/UnrealBuildTool/UnrealBuildTool.cs
--- a/Engine/Source/Programs/UnrealBuildTool/UnrealBuildTool.cs
+++ b/Engine/Source/Programs/UnrealBuildTool/UnrealBuildTool.cs
@@ -1155,8 +1155,8 @@ namespace UnrealBuildTool
 											{
 												// Execute the header tool
 												FileReference ModuleInfoFileName = FileReference.Combine(Target.ProjectIntermediateDirectory, Target.GetTargetName() + ".uhtmanifest");
-												ECompilationResult UHTResult = ECompilationResult.OtherCompilationError;
-												if (!ExternalExecution.ExecuteHeaderToolIfNecessary(BuildConfiguration, Target, GlobalCompileEnvironment: null, UObjectModules: TargetUObjectModules, ModuleInfoFileName: ModuleInfoFileName, UHTResult: ref UHTResult, HotReload: HotReload, bIsGatheringBuild: bIsGatheringBuild, bIsAssemblingBuild: bIsAssemblingBuild))
+												ECompilationResult UHTResult = ExternalExecution.ExecuteHeaderToolIfNecessary(BuildConfiguration, Target.ProjectFile, Target.TargetName, Target.TargetType, Target.bHasProjectScriptPlugin, GlobalCompileEnvironment: null, UObjectModules: TargetUObjectModules, ModuleInfoFileName: ModuleInfoFileName, HotReload: HotReload, bIsGatheringBuild: bIsGatheringBuild, bIsAssemblingBuild: bIsAssemblingBuild);
+												if(UHTResult.Succeeded())
 												{
 													Log.TraceInformation("UnrealHeaderTool failed for target '" + Target.GetTargetName() + "' (platform: " + Target.Platform.ToString() + ", module info: " + ModuleInfoFileName + ").");
 													BuildResult = UHTResult;
