UBT: Remove cross-platform override for modifying Mac-specific output paths.

#rnx

diff --git a/Engine/Source/Programs/UnrealBuildTool/Configuration/UEBuildTarget.cs b/Engine/Source/Programs/UnrealBuildTool/Configuration/UEBuildTarget.cs
--- a/Engine/Source/Programs/UnrealBuildTool/Configuration/UEBuildTarget.cs
+++ b/Engine/Source/Programs/UnrealBuildTool/Configuration/UEBuildTarget.cs
@@ -916,13 +916,13 @@ namespace UnrealBuildTool
 				ProjectDirectory = ProjectFile.Directory;
 			}
 			else if (Rules.File.IsUnderDirectory(UnrealBuildTool.EnterpriseDirectory))
-				{
-					ProjectDirectory = UnrealBuildTool.EnterpriseDirectory;
-				}
-				else
-				{
-					ProjectDirectory = UnrealBuildTool.EngineDirectory;
-				}
+			{
+				ProjectDirectory = UnrealBuildTool.EnterpriseDirectory;
+			}
+			else
+			{
+				ProjectDirectory = UnrealBuildTool.EngineDirectory;
+			}
 
 			// Build the project intermediate directory
 			if(bUseSharedBuildEnvironment && TargetRulesFile.IsUnderDirectory(UnrealBuildTool.EngineDirectory))
@@ -1525,7 +1525,7 @@ namespace UnrealBuildTool
 			SetupGlobalEnvironment(TargetToolChain, GlobalCompileEnvironment, GlobalLinkEnvironment);
 
 			// Create all the binaries and modules
-			PreBuildSetup(TargetToolChain);
+			PreBuildSetup();
 
 			// Save off the original list of binaries. We'll use this to figure out which PCHs to create later, to avoid switching PCHs when compiling single modules.
 			List<UEBuildBinary> OriginalBinaries = Binaries;
@@ -2214,7 +2214,7 @@ namespace UnrealBuildTool
 		/// <summary>
 		/// Setup target before build. This method finds dependencies, sets up global environment etc.
 		/// </summary>
-		public void PreBuildSetup(UEToolChain TargetToolChain)
+		public void PreBuildSetup()
 		{
 			// Describe what's being built.
 			Log.TraceVerbose("Building {0} - {1} - {2} - {3}", AppName, TargetName, Platform, Configuration);
@@ -2309,7 +2309,7 @@ namespace UnrealBuildTool
 			// On Mac AppBinaries paths for non-console targets need to be adjusted to be inside the app bundle
 			if (Platform == UnrealTargetPlatform.Mac && !Rules.bIsBuildingConsoleApplication)
 			{
-				TargetToolChain.FixBundleBinariesPaths(this, Binaries);
+				MacToolChain.FixBundleBinariesPaths(OutputPath, Binaries);
 			}
 		}
 
@@ -2379,25 +2379,25 @@ namespace UnrealBuildTool
 		/// <returns>Map of variable names to values</returns>
 		private Dictionary<string, string> GetTargetVariables(UEBuildPlugin Plugin)
 		{
-				Dictionary<string, string> Variables = new Dictionary<string,string>();
-				Variables.Add("RootDir", UnrealBuildTool.RootDirectory.FullName);
-				Variables.Add("EngineDir", UnrealBuildTool.EngineDirectory.FullName);
-				Variables.Add("EnterpriseDir", UnrealBuildTool.EnterpriseDirectory.FullName);
-				Variables.Add("ProjectDir", ProjectDirectory.FullName);
-				Variables.Add("TargetName", TargetName);
-				Variables.Add("TargetPlatform", Platform.ToString());
-				Variables.Add("TargetConfiguration", Configuration.ToString());
-				Variables.Add("TargetType", TargetType.ToString());
-				if(ProjectFile != null)
-				{
-					Variables.Add("ProjectFile", ProjectFile.FullName);
-				}
+			Dictionary<string, string> Variables = new Dictionary<string,string>();
+			Variables.Add("RootDir", UnrealBuildTool.RootDirectory.FullName);
+			Variables.Add("EngineDir", UnrealBuildTool.EngineDirectory.FullName);
+			Variables.Add("EnterpriseDir", UnrealBuildTool.EnterpriseDirectory.FullName);
+			Variables.Add("ProjectDir", ProjectDirectory.FullName);
+			Variables.Add("TargetName", TargetName);
+			Variables.Add("TargetPlatform", Platform.ToString());
+			Variables.Add("TargetConfiguration", Configuration.ToString());
+			Variables.Add("TargetType", TargetType.ToString());
+			if(ProjectFile != null)
+			{
+				Variables.Add("ProjectFile", ProjectFile.FullName);
+			}
 			if(Plugin != null)
-				{
+			{
 				Variables.Add("PluginDir", Plugin.Directory.FullName);
 			}
 			return Variables;
-				}
+		}
 
 		/// <summary>
 		/// Write scripts containing the custom build steps for the given host platform
diff --git a/Engine/Source/Programs/UnrealBuildTool/Platform/Mac/MacToolChain.cs b/Engine/Source/Programs/UnrealBuildTool/Platform/Mac/MacToolChain.cs
--- a/Engine/Source/Programs/UnrealBuildTool/Platform/Mac/MacToolChain.cs
+++ b/Engine/Source/Programs/UnrealBuildTool/Platform/Mac/MacToolChain.cs
@@ -1237,11 +1237,9 @@ namespace UnrealBuildTool
 			}
 		}
 
-		public override void FixBundleBinariesPaths(UEBuildTarget Target, List<UEBuildBinary> Binaries)
+		public static void FixBundleBinariesPaths(FileReference OutputPath, List<UEBuildBinary> Binaries)
 		{
-			base.FixBundleBinariesPaths(Target, Binaries);
-
-			string BundleContentsPath = Target.OutputPath.FullName + ".app/Contents/";
+			string BundleContentsPath = OutputPath.FullName + ".app/Contents/";
 			foreach (UEBuildBinary Binary in Binaries)
 			{
 				string BinaryFileName = Path.GetFileName(Binary.OutputFilePath.FullName);
@@ -1250,7 +1248,7 @@ namespace UnrealBuildTool
 					// Only dylibs from the same folder as the executable should be moved to the bundle. UE4Editor-*Game* dylibs and plugins will be loaded
 					// from their Binaries/Mac folders.
 					string DylibDir = Path.GetDirectoryName(Path.GetFullPath(Binary.OutputFilePath.FullName));
-					string ExeDir = Path.GetDirectoryName(Path.GetFullPath(Target.OutputPath.FullName));
+					string ExeDir = Path.GetDirectoryName(Path.GetFullPath(OutputPath.FullName));
 					if (DylibDir.StartsWith(ExeDir))
 					{
 						// get the subdir, which is the DylibDir - ExeDir
diff --git a/Engine/Source/Programs/UnrealBuildTool/ToolChain/UEToolChain.cs b/Engine/Source/Programs/UnrealBuildTool/ToolChain/UEToolChain.cs
--- a/Engine/Source/Programs/UnrealBuildTool/ToolChain/UEToolChain.cs
+++ b/Engine/Source/Programs/UnrealBuildTool/ToolChain/UEToolChain.cs
@@ -106,11 +106,6 @@ namespace UnrealBuildTool
 
 		}
 
-		public virtual void FixBundleBinariesPaths(UEBuildTarget Target, List<UEBuildBinary> Binaries)
-		{
-
-		}
-
         public virtual string GetSDKVersion()
         {
             return "Not Applicable";
