UBT: Allow project file generator to construct targets and query it directly. Removes a lot of global state, and simplifies code paths through building a target.

#rnx

diff --git a/Engine/Source/Programs/UnrealBuildTool/Configuration/TargetDescriptor.cs b/Engine/Source/Programs/UnrealBuildTool/Configuration/TargetDescriptor.cs
--- a/Engine/Source/Programs/UnrealBuildTool/Configuration/TargetDescriptor.cs
+++ b/Engine/Source/Programs/UnrealBuildTool/Configuration/TargetDescriptor.cs
@@ -19,7 +19,7 @@ namespace UnrealBuildTool
 		public UnrealTargetPlatform Platform;
 		public UnrealTargetConfiguration Configuration;
 		public string Architecture;
-		public List<OnlyModule> OnlyModules;
+		public List<OnlyModule> OnlyModules = new List<OnlyModule>();
 		public FileReference ForeignPlugin;
 		public string[] AdditionalArguments;
 
@@ -187,7 +187,7 @@ namespace UnrealBuildTool
 
 				// Create the target descriptor
 				TargetDescriptor Target = new TargetDescriptor(ProjectFile, TargetName, Platform, Configuration, Architecture, AdditionalArguments.ToArray());
-				Target.OnlyModules = OnlyModules;
+				Target.OnlyModules.AddRange(OnlyModules);
 				Target.ForeignPlugin = ForeignPlugin;
 				Targets.Add(Target);
 			}
diff --git a/Engine/Source/Programs/UnrealBuildTool/Configuration/UEBuildBinary.cs b/Engine/Source/Programs/UnrealBuildTool/Configuration/UEBuildBinary.cs
--- a/Engine/Source/Programs/UnrealBuildTool/Configuration/UEBuildBinary.cs
+++ b/Engine/Source/Programs/UnrealBuildTool/Configuration/UEBuildBinary.cs
@@ -676,20 +676,7 @@ namespace UnrealBuildTool
 			return Type == UEBuildBinaryType.StaticLibrary;
 		}
 
-		public void GatherDataForProjectFiles(ReadOnlyTargetRules Target, CppCompileEnvironment GlobalCompileEnvironment)
-		{
-			CppCompileEnvironment BinaryCompileEnvironment = CreateBinaryCompileEnvironment(GlobalCompileEnvironment);
-			foreach(UEBuildModuleCPP Module in Modules.OfType<UEBuildModuleCPP>())
-			{
-				ProjectFile ProjectFileForIDE;
-				if (ProjectFileGenerator.ModuleToProjectFileMap.TryGetValue(Module.Name, out ProjectFileForIDE))
-				{
-					Module.GatherDataForProjectFile(Target, BinaryCompileEnvironment, ProjectFileForIDE);
-				}
-			}
-		}
-
-		private CppCompileEnvironment CreateBinaryCompileEnvironment(CppCompileEnvironment GlobalCompileEnvironment)
+		public CppCompileEnvironment CreateBinaryCompileEnvironment(CppCompileEnvironment GlobalCompileEnvironment)
 		{
 			CppCompileEnvironment BinaryCompileEnvironment = new CppCompileEnvironment(GlobalCompileEnvironment);
 			BinaryCompileEnvironment.bIsBuildingDLL = IsBuildingDll(Type);
diff --git a/Engine/Source/Programs/UnrealBuildTool/Configuration/UEBuildModuleCPP.cs b/Engine/Source/Programs/UnrealBuildTool/Configuration/UEBuildModuleCPP.cs
--- a/Engine/Source/Programs/UnrealBuildTool/Configuration/UEBuildModuleCPP.cs
+++ b/Engine/Source/Programs/UnrealBuildTool/Configuration/UEBuildModuleCPP.cs
@@ -287,20 +287,6 @@ namespace UnrealBuildTool
 			get { return FileReference.Combine(IntermediateDirectory, String.Format("{0}.precompiled", Name)); }
 		}
 
-		/// <summary>
-		/// Gathers intellisense data for the project file containing this module
-		/// </summary>
-		/// <param name="Target">The target being built</param>
-		/// <param name="BinaryCompileEnvironment">The inherited compile environment for this module</param>
-		/// <param name="ProjectFile">The project file containing this module</param>
-		public void GatherDataForProjectFile(ReadOnlyTargetRules Target, CppCompileEnvironment BinaryCompileEnvironment, ProjectFile ProjectFile)
-		{
-			CppCompileEnvironment ModuleCompileEnvironment = CreateModuleCompileEnvironment(Target, BinaryCompileEnvironment);
-			ProjectFile.AddIntelliSensePreprocessorDefinitions(ModuleCompileEnvironment.Definitions);
-			ProjectFile.AddIntelliSenseIncludePaths(ModuleCompileEnvironment.IncludePaths.SystemIncludePaths, true);
-			ProjectFile.AddIntelliSenseIncludePaths(ModuleCompileEnvironment.IncludePaths.UserIncludePaths, false);
-		}
-
 		/// <summary>
 		/// Sets up the environment for compiling any module that includes the public interface of this module.
 		/// </summary>
diff --git a/Engine/Source/Programs/UnrealBuildTool/Configuration/UEBuildTarget.cs b/Engine/Source/Programs/UnrealBuildTool/Configuration/UEBuildTarget.cs
--- a/Engine/Source/Programs/UnrealBuildTool/Configuration/UEBuildTarget.cs
+++ b/Engine/Source/Programs/UnrealBuildTool/Configuration/UEBuildTarget.cs
@@ -983,6 +983,9 @@ namespace UnrealBuildTool
 				}
 				VersionFile = BuildVersion.GetFileNameForTarget(OutputDirectory, bCompileMonolithic? TargetName : AppName, Platform, VersionConfig, Architecture);
 			}
+
+			// Create all the binaries and modules
+			PreBuildSetup();
 		}
 
 		/// <summary>
@@ -1392,7 +1395,7 @@ namespace UnrealBuildTool
 		/// <param name="RuntimeDependencies">Output runtime dependencies</param>
 		/// <param name="HotReload">The hot-reload mode</param>
 		TargetReceipt PrepareReceipt(UEToolChain ToolChain, List<KeyValuePair<FileReference, BuildProductType>> BuildProducts, List<RuntimeDependency> RuntimeDependencies, EHotReload HotReload)
-			{
+		{
 			// Read the version file
 			BuildVersion Version;
 			if (!BuildVersion.TryRead(BuildVersion.GetDefaultFileName(), out Version))
@@ -1507,6 +1510,24 @@ namespace UnrealBuildTool
 			return Output;
 		}
 
+		/// <summary>
+		/// Creates a global compile environment suitable for generating project files.
+		/// </summary>
+		/// <returns>New compile environment</returns>
+		public CppCompileEnvironment CreateCompileEnvironmentForProjectFiles()
+		{
+			CppPlatform CppPlatform = UEBuildPlatform.GetBuildPlatform(Platform).DefaultCppPlatform;
+			CppConfiguration CppConfiguration = GetCppConfiguration(Configuration);
+
+			CppCompileEnvironment GlobalCompileEnvironment = new CppCompileEnvironment(CppPlatform, CppConfiguration, Architecture, null);
+			LinkEnvironment GlobalLinkEnvironment = new LinkEnvironment(GlobalCompileEnvironment.Platform, GlobalCompileEnvironment.Configuration, GlobalCompileEnvironment.Architecture);
+
+			UEToolChain TargetToolChain = CreateToolchain(CppPlatform);
+			SetupGlobalEnvironment(TargetToolChain, GlobalCompileEnvironment, GlobalLinkEnvironment);
+
+			return GlobalCompileEnvironment;
+		}
+
 		/// <summary>
 		/// Builds the target, appending list of output files and returns building result.
 		/// </summary>
@@ -1516,16 +1537,12 @@ namespace UnrealBuildTool
 			CppConfiguration CppConfiguration = GetCppConfiguration(Configuration);
 
 			CppCompileEnvironment GlobalCompileEnvironment = new CppCompileEnvironment(CppPlatform, CppConfiguration, Architecture, Headers);
-			UEToolChain TargetToolChain = CreateToolchain(CppPlatform);
-			if(!ProjectFileGenerator.bGenerateProjectFiles)
-			{
-				TargetToolChain.PrintVersionInfo();
-			}
 			LinkEnvironment GlobalLinkEnvironment = new LinkEnvironment(GlobalCompileEnvironment.Platform, GlobalCompileEnvironment.Configuration, GlobalCompileEnvironment.Architecture);
-			SetupGlobalEnvironment(TargetToolChain, GlobalCompileEnvironment, GlobalLinkEnvironment);
 
-			// Create all the binaries and modules
-			PreBuildSetup();
+			UEToolChain TargetToolChain = CreateToolchain(CppPlatform);
+			TargetToolChain.PrintVersionInfo();
+			
+			SetupGlobalEnvironment(TargetToolChain, GlobalCompileEnvironment, GlobalLinkEnvironment);
 
 			// Save off the original list of binaries. We'll use this to figure out which PCHs to create later, to avoid switching PCHs when compiling single modules.
 			List<UEBuildBinary> OriginalBinaries = Binaries;
@@ -1550,7 +1567,7 @@ namespace UnrealBuildTool
 			}
 
 			// For installed builds, filter out all the binaries that aren't in mods
-			if (!ProjectFileGenerator.bGenerateProjectFiles && UnrealBuildTool.IsProjectInstalled())
+			if (UnrealBuildTool.IsProjectInstalled())
 			{
 				List<DirectoryReference> ModDirectories = EnabledPlugins.Where(x => x.Type == PluginType.Mod).Select(x => x.Directory).ToList();
 
@@ -1600,110 +1617,104 @@ namespace UnrealBuildTool
 				}
 			}
 
-			if(!ProjectFileGenerator.bGenerateProjectFiles)
+			if(!Rules.bDisableLinking)
 			{
-				if(!Rules.bDisableLinking)
+				// Check the distribution level of all binaries based on the dependencies they have
+				if(ProjectFile == null && !Rules.bOutputPubliclyDistributable)
 				{
-					// Check the distribution level of all binaries based on the dependencies they have
-					if(ProjectFile == null && !Rules.bOutputPubliclyDistributable)
-					{
-						Dictionary<UEBuildModule, Dictionary<RestrictedFolder, DirectoryReference>> ModuleRestrictedFolderCache = new Dictionary<UEBuildModule, Dictionary<RestrictedFolder, DirectoryReference>>();
+					Dictionary<UEBuildModule, Dictionary<RestrictedFolder, DirectoryReference>> ModuleRestrictedFolderCache = new Dictionary<UEBuildModule, Dictionary<RestrictedFolder, DirectoryReference>>();
 
-						bool bResult = true;
-						foreach (UEBuildBinary Binary in Binaries)
-						{
-							bResult &= Binary.CheckRestrictedFolders(DirectoryReference.FromFile(ProjectFile), ModuleRestrictedFolderCache);
-						}
-
-						if(!bResult)
-						{
-							throw new BuildException("Unable to create binaries in less restricted locations than their input files.");
-						}
+					bool bResult = true;
+					foreach (UEBuildBinary Binary in Binaries)
+					{
+						bResult &= Binary.CheckRestrictedFolders(DirectoryReference.FromFile(ProjectFile), ModuleRestrictedFolderCache);
 					}
 
-					// Check for linking against modules prohibited by the EULA
-					CheckForEULAViolation();
+					if(!bResult)
+					{
+						throw new BuildException("Unable to create binaries in less restricted locations than their input files.");
+					}
 				}
 
-				// Build a mapping from module to its plugin
-				Dictionary<UEBuildModule, UEBuildPlugin> ModuleToPlugin = new Dictionary<UEBuildModule, UEBuildPlugin>();
-				foreach(UEBuildPlugin Plugin in BuildPlugins)
+				// Check for linking against modules prohibited by the EULA
+				CheckForEULAViolation();
+			}
+
+			// Build a mapping from module to its plugin
+			Dictionary<UEBuildModule, UEBuildPlugin> ModuleToPlugin = new Dictionary<UEBuildModule, UEBuildPlugin>();
+			foreach(UEBuildPlugin Plugin in BuildPlugins)
+			{
+				foreach(UEBuildModule Module in Plugin.Modules)
 				{
-					foreach(UEBuildModule Module in Plugin.Modules)
+					if (!ModuleToPlugin.ContainsKey(Module))
 					{
-						if (!ModuleToPlugin.ContainsKey(Module))
-						{
-							ModuleToPlugin.Add(Module, Plugin);
-						}
+						ModuleToPlugin.Add(Module, Plugin);
 					}
 				}
+			}
 
-				// Check there aren't any engine binaries with dependencies on game modules. This can happen when game-specific plugins override engine plugins.
-				foreach(UEBuildModule Module in Modules.Values)
+			// Check there aren't any engine binaries with dependencies on game modules. This can happen when game-specific plugins override engine plugins.
+			foreach(UEBuildModule Module in Modules.Values)
+			{
+				if(Module.Binary != null && UnrealBuildTool.IsUnderAnEngineDirectory(Module.RulesFile.Directory))
 				{
-					if(Module.Binary != null && UnrealBuildTool.IsUnderAnEngineDirectory(Module.RulesFile.Directory))
+					DirectoryReference RootDirectory = UnrealBuildTool.EngineDirectory;
+
+					if (Module.RulesFile.IsUnderDirectory(UnrealBuildTool.EnterpriseDirectory))
 					{
-						DirectoryReference RootDirectory = UnrealBuildTool.EngineDirectory;
+						RootDirectory = UnrealBuildTool.EnterpriseDirectory;
+					}
+
+					HashSet<UEBuildModule> ReferencedModules = Module.GetDependencies(bWithIncludePathModules: true, bWithDynamicallyLoadedModules: true);
 
-						if (Module.RulesFile.IsUnderDirectory(UnrealBuildTool.EnterpriseDirectory))
+					// Make sure engine modules don't depend on enterprise or game modules and that enterprise modules don't depend on game modules
+					foreach(UEBuildModule ReferencedModule in ReferencedModules)
+					{
+						if(ReferencedModule.RulesFile != null && !ReferencedModule.RulesFile.IsUnderDirectory(UnrealBuildTool.EngineDirectory) && !ReferencedModule.RulesFile.IsUnderDirectory(RootDirectory))
 						{
-							RootDirectory = UnrealBuildTool.EnterpriseDirectory;
+							string EngineModuleRelativePath = Module.RulesFile.MakeRelativeTo(UnrealBuildTool.EngineDirectory.ParentDirectory);
+							string ReferencedModuleRelativePath = (ProjectFile != null && ReferencedModule.RulesFile.IsUnderDirectory(ProjectFile.Directory)) ? ReferencedModule.RulesFile.MakeRelativeTo(ProjectFile.Directory.ParentDirectory) : ReferencedModule.RulesFile.FullName;
+							throw new BuildException("Engine module '{0}' should not depend on game module '{1}'", EngineModuleRelativePath, ReferencedModuleRelativePath);
 						}
+					}
 
-						HashSet<UEBuildModule> ReferencedModules = Module.GetDependencies(bWithIncludePathModules: true, bWithDynamicallyLoadedModules: true);
-
-						// Make sure engine modules don't depend on enterprise or game modules and that enterprise modules don't depend on game modules
+					// Make sure engine modules don't directly reference engine plugins
+					if(Module.RulesFile.IsUnderDirectory(UnrealBuildTool.EngineSourceDirectory) && !Module.RulesFile.IsUnderDirectory(TargetRulesFile.Directory))
+					{
 						foreach(UEBuildModule ReferencedModule in ReferencedModules)
 						{
-							if(ReferencedModule.RulesFile != null && !ReferencedModule.RulesFile.IsUnderDirectory(UnrealBuildTool.EngineDirectory) && !ReferencedModule.RulesFile.IsUnderDirectory(RootDirectory))
+							if(ReferencedModule.RulesFile != null && ModuleToPlugin.ContainsKey(ReferencedModule) && !IsWhitelistedEnginePluginReference(Module.Name, ReferencedModule.Name))
 							{
 								string EngineModuleRelativePath = Module.RulesFile.MakeRelativeTo(UnrealBuildTool.EngineDirectory.ParentDirectory);
-								string ReferencedModuleRelativePath = (ProjectFile != null && ReferencedModule.RulesFile.IsUnderDirectory(ProjectFile.Directory)) ? ReferencedModule.RulesFile.MakeRelativeTo(ProjectFile.Directory.ParentDirectory) : ReferencedModule.RulesFile.FullName;
-								throw new BuildException("Engine module '{0}' should not depend on game module '{1}'", EngineModuleRelativePath, ReferencedModuleRelativePath);
-							}
-						}
-
-						// Make sure engine modules don't directly reference engine plugins
-						if(Module.RulesFile.IsUnderDirectory(UnrealBuildTool.EngineSourceDirectory) && !Module.RulesFile.IsUnderDirectory(TargetRulesFile.Directory))
-						{
-							foreach(UEBuildModule ReferencedModule in ReferencedModules)
-							{
-								if(ReferencedModule.RulesFile != null && ModuleToPlugin.ContainsKey(ReferencedModule) && !IsWhitelistedEnginePluginReference(Module.Name, ReferencedModule.Name))
-								{
-									string EngineModuleRelativePath = Module.RulesFile.MakeRelativeTo(UnrealBuildTool.EngineDirectory.ParentDirectory);
-									string ReferencedModuleRelativePath = ReferencedModule.RulesFile.MakeRelativeTo(UnrealBuildTool.EngineDirectory.ParentDirectory);
-									Log.TraceWarning("Warning: Engine module '{0}' should not depend on plugin module '{1}'", EngineModuleRelativePath, ReferencedModuleRelativePath);
-								}
+								string ReferencedModuleRelativePath = ReferencedModule.RulesFile.MakeRelativeTo(UnrealBuildTool.EngineDirectory.ParentDirectory);
+								Log.TraceWarning("Warning: Engine module '{0}' should not depend on plugin module '{1}'", EngineModuleRelativePath, ReferencedModuleRelativePath);
 							}
 						}
 					}
 				}
+			}
 
-				// Check that each plugin declares its dependencies explicitly
-				foreach(UEBuildPlugin Plugin in BuildPlugins)
+			// Check that each plugin declares its dependencies explicitly
+			foreach(UEBuildPlugin Plugin in BuildPlugins)
+			{
+				foreach(UEBuildModule Module in Plugin.Modules)
 				{
-					foreach(UEBuildModule Module in Plugin.Modules)
+					HashSet<UEBuildModule> DependencyModules = Module.GetDependencies(bWithIncludePathModules: true, bWithDynamicallyLoadedModules: true);
+					foreach(UEBuildModule DependencyModule in DependencyModules)
 					{
-						HashSet<UEBuildModule> DependencyModules = Module.GetDependencies(bWithIncludePathModules: true, bWithDynamicallyLoadedModules: true);
-						foreach(UEBuildModule DependencyModule in DependencyModules)
+						UEBuildPlugin DependencyPlugin;
+						if(ModuleToPlugin.TryGetValue(DependencyModule, out DependencyPlugin) && DependencyPlugin != Plugin && !Plugin.Dependencies.Contains(DependencyPlugin))
 						{
-							UEBuildPlugin DependencyPlugin;
-							if(ModuleToPlugin.TryGetValue(DependencyModule, out DependencyPlugin) && DependencyPlugin != Plugin && !Plugin.Dependencies.Contains(DependencyPlugin))
-							{
-								Log.TraceWarning("Warning: Plugin '{0}' does not list plugin '{1}' as a dependency, but module '{2}' depends on '{3}'.", Plugin.Name, DependencyPlugin.Name, Module.Name, DependencyModule.Name);
-							}
+							Log.TraceWarning("Warning: Plugin '{0}' does not list plugin '{1}' as a dependency, but module '{2}' depends on '{3}'.", Plugin.Name, DependencyPlugin.Name, Module.Name, DependencyModule.Name);
 						}
 					}
 				}
 			}
 
 			// Execute the pre-build steps
-			if(!ProjectFileGenerator.bGenerateProjectFiles)
+			if(!ExecuteCustomPreBuildSteps())
 			{
-				if(!ExecuteCustomPreBuildSteps())
-				{
-					return ECompilationResult.OtherCompilationError;
-				}
+				return ECompilationResult.OtherCompilationError;
 			}
 
 			// If we're compiling monolithic, make sure the executable knows about all referenced modules
@@ -1758,39 +1769,37 @@ namespace UnrealBuildTool
 			// On Mac and Linux we have actions that should be executed after all the binaries are created
 			TargetToolChain.SetupBundleDependencies(Binaries, TargetName);
 
-			if (!ProjectFileGenerator.bGenerateProjectFiles)
+			// Generate headers
+			HashSet<UEBuildModuleCPP> ModulesToGenerateHeadersFor = GatherDependencyModules(OriginalBinaries.ToList());
+
+			if (OnlyModules.Count > 0)
 			{
-				HashSet<UEBuildModuleCPP> ModulesToGenerateHeadersFor = GatherDependencyModules(OriginalBinaries.ToList());
+				HashSet<UEBuildModuleCPP> CorrectlyOrderedModules = GatherDependencyModules(NonFilteredModules);
 
-				if (OnlyModules.Count > 0)
-				{
-					HashSet<UEBuildModuleCPP> CorrectlyOrderedModules = GatherDependencyModules(NonFilteredModules);
+				CorrectlyOrderedModules.RemoveWhere((Module) => !ModulesToGenerateHeadersFor.Contains(Module));
+				ModulesToGenerateHeadersFor = CorrectlyOrderedModules;
+			}
 
-					CorrectlyOrderedModules.RemoveWhere((Module) => !ModulesToGenerateHeadersFor.Contains(Module));
-					ModulesToGenerateHeadersFor = CorrectlyOrderedModules;
-				}
+			Dictionary<string, FlatModuleCsDataType> NameToFlatModuleData = new Dictionary<string, FlatModuleCsDataType>(StringComparer.InvariantCultureIgnoreCase);
+			foreach(FlatModuleCsDataType FlatModuleData in FlatModuleCsData)
+			{
+				NameToFlatModuleData[FlatModuleData.ModuleName] = FlatModuleData;
+			}
 
-				Dictionary<string, FlatModuleCsDataType> NameToFlatModuleData = new Dictionary<string, FlatModuleCsDataType>(StringComparer.InvariantCultureIgnoreCase);
-				foreach(FlatModuleCsDataType FlatModuleData in FlatModuleCsData)
-				{
-					NameToFlatModuleData[FlatModuleData.ModuleName] = FlatModuleData;
-				}
+			ExternalExecution.SetupUObjectModules(ModulesToGenerateHeadersFor, Rules.Platform, ProjectDescriptor, UObjectModules, NameToFlatModuleData, Rules.GeneratedCodeVersion, bIsAssemblingBuild);
 
-				ExternalExecution.SetupUObjectModules(ModulesToGenerateHeadersFor, Rules.Platform, ProjectDescriptor, UObjectModules, NameToFlatModuleData, Rules.GeneratedCodeVersion, bIsAssemblingBuild);
+			// NOTE: Even in Gather mode, we need to run UHT to make sure the files exist for the static action graph to be setup correctly.  This is because UHT generates .cpp
+			// files that are injected as top level prerequisites.  If UHT only emitted included header files, we wouldn't need to run it during the Gather phase at all.
+			if (UObjectModules.Count > 0)
+			{
+				FileReference ModuleInfoFileName = FileReference.Combine(ProjectIntermediateDirectory, GetTargetName() + ".uhtmanifest");
 
-				// NOTE: Even in Gather mode, we need to run UHT to make sure the files exist for the static action graph to be setup correctly.  This is because UHT generates .cpp
-				// files that are injected as top level prerequisites.  If UHT only emitted included header files, we wouldn't need to run it during the Gather phase at all.
-				if (UObjectModules.Count > 0)
+				// Execute the header tool
+				ECompilationResult UHTResult = ExternalExecution.ExecuteHeaderToolIfNecessary(BuildConfiguration, ProjectFile, TargetName, TargetType, bHasProjectScriptPlugin, GlobalCompileEnvironment, UObjectModules, ModuleInfoFileName, HotReload, true, bIsAssemblingBuild);
+				if(!UHTResult.Succeeded())
 				{
-					FileReference ModuleInfoFileName = FileReference.Combine(ProjectIntermediateDirectory, GetTargetName() + ".uhtmanifest");
-
-					// Execute the header tool
-					ECompilationResult UHTResult = ExternalExecution.ExecuteHeaderToolIfNecessary(BuildConfiguration, ProjectFile, TargetName, TargetType, bHasProjectScriptPlugin, GlobalCompileEnvironment, UObjectModules, ModuleInfoFileName, HotReload, true, bIsAssemblingBuild);
-					if(!UHTResult.Succeeded())
-					{
-						Log.TraceInformation(String.Format("Error: UnrealHeaderTool failed for target '{0}' (platform: {1}, module info: {2}, exit code: {3} ({4})).", GetTargetName(), Platform.ToString(), ModuleInfoFileName, UHTResult.ToString(), (int)UHTResult));
-						return UHTResult;
-					}
+					Log.TraceInformation(String.Format("Error: UnrealHeaderTool failed for target '{0}' (platform: {1}, module info: {2}, exit code: {3} ({4})).", GetTargetName(), Platform.ToString(), ModuleInfoFileName, UHTResult.ToString(), (int)UHTResult));
+					return UHTResult;
 				}
 			}
 
@@ -1798,7 +1807,7 @@ namespace UnrealBuildTool
 
 			// Find all the shared PCHs.
 			List<PrecompiledHeaderTemplate> SharedPCHs = new List<PrecompiledHeaderTemplate>();
-			if (!ProjectFileGenerator.bGenerateProjectFiles && Rules.bUseSharedPCHs)
+			if (Rules.bUseSharedPCHs)
 			{
 				SharedPCHs = FindSharedPCHs(OriginalBinaries, GlobalCompileEnvironment);
 			}
@@ -1823,16 +1832,6 @@ namespace UnrealBuildTool
 				}
 			}
 
-			// For the project file generator, just generate all the include paths
-			if(ProjectFileGenerator.bGenerateProjectFiles)
-			{
-				foreach(UEBuildBinary Binary in Binaries)
-				{
-					Binary.GatherDataForProjectFiles(Rules, GlobalCompileEnvironment);
-				}
-				return ECompilationResult.Succeeded;
-			}
-
 			// Build the target's binaries.
 			DirectoryReference ExeDir = OutputPaths[0].Directory;
 			foreach (UEBuildBinary Binary in Binaries)
@@ -1893,7 +1892,7 @@ namespace UnrealBuildTool
 			}
 
 			// Finalize and generate metadata for this target
-			if(!ProjectFileGenerator.bGenerateProjectFiles && !Rules.bDisableLinking)
+			if(!Rules.bDisableLinking)
 			{
 				// Also add any explicitly specified build products
 				if(Rules.AdditionalBuildProducts.Count > 0)
diff --git a/Engine/Source/Programs/UnrealBuildTool/ProjectFiles/Project.cs b/Engine/Source/Programs/UnrealBuildTool/ProjectFiles/Project.cs
--- a/Engine/Source/Programs/UnrealBuildTool/ProjectFiles/Project.cs
+++ b/Engine/Source/Programs/UnrealBuildTool/ProjectFiles/Project.cs
@@ -51,6 +51,11 @@ namespace UnrealBuildTool
 		/// Delegate for creating a rules instance for a given platform/configuration
 		public Func<UnrealTargetPlatform, UnrealTargetConfiguration, TargetRules> CreateRulesDelegate = null;
 
+		public string Name
+		{
+			get { return TargetFilePath.GetFileNameWithoutAnyExtensions(); }
+		}
+
 		public override string ToString()
 		{
 			return TargetFilePath.GetFileNameWithoutExtension();
@@ -286,35 +291,31 @@ namespace UnrealBuildTool
 		/// <param name="NewPreprocessorDefinitions">List of preprocessor definitons to add</param>
 		public void AddIntelliSensePreprocessorDefinitions(List<string> NewPreprocessorDefinitions)
 		{
-			if (ProjectFileGenerator.OnlyGenerateIntelliSenseDataForProject == null ||
-				ProjectFileGenerator.OnlyGenerateIntelliSenseDataForProject == this)
+			foreach (string CurDef in NewPreprocessorDefinitions)
 			{
-				foreach (string CurDef in NewPreprocessorDefinitions)
+				// Don't add definitions and value combinations that have already been added for this project
+				if (KnownIntelliSensePreprocessorDefinitions.Add(CurDef))
 				{
-					// Don't add definitions and value combinations that have already been added for this project
-					if (KnownIntelliSensePreprocessorDefinitions.Add(CurDef))
-					{
-						// Go ahead and check to see if the definition already exists, but the value is different
-						bool AlreadyExists = false;
+					// Go ahead and check to see if the definition already exists, but the value is different
+					bool AlreadyExists = false;
 
-						string Def, Value;
-						SplitDefinitionAndValue(CurDef, out Def, out Value);
-						for (int DefineIndex = 0; DefineIndex < IntelliSensePreprocessorDefinitions.Count; ++DefineIndex)
+					string Def, Value;
+					SplitDefinitionAndValue(CurDef, out Def, out Value);
+					for (int DefineIndex = 0; DefineIndex < IntelliSensePreprocessorDefinitions.Count; ++DefineIndex)
+					{
+						string ExistingDef, ExistingValue;
+						SplitDefinitionAndValue(IntelliSensePreprocessorDefinitions[DefineIndex], out ExistingDef, out ExistingValue);
+						if (ExistingDef == Def)
 						{
-							string ExistingDef, ExistingValue;
-							SplitDefinitionAndValue(IntelliSensePreprocessorDefinitions[DefineIndex], out ExistingDef, out ExistingValue);
-							if (ExistingDef == Def)
-							{
-								// Already exists, but the value is changing.  We don't bother clobbering values for existing defines for this project.
-								AlreadyExists = true;
-								break;
-							}
+							// Already exists, but the value is changing.  We don't bother clobbering values for existing defines for this project.
+							AlreadyExists = true;
+							break;
 						}
+					}
 
-						if (!AlreadyExists)
-						{
-							IntelliSensePreprocessorDefinitions.Add(CurDef);
-						}
+					if (!AlreadyExists)
+					{
+						IntelliSensePreprocessorDefinitions.Add(CurDef);
 					}
 				}
 			}
@@ -328,33 +329,29 @@ namespace UnrealBuildTool
 		/// <param name="bAddingSystemIncludes"></param>
 		public void AddIntelliSenseIncludePaths(HashSet<DirectoryReference> NewIncludePaths, bool bAddingSystemIncludes)
 		{
-			if (ProjectFileGenerator.OnlyGenerateIntelliSenseDataForProject == null ||
-				ProjectFileGenerator.OnlyGenerateIntelliSenseDataForProject == this)
+			foreach (DirectoryReference CurPath in NewIncludePaths)
 			{
-				foreach (DirectoryReference CurPath in NewIncludePaths)
+				if (bAddingSystemIncludes ? KnownIntelliSenseSystemIncludeSearchPaths.Add(CurPath.FullName) : KnownIntelliSenseIncludeSearchPaths.Add(CurPath.FullName))
 				{
-					if (bAddingSystemIncludes ? KnownIntelliSenseSystemIncludeSearchPaths.Add(CurPath.FullName) : KnownIntelliSenseIncludeSearchPaths.Add(CurPath.FullName))
+					// Incoming include paths are relative to the solution directory, but we need these paths to be
+					// relative to the project file's directory
+					string PathRelativeToProjectFile = NormalizeProjectPath(CurPath);
+
+					// Make sure that it doesn't exist already
+					bool AlreadyExists = false;
+					List<string> SearchPaths = bAddingSystemIncludes ? IntelliSenseSystemIncludeSearchPaths : IntelliSenseIncludeSearchPaths;
+					foreach (string ExistingPath in SearchPaths)
 					{
-						// Incoming include paths are relative to the solution directory, but we need these paths to be
-						// relative to the project file's directory
-						string PathRelativeToProjectFile = NormalizeProjectPath(CurPath);
-
-						// Make sure that it doesn't exist already
-						bool AlreadyExists = false;
-						List<string> SearchPaths = bAddingSystemIncludes ? IntelliSenseSystemIncludeSearchPaths : IntelliSenseIncludeSearchPaths;
-						foreach (string ExistingPath in SearchPaths)
+						if (PathRelativeToProjectFile == ExistingPath)
 						{
-							if (PathRelativeToProjectFile == ExistingPath)
-							{
-								AlreadyExists = true;
-								break;
-							}
+							AlreadyExists = true;
+							break;
 						}
+					}
 
-						if (!AlreadyExists)
-						{
-							SearchPaths.Add(PathRelativeToProjectFile);
-						}
+					if (!AlreadyExists)
+					{
+						SearchPaths.Add(PathRelativeToProjectFile);
 					}
 				}
 			}
diff --git a/Engine/Source/Programs/UnrealBuildTool/ProjectFiles/ProjectFileGenerator.cs b/Engine/Source/Programs/UnrealBuildTool/ProjectFiles/ProjectFileGenerator.cs
--- a/Engine/Source/Programs/UnrealBuildTool/ProjectFiles/ProjectFileGenerator.cs
+++ b/Engine/Source/Programs/UnrealBuildTool/ProjectFiles/ProjectFileGenerator.cs
@@ -223,7 +223,7 @@ namespace UnrealBuildTool
 		/// <summary>
 		/// True if build targets should pass the -useprecompiled argument
 		/// </summary>
-		public static bool bUsePrecompiled = false;
+		public bool bUsePrecompiled = false;
 
 		/// <summary>
 		/// True if we should include engine source in the generated solution
@@ -284,12 +284,12 @@ namespace UnrealBuildTool
 		/// <summary>
 		/// Name of the UE4 engine project that contains all of the engine code, config files and other files
 		/// </summary>
-		public static readonly string EngineProjectFileNameBase = "UE4";
+		public const string EngineProjectFileNameBase = "UE4";
 
 		/// <summary>
 		/// Name of the UE4 enterprise project that contains all of the enterprise code, config files and other files
 		/// </summary>
-		public static readonly string EnterpriseProjectFileNameBase = "Studio";
+		public const string EnterpriseProjectFileNameBase = "Studio";
 
 		/// <summary>
 		/// When ProjectsAreIntermediate is true, this is the directory to store generated project files
@@ -333,26 +333,13 @@ namespace UnrealBuildTool
 
 		/// <summary>
 		/// Maps all module names that were included in generated project files, to actual project file objects.
-		/// @todo projectfiles: Nasty global static list.  This is only really used for IntelliSense, and to avoid extra folder searches for projects we've already cached source files for.
 		/// </summary>
-		public static readonly Dictionary<string, ProjectFile> ModuleToProjectFileMap = new Dictionary<string, ProjectFile>( StringComparer.InvariantCultureIgnoreCase );
+		protected Dictionary<string, ProjectFile> ModuleToProjectFileMap = new Dictionary<string, ProjectFile>( StringComparer.InvariantCultureIgnoreCase );
 
 		/// <summary>
 		/// If generating project files for a single project, the path to its .uproject file.
 		/// </summary>
 		public readonly FileReference OnlyGameProject;
-
-		/// <summary>
-		/// When generating IntelliSense data, we may want to only generate data for a specific project file, even if other targets make use of modules
-		/// in this project file.  This is useful to prevent unusual or hacky global definitions from Programs affecting the Editor/Engine modules.  We
-		/// always want the most common and useful definitions to be set when working with solutions with many modules.
-		/// </summary>
-		public static ProjectFile OnlyGenerateIntelliSenseDataForProject
-		{
-			get;
-			private set;
-		}
-
 	
 		/// <summary>
 		/// File extension for project files we'll be generating (e.g. ".vcxproj")
@@ -957,11 +944,11 @@ namespace UnrealBuildTool
 
 				// Generate IntelliSense data if we need to.  This involves having UBT simulate the action compilation of
 				// the targets so that we can extra the compiler defines, include paths, etc.
-				if (GenerateIntelliSenseData(Arguments, IntelliSenseTargetFiles))
-				{
-					WriteProjectFiles(PlatformProjectGenerators);
-					Log.TraceVerbose( "Project generation complete ({0} generated, {1} imported)", GeneratedProjectFiles.Count, OtherProjectFiles.Count );
-				}
+				GenerateIntelliSenseData(Arguments, IntelliSenseTargetFiles);
+
+				// Write the project files
+				WriteProjectFiles(PlatformProjectGenerators);
+				Log.TraceVerbose( "Project generation complete ({0} generated, {1} imported)", GeneratedProjectFiles.Count, OtherProjectFiles.Count );
 			}
 
 			return bSuccess;
@@ -1686,12 +1673,11 @@ namespace UnrealBuildTool
 		/// <param name="Arguments">Incoming command-line arguments to UBT</param>
 		/// <param name="Targets">Targets to build for intellisense</param>
 		/// <return>Whether the process was successful or not</return>
-		private bool GenerateIntelliSenseData( String[] Arguments, List<Tuple<ProjectFile, ProjectTarget>> Targets )
+		private void GenerateIntelliSenseData(string[] Arguments, List<Tuple<ProjectFile, ProjectTarget>> Targets)
 		{
-			bool bSuccess = true;
-			if( ShouldGenerateIntelliSenseData() && Targets.Count > 0 )
+			if(ShouldGenerateIntelliSenseData() && Targets.Count > 0)
 			{
-				BuildConfiguration BuildConfiguration = new BuildConfiguration();
+				ReadOnlyBuildVersion CurrentBuildVersion = new ReadOnlyBuildVersion(BuildVersion.ReadDefault());
 
 				string ProgressInfoText = Utils.IsRunningOnMono ? "Generating data for project indexing..." : "Binding IntelliSense data...";
 				using(ProgressWriter Progress = new ProgressWriter(ProgressInfoText, true))
@@ -1701,12 +1687,9 @@ namespace UnrealBuildTool
 						ProjectFile TargetProjectFile = Targets[ TargetIndex ].Item1;
 						ProjectTarget CurTarget = Targets[ TargetIndex ].Item2;
 
-						Log.TraceVerbose( "Found target: " + CurTarget.TargetFilePath.GetFileNameWithoutAnyExtensions() );
+						Log.TraceVerbose("Found target: " + CurTarget.Name);
 
 						List<string> NewArguments = new List<string>(Arguments.Length + 4);
-						NewArguments.Add(CurTarget.TargetFilePath.GetFileNameWithoutAnyExtensions());
-						NewArguments.Add(BuildHostPlatform.Current.Platform.ToString());
-						NewArguments.Add(UnrealTargetConfiguration.Development.ToString());
 						if(CurTarget.TargetRules.Type != TargetType.Program)
 						{
 							NewArguments.Add("-precompile");
@@ -1714,29 +1697,50 @@ namespace UnrealBuildTool
 						NewArguments.Add("-Monolithic"); // Don't want any _API macros to be DLLIMPORT; this will cause Intellisense errors in the file that they are defined.
 						NewArguments.AddRange(Arguments);
 
-						// We only want to update definitions and include paths for modules that are part of this target's project file.
-						ProjectFileGenerator.OnlyGenerateIntelliSenseDataForProject = TargetProjectFile;
-
 						// Clear the CachedIncludePaths field on every FileItem. It's faster to keep the cache of file items around, but we need to wipe
 						// out any state.
 						FileItem.ClearCachedIncludePaths();
 
 						try
 						{
-							// Run UnrealBuildTool, pretending to build this target but instead only gathering data for IntelliSense (include paths and definitions).
-							// No actual compiling or linking will happen because we early out using the ProjectFileGenerator.bGenerateProjectFiles global
-							bSuccess = UnrealBuildTool.RunUBT( BuildConfiguration, NewArguments.ToArray(), CurTarget.UnrealProjectFilePath, false ) == ECompilationResult.Succeeded;
-						}
-						catch(Exception Ex)
-						{
-							Log.TraceWarning("Exception while generating include data for {0}: {1}", CurTarget.TargetFilePath.GetFileNameWithoutAnyExtensions(), Ex.ToString());
-						}
+							const string DefaultArchitecture = "";
+
+							// Create the target descriptor
+							TargetDescriptor TargetDesc = new TargetDescriptor(CurTarget.UnrealProjectFilePath, CurTarget.Name, BuildHostPlatform.Current.Platform, UnrealTargetConfiguration.Development, DefaultArchitecture, NewArguments.ToArray());
+
+							// Create the target
+							UEBuildTarget Target = UEBuildTarget.CreateTarget(TargetDesc, NewArguments.ToArray(), false, false, bUsePrecompiled, CurrentBuildVersion);
 
-						ProjectFileGenerator.OnlyGenerateIntelliSenseDataForProject = null;
+							// Generate a compile environment for each module in the binary
+							CppCompileEnvironment GlobalCompileEnvironment = Target.CreateCompileEnvironmentForProjectFiles();
+							foreach(UEBuildBinary Binary in Target.Binaries)
+							{
+								CppCompileEnvironment BinaryCompileEnvironment = Binary.CreateBinaryCompileEnvironment(GlobalCompileEnvironment);
+								foreach(UEBuildModuleCPP Module in Binary.Modules.OfType<UEBuildModuleCPP>())
+								{
+									ProjectFile ProjectFileForIDE;
+									if (ModuleToProjectFileMap.TryGetValue(Module.Name, out ProjectFileForIDE) && ProjectFileForIDE == TargetProjectFile)
+									{
+										CppCompileEnvironment ModuleCompileEnvironment = Module.CreateModuleCompileEnvironment(Target.Rules, BinaryCompileEnvironment);
+										ProjectFileForIDE.AddIntelliSensePreprocessorDefinitions(ModuleCompileEnvironment.Definitions);
+										ProjectFileForIDE.AddIntelliSenseIncludePaths(ModuleCompileEnvironment.IncludePaths.SystemIncludePaths, true);
+										ProjectFileForIDE.AddIntelliSenseIncludePaths(ModuleCompileEnvironment.IncludePaths.UserIncludePaths, false);
+									}
+								}
+							}
 
-						if( !bSuccess )
+							// If we're generating project files, then go ahead and wipe out the existing UBTMakefile for every target, to make sure that
+							// it gets a full dependency scan next time.
+							// NOTE: This is just a safeguard and doesn't have to be perfect.  We also check for newer project file timestamps in LoadUBTMakefile()
+							FileReference UBTMakefilePath = UBTMakefile.GetUBTMakefilePath(new List<TargetDescriptor>{ TargetDesc }, EHotReload.Disabled);
+							if (FileReference.Exists(UBTMakefilePath))
+							{
+								FileReference.Delete(UBTMakefilePath);
+							}
+						}
+						catch(Exception Ex)
 						{
-							break;
+							Log.TraceWarning("Exception while generating include data for {0}: {1}", CurTarget.Name, Ex.ToString());
 						}
 
 						// Display progress
@@ -1744,8 +1748,6 @@ namespace UnrealBuildTool
 					}
 				}
 			}
-
-			return bSuccess;
 		}
 
 
diff --git a/Engine/Source/Programs/UnrealBuildTool/ProjectFiles/VisualStudio/VCProject.cs b/Engine/Source/Programs/UnrealBuildTool/ProjectFiles/VisualStudio/VCProject.cs
--- a/Engine/Source/Programs/UnrealBuildTool/ProjectFiles/VisualStudio/VCProject.cs
+++ b/Engine/Source/Programs/UnrealBuildTool/ProjectFiles/VisualStudio/VCProject.cs
@@ -286,6 +286,7 @@ namespace UnrealBuildTool
 		VCProjectFileFormat ProjectFileFormat;
 		bool bUseFastPDB;
 		bool bUsePerFileIntellisense;
+		bool bUsePrecompiled;
 		string BuildToolOverride;
 
 		/// This is the platform name that Visual Studio is always guaranteed to support.  We'll use this as
@@ -307,14 +308,16 @@ namespace UnrealBuildTool
 		/// <param name="InProjectFileFormat">Visual C++ project file version</param>
 		/// <param name="bUseFastPDB">If true, adds the -FastPDB argument to build command lines</param>
 		/// <param name="bUsePerFileIntellisense">If true, generates per-file intellisense data</param>
+		/// <param name="bUsePrecompiled">Whether to add the -UsePrecompiled argumemnt when building targets</param>
 		/// <param name="BuildToolOverride">Optional arguments to pass to UBT when building</param>
-		public VCProjectFile(FileReference InFilePath, FileReference InOnlyGameProject, VCProjectFileFormat InProjectFileFormat, bool bUseFastPDB, bool bUsePerFileIntellisense, string BuildToolOverride)
+		public VCProjectFile(FileReference InFilePath, FileReference InOnlyGameProject, VCProjectFileFormat InProjectFileFormat, bool bUseFastPDB, bool bUsePerFileIntellisense, bool bUsePrecompiled, string BuildToolOverride)
 			: base(InFilePath)
 		{
 			OnlyGameProject = InOnlyGameProject;
 			ProjectFileFormat = InProjectFileFormat;
 			this.bUseFastPDB = bUseFastPDB;
 			this.bUsePerFileIntellisense = bUsePerFileIntellisense;
+			this.bUsePrecompiled = bUsePrecompiled;
 			this.BuildToolOverride = BuildToolOverride;
 		}
 
@@ -1424,7 +1427,7 @@ namespace UnrealBuildTool
 					//	ie ..\..\Build\BatchFiles\Build.bat BlankProgram Win64 Debug
 
 					string BuildArguments = " " + TargetName + " " + UBTPlatformName + " " + UBTConfigurationName;
-					if (ProjectFileGenerator.bUsePrecompiled)
+					if (bUsePrecompiled)
 					{
 						BuildArguments += " -UsePrecompiled";
 					}
diff --git a/Engine/Source/Programs/UnrealBuildTool/ProjectFiles/VisualStudio/VCProjectFileGenerator.cs b/Engine/Source/Programs/UnrealBuildTool/ProjectFiles/VisualStudio/VCProjectFileGenerator.cs
--- a/Engine/Source/Programs/UnrealBuildTool/ProjectFiles/VisualStudio/VCProjectFileGenerator.cs
+++ b/Engine/Source/Programs/UnrealBuildTool/ProjectFiles/VisualStudio/VCProjectFileGenerator.cs
@@ -174,7 +174,7 @@ namespace UnrealBuildTool
 		/// <returns>The newly allocated project file object</returns>
 		protected override ProjectFile AllocateProjectFile(FileReference InitFilePath)
 		{
-			return new VCProjectFile(InitFilePath, OnlyGameProject, ProjectFileFormat, bAddFastPDBToProjects, bUsePerFileIntellisense, BuildToolOverride);
+			return new VCProjectFile(InitFilePath, OnlyGameProject, ProjectFileFormat, bAddFastPDBToProjects, bUsePerFileIntellisense, bUsePrecompiled, BuildToolOverride);
 		}
 
 
diff --git a/Engine/Source/Programs/UnrealBuildTool/UnrealBuildTool.cs b/Engine/Source/Programs/UnrealBuildTool/UnrealBuildTool.cs
--- a/Engine/Source/Programs/UnrealBuildTool/UnrealBuildTool.cs
+++ b/Engine/Source/Programs/UnrealBuildTool/UnrealBuildTool.cs
@@ -741,58 +741,47 @@ namespace UnrealBuildTool
 				}
 				TargetDescriptor HotReloadTargetDesc = (HotReload != EHotReload.Disabled) ? TargetDescs[0] : null;
 
-				if (ProjectFileGenerator.bGenerateProjectFiles)
-				{
-					// Create empty timestamp file to record when was the last time we regenerated projects.
-					Directory.CreateDirectory(Path.GetDirectoryName(ProjectFileGenerator.ProjectTimestampFile));
-					File.Create(ProjectFileGenerator.ProjectTimestampFile).Dispose();
-				}
-
-
 				// Used when BuildConfiguration.bUseUBTMakefiles is enabled.  If true, it means that our cached includes may not longer be
 				// valid (or never were), and we need to recover by forcibly scanning included headers for all build prerequisites to make sure that our
 				// cached set of includes is actually correct, before determining which files are outdated.
 				bool bNeedsFullCPPIncludeRescan = false;
 
-				if (!ProjectFileGenerator.bGenerateProjectFiles)
+				if (BuildConfiguration.bUseUBTMakefiles)
 				{
-					if (BuildConfiguration.bUseUBTMakefiles)
+					// Only the modular editor and game targets will share build products.  Unfortunately, we can't determine at
+					// at this point whether we're dealing with modular or monolithic game binaries, so we opt to always invalidate
+					// cached includes if the target we're switching to is either a game target (has project file) or "UE4Editor".
+					bool bMightHaveSharedBuildProducts =
+						ProjectFile != null ||  // Is this a game? (has a .uproject file for the target)
+						TargetDescs[0].Name.Equals("UE4Editor", StringComparison.InvariantCultureIgnoreCase); // Is the engine?
+					if (bMightHaveSharedBuildProducts)
 					{
-						// Only the modular editor and game targets will share build products.  Unfortunately, we can't determine at
-						// at this point whether we're dealing with modular or monolithic game binaries, so we opt to always invalidate
-						// cached includes if the target we're switching to is either a game target (has project file) or "UE4Editor".
-						bool bMightHaveSharedBuildProducts =
-							ProjectFile != null ||  // Is this a game? (has a .uproject file for the target)
-							TargetDescs[0].Name.Equals("UE4Editor", StringComparison.InvariantCultureIgnoreCase); // Is the engine?
-						if (bMightHaveSharedBuildProducts)
-						{
-							bool bIsBuildingSameTargetsAsLastTime = false;
+						bool bIsBuildingSameTargetsAsLastTime = false;
 
-							string TargetCollectionName = UBTMakefile.MakeTargetCollectionName(TargetDescs);
+						string TargetCollectionName = UBTMakefile.MakeTargetCollectionName(TargetDescs);
 
-							string LastBuiltTargetsFileName = (HotReload != EHotReload.Disabled) ? "HotReloadLastBuiltTargets.txt" : "LastBuiltTargets.txt";
-							string LastBuiltTargetsFilePath = FileReference.Combine(UnrealBuildTool.EngineDirectory, "Intermediate", "Build", LastBuiltTargetsFileName).FullName;
-							if (File.Exists(LastBuiltTargetsFilePath) && Utils.ReadAllText(LastBuiltTargetsFilePath) == TargetCollectionName)
-							{
-								// @todo ubtmake: Because we're using separate files for hot reload vs. full compiles, it's actually possible that includes will
-								// become out of date without us knowing if the developer ping-pongs between hot reloading target A and building target B normally.
-								// To fix this we can not use a different file name for last built targets, but the downside is slower performance when
-								// performing the first hot reload after compiling normally (forces full include dependency scan)
-								bIsBuildingSameTargetsAsLastTime = true;
-							}
+						string LastBuiltTargetsFileName = (HotReload != EHotReload.Disabled) ? "HotReloadLastBuiltTargets.txt" : "LastBuiltTargets.txt";
+						string LastBuiltTargetsFilePath = FileReference.Combine(UnrealBuildTool.EngineDirectory, "Intermediate", "Build", LastBuiltTargetsFileName).FullName;
+						if (File.Exists(LastBuiltTargetsFilePath) && Utils.ReadAllText(LastBuiltTargetsFilePath) == TargetCollectionName)
+						{
+							// @todo ubtmake: Because we're using separate files for hot reload vs. full compiles, it's actually possible that includes will
+							// become out of date without us knowing if the developer ping-pongs between hot reloading target A and building target B normally.
+							// To fix this we can not use a different file name for last built targets, but the downside is slower performance when
+							// performing the first hot reload after compiling normally (forces full include dependency scan)
+							bIsBuildingSameTargetsAsLastTime = true;
+						}
 
-							// Save out the name of the targets we're building
-							if (!bIsBuildingSameTargetsAsLastTime)
-							{
-								Directory.CreateDirectory(Path.GetDirectoryName(LastBuiltTargetsFilePath));
-								File.WriteAllText(LastBuiltTargetsFilePath, TargetCollectionName, Encoding.UTF8);
-
-								// Can't use super fast include checking unless we're building the same set of targets as last time, because
-								// we might not know about all of the C++ include dependencies for already-up-to-date shared build products
-								// between the targets
-								bNeedsFullCPPIncludeRescan = true;
-								Log.TraceInformation("Performing full C++ include scan ({0} a new target)", (HotReload != EHotReload.Disabled) ? "hot reloading" : "building");
-							}
+						// Save out the name of the targets we're building
+						if (!bIsBuildingSameTargetsAsLastTime)
+						{
+							Directory.CreateDirectory(Path.GetDirectoryName(LastBuiltTargetsFilePath));
+							File.WriteAllText(LastBuiltTargetsFilePath, TargetCollectionName, Encoding.UTF8);
+
+							// Can't use super fast include checking unless we're building the same set of targets as last time, because
+							// we might not know about all of the C++ include dependencies for already-up-to-date shared build products
+							// between the targets
+							bNeedsFullCPPIncludeRescan = true;
+							Log.TraceInformation("Performing full C++ include scan ({0} a new target)", (HotReload != EHotReload.Disabled) ? "hot reloading" : "building");
 						}
 					}
 				}
@@ -852,23 +841,12 @@ namespace UnrealBuildTool
 				{
 					UBTMakefile UBTMakefile = null;
 					{
-						// If we're generating project files, then go ahead and wipe out the existing UBTMakefile for every target, to make sure that
-						// it gets a full dependency scan next time.
-						// NOTE: This is just a safeguard and doesn't have to be perfect.  We also check for newer project file timestamps in LoadUBTMakefile()
 						FileReference UBTMakefilePath = UBTMakefile.GetUBTMakefilePath(TargetDescs, HotReload);
-						if (ProjectFileGenerator.bGenerateProjectFiles) // @todo ubtmake: This is only hit when generating IntelliSense for project files.  Probably should be done right inside ProjectFileGenerator.bat
-						{                                                   // @todo ubtmake: Won't catch multi-target cases as GPF always builds one target at a time for Intellisense
-																			// Delete the UBTMakefile
-							if (FileReference.Exists(UBTMakefilePath))
-							{
-								FileReference.Delete(UBTMakefilePath);
-							}
-						}
 
 						// Make sure the gather phase is executed if we're not actually building anything
-						if (ProjectFileGenerator.bGenerateProjectFiles || BuildConfiguration.bGenerateManifest || BuildConfiguration.bCleanProject || BuildConfiguration.bXGEExport || GeneratingActionGraph)
+						if (BuildConfiguration.bGenerateManifest || BuildConfiguration.bCleanProject || BuildConfiguration.bXGEExport || GeneratingActionGraph)
 						{
-								bIsGatheringBuild = true;
+							bIsGatheringBuild = true;
 						}
 
 						// Were we asked to run in 'assembler only' mode?  If so, let's check to see if that's even possible by seeing if
@@ -1017,7 +995,7 @@ namespace UnrealBuildTool
 							{
 								if (!bNeedsFullCPPIncludeRescan)
 								{
-									if (!ProjectFileGenerator.bGenerateProjectFiles && !BuildConfiguration.bXGEExport && !BuildConfiguration.bGenerateManifest && !BuildConfiguration.bCleanProject)
+									if (!BuildConfiguration.bXGEExport && !BuildConfiguration.bGenerateManifest && !BuildConfiguration.bCleanProject)
 									{
 										bNeedsFullCPPIncludeRescan = true;
 										Log.TraceInformation("Performing full C++ include scan (no include cache file)");
@@ -1061,7 +1039,7 @@ namespace UnrealBuildTool
 					if (BuildResult == ECompilationResult.Succeeded && !BuildConfiguration.bSkipBuild &&
 						(
 							(BuildConfiguration.bXGEExport && BuildConfiguration.bGenerateManifest) ||
-							(!GeneratingActionGraph && !ProjectFileGenerator.bGenerateProjectFiles && !BuildConfiguration.bGenerateManifest && !BuildConfiguration.bCleanProject)
+							(!GeneratingActionGraph && !BuildConfiguration.bGenerateManifest && !BuildConfiguration.bCleanProject)
 						))
 					{
 						if (bIsGatheringBuild)
@@ -1131,7 +1109,7 @@ namespace UnrealBuildTool
 								}
 
 								// Execute all the pre-build steps
-								if (!ProjectFileGenerator.bGenerateProjectFiles && !BuildConfiguration.bXGEExport)
+								if (!BuildConfiguration.bXGEExport)
 								{
 									foreach (UEBuildTarget Target in Targets)
 									{
@@ -1223,7 +1201,6 @@ namespace UnrealBuildTool
 							// Run the deployment steps
 							if (BuildResult.Succeeded()
 								&& String.IsNullOrEmpty(BuildConfiguration.SingleFileToCompile)
-								&& !ProjectFileGenerator.bGenerateProjectFiles
 								&& !BuildConfiguration.bGenerateManifest
 								&& !BuildConfiguration.bCleanProject
 							&& !BuildConfiguration.bXGEExport)
@@ -1330,7 +1307,7 @@ namespace UnrealBuildTool
 			// this code must be able to execute before we create or load module rules DLLs so that hot reload can work with bUseUBTMakefiles
 			bool bIsEditorTarget = TargetDesc.Name.EndsWith("Editor", StringComparison.InvariantCultureIgnoreCase);
 
-			if (!ProjectFileGenerator.bGenerateProjectFiles && !BuildConfiguration.bGenerateManifest && bIsEditorTarget)
+			if (!BuildConfiguration.bGenerateManifest && bIsEditorTarget)
 			{
 				string EditorBaseFileName = "UE4Editor";
 				if (TargetDesc.Configuration != UnrealTargetConfiguration.Development && TargetDesc.Configuration != UnrealTargetConfiguration.DebugGame)
