UBT: Add a global ReadOnlyBuildVersion instance, rather than having to manually create it and pass it around as an argument.

#rnx

diff --git a/Engine/Source/Programs/AutomationTool/AutomationUtils/ProjectUtils.cs b/Engine/Source/Programs/AutomationTool/AutomationUtils/ProjectUtils.cs
--- a/Engine/Source/Programs/AutomationTool/AutomationUtils/ProjectUtils.cs
+++ b/Engine/Source/Programs/AutomationTool/AutomationUtils/ProjectUtils.cs
@@ -564,8 +564,7 @@ namespace AutomationTool
 				{
 					string TargetName = GetTargetName(TargetType);
 
-					ReadOnlyBuildVersion Version = new ReadOnlyBuildVersion(BuildVersion.ReadDefault());
-					TargetInfo DummyTargetInfo = new TargetInfo(TargetName, BuildHostPlatform.Current.Platform, UnrealTargetConfiguration.Development, "", Properties.RawProjectPath, Version);
+					TargetInfo DummyTargetInfo = new TargetInfo(TargetName, BuildHostPlatform.Current.Platform, UnrealTargetConfiguration.Development, "", Properties.RawProjectPath);
 
 					// Create an instance of this type
 					CommandUtils.LogVerbose("Creating target rules object: {0}", TargetType.Name);
diff --git a/Engine/Source/Programs/UnrealBuildTool/Configuration/TargetDescriptor.cs b/Engine/Source/Programs/UnrealBuildTool/Configuration/TargetDescriptor.cs
--- a/Engine/Source/Programs/UnrealBuildTool/Configuration/TargetDescriptor.cs
+++ b/Engine/Source/Programs/UnrealBuildTool/Configuration/TargetDescriptor.cs
@@ -171,7 +171,7 @@ namespace UnrealBuildTool
 				}
 				else
 				{
-					TargetNames.Add(RulesCompiler.CreateProjectRulesAssembly(ProjectFile, bUsePrecompiled, bSkipRulesCompile).GetTargetNameByType(Type, Platform, Configuration, Architecture, ProjectFile, new ReadOnlyBuildVersion(BuildVersion.ReadDefault())));
+					TargetNames.Add(RulesCompiler.CreateProjectRulesAssembly(ProjectFile, bUsePrecompiled, bSkipRulesCompile).GetTargetNameByType(Type, Platform, Configuration, Architecture, ProjectFile));
 				}
 			}
 
diff --git a/Engine/Source/Programs/UnrealBuildTool/Configuration/TargetInfo.cs b/Engine/Source/Programs/UnrealBuildTool/Configuration/TargetInfo.cs
--- a/Engine/Source/Programs/UnrealBuildTool/Configuration/TargetInfo.cs
+++ b/Engine/Source/Programs/UnrealBuildTool/Configuration/TargetInfo.cs
@@ -43,7 +43,10 @@ namespace UnrealBuildTool
 		/// <summary>
 		/// The current build version
 		/// </summary>
-		public readonly ReadOnlyBuildVersion Version;
+		public ReadOnlyBuildVersion Version
+		{
+			get { return ReadOnlyBuildVersion.Current; }
+		}
 
 		/// <summary>
 		/// Constructs a TargetInfo for passing to the TargetRules constructor.
@@ -53,15 +56,13 @@ namespace UnrealBuildTool
 		/// <param name="Configuration">The configuration being built</param>
 		/// <param name="Architecture">The architecture being built for</param>
 		/// <param name="ProjectFile">Path to the project file containing the target</param>
-		/// <param name="Version">The current build version</param>
-		public TargetInfo(string Name, UnrealTargetPlatform Platform, UnrealTargetConfiguration Configuration, string Architecture, FileReference ProjectFile, ReadOnlyBuildVersion Version)
+		public TargetInfo(string Name, UnrealTargetPlatform Platform, UnrealTargetConfiguration Configuration, string Architecture, FileReference ProjectFile)
 		{
 			this.Name = Name;
 			this.Platform = Platform;
 			this.Configuration = Configuration;
 			this.Architecture = Architecture;
 			this.ProjectFile = ProjectFile;
-			this.Version = Version;
 		}
 	}
 }
diff --git a/Engine/Source/Programs/UnrealBuildTool/Configuration/UEBuildTarget.cs b/Engine/Source/Programs/UnrealBuildTool/Configuration/UEBuildTarget.cs
--- a/Engine/Source/Programs/UnrealBuildTool/Configuration/UEBuildTarget.cs
+++ b/Engine/Source/Programs/UnrealBuildTool/Configuration/UEBuildTarget.cs
@@ -377,15 +377,14 @@ namespace UnrealBuildTool
 		/// <param name="bSkipRulesCompile">Whether to skip compiling any rules assemblies</param>
 		/// <param name="bCompilingSingleFile">Whether we're compiling a single file</param>
 		/// <param name="bUsePrecompiled">Whether to use a precompiled engine/enterprise build</param>
-		/// <param name="Version">The current build version</param>
 		/// <returns>The build target object for the specified build rules source file</returns>
-		public static UEBuildTarget CreateTarget(TargetDescriptor Desc, string[] Arguments, bool bSkipRulesCompile, bool bCompilingSingleFile, bool bUsePrecompiled, ReadOnlyBuildVersion Version)
+		public static UEBuildTarget CreateTarget(TargetDescriptor Desc, string[] Arguments, bool bSkipRulesCompile, bool bCompilingSingleFile, bool bUsePrecompiled)
 		{
 			DateTime CreateTargetStartTime = DateTime.UtcNow;
 
 			RulesAssembly RulesAssembly = RulesCompiler.CreateTargetRulesAssembly(Desc.ProjectFile, Desc.Name, bSkipRulesCompile, bUsePrecompiled, Desc.ForeignPlugin);
 
-			TargetRules RulesObject = RulesAssembly.CreateTargetRules(Desc.Name, Desc.Platform, Desc.Configuration, Desc.Architecture, Desc.ProjectFile, Version, Arguments);
+			TargetRules RulesObject = RulesAssembly.CreateTargetRules(Desc.Name, Desc.Platform, Desc.Configuration, Desc.Architecture, Desc.ProjectFile, Arguments);
 			if ((ProjectFileGenerator.bGenerateProjectFiles == false) && !GetSupportedPlatforms(RulesObject).Contains(Desc.Platform))
 			{
 				throw new BuildException("{0} does not support the {1} platform.", Desc.Name, Desc.Platform.ToString());
@@ -527,7 +526,7 @@ namespace UnrealBuildTool
 			}
 
 			// Create the target rules for it
-			TargetRules BaseRules = RulesAssembly.CreateTargetRules(BaseTargetName, ThisRules.Platform, ThisRules.Configuration, ThisRules.Architecture, null, ThisRules.Version, null);
+			TargetRules BaseRules = RulesAssembly.CreateTargetRules(BaseTargetName, ThisRules.Platform, ThisRules.Configuration, ThisRules.Architecture, null, null);
 
 			// Get all the configurable objects
 			object[] BaseObjects = BaseRules.GetConfigurableObjects().ToArray();
diff --git a/Engine/Source/Programs/UnrealBuildTool/Platform/Android/UEDeployAndroid.cs b/Engine/Source/Programs/UnrealBuildTool/Platform/Android/UEDeployAndroid.cs
--- a/Engine/Source/Programs/UnrealBuildTool/Platform/Android/UEDeployAndroid.cs
+++ b/Engine/Source/Programs/UnrealBuildTool/Platform/Android/UEDeployAndroid.cs
@@ -2056,7 +2056,7 @@ namespace UnrealBuildTool
 		{
 			if (!bHaveReadEngineVersion)
 			{
-				BuildVersion Version = BuildVersion.ReadDefault();
+				ReadOnlyBuildVersion Version = ReadOnlyBuildVersion.Current;
 
 				EngineMajorVersion = Version.MajorVersion.ToString();
 				EngineMinorVersion = Version.MinorVersion.ToString();
@@ -3215,7 +3215,7 @@ namespace UnrealBuildTool
 					NDKArches.Add(NDKArch);
 				}
 			}
-			UPL.Init(NDKArches, bForDistribution, EngineDirectory, UE4BuildPath, ProjectDirectory, Configuration, BuildVersion.ReadDefault());
+			UPL.Init(NDKArches, bForDistribution, EngineDirectory, UE4BuildPath, ProjectDirectory, Configuration);
 
 			IEnumerable<Tuple<string, string, string>> BuildList = null;
 
diff --git a/Engine/Source/Programs/UnrealBuildTool/Platform/IOS/UEDeployIOS.cs b/Engine/Source/Programs/UnrealBuildTool/Platform/IOS/UEDeployIOS.cs
--- a/Engine/Source/Programs/UnrealBuildTool/Platform/IOS/UEDeployIOS.cs
+++ b/Engine/Source/Programs/UnrealBuildTool/Platform/IOS/UEDeployIOS.cs
@@ -161,7 +161,7 @@ namespace UnrealBuildTool
 				else
 				{
                     // get the changelist
-                    CFBundleVersion = BuildVersion.ReadDefault().Changelist.ToString();
+                    CFBundleVersion = ReadOnlyBuildVersion.Current.Changelist.ToString();
 
 				}
 
@@ -795,7 +795,7 @@ namespace UnrealBuildTool
 			UnrealPluginLanguage UPL = new UnrealPluginLanguage(ProjectFile, UPLScripts, ProjectArches, "", "", UnrealTargetPlatform.IOS);
 
 			// Passing in true for distribution is not ideal here but given the way that ios packaging happens and this call chain it seems unavoidable for now, maybe there is a way to correctly pass it in that I can't find?
-			UPL.Init(ProjectArches, true, RelativeEnginePath, BundlePath, ProjectDirectory, Config.ToString(), BuildVersion.ReadDefault());
+			UPL.Init(ProjectArches, true, RelativeEnginePath, BundlePath, ProjectDirectory, Config.ToString());
 
 			return GenerateIOSPList(ProjectFile, Config, ProjectDirectory, bIsUE4Game, GameName, ProjectName, InEngineDir, AppDirectory, SdkVersion, UPL, out bSupportsPortrait, out bSupportsLandscape, out bSkipIcons);
 		}
diff --git a/Engine/Source/Programs/UnrealBuildTool/Platform/Lumin/UEDeployLumin.cs b/Engine/Source/Programs/UnrealBuildTool/Platform/Lumin/UEDeployLumin.cs
--- a/Engine/Source/Programs/UnrealBuildTool/Platform/Lumin/UEDeployLumin.cs
+++ b/Engine/Source/Programs/UnrealBuildTool/Platform/Lumin/UEDeployLumin.cs
@@ -285,7 +285,7 @@ namespace UnrealBuildTool
 			SetLuminPluginData(MLSDKArches, CollectPluginDataPaths(TargetReceipt.Read(ReceiptFilename)));
 
 			//gather all of the xml
-			UPL.Init(MLSDKArches, true, RelativeEnginePath, UE4BuildPath, RelativeProjectPath, ConfigurationString, BuildVersion.ReadDefault());
+			UPL.Init(MLSDKArches, true, RelativeEnginePath, UE4BuildPath, RelativeProjectPath, ConfigurationString);
 		}
 
 		public string StageFiles()
diff --git a/Engine/Source/Programs/UnrealBuildTool/ProjectFiles/ProjectFileGenerator.cs b/Engine/Source/Programs/UnrealBuildTool/ProjectFiles/ProjectFileGenerator.cs
--- a/Engine/Source/Programs/UnrealBuildTool/ProjectFiles/ProjectFileGenerator.cs
+++ b/Engine/Source/Programs/UnrealBuildTool/ProjectFiles/ProjectFileGenerator.cs
@@ -1677,8 +1677,6 @@ namespace UnrealBuildTool
 		{
 			if(ShouldGenerateIntelliSenseData() && Targets.Count > 0)
 			{
-				ReadOnlyBuildVersion CurrentBuildVersion = new ReadOnlyBuildVersion(BuildVersion.ReadDefault());
-
 				string ProgressInfoText = Utils.IsRunningOnMono ? "Generating data for project indexing..." : "Binding IntelliSense data...";
 				using(ProgressWriter Progress = new ProgressWriter(ProgressInfoText, true))
 				{
@@ -1709,7 +1707,7 @@ namespace UnrealBuildTool
 							TargetDescriptor TargetDesc = new TargetDescriptor(CurTarget.UnrealProjectFilePath, CurTarget.Name, BuildHostPlatform.Current.Platform, UnrealTargetConfiguration.Development, DefaultArchitecture, NewArguments.ToArray());
 
 							// Create the target
-							UEBuildTarget Target = UEBuildTarget.CreateTarget(TargetDesc, NewArguments.ToArray(), false, false, bUsePrecompiled, CurrentBuildVersion);
+							UEBuildTarget Target = UEBuildTarget.CreateTarget(TargetDesc, NewArguments.ToArray(), false, false, bUsePrecompiled);
 
 							// Generate a compile environment for each module in the binary
 							CppCompileEnvironment GlobalCompileEnvironment = Target.CreateCompileEnvironmentForProjectFiles();
@@ -2031,9 +2029,6 @@ namespace UnrealBuildTool
 			DirectoryReference EngineSourceProgramsDirectory = DirectoryReference.Combine(UnrealBuildTool.EngineSourceDirectory, "Programs");
 			DirectoryReference EnterpriseSourceProgramsDirectory = DirectoryReference.Combine(UnrealBuildTool.EnterpriseSourceDirectory, "Programs");
 
-			// Read the current engine version
-			ReadOnlyBuildVersion Version = new ReadOnlyBuildVersion(BuildVersion.ReadDefault());
-
 			// Find all of the target files.  This will filter out any modules or targets that don't
 			// belong to platforms we're generating project files for.
 			List<FileReference> AllTargetFiles = DiscoverTargets(AllGames);
@@ -2101,7 +2096,7 @@ namespace UnrealBuildTool
 
 					// Create target rules for all of the platforms and configuration combinations that we want to enable support for.
 					// Just use the current platform as we only need to recover the target type and both should be supported for all targets...
-					TargetRules TargetRulesObject = RulesAssembly.CreateTargetRules(TargetName, BuildHostPlatform.Current.Platform, UnrealTargetConfiguration.Development, "", CheckProjectFile, Version, null);
+					TargetRules TargetRulesObject = RulesAssembly.CreateTargetRules(TargetName, BuildHostPlatform.Current.Platform, UnrealTargetConfiguration.Development, "", CheckProjectFile, null);
 
 					bool IsProgramTarget = false;
 
@@ -2257,7 +2252,7 @@ namespace UnrealBuildTool
 							ProjectFilePath = ProjectFilePath,
 							UnrealProjectFilePath = CheckProjectFile,
 							SupportedPlatforms = UEBuildTarget.GetSupportedPlatforms(TargetRulesObject).Where(x => UEBuildPlatform.GetBuildPlatform(x, true) != null).ToArray(),
-							CreateRulesDelegate = (Platform, Configuration) => RulesAssembly.CreateTargetRules(TargetName, Platform, Configuration, "", CheckProjectFile, Version, null)
+							CreateRulesDelegate = (Platform, Configuration) => RulesAssembly.CreateTargetRules(TargetName, Platform, Configuration, "", CheckProjectFile, null)
                         };
 
 					if (TargetName == "ShaderCompileWorker")		// @todo projectfiles: Ideally, the target rules file should set this
diff --git a/Engine/Source/Programs/UnrealBuildTool/System/BuildVersion.cs b/Engine/Source/Programs/UnrealBuildTool/System/BuildVersion.cs
--- a/Engine/Source/Programs/UnrealBuildTool/System/BuildVersion.cs
+++ b/Engine/Source/Programs/UnrealBuildTool/System/BuildVersion.cs
@@ -75,27 +75,6 @@ namespace UnrealBuildTool
 			get { return (Changelist != 0 && CompatibleChangelist != 0)? CompatibleChangelist : Changelist; }
 		}
 
-		/// <summary>
-		/// Reads the default build version, throwing an exception on error.
-		/// </summary>
-		/// <returns>New BuildVersion instance</returns>
-		public static BuildVersion ReadDefault()
-		{
-			FileReference File = GetDefaultFileName();
-			if(!FileReference.Exists(File))
-			{
-				throw new BuildException("Version file is missing ({0})", File);
-			}
-
-			BuildVersion Version;
-			if(!TryRead(File, out Version))
-			{
-				throw new BuildException("Unable to read version file ({0}). Check that this file is present and well-formed JSON.", File);
-			}
-
-			return Version;
-		}
-
 		/// <summary>
 		/// Try to read a version file from disk
 		/// </summary>
@@ -231,7 +210,15 @@ namespace UnrealBuildTool
 	/// </summary>
 	public class ReadOnlyBuildVersion
 	{
-		BuildVersion Inner;
+		/// <summary>
+		/// The inner build version
+		/// </summary>
+		private BuildVersion Inner;
+
+		/// <summary>
+		/// Cached copy of the current build version
+		/// </summary>
+		private static ReadOnlyBuildVersion CurrentCached;
 
 		/// <summary>
 		/// Constructor
@@ -242,6 +229,33 @@ namespace UnrealBuildTool
 			this.Inner = Inner;
 		}
 
+		/// <summary>
+		/// Gets the current build version
+		/// </summary>
+		public static ReadOnlyBuildVersion Current
+		{
+			get
+			{
+				if(CurrentCached == null)
+				{
+					FileReference File = BuildVersion.GetDefaultFileName();
+					if(!FileReference.Exists(File))
+					{
+						throw new BuildException("Version file is missing ({0})", File);
+					}
+
+					BuildVersion Version;
+					if(!BuildVersion.TryRead(File, out Version))
+					{
+						throw new BuildException("Unable to read version file ({0}). Check that this file is present and well-formed JSON.", File);
+					}
+
+					CurrentCached = new ReadOnlyBuildVersion(Version);
+				}
+				return CurrentCached;
+			}
+		}
+
 		/// <summary>
 		/// Accessors for fields on the inner BuildVersion instance
 		/// </summary>
diff --git a/Engine/Source/Programs/UnrealBuildTool/System/RulesAssembly.cs b/Engine/Source/Programs/UnrealBuildTool/System/RulesAssembly.cs
--- a/Engine/Source/Programs/UnrealBuildTool/System/RulesAssembly.cs
+++ b/Engine/Source/Programs/UnrealBuildTool/System/RulesAssembly.cs
@@ -521,10 +521,9 @@ namespace UnrealBuildTool
 		/// <param name="Configuration">Configuration being compiled</param>
 		/// <param name="Architecture">Architecture being built</param>
 		/// <param name="ProjectFile">Path to the project file for this target</param>
-		/// <param name="Version">The current build version</param>
 		/// <param name="Arguments">Command line arguments for this target</param>
 		/// <returns>The build target rules for the specified target</returns>
-		public TargetRules CreateTargetRules(string TargetName, UnrealTargetPlatform Platform, UnrealTargetConfiguration Configuration, string Architecture, FileReference ProjectFile, ReadOnlyBuildVersion Version, string[] Arguments)
+		public TargetRules CreateTargetRules(string TargetName, UnrealTargetPlatform Platform, UnrealTargetConfiguration Configuration, string Architecture, FileReference ProjectFile, string[] Arguments)
 		{
 			bool bFoundTargetName = TargetNameToTargetFile.ContainsKey(TargetName);
 			if (bFoundTargetName == false)
@@ -550,7 +549,7 @@ namespace UnrealBuildTool
 				}
 				else
 				{
-					return Parent.CreateTargetRules(TargetName, Platform, Configuration, Architecture, ProjectFile, Version, Arguments);
+					return Parent.CreateTargetRules(TargetName, Platform, Configuration, Architecture, ProjectFile, Arguments);
 				}
 			}
 
@@ -558,7 +557,7 @@ namespace UnrealBuildTool
 			string TargetTypeName = TargetName + "Target";
 
 			// The build module must define a type named '<TargetName>Target' that derives from our 'TargetRules' type.  
-			return CreateTargetRulesInstance(TargetTypeName, new TargetInfo(TargetName, Platform, Configuration, Architecture, ProjectFile, Version), Arguments);
+			return CreateTargetRulesInstance(TargetTypeName, new TargetInfo(TargetName, Platform, Configuration, Architecture, ProjectFile), Arguments);
 		}
 
 		/// <summary>
@@ -569,15 +568,14 @@ namespace UnrealBuildTool
 		/// <param name="Configuration">The configuration being built</param>
 		/// <param name="Architecture">The architecture being built</param>
 		/// <param name="ProjectFile">Project file for the target being built</param>
-		/// <param name="Version">The current engine version information</param>
 		/// <returns>Name of the target for the given type</returns>
-		public string GetTargetNameByType(TargetType Type, UnrealTargetPlatform Platform, UnrealTargetConfiguration Configuration, string Architecture, FileReference ProjectFile, ReadOnlyBuildVersion Version)
+		public string GetTargetNameByType(TargetType Type, UnrealTargetPlatform Platform, UnrealTargetConfiguration Configuration, string Architecture, FileReference ProjectFile)
 		{
 			// Create all the targets in this assembly 
 			List<string> Matches = new List<string>();
 			foreach(KeyValuePair<string, FileReference> TargetPair in TargetNameToTargetFile)
 			{
-				TargetRules Rules = CreateTargetRulesInstance(TargetPair.Key + "Target", new TargetInfo(TargetPair.Key, Platform, Configuration, Architecture, ProjectFile, Version), null);
+				TargetRules Rules = CreateTargetRulesInstance(TargetPair.Key + "Target", new TargetInfo(TargetPair.Key, Platform, Configuration, Architecture, ProjectFile), null);
 				if(Rules.Type == Type)
 				{
 					Matches.Add(TargetPair.Key);
@@ -593,7 +591,7 @@ namespace UnrealBuildTool
 				}
 				else
 				{
-					return Parent.GetTargetNameByType(Type, Platform, Configuration, Architecture, ProjectFile, Version);
+					return Parent.GetTargetNameByType(Type, Platform, Configuration, Architecture, ProjectFile);
 				}
 			}
 			else
diff --git a/Engine/Source/Programs/UnrealBuildTool/System/UnrealPluginLanguage.cs b/Engine/Source/Programs/UnrealBuildTool/System/UnrealPluginLanguage.cs
--- a/Engine/Source/Programs/UnrealBuildTool/System/UnrealPluginLanguage.cs
+++ b/Engine/Source/Programs/UnrealBuildTool/System/UnrealPluginLanguage.cs
@@ -2401,7 +2401,7 @@ namespace UnrealBuildTool
 			return GlobalContext.StringVariables["Output"];
 		}
 
-		public void Init(List<string> Architectures, bool bDistribution, string EngineDirectory, string BuildDirectory, string ProjectDirectory, string Configuration, BuildVersion Version)
+		public void Init(List<string> Architectures, bool bDistribution, string EngineDirectory, string BuildDirectory, string ProjectDirectory, string Configuration)
 		{
 			GlobalContext.BoolVariables["Distribution"] = bDistribution;
 			GlobalContext.StringVariables["Configuration"] = Configuration;
@@ -2410,6 +2410,7 @@ namespace UnrealBuildTool
 			GlobalContext.StringVariables["BuildDir"] = BuildDirectory.Replace("\\", "/");
 			GlobalContext.StringVariables["ProjectDir"] = ProjectDirectory.Replace("\\", "/");
 
+			ReadOnlyBuildVersion Version = ReadOnlyBuildVersion.Current;
 			GlobalContext.IntVariables["EngineMajorVersion"] = Version.MajorVersion;
 			GlobalContext.IntVariables["EngineMinorVersion"] = Version.MinorVersion;
 			GlobalContext.IntVariables["EnginePatchVersion"] = Version.PatchVersion;
diff --git a/Engine/Source/Programs/UnrealBuildTool/ToolChain/RemoteMac.cs b/Engine/Source/Programs/UnrealBuildTool/ToolChain/RemoteMac.cs
--- a/Engine/Source/Programs/UnrealBuildTool/ToolChain/RemoteMac.cs
+++ b/Engine/Source/Programs/UnrealBuildTool/ToolChain/RemoteMac.cs
@@ -362,7 +362,7 @@ namespace UnrealBuildTool
 			RulesAssembly RulesAssembly = RulesCompiler.CreateTargetRulesAssembly(TargetDesc.ProjectFile, TargetDesc.Name, false, false, TargetDesc.ForeignPlugin);
 
 			// Create the target rules
-			TargetRules Rules = RulesAssembly.CreateTargetRules(TargetDesc.Name, TargetDesc.Platform, TargetDesc.Configuration, TargetDesc.Architecture, TargetDesc.ProjectFile, new ReadOnlyBuildVersion(BuildVersion.ReadDefault()), new string[0]);
+			TargetRules Rules = RulesAssembly.CreateTargetRules(TargetDesc.Name, TargetDesc.Platform, TargetDesc.Configuration, TargetDesc.Architecture, TargetDesc.ProjectFile, new string[0]);
 
 			// Check if we need to enable a nativized plugin, and compile the assembly for that if we do
 			FileReference NativizedPluginFile = Rules.GetNativizedPlugin();
diff --git a/Engine/Source/Programs/UnrealBuildTool/UnrealBuildTool.cs b/Engine/Source/Programs/UnrealBuildTool/UnrealBuildTool.cs
--- a/Engine/Source/Programs/UnrealBuildTool/UnrealBuildTool.cs
+++ b/Engine/Source/Programs/UnrealBuildTool/UnrealBuildTool.cs
@@ -933,12 +933,10 @@ namespace UnrealBuildTool
 					{
 						DateTime TargetInitStartTime = DateTime.UtcNow;
 
-						ReadOnlyBuildVersion Version = new ReadOnlyBuildVersion(BuildVersion.ReadDefault());
-
 						Targets = new List<UEBuildTarget>();
 						foreach (TargetDescriptor TargetDesc in TargetDescs)
 						{
-							UEBuildTarget Target = UEBuildTarget.CreateTarget(TargetDesc, Arguments, bSkipRulesCompile, BuildConfiguration.SingleFileToCompile != null, BuildConfiguration.bUsePrecompiled, Version);
+							UEBuildTarget Target = UEBuildTarget.CreateTarget(TargetDesc, Arguments, bSkipRulesCompile, BuildConfiguration.SingleFileToCompile != null, BuildConfiguration.bUsePrecompiled);
 							if ((Target == null) && (BuildConfiguration.bCleanProject))
 							{
 								continue;
