UBT: Allow creating paths to the dependency cache without passing the target around.

#rnx

diff --git a/Engine/Source/Programs/UnrealBuildTool/Configuration/UEBuildTarget.cs b/Engine/Source/Programs/UnrealBuildTool/Configuration/UEBuildTarget.cs
--- a/Engine/Source/Programs/UnrealBuildTool/Configuration/UEBuildTarget.cs
+++ b/Engine/Source/Programs/UnrealBuildTool/Configuration/UEBuildTarget.cs
@@ -883,7 +883,7 @@ namespace UnrealBuildTool
 			bDeployAfterCompile = InRules.bDeployAfterCompile && !InRules.bDisableLinking;
 
 			// now that we have the platform, we can set the intermediate path to include the platform/architecture name
-			PlatformIntermediateFolder = Path.Combine("Intermediate", "Build", Platform.ToString(), UEBuildPlatform.GetBuildPlatform(Platform).GetFolderNameForArchitecture(Architecture));
+			PlatformIntermediateFolder = GetPlatformIntermediateFolder(Platform, Architecture);
 
 			TargetRulesFile = InRules.File;
 
@@ -987,6 +987,18 @@ namespace UnrealBuildTool
 			PreBuildSetup();
 		}
 
+		/// <summary>
+		/// Gets the intermediate directory for a given platform
+		/// </summary>
+		/// <param name="Platform">Platform to get the folder for</param>
+		/// <param name="Architecture">Architecture to get the folder for</param>
+		/// <returns>The output directory for intermediates</returns>
+		public static string GetPlatformIntermediateFolder(UnrealTargetPlatform Platform, string Architecture)
+		{
+			// now that we have the platform, we can set the intermediate path to include the platform/architecture name
+			return Path.Combine("Intermediate", "Build", Platform.ToString(), UEBuildPlatform.GetBuildPlatform(Platform).GetFolderNameForArchitecture(Architecture));
+		}
+
 		/// <summary>
 		/// Gets the app name for a given target type
 		/// </summary>
@@ -1055,7 +1067,7 @@ namespace UnrealBuildTool
 
 			// Add all the makefiles and caches to be deleted
 			List<FileReference> FilesToDelete = new List<FileReference>();
-			FilesToDelete.Add(FlatCPPIncludeDependencyCache.GetDependencyCachePathForTarget(this));
+			FilesToDelete.Add(FlatCPPIncludeDependencyCache.GetDependencyCachePathForTarget(ProjectFile, TargetName, Platform, Architecture));
 			FilesToDelete.Add(DependencyCache.GetDependencyCachePathForTarget(ProjectFile, Platform, TargetName));
 			FilesToDelete.Add(UBTMakefile.GetUBTMakefilePath(ProjectFile, Platform, Configuration, TargetName, false));
 			FilesToDelete.Add(UBTMakefile.GetUBTMakefilePath(ProjectFile, Platform, Configuration, TargetName, true));
diff --git a/Engine/Source/Programs/UnrealBuildTool/System/FlatCPPIncludeDepencencyCache.cs b/Engine/Source/Programs/UnrealBuildTool/System/FlatCPPIncludeDepencencyCache.cs
--- a/Engine/Source/Programs/UnrealBuildTool/System/FlatCPPIncludeDepencencyCache.cs
+++ b/Engine/Source/Programs/UnrealBuildTool/System/FlatCPPIncludeDepencencyCache.cs
@@ -285,20 +285,26 @@ namespace UnrealBuildTool
 		/// <summary>
 		/// Gets the dependency cache path and filename for the specified target.
 		/// </summary>
-		/// <param name="Target">Current build target</param>
+		/// <param name="ProjectFile">Project file containing this target</param>
+		/// <param name="TargetName">Name of the target being built</param>
+		/// <param name="Platform">Platform for the target</param>
+		/// <param name="Architecture">Architecture for the target</param>
 		/// <returns>Cache Path</returns>
-		public static FileReference GetDependencyCachePathForTarget(UEBuildTarget Target)
+		public static FileReference GetDependencyCachePathForTarget(FileReference ProjectFile, string TargetName, UnrealTargetPlatform Platform, string Architecture)
 		{
+			string PlatformIntermediateFolder = UEBuildTarget.GetPlatformIntermediateFolder(Platform, Architecture);
+
 			DirectoryReference PlatformIntermediatePath;
-			if (Target.ProjectFile != null)
+			if (ProjectFile != null)
 			{
-				PlatformIntermediatePath = DirectoryReference.Combine(Target.ProjectFile.Directory, Target.PlatformIntermediateFolder);
+				PlatformIntermediatePath = DirectoryReference.Combine(ProjectFile.Directory, PlatformIntermediateFolder);
 			}
 			else
 			{
-				PlatformIntermediatePath = DirectoryReference.Combine(UnrealBuildTool.EngineDirectory, Target.PlatformIntermediateFolder);
+				PlatformIntermediatePath = DirectoryReference.Combine(UnrealBuildTool.EngineDirectory, PlatformIntermediateFolder);
 			}
-			return FileReference.Combine(PlatformIntermediatePath, Target.GetTargetName(), "FlatCPPIncludes.bin");
+
+			return FileReference.Combine(PlatformIntermediatePath, TargetName, "FlatCPPIncludes.bin");
 		}
 	}
 }
diff --git a/Engine/Source/Programs/UnrealBuildTool/UnrealBuildTool.cs b/Engine/Source/Programs/UnrealBuildTool/UnrealBuildTool.cs
--- a/Engine/Source/Programs/UnrealBuildTool/UnrealBuildTool.cs
+++ b/Engine/Source/Programs/UnrealBuildTool/UnrealBuildTool.cs
@@ -988,7 +988,7 @@ namespace UnrealBuildTool
 						{
 							// Load the cache that contains the list of flattened resolved includes for resolved source files
 							// @todo ubtmake: Ideally load this asynchronously at startup and only block when it is first needed and not finished loading
-							FileReference CacheFile = FlatCPPIncludeDependencyCache.GetDependencyCachePathForTarget(Target);
+							FileReference CacheFile = FlatCPPIncludeDependencyCache.GetDependencyCachePathForTarget(Target.ProjectFile, Target.TargetName, Target.Platform, Target.Architecture);
 							if (!FlatCPPIncludeDependencyCache.TryRead(CacheFile, out Headers.FlatCPPIncludeDependencyCache))
 							{
 								if (!bNeedsFullCPPIncludeRescan)
