UBT: Pass explicit arguments to ActionHistory.GeneratePathForTarget

#rnx

diff --git a/Engine/Source/Programs/UnrealBuildTool/Configuration/UEBuildTarget.cs b/Engine/Source/Programs/UnrealBuildTool/Configuration/UEBuildTarget.cs
--- a/Engine/Source/Programs/UnrealBuildTool/Configuration/UEBuildTarget.cs
+++ b/Engine/Source/Programs/UnrealBuildTool/Configuration/UEBuildTarget.cs
@@ -1071,7 +1071,7 @@ namespace UnrealBuildTool
 			FilesToDelete.Add(DependencyCache.GetDependencyCachePathForTarget(ProjectFile, Platform, TargetName));
 			FilesToDelete.Add(UBTMakefile.GetUBTMakefilePath(ProjectFile, Platform, Configuration, TargetName, false));
 			FilesToDelete.Add(UBTMakefile.GetUBTMakefilePath(ProjectFile, Platform, Configuration, TargetName, true));
-			FilesToDelete.Add(ActionHistory.GeneratePathForTarget(this));
+			FilesToDelete.Add(ActionHistory.GeneratePathForTarget(ProjectFile, TargetName, Platform, Architecture, bUseSharedBuildEnvironment));
 
 			// Add all the intermediate folders to be deleted
 			List<DirectoryReference> DirectoriesToDelete = new List<DirectoryReference>();
diff --git a/Engine/Source/Programs/UnrealBuildTool/System/ActionGraph.cs b/Engine/Source/Programs/UnrealBuildTool/System/ActionGraph.cs
--- a/Engine/Source/Programs/UnrealBuildTool/System/ActionGraph.cs
+++ b/Engine/Source/Programs/UnrealBuildTool/System/ActionGraph.cs
@@ -326,7 +326,7 @@ namespace UnrealBuildTool
 			TargetToOutdatedPrerequisitesMap = new Dictionary<UEBuildTarget, List<FileItem>>();
 			foreach (UEBuildTarget BuildTarget in Targets)	// @todo ubtmake: Optimization: Ideally we don't even need to know about targets for ubtmake -- everything would come from the files
 			{
-				FileReference HistoryFilename = ActionHistory.GeneratePathForTarget(BuildTarget);
+				FileReference HistoryFilename = ActionHistory.GeneratePathForTarget(BuildTarget.ProjectFile, BuildTarget.TargetName, BuildTarget.Platform, BuildTarget.Architecture, BuildTarget.bUseSharedBuildEnvironment);
 				if (!OpenHistoryFiles.Contains(HistoryFilename))		// @todo ubtmake: Optimization: We should be able to move the command-line outdatedness and build product deletion over to the 'gather' phase, as the command-lines won't change between assembler runs
 				{
 					ActionHistory History = new ActionHistory(HistoryFilename.FullName);
diff --git a/Engine/Source/Programs/UnrealBuildTool/System/ActionHistory.cs b/Engine/Source/Programs/UnrealBuildTool/System/ActionHistory.cs
--- a/Engine/Source/Programs/UnrealBuildTool/System/ActionHistory.cs
+++ b/Engine/Source/Programs/UnrealBuildTool/System/ActionHistory.cs
@@ -99,32 +99,32 @@ namespace UnrealBuildTool
 		/// <summary>
 		/// Generates a full path to action history file for the specified target.
 		/// </summary>
-		public static FileReference GeneratePathForTarget(UEBuildTarget Target)
+		public static FileReference GeneratePathForTarget(FileReference ProjectFile, string TargetName, UnrealTargetPlatform Platform, string Architecture, bool bUseSharedBuildEnvironment)
 		{
-			DirectoryReference Folder = null;
-			if (Target.ShouldCompileMonolithic() || Target.TargetType == TargetType.Program)
+			if(bUseSharedBuildEnvironment)
 			{
-				// Monolithic configs and programs have their Action History stored in their respective project folders
-				// or under engine intermediate folder + program name folder
-				DirectoryReference RootDirectory;
-				if (Target.ProjectFile != null)
+				if(UnrealBuildTool.IsEngineInstalled() && ProjectFile != null)
 				{
-					RootDirectory = Target.ProjectFile.Directory;
+					return FileReference.Combine(ProjectFile.Directory, "Intermediate", "Build", "ActionHistory.bin");
 				}
 				else
 				{
-					RootDirectory = UnrealBuildTool.EngineDirectory;
+					return FileReference.Combine(UnrealBuildTool.EngineDirectory, "Intermediate", "Build", "ActionHistory.bin");
 				}
-				Folder = DirectoryReference.Combine(RootDirectory, Target.PlatformIntermediateFolder, Target.GetTargetName());
 			}
 			else
 			{
-				// Shared action history (unless this is an installed build target)
-				Folder = (UnrealBuildTool.IsEngineInstalled() && Target.ProjectFile != null) ?
-					DirectoryReference.Combine(Target.ProjectFile.Directory, "Intermediate", "Build") :
-					DirectoryReference.Combine(UnrealBuildTool.EngineDirectory, "Intermediate", "Build");
+				string PlatformIntermediateFolder = UEBuildTarget.GetPlatformIntermediateFolder(Platform, Architecture);
+
+				if(ProjectFile != null)
+				{
+					return FileReference.Combine(ProjectFile.Directory, PlatformIntermediateFolder, TargetName, "ActionHistory.bin");
+				}
+				else
+				{
+					return FileReference.Combine(UnrealBuildTool.EngineDirectory, PlatformIntermediateFolder, TargetName, "ActionHistory.bin");
+				}
 			}
-			return FileReference.Combine(Folder, "ActionHistory.bin");
 		}
 	}
 }
