UBT: Remove hot reload state from UHT invocation.

#rnx

diff --git a/Engine/Source/Programs/UnrealBuildTool/Configuration/UEBuildTarget.cs b/Engine/Source/Programs/UnrealBuildTool/Configuration/UEBuildTarget.cs
--- a/Engine/Source/Programs/UnrealBuildTool/Configuration/UEBuildTarget.cs
+++ b/Engine/Source/Programs/UnrealBuildTool/Configuration/UEBuildTarget.cs
@@ -1806,7 +1806,7 @@ namespace UnrealBuildTool
 				FileReference ModuleInfoFileName = FileReference.Combine(ProjectIntermediateDirectory, GetTargetName() + ".uhtmanifest");
 
 				// Execute the header tool
-				ECompilationResult UHTResult = ExternalExecution.ExecuteHeaderToolIfNecessary(BuildConfiguration, ProjectFile, TargetName, TargetType, bHasProjectScriptPlugin, GlobalCompileEnvironment, UObjectModules, ModuleInfoFileName, HotReload, true, bIsAssemblingBuild);
+				ECompilationResult UHTResult = ExternalExecution.ExecuteHeaderToolIfNecessary(BuildConfiguration, ProjectFile, TargetName, TargetType, bHasProjectScriptPlugin, UObjectModules, ModuleInfoFileName, true, bIsAssemblingBuild);
 				if(!UHTResult.Succeeded())
 				{
 					Log.TraceInformation(String.Format("Error: UnrealHeaderTool failed for target '{0}' (platform: {1}, module info: {2}, exit code: {3} ({4})).", GetTargetName(), Platform.ToString(), ModuleInfoFileName, UHTResult.ToString(), (int)UHTResult));
diff --git a/Engine/Source/Programs/UnrealBuildTool/System/ExternalExecution.cs b/Engine/Source/Programs/UnrealBuildTool/System/ExternalExecution.cs
--- a/Engine/Source/Programs/UnrealBuildTool/System/ExternalExecution.cs
+++ b/Engine/Source/Programs/UnrealBuildTool/System/ExternalExecution.cs
@@ -817,11 +817,10 @@ namespace UnrealBuildTool
 		/// <param name="BuildConfiguration">Build configuration</param>
 		/// <param name="UObjectModules">Modules that we generate headers for</param>
 		/// <param name="HeaderToolTimestamp">Timestamp for UHT</param>
-		/// <param name="HotReload">The hot reload state</param>
 		/// <param name="bIsGatheringBuild"></param>
 		/// <param name="bIsAssemblingBuild"></param>
 		/// <returns>True if the code files are out of date</returns>
-		private static bool AreGeneratedCodeFilesOutOfDate(BuildConfiguration BuildConfiguration, List<UHTModuleInfo> UObjectModules, DateTime HeaderToolTimestamp, EHotReload HotReload, bool bIsGatheringBuild, bool bIsAssemblingBuild)
+		private static bool AreGeneratedCodeFilesOutOfDate(BuildConfiguration BuildConfiguration, List<UHTModuleInfo> UObjectModules, DateTime HeaderToolTimestamp, bool bIsGatheringBuild, bool bIsAssemblingBuild)
 		{
 			// Get CoreUObject.init.gen.cpp timestamp.  If the source files are older than the CoreUObject generated code, we'll
 			// need to regenerate code for the module
@@ -1087,7 +1086,7 @@ namespace UnrealBuildTool
 		/// Builds and runs the header tool and touches the header directories.
 		/// Performs any early outs if headers need no changes, given the UObject modules, tool path, game name, and configuration
 		/// </summary>
-		public static ECompilationResult ExecuteHeaderToolIfNecessary(BuildConfiguration BuildConfiguration, FileReference ProjectFile, string TargetName, TargetType TargetType, bool bHasProjectScriptPlugin, CppCompileEnvironment GlobalCompileEnvironment, List<UHTModuleInfo> UObjectModules, FileReference ModuleInfoFileName, EHotReload HotReload, bool bIsGatheringBuild, bool bIsAssemblingBuild)
+		public static ECompilationResult ExecuteHeaderToolIfNecessary(BuildConfiguration BuildConfiguration, FileReference ProjectFile, string TargetName, TargetType TargetType, bool bHasProjectScriptPlugin, List<UHTModuleInfo> UObjectModules, FileReference ModuleInfoFileName, bool bIsGatheringBuild, bool bIsAssemblingBuild)
 		{
 			if (ProgressWriter.bWriteMarkup)
 			{
@@ -1110,7 +1109,7 @@ namespace UnrealBuildTool
 				bool bHaveHeaderTool = !bIsBuildingUHT && GetHeaderToolTimestamp(HeaderToolReceipt, out HeaderToolTimestamp);
 
 				// ensure the headers are up to date
-				bool bUHTNeedsToRun = (BuildConfiguration.bForceHeaderGeneration || !bHaveHeaderTool || AreGeneratedCodeFilesOutOfDate(BuildConfiguration, UObjectModules, HeaderToolTimestamp, HotReload, bIsGatheringBuild, bIsAssemblingBuild));
+				bool bUHTNeedsToRun = (BuildConfiguration.bForceHeaderGeneration || !bHaveHeaderTool || AreGeneratedCodeFilesOutOfDate(BuildConfiguration, UObjectModules, HeaderToolTimestamp, bIsGatheringBuild, bIsAssemblingBuild));
 
 				// Get the file containing dependencies for the generated code
 				FileReference ExternalDependenciesFile = ModuleInfoFileName.ChangeExtension(".deps");
@@ -1134,7 +1133,6 @@ namespace UnrealBuildTool
 					// Always build UnrealHeaderTool if header regeneration is required, unless we're running within an installed ecosystem or hot-reloading
 					if ((!UnrealBuildTool.IsEngineInstalled() || bHasProjectScriptPlugin) &&
 						!BuildConfiguration.bDoNotBuildUHT &&
-						HotReload != EHotReload.FromIDE &&
 						!(bHaveHeaderTool && !bIsGatheringBuild && bIsAssemblingBuild))	// If running in "assembler only" mode, we assume UHT is already up to date for much faster iteration!
 					{
 						// If it is out of date or not there it will be built.
diff --git a/Engine/Source/Programs/UnrealBuildTool/UnrealBuildTool.cs b/Engine/Source/Programs/UnrealBuildTool/UnrealBuildTool.cs
--- a/Engine/Source/Programs/UnrealBuildTool/UnrealBuildTool.cs
+++ b/Engine/Source/Programs/UnrealBuildTool/UnrealBuildTool.cs
@@ -1131,7 +1131,7 @@ namespace UnrealBuildTool
 											{
 												// Execute the header tool
 												FileReference ModuleInfoFileName = FileReference.Combine(Target.ProjectIntermediateDirectory, Target.GetTargetName() + ".uhtmanifest");
-												ECompilationResult UHTResult = ExternalExecution.ExecuteHeaderToolIfNecessary(BuildConfiguration, Target.ProjectFile, Target.TargetName, Target.TargetType, Target.bHasProjectScriptPlugin, GlobalCompileEnvironment: null, UObjectModules: TargetUObjectModules, ModuleInfoFileName: ModuleInfoFileName, HotReload: HotReload, bIsGatheringBuild: bIsGatheringBuild, bIsAssemblingBuild: bIsAssemblingBuild);
+												ECompilationResult UHTResult = ExternalExecution.ExecuteHeaderToolIfNecessary(BuildConfiguration, Target.ProjectFile, Target.TargetName, Target.TargetType, Target.bHasProjectScriptPlugin,  UObjectModules: TargetUObjectModules, ModuleInfoFileName: ModuleInfoFileName, bIsGatheringBuild: bIsGatheringBuild, bIsAssemblingBuild: bIsAssemblingBuild);
 												if(UHTResult.Succeeded())
 												{
 													Log.TraceInformation("UnrealHeaderTool failed for target '" + Target.GetTargetName() + "' (platform: " + Target.Platform.ToString() + ", module info: " + ModuleInfoFileName + ").");
