Various changes to hot reload build process.

* Hot reload now shares the regular UBT makefile, significantly improving performance for first-time compile. The history of previous hot-reloads is stored in an intermediate file (see GetHotReloadStateFile()), allowing the action graph to be patched to the last hot reload state before executing the build rather than creating a separate makefile. (A nice side effect of this is that it allows tracking hot reload attempts, and assigning incremental suffixes rather than random numbers for each run).
* The list of modules eligable for hot reload is now explicitly saved to the makefile. During a hot reload from IDE invocation, we always try to build all out of date modules and only apply suffixes to those that support it. This prevents the confusing behavior where an editor open in the background will not attempt to build certain files.
* Logic for building dependent modules is now significantly more reliable. Only modules with changes are rebuilt, and all dependent modules of those modules are correctly rebuilt too. The -CanSkipLink option is no longer necessary, and has been removed.
* All code for dealing with hot-reload only happens when the action graph is executed. The construction of a target no longer has any knowledge of whether hot reload is enabled or not.

diff --git a/Engine/Source/Developer/HotReload/Private/HotReload.cpp b/Engine/Source/Developer/HotReload/Private/HotReload.cpp
--- a/Engine/Source/Developer/HotReload/Private/HotReload.cpp
+++ b/Engine/Source/Developer/HotReload/Private/HotReload.cpp
@@ -1571,21 +1571,8 @@ FString FHotReloadModule::MakeUBTArgumentsForModuleCompiling()
 	{
 		// We have to pass FULL paths to UBT
 		FString FullProjectPath = FPaths::ConvertRelativePathToFull(FPaths::GetProjectFilePath());
-
-		// @todo projectdirs: Currently non-installed projects that exist under the UE4 root are compiled by UBT with no .uproject file
-		//     name passed in (see bIsProjectTarget in VCProject.cs), which causes intermediate libraries to be saved to the Engine
-		//     intermediate folder instead of the project's intermediate folder.  We're emulating this behavior here for module
-		//     recompiling, so that compiled modules will be able to find their import libraries in the original folder they were compiled.
-		if( FApp::IsEngineInstalled() || !FullProjectPath.StartsWith( FPaths::ConvertRelativePathToFull( FPaths::RootDir() ) ) )
-		{
-			const FString ProjectFilenameWithQuotes = FString::Printf(TEXT("\"%s\""), *FullProjectPath);
-			AdditionalArguments += FString::Printf(TEXT("%s "), *ProjectFilenameWithQuotes);
-		}
+		AdditionalArguments += FString::Printf(TEXT("\"%s\" "), *FullProjectPath);
 	}
-
-	// Use new FastPDB option to cut down linking time. Currently disabled due to problems with missing symbols in VS2015.
-	// AdditionalArguments += TEXT(" -FastPDB");
-
 	return AdditionalArguments;
 }
 
@@ -1636,7 +1623,7 @@ bool FHotReloadModule::StartCompilingModuleDLLs(const TArray< FModuleToRecompile
 
 	if (FPaths::IsProjectFilePathSet())
 	{
-		ExtraArg += FString::Printf(TEXT("-Project=\"%s\" "), *FPaths::GetProjectFilePath());
+		ExtraArg += FString::Printf(TEXT("-Project=\"%s\" "), *FPaths::ConvertRelativePathToFull(FPaths::GetProjectFilePath()));
 	}
 
 	if (bInFailIfGeneratedCodeChanges)
@@ -1645,9 +1632,6 @@ bool FHotReloadModule::StartCompilingModuleDLLs(const TArray< FModuleToRecompile
 		ExtraArg += TEXT( "-FailIfGeneratedCodeChanges " );
 	}
 
-	// If there's nothing to compile, don't bother linking the DLLs as the old ones are up-to-date
-	ExtraArg += TEXT("-canskiplink ");
-
 	// Shared PCH does no work with hot-reloading engine/editor modules as we don't scan all modules for them.
 	if (!ContainsOnlyGameModules(ModuleNames))
 	{
diff --git a/Engine/Source/Programs/UnrealBuildTool/Configuration/BuildConfiguration.cs b/Engine/Source/Programs/UnrealBuildTool/Configuration/BuildConfiguration.cs
--- a/Engine/Source/Programs/UnrealBuildTool/Configuration/BuildConfiguration.cs
+++ b/Engine/Source/Programs/UnrealBuildTool/Configuration/BuildConfiguration.cs
@@ -205,12 +205,6 @@ namespace UnrealBuildTool
 		[CommandLine("-2017", Value = "-2017")]
 		public string CompilerArgumentForUnrealHeaderTool = null;
 
-		/// <summary>
-		/// When true, the targets won't execute their link actions if there was nothing to compile
-		/// </summary>
-		[CommandLine("-CanSkipLink")]
-		public bool bSkipLinkingWhenNothingToCompile = false;
-
 		/// <summary>
 		/// Default constructor. Reads settings from the XmlConfig files.
 		/// </summary>
diff --git a/Engine/Source/Programs/UnrealBuildTool/Configuration/TargetDescriptor.cs b/Engine/Source/Programs/UnrealBuildTool/Configuration/TargetDescriptor.cs
--- a/Engine/Source/Programs/UnrealBuildTool/Configuration/TargetDescriptor.cs
+++ b/Engine/Source/Programs/UnrealBuildTool/Configuration/TargetDescriptor.cs
@@ -19,7 +19,6 @@ namespace UnrealBuildTool
 		public UnrealTargetPlatform Platform;
 		public UnrealTargetConfiguration Configuration;
 		public string Architecture;
-		public List<OnlyModule> OnlyModules = new List<OnlyModule>();
 		public FileReference ForeignPlugin;
 		public string[] AdditionalArguments;
 
@@ -57,7 +56,6 @@ namespace UnrealBuildTool
 			List<string> TargetNames = new List<string>();
 			List<TargetType> TargetTypes = new List<TargetType>();
 			string Architecture = null;
-			List<OnlyModule> OnlyModules = new List<OnlyModule>();
 			FileReference ForeignPlugin = null;
 			List<string> AdditionalArguments = new List<string>();
 
@@ -107,19 +105,6 @@ namespace UnrealBuildTool
 						}
 						TargetTypes.Add(Type);
 					}
-					else if(ParseArgumentValue(Argument, "-Module=", out Value))
-					{
-						OnlyModules.Add(new OnlyModule(Value));
-					}
-					else if(ParseArgumentValue(Argument, "-ModuleWithSuffix=", out Value))
-					{
-						int SuffixIdx = Value.LastIndexOf(',');
-						if(SuffixIdx == -1)
-						{
-							throw new BuildException("Missing suffix argument from -ModuleWithSuffix=Name,Suffix");
-						}
-						OnlyModules.Add(new OnlyModule(Value.Substring(0, SuffixIdx), Value.Substring(SuffixIdx + 1)));
-					}
 					else if(ParseArgumentValue(Argument, "-Plugin=", out Value))
 					{
 						if(ForeignPlugin != null)
@@ -187,7 +172,6 @@ namespace UnrealBuildTool
 
 				// Create the target descriptor
 				TargetDescriptor Target = new TargetDescriptor(ProjectFile, TargetName, Platform, Configuration, Architecture, AdditionalArguments.ToArray());
-				Target.OnlyModules.AddRange(OnlyModules);
 				Target.ForeignPlugin = ForeignPlugin;
 				Targets.Add(Target);
 			}
diff --git a/Engine/Source/Programs/UnrealBuildTool/Configuration/UEBuildBinary.cs b/Engine/Source/Programs/UnrealBuildTool/Configuration/UEBuildBinary.cs
--- a/Engine/Source/Programs/UnrealBuildTool/Configuration/UEBuildBinary.cs
+++ b/Engine/Source/Programs/UnrealBuildTool/Configuration/UEBuildBinary.cs
@@ -68,26 +68,6 @@ namespace UnrealBuildTool
 			}
 		}
 
-		/// <summary>
-		/// Original output filepath. This is the original binary name before hot-reload suffix has been appended to it.
-		/// </summary>
-		public List<FileReference> OriginalOutputFilePaths;
-
-		/// <summary>
-		/// Returns the OriginalOutputFilePath if there is only one entry in OriginalOutputFilePaths
-		/// </summary>
-		public FileReference OriginalOutputFilePath
-		{
-			get
-			{
-				if (OriginalOutputFilePaths.Count != 1)
-				{
-					throw new BuildException("Attempted to use UEBuildBinaryConfiguration.OriginalOutputFilePath property, but there are multiple (or no) OriginalOutputFilePaths. You need to handle multiple in the code that called this (size = {0})", OriginalOutputFilePaths.Count);
-				}
-				return OriginalOutputFilePaths[0];
-			}
-		}
-
 		/// <summary>
 		/// The intermediate directory for this binary. Modules should create separate intermediate directories below this. Must be set before a binary can be built using it.
 		/// </summary>
@@ -181,7 +161,7 @@ namespace UnrealBuildTool
 		/// <param name="ExeDir">Directory containing the output executable</param>
 		/// <param name="ActionGraph">Graph to add build actions to</param>
 		/// <returns>Set of built products</returns>
-		public IEnumerable<FileItem> Build(ReadOnlyTargetRules Target, UEToolChain ToolChain, CppCompileEnvironment CompileEnvironment, LinkEnvironment LinkEnvironment, List<PrecompiledHeaderTemplate> SharedPCHs, ISourceFileWorkingSet WorkingSet, DirectoryReference ExeDir, ActionGraph ActionGraph)
+		public List<FileItem> Build(ReadOnlyTargetRules Target, UEToolChain ToolChain, CppCompileEnvironment CompileEnvironment, LinkEnvironment LinkEnvironment, List<PrecompiledHeaderTemplate> SharedPCHs, ISourceFileWorkingSet WorkingSet, DirectoryReference ExeDir, ActionGraph ActionGraph)
 		{
 			// Return nothing if we're using precompiled binaries. If we're not linking, we might want just one module to be compiled (eg. a foreign plugin), so allow any actions to run.
 			if (bUsePrecompiled && !Target.bDisableLinking)
@@ -333,26 +313,6 @@ namespace UnrealBuildTool
 			DependentLinkEnvironment.AdditionalLibraries.AddRange(DependentLinkLibraries);
 		}
 
-		/// <summary>
-		/// Called to allow the binary to to determine if it matches the Only module "short module name".
-		/// </summary>
-		/// <param name="OnlyModules"></param>
-		/// <returns>The OnlyModule if found, null if not</returns>
-		public OnlyModule FindOnlyModule(List<OnlyModule> OnlyModules)
-		{
-			foreach (UEBuildModule Module in Modules)
-			{
-				foreach (OnlyModule OnlyModule in OnlyModules)
-				{
-					if (OnlyModule.OnlyModuleName.ToLower() == Module.Name.ToLower())
-					{
-						return OnlyModule;
-					}
-				}
-			}
-			return null;
-		}
-
 		/// <summary>
 		/// Called to allow the binary to find game modules.
 		/// </summary>
@@ -684,10 +644,7 @@ namespace UnrealBuildTool
 
 			// @todo: This should be in some Windows code somewhere...
 			// Set the original file name macro; used in PCLaunch.rc to set the binary metadata fields.
-			string OriginalFilename = (OriginalOutputFilePaths != null) ?
-				OriginalOutputFilePaths[0].GetFileName() :
-				OutputFilePaths[0].GetFileName();
-			BinaryCompileEnvironment.Definitions.Add("ORIGINAL_FILE_NAME=\"" + OriginalFilename + "\"");
+			BinaryCompileEnvironment.Definitions.Add("ORIGINAL_FILE_NAME=\"" + OutputFilePaths[0].GetFileName() + "\"");
 
 			return BinaryCompileEnvironment;
 		}
@@ -779,8 +736,7 @@ namespace UnrealBuildTool
 
 						// @todo: This should be in some Windows code somewhere...
 						// Set the original file name macro; used in PCLaunch.rc to set the binary metadata fields.
-						string OriginalFilename = (OriginalOutputFilePaths != null) ? OriginalOutputFilePaths[0].GetFileName() : OutputFilePaths[0].GetFileName();
-						ResourceCompileEnvironment.Definitions.Add("ORIGINAL_FILE_NAME=\"" + OriginalFilename + "\"");
+						ResourceCompileEnvironment.Definitions.Add("ORIGINAL_FILE_NAME=\"" + OutputFilePaths[0].GetFileName() + "\"");
 
 						// Set the other version fields
 						ResourceCompileEnvironment.Definitions.Add(String.Format("BUILT_FROM_CHANGELIST={0}", Target.Version.Changelist));
diff --git a/Engine/Source/Programs/UnrealBuildTool/Configuration/UEBuildTarget.cs b/Engine/Source/Programs/UnrealBuildTool/Configuration/UEBuildTarget.cs
--- a/Engine/Source/Programs/UnrealBuildTool/Configuration/UEBuildTarget.cs
+++ b/Engine/Source/Programs/UnrealBuildTool/Configuration/UEBuildTarget.cs
@@ -288,44 +288,6 @@ namespace UnrealBuildTool
 		public List<string> UHTHeaderNames = new List<string>();
 	}
 
-	[Serializable]
-	class OnlyModule : ISerializable
-	{
-		public OnlyModule(SerializationInfo Info, StreamingContext Context)
-		{
-			OnlyModuleName = Info.GetString("mn");
-			OnlyModuleSuffix = Info.GetString("ms");
-		}
-
-		public void GetObjectData(SerializationInfo Info, StreamingContext Context)
-		{
-			Info.AddValue("mn", OnlyModuleName);
-			Info.AddValue("ms", OnlyModuleSuffix);
-		}
-
-		public OnlyModule(string InitOnlyModuleName)
-		{
-			OnlyModuleName = InitOnlyModuleName;
-			OnlyModuleSuffix = String.Empty;
-		}
-
-		public OnlyModule(string InitOnlyModuleName, string InitOnlyModuleSuffix)
-		{
-			OnlyModuleName = InitOnlyModuleName;
-			OnlyModuleSuffix = InitOnlyModuleSuffix;
-		}
-
-		/// <summary>
-		/// If building only a single module, this is the module name to build
-		/// </summary>
-		public readonly string OnlyModuleName;
-
-		/// <summary>
-		/// When building only a single module, the optional suffix for the module file name
-		/// </summary>
-		public readonly string OnlyModuleSuffix;
-	}
-
 	/// <summary>
 	/// A target that can be built
 	/// </summary>
@@ -399,15 +361,6 @@ namespace UnrealBuildTool
 				ValidateSharedEnvironment(RulesAssembly, Desc.Name, RulesObject);
 			}
 
-			// Make sure that we don't explicitly enable or disable any plugins through the target rules. We can't 
-			if(RulesObject.BuildEnvironment == TargetBuildEnvironment.Shared && Desc.OnlyModules.Count == 0)
-			{
-				if(RulesObject.EnablePlugins.Count > 0 || RulesObject.DisablePlugins.Count > 0)
-				{
-					throw new BuildException("Explicitly enabling and disabling plugins for a target is only supported when using a unique build environment (eg. for monolithic game targets).");
-				}
-			}
-
 			// If we're precompiling, generate a list of all the files that we depend on
 			if (RulesObject.bPrecompile)
 			{
@@ -559,6 +512,12 @@ namespace UnrealBuildTool
 					}
 				}
 			}
+
+			// Make sure that we don't explicitly enable or disable any plugins through the target rules. We can't do this with the shared build environment because it requires recompiling the "Projects" engine module.
+			if(ThisRules.EnablePlugins.Count > 0 || ThisRules.DisablePlugins.Count > 0)
+			{
+				throw new BuildException("Explicitly enabling and disabling plugins for a target is only supported when using a unique build environment (eg. for monolithic game targets).");
+			}
 		}
 
 		/// <summary>
@@ -724,11 +683,6 @@ namespace UnrealBuildTool
 		[NonSerialized]
 		public List<UEBuildBinary> Binaries = new List<UEBuildBinary>();
 
-		/// <summary>
-		/// If building only a specific set of modules, these are the modules to build
-		/// </summary>
-		public List<OnlyModule> OnlyModules = new List<OnlyModule>();
-
 		/// <summary>
 		/// Kept to determine the correct module parsing order when filtering modules.
 		/// </summary>
@@ -820,7 +774,6 @@ namespace UnrealBuildTool
 			OutputPaths = (List<FileReference>)Info.GetValue("op", typeof(List<FileReference>));
 			VersionFile = (FileReference)Info.GetValue("vf", typeof(FileReference));
 			bPrecompile = Info.GetBoolean("pc");
-			OnlyModules = (List<OnlyModule>)Info.GetValue("om", typeof(List<OnlyModule>));
 			bCompileMonolithic = Info.GetBoolean("cm");
 			FlatModuleCsData = (List<FlatModuleCsDataType>)Info.GetValue("fm", typeof(List<FlatModuleCsDataType>));
 			ReceiptFileName = (FileReference)Info.GetValue("rf", typeof(FileReference));
@@ -849,7 +802,6 @@ namespace UnrealBuildTool
 			Info.AddValue("op", OutputPaths);
 			Info.AddValue("vf", VersionFile);
 			Info.AddValue("pc", bPrecompile);
-			Info.AddValue("om", OnlyModules);
 			Info.AddValue("cm", bCompileMonolithic);
 			Info.AddValue("fm", FlatModuleCsData);
 			Info.AddValue("rf", ReceiptFileName);
@@ -956,8 +908,6 @@ namespace UnrealBuildTool
 				ProjectDescriptor = ProjectDescriptor.FromFile(ProjectFile);
 			}
 
-			OnlyModules = InDesc.OnlyModules;
-
 			// Construct the output paths for this target's executable
 			DirectoryReference OutputDirectory;
 			if (bCompileMonolithic || !bUseSharedBuildEnvironment)
@@ -1338,10 +1288,7 @@ namespace UnrealBuildTool
 
 			if (!Rules.bDisableLinking)
 			{
-				if (OnlyModules.Count == 0)
-				{
-					Manifest.AddBuildProduct(ReceiptFileName.FullName);
-				}
+				Manifest.AddBuildProduct(ReceiptFileName.FullName);
 
 				if (bDeployAfterCompile)
 				{
@@ -1404,8 +1351,7 @@ namespace UnrealBuildTool
 		/// <param name="ToolChain">The toolchain used to build the target</param>
 		/// <param name="BuildProducts">Artifacts from the build</param>
 		/// <param name="RuntimeDependencies">Output runtime dependencies</param>
-		/// <param name="HotReload">The hot-reload mode</param>
-		TargetReceipt PrepareReceipt(UEToolChain ToolChain, List<KeyValuePair<FileReference, BuildProductType>> BuildProducts, List<RuntimeDependency> RuntimeDependencies, EHotReload HotReload)
+		TargetReceipt PrepareReceipt(UEToolChain ToolChain, List<KeyValuePair<FileReference, BuildProductType>> BuildProducts, List<RuntimeDependency> RuntimeDependencies)
 		{
 			// Read the version file
 			BuildVersion Version;
@@ -1424,15 +1370,6 @@ namespace UnrealBuildTool
 					// If this is a formal build, we can just the compatible changelist as the unique id.
 					Version.BuildId = String.Format("{0}", Version.EffectiveCompatibleChangelist);
 				}
-				else if(HotReload != EHotReload.Disabled || OnlyModules.Count > 0)
-				{
-					// If we're hot reloading or doing a partial build, just use the last version number.
-					BuildVersion LastVersion;
-					if(VersionFile != null && BuildVersion.TryRead(VersionFile, out LastVersion))
-					{
-						Version = LastVersion;
-					}
-				}
 			}
 
 			// Create the receipt
@@ -1542,7 +1479,7 @@ namespace UnrealBuildTool
 		/// <summary>
 		/// Builds the target, appending list of output files and returns building result.
 		/// </summary>
-		public ECompilationResult Build(BuildConfiguration BuildConfiguration, CPPHeaders Headers, List<FileItem> OutputItems, List<UHTModuleInfo> UObjectModules, ISourceFileWorkingSet WorkingSet, ActionGraph ActionGraph, EHotReload HotReload, bool bIsAssemblingBuild)
+		public ECompilationResult Build(BuildConfiguration BuildConfiguration, CPPHeaders Headers, List<FileItem> OutputItems, Dictionary<string, FileItem[]> ModuleNameToOutputItems, List<UHTModuleInfo> UObjectModules, ISourceFileWorkingSet WorkingSet, ActionGraph ActionGraph, bool bIsAssemblingBuild)
 		{
 			CppPlatform CppPlatform = UEBuildPlatform.GetBuildPlatform(Platform).DefaultCppPlatform;
 			CppConfiguration CppConfiguration = GetCppConfiguration(Configuration);
@@ -1558,25 +1495,6 @@ namespace UnrealBuildTool
 			// Save off the original list of binaries. We'll use this to figure out which PCHs to create later, to avoid switching PCHs when compiling single modules.
 			List<UEBuildBinary> OriginalBinaries = Binaries;
 
-			// If we're building a single module, then find the binary for that module and add it to our target
-			if (OnlyModules.Count > 0)
-			{
-				NonFilteredModules = Binaries;
-				Binaries = GetFilteredOnlyModules(Binaries, OnlyModules);
-				if (Binaries.Count == 0)
-				{
-					throw new BuildException("One or more of the modules specified using the '-module' argument could not be found.");
-				}
-			}
-			else if (HotReload == EHotReload.FromIDE)
-			{
-				Binaries = GetFilteredGameModules(Binaries);
-				if (Binaries.Count == 0)
-				{
-					throw new BuildException("One or more of the modules specified using the '-module' argument could not be found.");
-				}
-			}
-
 			// For installed builds, filter out all the binaries that aren't in mods
 			if (UnrealBuildTool.IsProjectInstalled())
 			{
@@ -1783,14 +1701,6 @@ namespace UnrealBuildTool
 			// Generate headers
 			HashSet<UEBuildModuleCPP> ModulesToGenerateHeadersFor = GatherDependencyModules(OriginalBinaries.ToList());
 
-			if (OnlyModules.Count > 0)
-			{
-				HashSet<UEBuildModuleCPP> CorrectlyOrderedModules = GatherDependencyModules(NonFilteredModules);
-
-				CorrectlyOrderedModules.RemoveWhere((Module) => !ModulesToGenerateHeadersFor.Contains(Module));
-				ModulesToGenerateHeadersFor = CorrectlyOrderedModules;
-			}
-
 			Dictionary<string, FlatModuleCsDataType> NameToFlatModuleData = new Dictionary<string, FlatModuleCsDataType>(StringComparer.InvariantCultureIgnoreCase);
 			foreach(FlatModuleCsDataType FlatModuleData in FlatModuleCsData)
 			{
@@ -1847,7 +1757,12 @@ namespace UnrealBuildTool
 			DirectoryReference ExeDir = OutputPaths[0].Directory;
 			foreach (UEBuildBinary Binary in Binaries)
 			{
-				OutputItems.AddRange(Binary.Build(Rules, TargetToolChain, GlobalCompileEnvironment, GlobalLinkEnvironment, SharedPCHs, WorkingSet, ExeDir, ActionGraph));
+				List<FileItem> BinaryOutputItems = Binary.Build(Rules, TargetToolChain, GlobalCompileEnvironment, GlobalLinkEnvironment, SharedPCHs, WorkingSet, ExeDir, ActionGraph);
+				if(!bCompileMonolithic)
+				{
+					ModuleNameToOutputItems[Binary.PrimaryModule.Name] = BinaryOutputItems.ToArray();
+				}
+				OutputItems.AddRange(BinaryOutputItems);
 			}
 
 			// Prepare all the runtime dependencies, copying them from their source folders if necessary
@@ -1927,7 +1842,7 @@ namespace UnrealBuildTool
 				BuildProducts.AddRange(FileNameToModuleManifest.Select(x => new KeyValuePair<FileReference, BuildProductType>(x.Key, BuildProductType.RequiredResource)));
 
 				// Prepare the receipt
-				TargetReceipt Receipt = PrepareReceipt(TargetToolChain, BuildProducts, RuntimeDependencies, HotReload);
+				TargetReceipt Receipt = PrepareReceipt(TargetToolChain, BuildProducts, RuntimeDependencies);
 
 				// Create an action which to generate the receipts
 				WriteMetadataTargetInfo MetadataTargetInfo = new WriteMetadataTargetInfo(ProjectFile, VersionFile, ReceiptFileName, Receipt, FileNameToModuleManifest);
@@ -2478,54 +2393,25 @@ namespace UnrealBuildTool
 			return new FileReference(FilePath.FullName.Insert(MatchPos + ModuleName.Length, Appendage));
 		}
 
-		private static List<UEBuildBinary> GetFilteredOnlyModules(List<UEBuildBinary> Binaries, List<OnlyModule> OnlyModules)
+		/// <summary>
+		/// Finds a list of module names which can be hot-reloaded
+		/// </summary>
+		/// <returns>Set of module names</returns>
+		public HashSet<string> GetHotReloadModuleNames()
 		{
-			List<UEBuildBinary> Result = new List<UEBuildBinary>();
-
+			HashSet<string> HotReloadModuleNames = new HashSet<string>(StringComparer.OrdinalIgnoreCase);
 			foreach (UEBuildBinary Binary in Binaries)
 			{
-				// If we're doing an OnlyModule compile, we never want the executable that static libraries are linked into for monolithic builds
-				if(Binary.Type != UEBuildBinaryType.Executable)
-				{
-					OnlyModule FoundOnlyModule = Binary.FindOnlyModule(OnlyModules);
-					if (FoundOnlyModule != null)
-					{
-						Result.Add(Binary);
-
-						if (!String.IsNullOrEmpty(FoundOnlyModule.OnlyModuleSuffix))
-						{
-							Binary.OriginalOutputFilePaths = Binary.OutputFilePaths;
-							Binary.OutputFilePaths = Binary.OutputFilePaths.Select(Path => AddModuleFilenameSuffix(FoundOnlyModule.OnlyModuleName, Path, FoundOnlyModule.OnlyModuleSuffix)).ToList();
-						}
-					}
-				}
-			}
-
-			return Result;
-		}
-
-		private List<UEBuildBinary> GetFilteredGameModules(List<UEBuildBinary> Binaries)
-		{
-			List<UEBuildBinary> Result = new List<UEBuildBinary>();
-
-			foreach (UEBuildBinary DLLBinary in Binaries)
-			{
-				List<UEBuildModule> GameModules = DLLBinary.FindGameModules();
+				List<UEBuildModule> GameModules = Binary.FindGameModules();
 				if (GameModules != null && GameModules.Count > 0)
 				{
-					if(!UnrealBuildTool.IsProjectInstalled() || EnabledPlugins.Where(x => x.Type == PluginType.Mod).Any(x => DLLBinary.OutputFilePaths[0].IsUnderDirectory(x.Directory)))
+					if(!UnrealBuildTool.IsProjectInstalled() || EnabledPlugins.Where(x => x.Type == PluginType.Mod).Any(x => Binary.OutputFilePaths[0].IsUnderDirectory(x.Directory)))
 					{
-						Result.Add(DLLBinary);
-
-						string UniqueSuffix = (new Random((int)(DateTime.Now.Ticks % Int32.MaxValue)).Next(10000)).ToString();
-
-						DLLBinary.OriginalOutputFilePaths = DLLBinary.OutputFilePaths;
-						DLLBinary.OutputFilePaths = DLLBinary.OutputFilePaths.Select(Path => AddModuleFilenameSuffix(GameModules[0].Name, Path, UniqueSuffix)).ToList();
+						HotReloadModuleNames.UnionWith(GameModules.Select(x => x.Name));
 					}
 				}
 			}
-
-			return Result;
+			return HotReloadModuleNames;
 		}
 
 		/// <summary>
@@ -3846,7 +3732,7 @@ namespace UnrealBuildTool
 				// Figure out whether we need to build this module
 				// We don't care about actual source files when generating projects, as these are discovered separately
 				bool bDiscoverFiles = !ProjectFileGenerator.bGenerateProjectFiles;
-				bool bBuildFiles = bDiscoverFiles && (OnlyModules.Count == 0 || OnlyModules.Any(x => string.Equals(x.OnlyModuleName, ModuleName, StringComparison.InvariantCultureIgnoreCase)));
+				bool bBuildFiles = bDiscoverFiles;
 
 				List<FileItem> FoundSourceFiles = new List<FileItem>();
 				if (RulesObject.Type == ModuleRules.ModuleType.CPlusPlus)
diff --git a/Engine/Source/Programs/UnrealBuildTool/ProjectFiles/ProjectFileGenerator.cs b/Engine/Source/Programs/UnrealBuildTool/ProjectFiles/ProjectFileGenerator.cs
--- a/Engine/Source/Programs/UnrealBuildTool/ProjectFiles/ProjectFileGenerator.cs
+++ b/Engine/Source/Programs/UnrealBuildTool/ProjectFiles/ProjectFileGenerator.cs
@@ -1730,7 +1730,7 @@ namespace UnrealBuildTool
 							// If we're generating project files, then go ahead and wipe out the existing UBTMakefile for every target, to make sure that
 							// it gets a full dependency scan next time.
 							// NOTE: This is just a safeguard and doesn't have to be perfect.  We also check for newer project file timestamps in LoadUBTMakefile()
-							FileReference UBTMakefilePath = UBTMakefile.GetUBTMakefilePath(new List<TargetDescriptor>{ TargetDesc }, EHotReload.Disabled);
+							FileReference UBTMakefilePath = UBTMakefile.GetUBTMakefilePath(new List<TargetDescriptor>{ TargetDesc });
 							if (FileReference.Exists(UBTMakefilePath))
 							{
 								FileReference.Delete(UBTMakefilePath);
diff --git a/Engine/Source/Programs/UnrealBuildTool/System/ActionGraph.cs b/Engine/Source/Programs/UnrealBuildTool/System/ActionGraph.cs
--- a/Engine/Source/Programs/UnrealBuildTool/System/ActionGraph.cs
+++ b/Engine/Source/Programs/UnrealBuildTool/System/ActionGraph.cs
@@ -308,7 +308,7 @@ namespace UnrealBuildTool
 		/// <summary>
 		/// Builds a list of actions that need to be executed to produce the specified output items.
 		/// </summary>
-		public List<Action> GetActionsToExecute(BuildConfiguration BuildConfiguration, Action[] PrerequisiteActions, List<UEBuildTarget> Targets, Dictionary<UEBuildTarget, CPPHeaders> TargetToHeaders, bool bIsAssemblingBuild, bool bNeedsFullCPPIncludeRescan, out Dictionary<UEBuildTarget, List<FileItem>> TargetToOutdatedPrerequisitesMap)
+		public HashSet<Action> GetActionsToExecute(BuildConfiguration BuildConfiguration, Action[] PrerequisiteActions, List<UEBuildTarget> Targets, Dictionary<UEBuildTarget, CPPHeaders> TargetToHeaders, bool bIsAssemblingBuild, bool bNeedsFullCPPIncludeRescan, out Dictionary<UEBuildTarget, List<FileItem>> TargetToOutdatedPrerequisitesMap)
 		{
 			DateTime CheckOutdatednessStartTime = DateTime.UtcNow;
 
@@ -371,63 +371,13 @@ namespace UnrealBuildTool
 			// Build a list of actions that are both needed for this target and outdated.
 			HashSet<Action> ActionsToExecute = new HashSet<Action>(AllActions.Where(Action => Action.CommandPath != null && IsActionOutdatedMap.ContainsKey(Action) && OutdatedActionDictionary[Action]));
 
-			// Remove link actions if asked to
-			if (BuildConfiguration.bSkipLinkingWhenNothingToCompile)
-			{
-				// Get all items produced by a compile action
-				HashSet<FileItem> ProducedItems = new HashSet<FileItem>(ActionsToExecute.Where(Action => Action.ActionType == ActionType.Compile).SelectMany(x => x.ProducedItems));
-
-				// Get all link actions which have no out-of-date prerequisites
-				HashSet<Action> UnlinkedActions = new HashSet<Action>(ActionsToExecute.Where(Action => Action.ActionType == ActionType.Link && !ProducedItems.Overlaps(Action.PrerequisiteItems)));
-
-				// Don't regard an action as unlinked if there is an associated 'failed.hotreload' file.
-				UnlinkedActions.RemoveWhere(Action => Action.ProducedItems.Any(Item => File.Exists(Path.Combine(Path.GetDirectoryName(Item.AbsolutePath), "failed.hotreload"))));
-
-				HashSet<Action> UnlinkedActionsWithFailedHotreload = new HashSet<Action>(ActionsToExecute.Where(Action => Action.ActionType == ActionType.Link && !ProducedItems.Overlaps(Action.PrerequisiteItems)));
-
-				// Remove unlinked items
-				ActionsToExecute.ExceptWith(UnlinkedActions);
-
-				// Re-add unlinked items which produce things which are dependencies of other actions
-				for (;;)
-				{
-					// Get all prerequisite items of a link action
-					HashSet<Action> PrerequisiteLinkActions = new HashSet<Action>(ActionsToExecute.Where(Action => Action.ActionType == ActionType.Link).SelectMany(x => x.PrerequisiteItems).Select(Item => Item.ProducingAction));
-
-					// Find all unlinked actions that need readding
-					HashSet<Action> UnlinkedActionsToReadd = new HashSet<Action>(UnlinkedActions.Where(Action => PrerequisiteLinkActions.Contains(Action)));
-
-					// Also re-add any DLL whose import library is being rebuilt. These may be separate actions, and the import library will reference the new DLL even if it isn't being compiled itself, so it must exist.
-					HashSet<string> ProducedDllsToReAdd = new HashSet<string>(UnlinkedActionsToReadd.SelectMany(x => x.ProducedItems).Select(x => x.Location).Where(x => x.HasExtension(".lib")).Select(x => x.GetFileNameWithoutExtension() + ".dll"), StringComparer.OrdinalIgnoreCase);
-					UnlinkedActionsToReadd.UnionWith(UnlinkedActions.Where(x => x.ProducedItems.Any(y => ProducedDllsToReAdd.Contains(y.Location.GetFileName()))));
-
-					// Bail if we didn't find anything
-					if (UnlinkedActionsToReadd.Count == 0)
-					{
-						break;
-					}
-
-					ActionsToExecute.UnionWith(UnlinkedActionsToReadd);
-					UnlinkedActions.ExceptWith(UnlinkedActionsToReadd);
-
-					// Break early if there are no more unlinked actions to readd
-					if (UnlinkedActions.Count == 0)
-					{
-						break;
-					}
-				}
-
-				// Remove actions that are wholly dependent on unlinked actions
-				ActionsToExecute = new HashSet<Action>(ActionsToExecute.Where(Action => Action.PrerequisiteItems.Count == 0 || !new HashSet<Action>(Action.PrerequisiteItems.Select(Item => Item.ProducingAction)).IsSubsetOf(UnlinkedActions)));
-			}
-
 			if (UnrealBuildTool.bPrintPerformanceInfo)
 			{
 				double CheckOutdatednessTime = (DateTime.UtcNow - CheckOutdatednessStartTime).TotalSeconds;
 				Log.TraceInformation("Checking outdatedness took " + CheckOutdatednessTime + "s");
 			}
 
-			return ActionsToExecute.ToList();
+			return ActionsToExecute;
 		}
 
 		/// <summary>
@@ -573,82 +523,6 @@ namespace UnrealBuildTool
 			return WorkingString;
 		}
 
-
-		/// <summary>
-		/// Finds and deletes stale hot reload DLLs.
-		/// </summary>
-		public void DeleteStaleHotReloadDLLs()
-		{
-			DateTime DeleteStartTime = DateTime.UtcNow;
-
-			foreach (Action BuildAction in AllActions)
-			{
-				if (BuildAction.ActionType == ActionType.Link)
-				{
-					foreach (FileItem Item in BuildAction.ProducedItems)
-					{
-						if (Item.bNeedsHotReloadNumbersDLLCleanUp)
-						{
-							string PlatformSuffix, ConfigSuffix, ProducedItemExtension;
-							string Base = SplitFilename(Item.AbsolutePath, out PlatformSuffix, out ConfigSuffix, out ProducedItemExtension);
-							String WildCard = Base + "-*" + PlatformSuffix + ConfigSuffix + ProducedItemExtension;
-							// Log.TraceInformation("Deleting old hot reload wildcard: \"{0}\".", WildCard);
-							// Wildcard search and delete
-							string DirectoryToLookIn = Path.GetDirectoryName(WildCard);
-							string FileName = Path.GetFileName(WildCard);
-							if (Directory.Exists(DirectoryToLookIn))
-							{
-								// Delete all files within the specified folder
-								string[] FilesToDelete = Directory.GetFiles(DirectoryToLookIn, FileName, SearchOption.TopDirectoryOnly);
-								foreach (string JunkFile in FilesToDelete)
-								{
-
-									string JunkPlatformSuffix, JunkConfigSuffix, JunkProducedItemExtension;
-									SplitFilename(JunkFile, out JunkPlatformSuffix, out JunkConfigSuffix, out JunkProducedItemExtension);
-									// now make sure that this file has the same config and platform
-									if (JunkPlatformSuffix == PlatformSuffix && JunkConfigSuffix == ConfigSuffix)
-									{
-										try
-										{
-											Log.TraceInformation("Deleting old hot reload file: \"{0}\".", JunkFile);
-											File.Delete(JunkFile);
-										}
-										catch (Exception Ex)
-										{
-											// Ignore all exceptions
-											Log.TraceInformation("Unable to delete old hot reload file: \"{0}\". Error: {1}", JunkFile, Ex.Message.TrimEnd());
-										}
-
-										// Delete the PDB file.
-										string JunkPDBFile = JunkFile.Replace(ProducedItemExtension, ".pdb");
-										if (System.IO.File.Exists(JunkPDBFile))
-										{
-											try
-											{
-												Log.TraceInformation("Deleting old hot reload file: \"{0}\".", JunkPDBFile);
-												File.Delete(JunkPDBFile);
-											}
-											catch (Exception Ex)
-											{
-												// Ignore all exceptions
-												Log.TraceInformation("Unable to delete old hot reload file: \"{0}\". Error: {1}", JunkPDBFile, Ex.Message.TrimEnd());
-											}
-										}
-									}
-								}
-							}
-						}
-					}
-				}
-			}
-
-			if (UnrealBuildTool.bPrintPerformanceInfo)
-			{
-				double DeleteTime = (DateTime.UtcNow - DeleteStartTime).TotalSeconds;
-				Log.TraceInformation("Deleting stale hot reload DLLs took " + DeleteTime + "s");
-			}
-		}
-
 		/// <summary>
 		/// Sorts the action list for improved parallelism with local execution.
 		/// </summary>
diff --git a/Engine/Source/Programs/UnrealBuildTool/System/UBTMakefile.cs b/Engine/Source/Programs/UnrealBuildTool/System/UBTMakefile.cs
--- a/Engine/Source/Programs/UnrealBuildTool/System/UBTMakefile.cs
+++ b/Engine/Source/Programs/UnrealBuildTool/System/UBTMakefile.cs
@@ -33,11 +33,6 @@ namespace UnrealBuildTool
 		/// </summary>
 		public List<Action> AllActions;
 
-		/// <summary>
-		/// List of the actions that need to be run in order to build the targets' final output items
-		/// </summary>
-		public Action[] PrerequisiteActions;
-
 		/// <summary>
 		/// Environment variables that we'll need in order to invoke the platform's compiler and linker
 		/// </summary>
@@ -49,6 +44,21 @@ namespace UnrealBuildTool
 		/// </summary>
 		public Dictionary<string, List<UHTModuleInfo>> TargetNameToUObjectModules;
 
+		/// <summary>
+		/// The final output items for all target
+		/// </summary>
+		public List<FileItem> OutputItemsForAllTargets;
+
+		/// <summary>
+		/// Maps module names to output items
+		/// </summary>
+		public Dictionary<string, FileItem[]> ModuleNameToOutputItems;
+
+		/// <summary>
+		/// List of game module names, for hot-reload
+		/// </summary>
+		public HashSet<string> HotReloadModuleNamesForAllTargets;
+
 		/// <summary>
 		/// List of targets being built
 		/// </summary>
@@ -83,9 +93,11 @@ namespace UnrealBuildTool
 			}
 
 			AllActions = (List<Action>)Info.GetValue("ac", typeof(List<Action>));
-			PrerequisiteActions = (Action[])Info.GetValue("pa", typeof(Action[]));
 			EnvironmentVariables = ((string[])Info.GetValue("e1", typeof(string[]))).Zip((string[])Info.GetValue("e2", typeof(string[])), (i1, i2) => new Tuple<string, string>(i1, i2)).ToList();
 			TargetNameToUObjectModules = (Dictionary<string, List<UHTModuleInfo>>)Info.GetValue("nu", typeof(Dictionary<string, List<UHTModuleInfo>>));
+			OutputItemsForAllTargets = (List<FileItem>)Info.GetValue("oi", typeof(List<FileItem>));
+			ModuleNameToOutputItems = (Dictionary<string, FileItem[]>)Info.GetValue("mn", typeof(Dictionary<string, FileItem[]>));
+			HotReloadModuleNamesForAllTargets = (HashSet<string>)Info.GetValue("gm", typeof(HashSet<string>));
 			Targets = (List<UEBuildTarget>)Info.GetValue("ta", typeof(List<UEBuildTarget>));
 			bUseAdaptiveUnityBuild = Info.GetBoolean("ua");
 			SourceFileWorkingSet = (HashSet<FileItem>)Info.GetValue("ws", typeof(HashSet<FileItem>));
@@ -96,10 +108,12 @@ namespace UnrealBuildTool
 		{
 			Info.AddValue("ve", Version);
 			Info.AddValue("ac", AllActions);
-			Info.AddValue("pa", PrerequisiteActions);
 			Info.AddValue("e1", EnvironmentVariables.Select(x => x.Item1).ToArray());
 			Info.AddValue("e2", EnvironmentVariables.Select(x => x.Item2).ToArray());
 			Info.AddValue("nu", TargetNameToUObjectModules);
+			Info.AddValue("oi", OutputItemsForAllTargets);
+			Info.AddValue("mn", ModuleNameToOutputItems);
+			Info.AddValue("gm", HotReloadModuleNamesForAllTargets);
 			Info.AddValue("ta", Targets);
 			Info.AddValue("ua", bUseAdaptiveUnityBuild);
 			Info.AddValue("ws", SourceFileWorkingSet);
@@ -112,9 +126,11 @@ namespace UnrealBuildTool
 		{
 			return
 				AllActions != null && AllActions.Count > 0 &&
-				PrerequisiteActions != null && PrerequisiteActions.Length > 0 &&
 				EnvironmentVariables != null &&
 				TargetNameToUObjectModules != null && TargetNameToUObjectModules.Count > 0 &&
+				OutputItemsForAllTargets != null && 
+				ModuleNameToOutputItems != null && 
+				HotReloadModuleNamesForAllTargets != null &&
 				Targets != null && Targets.Count > 0 &&
 				SourceFileWorkingSet != null &&
 				CandidateSourceFilesForWorkingSet != null;
@@ -125,9 +141,8 @@ namespace UnrealBuildTool
 		/// Saves a UBTMakefile to disk
 		/// </summary>
 		/// <param name="TargetDescs">List of targets.  Order is not important</param>
-		/// <param name="HotReload">The hot reload state</param>
 		/// <param name="UBTMakefile">The UBT makefile</param>
-		public static void SaveUBTMakefile(List<TargetDescriptor> TargetDescs, EHotReload HotReload, UBTMakefile UBTMakefile)
+		public static void SaveUBTMakefile(List<TargetDescriptor> TargetDescs, UBTMakefile UBTMakefile)
 		{
 			if (!UBTMakefile.IsValidMakefile())
 			{
@@ -136,7 +151,7 @@ namespace UnrealBuildTool
 
 			DateTime TimerStartTime = DateTime.UtcNow;
 
-			FileItem UBTMakefileItem = FileItem.GetItemByFileReference(GetUBTMakefilePath(TargetDescs, HotReload));
+			FileItem UBTMakefileItem = FileItem.GetItemByFileReference(GetUBTMakefilePath(TargetDescs));
 
 			// @todo ubtmake: Optimization: The UBTMakefile saved for game projects is upwards of 9 MB.  We should try to shrink its content if possible
 			// @todo ubtmake: Optimization: C# Serialization may be too slow for these big Makefiles.  Loading these files often shows up as the slower part of the assembling phase.
@@ -188,7 +203,7 @@ namespace UnrealBuildTool
 			FileInfo BuildVersionFileInfo = new FileInfo(BuildVersion.GetDefaultFileName().FullName);
 			if (BuildVersionFileInfo.Exists && UBTMakefileInfo.LastWriteTime.CompareTo(BuildVersionFileInfo.LastWriteTime) < 0)
 			{
-				Log.TraceVerbose("Existing makefile is older than Build.version, ignoring it");
+				Log.TraceLog("Existing makefile is older than Build.version, ignoring it");
 				ReasonNotLoaded = "Build.version is newer";
 				return null;
 			}
@@ -206,7 +221,7 @@ namespace UnrealBuildTool
 					if (UBTMakefileInfo.LastWriteTime.CompareTo(EngineProjectFilesLastUpdateTime) < 0)
 					{
 						// Engine project files are newer than UBTMakefile
-						Log.TraceVerbose("Existing makefile is older than generated engine project files, ignoring it");
+						Log.TraceLog("Existing makefile is older than generated engine project files, ignoring it");
 						ReasonNotLoaded = "project files are newer";
 						return null;
 					}
@@ -221,7 +236,7 @@ namespace UnrealBuildTool
 				if (!ProjectFileInfo.Exists || UBTMakefileInfo.LastWriteTime.CompareTo(ProjectFileInfo.LastWriteTime) < 0)
 				{
 					// .uproject file is newer than UBTMakefile
-					Log.TraceVerbose("Makefile is older than .uproject file, ignoring it");
+					Log.TraceLog("Makefile is older than .uproject file, ignoring it");
 					ReasonNotLoaded = ".uproject file is newer";
 					return null;
 				}
@@ -234,7 +249,7 @@ namespace UnrealBuildTool
 					if (UBTMakefileInfo.LastWriteTime.CompareTo(GameProjectFilesLastUpdateTime) < 0)
 					{
 						// Game project files are newer than UBTMakefile
-						Log.TraceVerbose("Makefile is older than generated game project files, ignoring it");
+						Log.TraceLog("Makefile is older than generated game project files, ignoring it");
 						ReasonNotLoaded = "game project files are newer";
 						return null;
 					}
@@ -246,7 +261,7 @@ namespace UnrealBuildTool
 			if (UBTMakefileInfo.LastWriteTime.CompareTo(UnrealBuildToolTimestamp) < 0)
 			{
 				// UnrealBuildTool.exe was compiled more recently than the UBTMakefile
-				Log.TraceVerbose("Makefile is older than UnrealBuildTool.exe, ignoring it");
+				Log.TraceLog("Makefile is older than UnrealBuildTool.exe, ignoring it");
 				ReasonNotLoaded = "UnrealBuildTool.exe is newer";
 				return null;
 			}
@@ -258,7 +273,7 @@ namespace UnrealBuildTool
 				FileInfo InputFileInfo = new FileInfo(InputFile.Location.FullName);
 				if (InputFileInfo.LastWriteTime > UBTMakefileInfo.LastWriteTime)
 				{
-					Log.TraceVerbose("Makefile is older than BuildConfiguration.xml, ignoring it");
+					Log.TraceLog("Makefile is older than BuildConfiguration.xml, ignoring it");
 					ReasonNotLoaded = "BuildConfiguration.xml is newer";
 					return null;
 				}
@@ -307,7 +322,7 @@ namespace UnrealBuildTool
 					bool bTargetCsFileExists = TargetCsFile.Exists;
 					if (!bTargetCsFileExists || TargetCsFile.LastWriteTime > UBTMakefileInfo.LastWriteTime)
 					{
-						Log.TraceVerbose("{0} has been {1} since makefile was built, ignoring it ({2})", TargetCsFilename, bTargetCsFileExists ? "changed" : "deleted", UBTMakefileInfo.FullName);
+						Log.TraceLog("{0} has been {1} since makefile was built, ignoring it ({2})", TargetCsFilename, bTargetCsFileExists ? "changed" : "deleted", UBTMakefileInfo.FullName);
 						ReasonNotLoaded = string.Format("changes to target files");
 						return null;
 					}
@@ -322,7 +337,7 @@ namespace UnrealBuildTool
 						bool bBuildCsFileExists = BuildCsFile.Exists;
 						if (!bBuildCsFileExists || BuildCsFile.LastWriteTime > UBTMakefileInfo.LastWriteTime)
 						{
-							Log.TraceVerbose("{0} has been {1} since makefile was built, ignoring it ({2})", BuildCsFilename, bBuildCsFileExists ? "changed" : "deleted", UBTMakefileInfo.FullName);
+							Log.TraceLog("{0} has been {1} since makefile was built, ignoring it ({2})", BuildCsFilename, bBuildCsFileExists ? "changed" : "deleted", UBTMakefileInfo.FullName);
 							ReasonNotLoaded = string.Format("changes to module files");
 							return null;
 						}
@@ -340,7 +355,7 @@ namespace UnrealBuildTool
 							bool bDependencyFileExists = DependencyFile.Exists;
 							if (!bDependencyFileExists || DependencyFile.LastWriteTime > UBTMakefileInfo.LastWriteTime)
 							{
-								Log.TraceVerbose("{0} has been {1} since makefile was built, ignoring it ({2})", DependencyFile.FullName, bDependencyFileExists ? "changed" : "deleted", UBTMakefileInfo.FullName);
+								Log.TraceLog("{0} has been {1} since makefile was built, ignoring it ({2})", DependencyFile.FullName, bDependencyFileExists ? "changed" : "deleted", UBTMakefileInfo.FullName);
 								ReasonNotLoaded = string.Format("changes to external dependency");
 								return null;
 							}
@@ -353,7 +368,7 @@ namespace UnrealBuildTool
 					DirectoryInfo SourceDirectoryInfo = new DirectoryInfo(SourceDirectory.FullName);
 					if(!SourceDirectoryInfo.Exists || SourceDirectoryInfo.LastWriteTimeUtc > UBTMakefileInfo.LastWriteTimeUtc)
 					{
-						Log.TraceVerbose("Timestamp of {0} ({1}) is newer than makefile ({2})", SourceDirectory, SourceDirectoryInfo.LastWriteTimeUtc, UBTMakefileInfo.LastWriteTimeUtc);
+						Log.TraceLog("Timestamp of {0} ({1}) is newer than makefile ({2})", SourceDirectory, SourceDirectoryInfo.LastWriteTimeUtc, UBTMakefileInfo.LastWriteTimeUtc);
 						ReasonNotLoaded = "source directory changed";
 						return null;
 					}
@@ -383,7 +398,7 @@ namespace UnrealBuildTool
 				{
 					if (!File.Exists(FileName))
 					{
-						Log.TraceVerbose("File processed by UHT was deleted ({0}); invalidating makefile", FileName);
+						Log.TraceLog("File processed by UHT was deleted ({0}); invalidating makefile", FileName);
 						ReasonNotLoaded = string.Format("UHT file was deleted");
 						return null;
 					}
@@ -398,7 +413,7 @@ namespace UnrealBuildTool
 					bool bWasProcessed = AllUHTHeaders.Contains(Filename);
 					if (bContainsUHTData != bWasProcessed)
 					{
-						Log.TraceVerbose("{0} {1} contain UHT types and now {2} , ignoring it ({3})", Filename, bWasProcessed ? "used to" : "didn't", bWasProcessed ? "doesn't" : "does", UBTMakefileInfo.FullName);
+						Log.TraceLog("{0} {1} contain UHT types and now {2} , ignoring it ({3})", Filename, bWasProcessed ? "used to" : "didn't", bWasProcessed ? "doesn't" : "does", UBTMakefileInfo.FullName);
 						ReasonNotLoaded = string.Format("new files with reflected types");
 						return null;
 					}
@@ -417,7 +432,7 @@ namespace UnrealBuildTool
 				{
 					if (!WorkingSet.Contains(SourceFile.Location) && File.GetLastWriteTimeUtc(SourceFile.AbsolutePath) > UBTMakefileInfo.LastWriteTimeUtc)
 					{
-						Log.TraceVerbose("{0} was part of source working set and now is not; invalidating makefile ({1})", SourceFile.AbsolutePath, UBTMakefileInfo.FullName);
+						Log.TraceLog("{0} was part of source working set and now is not; invalidating makefile ({1})", SourceFile.AbsolutePath, UBTMakefileInfo.FullName);
 						ReasonNotLoaded = string.Format("working set of source files changed");
 						return null;
 					}
@@ -428,7 +443,7 @@ namespace UnrealBuildTool
 				{
 					if (WorkingSet.Contains(SourceFile.Location) && File.GetLastWriteTimeUtc(SourceFile.AbsolutePath) > UBTMakefileInfo.LastWriteTimeUtc)
 					{
-						Log.TraceVerbose("{0} was part of source working set and now is not; invalidating makefile ({1})", SourceFile.AbsolutePath, UBTMakefileInfo.FullName);
+						Log.TraceLog("{0} was part of source working set and now is not; invalidating makefile ({1})", SourceFile.AbsolutePath, UBTMakefileInfo.FullName);
 						ReasonNotLoaded = string.Format("working set of source files changed");
 						return null;
 					}
@@ -444,9 +459,8 @@ namespace UnrealBuildTool
 		/// Gets the file path for a UBTMakefile
 		/// </summary>
 		/// <param name="TargetDescs">List of targets.  Order is not important</param>
-		/// <param name="HotReload">The hot reload state.</param>
 		/// <returns>UBTMakefile path</returns>
-		public static FileReference GetUBTMakefilePath(List<TargetDescriptor> TargetDescs, EHotReload HotReload)
+		public static FileReference GetUBTMakefilePath(List<TargetDescriptor> TargetDescs)
 		{
 			FileReference UBTMakefilePath;
 
@@ -454,7 +468,7 @@ namespace UnrealBuildTool
 			{
 				TargetDescriptor TargetDesc = TargetDescs[0];
 
-				string UBTMakefileName = (HotReload != EHotReload.Disabled) ? "HotReloadMakefile.ubt" : "Makefile.ubt";
+				string UBTMakefileName = "Makefile.ubt";
 
 				UBTMakefilePath = FileReference.Combine(GetUBTMakefileDirectoryPathForSingleTarget(TargetDesc), UBTMakefileName);
 			}
diff --git a/Engine/Source/Programs/UnrealBuildTool/UnrealBuildTool.cs b/Engine/Source/Programs/UnrealBuildTool/UnrealBuildTool.cs
--- a/Engine/Source/Programs/UnrealBuildTool/UnrealBuildTool.cs
+++ b/Engine/Source/Programs/UnrealBuildTool/UnrealBuildTool.cs
@@ -2,14 +2,12 @@
 
 using System;
 using System.Collections.Generic;
-using System.Text;
-using System.IO;
 using System.Diagnostics;
-using System.Threading;
-using System.Reflection;
+using System.IO;
 using System.Linq;
-using System.Runtime.Serialization.Formatters.Binary;
-using System.Runtime.Serialization;
+using System.Reflection;
+using System.Text;
+using System.Threading;
 using Tools.DotNETCommon;
 
 namespace UnrealBuildTool
@@ -21,6 +19,14 @@ namespace UnrealBuildTool
 		FromEditor
 	}
 
+	[Serializable]
+	class HotReloadState
+	{
+		public int NextSuffix = 1;
+		public Dictionary<FileReference, FileReference> OriginalFileToHotReloadFile = new Dictionary<FileReference, FileReference>();
+		public HashSet<FileReference> AllHotReloadFiles = new HashSet<FileReference>();
+	}
+
 	static class UnrealBuildTool
 	{
 		/// <summary>
@@ -722,25 +728,6 @@ namespace UnrealBuildTool
 					return ECompilationResult.Succeeded;
 				}
 
-				EHotReload HotReload = EHotReload.Disabled;
-				if (TargetDescs.Count == 1 && !Arguments.Any(x => x.Equals("-NoHotReload", StringComparison.InvariantCultureIgnoreCase)))
-				{
-					if (BuildConfiguration.bAllowHotReloadFromIDE && ShouldDoHotReloadFromIDE(BuildConfiguration, Arguments, TargetDescs[0]))
-					{
-						HotReload = EHotReload.FromIDE;
-					}
-					else if (TargetDescs[0].OnlyModules.Count > 0 && TargetDescs[0].ForeignPlugin == null)
-					{
-						HotReload = EHotReload.FromEditor;
-					}
-
-					if (HotReload != EHotReload.Disabled && BuildConfiguration.bCleanProject)
-					{
-						throw new BuildException("Unable to clean target while hot-reloading. Close the editor and try again.");
-					}
-				}
-				TargetDescriptor HotReloadTargetDesc = (HotReload != EHotReload.Disabled) ? TargetDescs[0] : null;
-
 				// Used when BuildConfiguration.bUseUBTMakefiles is enabled.  If true, it means that our cached includes may not longer be
 				// valid (or never were), and we need to recover by forcibly scanning included headers for all build prerequisites to make sure that our
 				// cached set of includes is actually correct, before determining which files are outdated.
@@ -760,7 +747,7 @@ namespace UnrealBuildTool
 
 						string TargetCollectionName = UBTMakefile.MakeTargetCollectionName(TargetDescs);
 
-						string LastBuiltTargetsFileName = (HotReload != EHotReload.Disabled) ? "HotReloadLastBuiltTargets.txt" : "LastBuiltTargets.txt";
+						string LastBuiltTargetsFileName = "LastBuiltTargets.txt";
 						string LastBuiltTargetsFilePath = FileReference.Combine(UnrealBuildTool.EngineDirectory, "Intermediate", "Build", LastBuiltTargetsFileName).FullName;
 						if (File.Exists(LastBuiltTargetsFilePath) && Utils.ReadAllText(LastBuiltTargetsFilePath) == TargetCollectionName)
 						{
@@ -781,7 +768,7 @@ namespace UnrealBuildTool
 							// we might not know about all of the C++ include dependencies for already-up-to-date shared build products
 							// between the targets
 							bNeedsFullCPPIncludeRescan = true;
-							Log.TraceInformation("Performing full C++ include scan ({0} a new target)", (HotReload != EHotReload.Disabled) ? "hot reloading" : "building");
+							Log.TraceInformation("Performing full C++ include scan (building a new target)");
 						}
 					}
 				}
@@ -841,7 +828,7 @@ namespace UnrealBuildTool
 				{
 					UBTMakefile UBTMakefile = null;
 					{
-						FileReference UBTMakefilePath = UBTMakefile.GetUBTMakefilePath(TargetDescs, HotReload);
+						FileReference UBTMakefilePath = UBTMakefile.GetUBTMakefilePath(TargetDescs);
 
 						// Make sure the gather phase is executed if we're not actually building anything
 						if (BuildConfiguration.bGenerateManifest || BuildConfiguration.bCleanProject || BuildConfiguration.bXGEExport || GeneratingActionGraph)
@@ -860,19 +847,6 @@ namespace UnrealBuildTool
 							string ReasonNotLoaded;
 							UBTMakefile = UBTMakefile.LoadUBTMakefile(UBTMakefilePath, ProjectFile, WorkingSet, out ReasonNotLoaded);
 
-							// Invalid makefile if only modules have changed
-							if (UBTMakefile != null && !TargetDescs.SelectMany(x => x.OnlyModules)
-									.Select(x => new Tuple<string, bool>(x.OnlyModuleName.ToLower(), string.IsNullOrWhiteSpace(x.OnlyModuleSuffix)))
-									.SequenceEqual(
-										UBTMakefile.Targets.SelectMany(x => x.OnlyModules)
-										.Select(x => new Tuple<string, bool>(x.OnlyModuleName.ToLower(), string.IsNullOrWhiteSpace(x.OnlyModuleSuffix)))
-									)
-								)
-							{
-								UBTMakefile = null;
-								ReasonNotLoaded = "modules to compile have changed";
-							}
-
 							if (UBTMakefile != null)
 							{
 								// Check if ini files are newer. Ini files contain build settings too.
@@ -914,9 +888,8 @@ namespace UnrealBuildTool
 
 								FileItem.ClearCaches();
 
-								Log.TraceInformation("Creating makefile for {0}{1}{2} ({3})",
-									(HotReload != EHotReload.Disabled) ? "hot reloading " : "",
-										TargetDescs[0].Name,
+								Log.TraceInformation("Creating makefile for {0}{1} ({2})",
+									TargetDescs[0].Name,
 									TargetDescs.Count > 1 ? (" (and " + (TargetDescs.Count - 1).ToString() + " more)") : "",
 									ReasonNotLoaded);
 							}
@@ -953,16 +926,11 @@ namespace UnrealBuildTool
 
 					// Build action lists for all passed in targets.
 					List<FileItem> OutputItemsForAllTargets = new List<FileItem>();
+					Dictionary<string, FileItem[]> ModuleNameToOutputItems = new Dictionary<string, FileItem[]>(StringComparer.OrdinalIgnoreCase);
 					Dictionary<string, List<UHTModuleInfo>> TargetNameToUObjectModules = new Dictionary<string, List<UHTModuleInfo>>(StringComparer.InvariantCultureIgnoreCase);
+					HashSet<string> HotReloadModuleNamesForAllTargets = new HashSet<string>(StringComparer.OrdinalIgnoreCase);
 					foreach (UEBuildTarget Target in Targets)
 					{
-						if (HotReload != EHotReload.Disabled)
-						{
-							// Don't produce new DLLs if there's been no code changes
-							BuildConfiguration.bSkipLinkingWhenNothingToCompile = true;
-							Log.TraceInformation("Compiling game modules for hot reload");
-						}
-
 						// Create the header cache for this target
 						FileReference DependencyCacheFile = DependencyCache.GetDependencyCachePathForTarget(ProjectFile, Target.Platform, Target.TargetName);
 						bool bUseFlatCPPIncludeDependencyCache = BuildConfiguration.bUseUBTMakefiles && bIsAssemblingBuild;
@@ -1007,13 +975,14 @@ namespace UnrealBuildTool
 						{
 							List<FileItem> TargetOutputItems = new List<FileItem>();
 							List<UHTModuleInfo> TargetUObjectModules = new List<UHTModuleInfo>();
+							Dictionary<string, FileItem[]> TargetModuleNameToOutputItems = new Dictionary<string, FileItem[]>(StringComparer.OrdinalIgnoreCase);
 							if (BuildConfiguration.bCleanProject)
 							{
-								BuildResult = Target.Clean(!BuildConfiguration.bDoNotBuildUHT && HotReload != EHotReload.FromIDE);
+								BuildResult = Target.Clean(!BuildConfiguration.bDoNotBuildUHT);
 							}
 							else
 							{
-								BuildResult = Target.Build(BuildConfiguration, TargetToHeaders[Target], TargetOutputItems, TargetUObjectModules, WorkingSet, ActionGraph, HotReload, bIsAssemblingBuild);
+								BuildResult = Target.Build(BuildConfiguration, TargetToHeaders[Target], TargetOutputItems, TargetModuleNameToOutputItems, TargetUObjectModules, WorkingSet, ActionGraph, bIsAssemblingBuild);
 							}
 							if (BuildResult != ECompilationResult.Succeeded)
 							{
@@ -1031,6 +1000,23 @@ namespace UnrealBuildTool
 
 							// Update mapping of the target name to the list of UObject modules in this target
 							TargetNameToUObjectModules[Target.GetTargetName()] = TargetUObjectModules;
+
+							// Merge the module name to output items map
+							foreach(KeyValuePair<string, FileItem[]> Pair in TargetModuleNameToOutputItems)
+							{
+								FileItem[] Items;
+								if(ModuleNameToOutputItems.TryGetValue(Pair.Key, out Items))
+								{
+									ModuleNameToOutputItems[Pair.Key] = Items.Concat(Pair.Value).ToArray();
+								}
+								else
+								{
+									ModuleNameToOutputItems[Pair.Key] = Pair.Value;
+								}
+							}
+
+							// Get the game module names
+							HotReloadModuleNamesForAllTargets.UnionWith(Target.GetHotReloadModuleNames());
 						}
 					}
 
@@ -1046,22 +1032,14 @@ namespace UnrealBuildTool
 
 							UBTMakefile = new UBTMakefile();
 							UBTMakefile.AllActions = ActionGraph.AllActions;
-
-							// For now simply treat all object files as the root target.
-							{
-								HashSet<Action> PrerequisiteActionsSet = new HashSet<Action>();
-								foreach (FileItem OutputItem in OutputItemsForAllTargets)
-								{
-									ActionGraph.GatherPrerequisiteActions(OutputItem, ref PrerequisiteActionsSet);
-								}
-								UBTMakefile.PrerequisiteActions = PrerequisiteActionsSet.ToArray();
-							}
-
+							UBTMakefile.OutputItemsForAllTargets = OutputItemsForAllTargets;
 							foreach (System.Collections.DictionaryEntry EnvironmentVariable in Environment.GetEnvironmentVariables())
 							{
 								UBTMakefile.EnvironmentVariables.Add(Tuple.Create((string)EnvironmentVariable.Key, (string)EnvironmentVariable.Value));
 							}
 							UBTMakefile.TargetNameToUObjectModules = TargetNameToUObjectModules;
+							UBTMakefile.ModuleNameToOutputItems = ModuleNameToOutputItems;
+							UBTMakefile.HotReloadModuleNamesForAllTargets = HotReloadModuleNamesForAllTargets;
 							UBTMakefile.Targets = Targets;
 							UBTMakefile.bUseAdaptiveUnityBuild = Targets.Any(x => x.Rules.bUseAdaptiveUnityBuild);
 							UBTMakefile.SourceFileWorkingSet = Unity.SourceFileWorkingSet;
@@ -1069,22 +1047,73 @@ namespace UnrealBuildTool
 
 							if (BuildConfiguration.bUseUBTMakefiles)
 							{
-								if(UBTMakefile.PrerequisiteActions.Length == 0)
+								// We've been told to prepare to build, so let's go ahead and save out our action graph so that we can use in a later invocation 
+								// to assemble the build.  Even if we are configured to assemble the build in this same invocation, we want to save out the
+								// Makefile so that it can be used on subsequent 'assemble only' runs, for the fastest possible iteration times
+								// @todo ubtmake: Optimization: We could make 'gather + assemble' mode slightly faster by saving this while busy compiling (on our worker thread)
+								UBTMakefile.SaveUBTMakefile(TargetDescs, UBTMakefile);
+							}
+						}
+
+						HashSet<string> OnlyModuleNames = new HashSet<string>(StringComparer.OrdinalIgnoreCase);
+						foreach(string Argument in Arguments)
+						{
+							const string ModuleArgument = "-Module=";
+							if(Argument.StartsWith(ModuleArgument, StringComparison.OrdinalIgnoreCase))
+							{
+								OnlyModuleNames.Add(Argument.Substring(ModuleArgument.Length));
+							}
+						}
+
+						Dictionary<string, int> HotReloadModuleNameToSuffix = new Dictionary<string, int>(StringComparer.OrdinalIgnoreCase);
+						foreach(string Argument in Arguments)
+						{
+							const string ModuleWithSuffixArgument = "-ModuleWithSuffix=";
+							if(Argument.StartsWith(ModuleWithSuffixArgument, StringComparison.OrdinalIgnoreCase))
+							{
+								string Value = Argument.Substring(ModuleWithSuffixArgument.Length);
+
+								int SuffixIdx = Value.LastIndexOf(',');
+								if(SuffixIdx == -1)
 								{
-									// Nothing to build for this target
-									Log.TraceInformation("No buildable actions for this target. Makefile will not be saved.");
+									throw new BuildException("Missing suffix argument from -ModuleWithSuffix=Name,Suffix");
 								}
-								else
+
+								string ModuleName = Value.Substring(0, SuffixIdx);
+
+								int Suffix;
+								if(!Int32.TryParse(Value.Substring(SuffixIdx + 1), out Suffix))
 								{
-									// We've been told to prepare to build, so let's go ahead and save out our action graph so that we can use in a later invocation 
-									// to assemble the build.  Even if we are configured to assemble the build in this same invocation, we want to save out the
-									// Makefile so that it can be used on subsequent 'assemble only' runs, for the fastest possible iteration times
-									// @todo ubtmake: Optimization: We could make 'gather + assemble' mode slightly faster by saving this while busy compiling (on our worker thread)
-									UBTMakefile.SaveUBTMakefile(TargetDescs, HotReload, UBTMakefile);
+									throw new BuildException("Suffix for modules must be an integer");
 								}
+
+								HotReloadModuleNameToSuffix[ModuleName] = Suffix;
 							}
 						}
 
+						EHotReload HotReloadMode = EHotReload.Disabled;
+						if(Arguments.Any(x => x.Equals("-ForceHotReload", StringComparison.InvariantCultureIgnoreCase)))
+						{
+							HotReloadMode = EHotReload.FromIDE;
+						}
+						else if (TargetDescs.Count == 1 && !Arguments.Any(x => x.Equals("-NoHotReload", StringComparison.InvariantCultureIgnoreCase)))
+						{
+							if (BuildConfiguration.bAllowHotReloadFromIDE && ShouldDoHotReloadFromIDE(BuildConfiguration, Arguments, TargetDescs[0]))
+							{
+								HotReloadMode = EHotReload.FromIDE;
+							}
+							else if (HotReloadModuleNameToSuffix.Count > 0 && TargetDescs[0].ForeignPlugin == null)
+							{
+								HotReloadMode = EHotReload.FromEditor;
+							}
+
+							if (HotReloadMode != EHotReload.Disabled && BuildConfiguration.bCleanProject)
+							{
+								throw new BuildException("Unable to clean target while hot-reloading. Close the editor and try again.");
+							}
+						}
+						TargetDescriptor HotReloadTargetDesc = (HotReloadMode != EHotReload.Disabled) ? TargetDescs[0] : null;
+
 						if (bIsAssemblingBuild)
 						{
 							// If we didn't build the graph in this session, then we'll need to load a cached one
@@ -1092,13 +1121,11 @@ namespace UnrealBuildTool
 							{
 								ActionGraph.AllActions = UBTMakefile.AllActions;
 
-								// Patch action history for hot reload when running in assembler mode.  In assembler mode, the suffix on the output file will be
-								// the same for every invocation on that makefile, but we need a new suffix each time.
-								if (HotReload != EHotReload.Disabled)
-								{
-									PatchActionHistoryForHotReloadAssembling(ActionGraph, HotReloadTargetDesc.OnlyModules, Targets);
-								}
+								OutputItemsForAllTargets = UBTMakefile.OutputItemsForAllTargets;
 
+								ModuleNameToOutputItems = UBTMakefile.ModuleNameToOutputItems;
+
+								HotReloadModuleNamesForAllTargets = UBTMakefile.HotReloadModuleNamesForAllTargets;
 
 								foreach (Tuple<string, string> EnvironmentVariable in UBTMakefile.EnvironmentVariables)
 								{
@@ -1132,7 +1159,7 @@ namespace UnrealBuildTool
 												// Execute the header tool
 												FileReference ModuleInfoFileName = FileReference.Combine(Target.ProjectIntermediateDirectory, Target.GetTargetName() + ".uhtmanifest");
 												ECompilationResult UHTResult = ExternalExecution.ExecuteHeaderToolIfNecessary(BuildConfiguration, Target.ProjectFile, Target.TargetName, Target.TargetType, Target.bHasProjectScriptPlugin,  UObjectModules: TargetUObjectModules, ModuleInfoFileName: ModuleInfoFileName, bIsGatheringBuild: bIsGatheringBuild, bIsAssemblingBuild: bIsAssemblingBuild);
-												if(UHTResult.Succeeded())
+												if(!UHTResult.Succeeded())
 												{
 													Log.TraceInformation("UnrealHeaderTool failed for target '" + Target.GetTargetName() + "' (platform: " + Target.Platform.ToString() + ", module info: " + ModuleInfoFileName + ").");
 													BuildResult = UHTResult;
@@ -1146,12 +1173,241 @@ namespace UnrealBuildTool
 
 							if (BuildResult.Succeeded())
 							{
-								// Make sure any old DLL files from in-engine recompiles aren't lying around.  Must be called after the action graph is finalized.
-								ActionGraph.DeleteStaleHotReloadDLLs();
+								// Get the output items we're trying to build right now
+								List<FileItem> OutputItems = OutputItemsForAllTargets;
+
+								// Filter them if we have any specific modules specified
+								if (OnlyModuleNames.Count > 0)
+								{
+									OutputItems = new List<FileItem>();
+									foreach(string OnlyModuleName in OnlyModuleNames)
+									{
+										FileItem[] OutputItemsForModule;
+										if(!ModuleNameToOutputItems.TryGetValue(OnlyModuleName, out OutputItemsForModule))
+										{
+											throw new BuildException("Unable to find output items for module '{0}'", OnlyModuleName);
+										}
+										OutputItems.AddRange(OutputItemsForModule);
+									}
+								}
+
+								// For now simply treat all object files as the root target.
+								Action[] PrerequisiteActions;
+								{
+									HashSet<Action> PrerequisiteActionsSet = new HashSet<Action>();
+									foreach (FileItem OutputItem in OutputItems)
+									{
+										ActionGraph.GatherPrerequisiteActions(OutputItem, ref PrerequisiteActionsSet);
+									}
+									PrerequisiteActions = PrerequisiteActionsSet.ToArray();
+								}
+
+								// Apply the previous hot reload state
+								HotReloadState HotReloadState = null;
+								if(HotReloadMode == EHotReload.Disabled)
+								{
+									// Delete any previous state files; they are no longer valid
+									foreach(TargetDescriptor TargetDesc in TargetDescs)
+									{
+										FileReference HotReloadStateFile = GetHotReloadStateFile(TargetDesc.ProjectFile, TargetDesc.Name, TargetDesc.Platform, TargetDesc.Configuration, TargetDesc.Architecture);
+										if(FileReference.Exists(HotReloadStateFile))
+										{
+											try
+											{
+												HotReloadState = ReadHotReloadState(HotReloadStateFile);
+											}
+											catch(Exception Ex)
+											{
+												Log.TraceWarning("Unable to read hot reload state file: {0}", HotReloadStateFile);
+												Log.WriteException(Ex, null);
+											}
+
+											if(HotReloadState != null)
+											{
+												foreach(FileReference Location in HotReloadState.AllHotReloadFiles.OrderBy(x => x.FullName, StringComparer.OrdinalIgnoreCase))
+												{
+													if(FileReference.Exists(Location))
+													{
+														try
+														{
+															FileReference.Delete(Location);
+														}
+														catch(Exception Ex)
+														{
+															throw new BuildException(Ex, "Unable to delete hot-reload file: {0}", Location);
+														}
+														Log.TraceInformation("Deleted hot-reload file: {0}", Location);
+													}
+												}
+												try
+												{
+													FileReference.Delete(HotReloadStateFile);
+												}
+												catch(Exception Ex)
+												{
+													throw new BuildException(Ex, "Unable to delete hot-reload state file: {0}", HotReloadStateFile);
+												}
+											}
+										}
+									}
+								}
+								else
+								{
+									// Read the previous state file and apply it to the action graph
+									FileReference HotReloadStateFile = GetHotReloadStateFile(TargetDescs[0].ProjectFile, TargetDescs[0].Name, TargetDescs[0].Platform, TargetDescs[0].Configuration, TargetDescs[0].Architecture);
+									if(FileReference.Exists(HotReloadStateFile))
+									{
+										HotReloadState = ReadHotReloadState(HotReloadStateFile);
+									}
+									else
+									{
+										HotReloadState = new HotReloadState();
+									}
+
+									// Update the action graph to produce these new files
+									PatchActionGraphForHotReload(PrerequisiteActions, HotReloadState.OriginalFileToHotReloadFile);
+
+									// Update the module to output file mapping
+									foreach(string HotReloadModuleName in HotReloadModuleNamesForAllTargets)
+									{
+										FileItem[] ModuleOutputItems = ModuleNameToOutputItems[HotReloadModuleName];
+										for(int Idx = 0; Idx < ModuleOutputItems.Length; Idx++)
+										{
+											FileReference NewLocation;
+											if(HotReloadState.OriginalFileToHotReloadFile.TryGetValue(ModuleOutputItems[Idx].Location, out NewLocation))
+											{
+												ModuleOutputItems[Idx] = FileItem.GetItemByFileReference(NewLocation);
+											}
+										}
+									}
+
+									// If we want a specific suffix on any modules, apply that now. We'll track the outputs later, but the suffix has to be forced  (and is always out of date if it doesn't exist).
+									if(HotReloadModuleNameToSuffix.Count > 0)
+									{
+										Dictionary<FileReference, FileReference> OldLocationToNewLocation = new Dictionary<FileReference, FileReference>();
+										foreach(string HotReloadModuleName in HotReloadModuleNamesForAllTargets)
+										{
+											int ModuleSuffix;
+											if(HotReloadModuleNameToSuffix.TryGetValue(HotReloadModuleName, out ModuleSuffix))
+											{
+												FileItem[] ModuleOutputItems = ModuleNameToOutputItems[HotReloadModuleName];
+												foreach(FileItem ModuleOutputItem in ModuleOutputItems)
+												{
+													FileReference OldLocation = ModuleOutputItem.Location;
+													FileReference NewLocation = ReplaceHotReloadSuffix(OldLocation, ModuleSuffix);
+													OldLocationToNewLocation[OldLocation] = NewLocation;
+												}
+											}
+										}
+										PatchActionGraphForHotReload(PrerequisiteActions, OldLocationToNewLocation);
+									}
+								}
 
 								// Plan the actions to execute for the build.
 								Dictionary<UEBuildTarget, List<FileItem>> TargetToOutdatedPrerequisitesMap;
-								List<Action> ActionsToExecute = ActionGraph.GetActionsToExecute(BuildConfiguration, UBTMakefile.PrerequisiteActions, Targets, TargetToHeaders, bIsAssemblingBuild, bNeedsFullCPPIncludeRescan, out TargetToOutdatedPrerequisitesMap);
+								HashSet<Action> ActionsToExecute = ActionGraph.GetActionsToExecute(BuildConfiguration, PrerequisiteActions, Targets, TargetToHeaders, bIsAssemblingBuild, bNeedsFullCPPIncludeRescan, out TargetToOutdatedPrerequisitesMap);
+
+								// Patch action history for hot reload when running in assembler mode.  In assembler mode, the suffix on the output file will be
+								// the same for every invocation on that makefile, but we need a new suffix each time.
+								if (HotReloadMode != EHotReload.Disabled)
+								{
+									// For all the hot-reloadable modules that may need a unique suffix appended, build a mapping from output item to all the output items in that module. We can't 
+									// apply a suffix to one without applying a suffix to all of them.
+									Dictionary<FileItem, FileItem[]> HotReloadItemToDependentItems = new Dictionary<FileItem, FileItem[]>();
+									foreach(string HotReloadModuleName in HotReloadModuleNamesForAllTargets)
+									{
+										int ModuleSuffix;
+										if(!HotReloadModuleNameToSuffix.TryGetValue(HotReloadModuleName, out ModuleSuffix) || ModuleSuffix == -1)
+										{
+											FileItem[] ModuleOutputItems;
+											if(ModuleNameToOutputItems.TryGetValue(HotReloadModuleName, out ModuleOutputItems))
+											{
+												foreach(FileItem ModuleOutputItem in ModuleOutputItems)
+												{
+													HotReloadItemToDependentItems[ModuleOutputItem] = ModuleOutputItems;
+												}
+											}
+										} 
+									}
+
+									// Expand the list of actions to execute to include everything that references any files with a new suffix. Unlike a regular build, we can't ignore
+									// dependencies on import libraries under the assumption that a header would change if the API changes, because the dependency will be on a different DLL.
+									HashSet<FileItem> FilesRequiringSuffix = new HashSet<FileItem>(ActionsToExecute.SelectMany(x => x.ProducedItems).Where(x => HotReloadItemToDependentItems.ContainsKey(x)));
+									for(int LastNumFilesWithNewSuffix = 0; FilesRequiringSuffix.Count > LastNumFilesWithNewSuffix;)
+									{
+										LastNumFilesWithNewSuffix = FilesRequiringSuffix.Count;
+										foreach(Action PrerequisiteAction in PrerequisiteActions)
+										{
+											if(!ActionsToExecute.Contains(PrerequisiteAction))
+											{
+												foreach(FileItem ProducedItem in PrerequisiteAction.ProducedItems)
+												{
+													FileItem[] DependentItems;
+													if(HotReloadItemToDependentItems.TryGetValue(ProducedItem, out DependentItems))
+													{
+														ActionsToExecute.Add(PrerequisiteAction);
+														FilesRequiringSuffix.UnionWith(DependentItems);
+													}
+												}
+											}
+										}
+									}
+
+									// Build a list of file mappings
+									Dictionary<FileReference, FileReference> OldLocationToNewLocation = new Dictionary<FileReference, FileReference>();
+									foreach(FileItem FileRequiringSuffix in FilesRequiringSuffix)
+									{
+										FileReference OldLocation = FileRequiringSuffix.Location;
+										FileReference NewLocation = ReplaceHotReloadSuffix(OldLocation, HotReloadState.NextSuffix);
+										OldLocationToNewLocation[OldLocation] = NewLocation;
+									}
+
+									// Update the action graph with these new paths
+									PatchActionGraphForHotReload(PrerequisiteActions, OldLocationToNewLocation);
+
+									// Get a new list of actions to execute now that the graph has been modified
+									ActionsToExecute = ActionGraph.GetActionsToExecute(BuildConfiguration, PrerequisiteActions, Targets, TargetToHeaders, bIsAssemblingBuild, bNeedsFullCPPIncludeRescan, out TargetToOutdatedPrerequisitesMap);
+
+									// Build a mapping of all file items to their original
+									Dictionary<FileReference, FileReference> HotReloadFileToOriginalFile = new Dictionary<FileReference, FileReference>();
+									foreach(KeyValuePair<FileReference, FileReference> Pair in HotReloadState.OriginalFileToHotReloadFile)
+									{
+										HotReloadFileToOriginalFile[Pair.Value] = Pair.Key;
+									}
+									foreach(KeyValuePair<FileReference, FileReference> Pair in OldLocationToNewLocation)
+									{
+										FileReference OriginalLocation;
+										if(!HotReloadFileToOriginalFile.TryGetValue(Pair.Key, out OriginalLocation))
+										{
+											OriginalLocation = Pair.Key;
+										}
+										HotReloadFileToOriginalFile[Pair.Value] = OriginalLocation;
+									}
+
+									// Now filter out all the hot reload files and update the state
+									foreach(Action Action in ActionsToExecute)
+									{
+										foreach(FileItem ProducedItem in Action.ProducedItems)
+										{
+											FileReference OriginalLocation;
+											if(HotReloadFileToOriginalFile.TryGetValue(ProducedItem.Location, out OriginalLocation))
+											{
+												HotReloadState.OriginalFileToHotReloadFile[OriginalLocation] = ProducedItem.Location;
+												HotReloadState.AllHotReloadFiles.Add(ProducedItem.Location);
+											}
+										}
+									}
+
+									// Increment the suffix for the next iteration
+									if(ActionsToExecute.Count > 0)
+									{
+										HotReloadState.NextSuffix++;
+									}
+
+									// Save the new state
+									FileReference HotReloadStateFile = GetHotReloadStateFile(TargetDescs[0].ProjectFile, TargetDescs[0].Name, TargetDescs[0].Platform, TargetDescs[0].Configuration, TargetDescs[0].Architecture);
+									WriteHotReloadState(HotReloadStateFile, HotReloadState);
+								}
 
 								// Display some stats to the user.
 								Log.TraceVerbose(
@@ -1173,24 +1429,17 @@ namespace UnrealBuildTool
 								DateTime ExecutorStartTime = DateTime.UtcNow;
 								if (BuildConfiguration.bXGEExport)
 								{
-									XGE.ExportActions(ActionsToExecute);
+									XGE.ExportActions(ActionsToExecute.ToList());
 									bSuccess = true;
 								}
 								else
 								{
-									bSuccess = ActionGraph.ExecuteActions(BuildConfiguration, ActionsToExecute, out ExecutorName, TargetInfoForTelemetry, HotReload);
+									bSuccess = ActionGraph.ExecuteActions(BuildConfiguration, ActionsToExecute.ToList(), out ExecutorName, TargetInfoForTelemetry, HotReloadMode);
 								}
 								TotalExecutorTime += (DateTime.UtcNow - ExecutorStartTime).TotalSeconds;
 
 								// if the build succeeded, write the receipts and do any needed syncing
-								if (bSuccess)
-								{
-									if (ActionsToExecute.Count == 0 && BuildConfiguration.bSkipLinkingWhenNothingToCompile)
-									{
-										BuildResult = ECompilationResult.UpToDate;
-									}
-								}
-								else
+								if (!bSuccess)
 								{
 									BuildResult = ECompilationResult.OtherCompilationError;
 								}
@@ -1205,7 +1454,7 @@ namespace UnrealBuildTool
 							{
 								foreach (UEBuildTarget Target in Targets)
 								{
-									if (Target.bDeployAfterCompile && Target.OnlyModules.Count == 0)
+									if (Target.bDeployAfterCompile && HotReloadModuleNameToSuffix.Count == 0)
 									{
 										Log.WriteLine(LogEventType.Console, "Deploying {0} {1} {2}...", Target.TargetName, Target.Platform, Target.Configuration);
 										TargetReceipt Receipt = TargetReceipt.Read(Target.ReceiptFileName);
@@ -1280,6 +1529,7 @@ namespace UnrealBuildTool
 				}
 			}
 		}
+
 		/// <summary>
 		/// Checks if the editor is currently running and this is a hot-reload
 		/// </summary>
@@ -1448,111 +1698,108 @@ namespace UnrealBuildTool
 			}
 		}
 
+		public static FileReference GetHotReloadStateFile(FileReference ProjectFile, string TargetName, UnrealTargetPlatform Platform, UnrealTargetConfiguration Configuration, string Architecture)
+		{
+			DirectoryReference BaseDir = DirectoryReference.FromFile(ProjectFile) ?? UnrealBuildTool.EngineDirectory;
+			return FileReference.Combine(BaseDir, UEBuildTarget.GetPlatformIntermediateFolder(Platform, Architecture), TargetName, Configuration.ToString(), "HotReload.state");
+		}
+
+		public static HotReloadState ReadHotReloadState(FileReference Location)
+		{
+			return BinaryFormatterUtils.Load<HotReloadState>(Location);
+		}
+
+		public static void WriteHotReloadState(FileReference Location, HotReloadState State)
+		{
+			DirectoryReference.CreateDirectory(Location.Directory);
+			BinaryFormatterUtils.Save(Location, State);
+		}
+
 		/// <summary>
 		/// Replaces a hot reload suffix in a filename.
 		/// </summary>
-		/// <param name="InOutFilename">The filename to replace the suffix in.</param>
-		/// <param name="ModuleNameToSuffix">A function which returns a replacement suffix for a given module name.</param>
-		/// <returns>true is a suffix replacement was made, false otherwise.</returns>
-		public static bool ReplaceHotReloadFilenameSuffix(ref string InOutFilename, Func<string, string> ModuleNameToSuffix)
+		public static FileReference ReplaceHotReloadSuffix(FileReference File, int Suffix)
 		{
-			// Remove the extension
-			string FilenameWithoutExtension = Path.GetFileNameWithoutExtension(InOutFilename);
-			string Extension = Path.GetExtension(InOutFilename);
-
-			// Remove the numbered suffix from the end of the filename
-			int IndexOfLastHyphen = FilenameWithoutExtension.LastIndexOf('-');
-			string OriginalNumberSuffix = FilenameWithoutExtension.Substring(IndexOfLastHyphen);
-			int NumberSuffix;
-			bool bHasNumberSuffix = int.TryParse(OriginalNumberSuffix, out NumberSuffix);
-
-			string OriginalFileNameWithoutNumberSuffix = null;
-			string PlatformConfigSuffix = "";
-			if (bHasNumberSuffix)
-			{
-				// Remove "-####" suffix in Development configuration
-				OriginalFileNameWithoutNumberSuffix = FilenameWithoutExtension.Substring(0, IndexOfLastHyphen);
-			}
-			else
-			{
-				OriginalFileNameWithoutNumberSuffix = FilenameWithoutExtension;
+			string FileName = File.GetFileName();
 
-				// Remove "-####-Platform-Configuration" suffix in Debug configuration
-				int IndexOfSecondLastHyphen = FilenameWithoutExtension.LastIndexOf('-', IndexOfLastHyphen - 1, IndexOfLastHyphen);
-				if (IndexOfSecondLastHyphen > 0)
-				{
-					int IndexOfThirdLastHyphen = FilenameWithoutExtension.LastIndexOf('-', IndexOfSecondLastHyphen - 1, IndexOfSecondLastHyphen);
-					if (IndexOfThirdLastHyphen > 0)
-					{
-						if (int.TryParse(FilenameWithoutExtension.Substring(IndexOfThirdLastHyphen + 1, IndexOfSecondLastHyphen - IndexOfThirdLastHyphen - 1), out NumberSuffix))
-						{
-							OriginalFileNameWithoutNumberSuffix = FilenameWithoutExtension.Substring(0, IndexOfThirdLastHyphen);
-							PlatformConfigSuffix = FilenameWithoutExtension.Substring(IndexOfSecondLastHyphen);
-							bHasNumberSuffix = true;
-						}
-					}
-				}
+			// Find the end of the target and module name
+			int HyphenIdx = FileName.IndexOf('-');
+			if (HyphenIdx == -1)
+			{
+				throw new BuildException("Hot-reloadable files are expected to contain a hyphen, eg. UE4Editor-Core");
 			}
 
-			if (!bHasNumberSuffix)
+			int NameEndIdx = HyphenIdx + 1;
+			while(NameEndIdx < FileName.Length && FileName[NameEndIdx] != '.' && FileName[NameEndIdx] != '-')
 			{
-				return false;
+				NameEndIdx++;
 			}
 
-			// Figure out which suffix to use
-			int FirstHyphenIndex = OriginalFileNameWithoutNumberSuffix.IndexOf('-');
-			if (FirstHyphenIndex == -1)
+			// Strip any existing suffix
+			if(NameEndIdx + 1 < FileName.Length && Char.IsDigit(FileName[NameEndIdx + 1]))
 			{
-				throw new BuildException("Expected produced item file name '{0}' to start with a prefix (such as 'UE4Editor-') when hot reloading", FilenameWithoutExtension);
+				int SuffixEndIdx = NameEndIdx + 2;
+				while(SuffixEndIdx < FileName.Length && Char.IsDigit(FileName[SuffixEndIdx]))
+				{
+					SuffixEndIdx++;
+				}
+				if(SuffixEndIdx == FileName.Length || FileName[SuffixEndIdx] == '-' || FileName[SuffixEndIdx] == '.')
+				{
+					FileName = FileName.Substring(0, NameEndIdx) + FileName.Substring(SuffixEndIdx);
+				}
 			}
-			string OutputFileNamePrefix = OriginalFileNameWithoutNumberSuffix.Substring(0, FirstHyphenIndex + 1);
 
-			// Get the module name from the file name
-			string ModuleName = OriginalFileNameWithoutNumberSuffix.Substring(OutputFileNamePrefix.Length);
+			string NewFileName = String.Format("{0}-{1:D4}{2}", FileName.Substring(0, NameEndIdx), Suffix, FileName.Substring(NameEndIdx));
 
-			// Add a new suffix
-			string UniqueSuffix = "-" + ModuleNameToSuffix(ModuleName);
-
-			InOutFilename = OriginalFileNameWithoutNumberSuffix + UniqueSuffix + PlatformConfigSuffix + Extension;
-			return true;
+			return FileReference.Combine(File.Directory, NewFileName);
 		}
 
 		/// <summary>
-		/// Returns a module suffix from a list of OnlyModules if one exists, otherwise generates a new one and adds it to the list.
+		/// Replaces a base filename within a string. Ensures that the filename is not a substring of another longer string (eg. replacing "Foo" will match "Foo.Bar" but not "FooBar" or "BarFooBar"). 
 		/// </summary>
-		/// <param name="OnlyModules">A list of modules with suffixes to search and add to if necessary.</param>
-		/// <param name="ModuleName">The module name to find in the list.</param>
-		/// <returns>The existing or new suffix for the module name.</returns>
-		public static string GetReplacementModuleSuffix(List<OnlyModule> OnlyModules, string ModuleName)
+		/// <param name="Text">Text to replace within</param>
+		/// <param name="OldFileName">Old filename</param>
+		/// <param name="NewFileName">New filename</param>
+		/// <returns>Text with file names replaced</returns>
+		static string ReplaceBaseFileName(string Text, string OldFileName, string NewFileName)
 		{
-			// Check the OnlyModules list for a prefix
-			if (OnlyModules != null)
+			int StartIdx = 0;
+			for(;;)
 			{
-				OnlyModule FoundModule = OnlyModules.Find((OnlyModule) => OnlyModule.OnlyModuleName.Equals(ModuleName, StringComparison.InvariantCultureIgnoreCase));
-				if (FoundModule != null)
+				int Idx = Text.IndexOf(OldFileName, StartIdx, StringComparison.OrdinalIgnoreCase);
+				if(Idx == -1)
+				{
+					break;
+				}
+				else if((Idx == 0 || !IsBaseFileNameCharacter(Text[Idx - 1])) && (Idx + OldFileName.Length == Text.Length || !IsBaseFileNameCharacter(Text[Idx + OldFileName.Length])))
 				{
-					return FoundModule.OnlyModuleSuffix;
+					Text = Text.Substring(0, Idx) + NewFileName + Text.Substring(Idx + OldFileName.Length);
+					StartIdx = Idx + NewFileName.Length;
+				}
+				else
+				{
+					StartIdx = Idx + 1;
 				}
 			}
+			return Text;
+		}
 
-			// Generate a new random suffix
-			string Result = (new Random((int)(DateTime.Now.Ticks % Int32.MaxValue)).Next(10000)).ToString();
-
-			// Add it to the OnlyModules list
-			OnlyModules.Add(new OnlyModule(ModuleName, Result));
-
-			return Result;
+		/// <summary>
+		/// Determines if a character should be treated as part of a base filename, when updating strings for hot reload
+		/// </summary>
+		/// <param name="Character">The character to check</param>
+		/// <returns>True if the character is part of a base filename, false otherwise</returns>
+		static bool IsBaseFileNameCharacter(char Character)
+		{
+			return Char.IsLetterOrDigit(Character) || Character == '_';
 		}
 
 		/// <summary>
 		/// Patch action history for hot reload when running in assembler mode.  In assembler mode, the suffix on the output file will be
 		/// the same for every invocation on that makefile, but we need a new suffix each time.
 		/// </summary>
-		static void PatchActionHistoryForHotReloadAssembling(ActionGraph ActionGraph, List<OnlyModule> OnlyModules, List<UEBuildTarget> Targets)
+		static void PatchActionGraphForHotReload(IEnumerable<Action> Actions, Dictionary<FileReference, FileReference> OriginalFileToHotReloadFile)
 		{
-			// Take a copy of the array so that we don't modify the original in GetReplacementModuleSuffix
-			List<OnlyModule> OnlyModulesLocal = new List<OnlyModule>(OnlyModules);
-
 			// Gather all of the response files for link actions.  We're going to need to patch 'em up after we figure out new
 			// names for all of the output files and import libraries
 			List<string> ResponseFilePaths = new List<string>();
@@ -1562,22 +1809,23 @@ namespace UnrealBuildTool
 			List<string> LinkScriptFilePaths = new List<string>();
 
 			// Keep a map of the original file names and their new file names, so we can fix up response files after
-			List<Tuple<string, string>> OriginalFileNameAndNewFileNameList_NoExtensions = new List<Tuple<string, string>>();
+			Dictionary<string, string> OriginalFileNameAndNewFileNameList_NoExtensions = new Dictionary<string, string>();
 
 			// Finally, we'll keep track of any file items that we had to create counterparts for change file names, so we can fix those up too
 			Dictionary<FileItem, FileItem> AffectedOriginalFileItemAndNewFileItemMap = new Dictionary<FileItem, FileItem>();
 
-			foreach (Action Action in ActionGraph.AllActions.Where((Action) => Action.ActionType == ActionType.Link))
+			foreach (Action Action in Actions.Where((Action) => Action.ActionType == ActionType.Link))
 			{
 				// Assume that the first produced item (with no extension) is our output file name
-				string OriginalFileNameWithoutExtension = Utils.GetFilenameWithoutAnyExtensions(Action.ProducedItems[0].AbsolutePath);
-
-				string NewFileNameWithoutExtension = OriginalFileNameWithoutExtension;
-				if (!ReplaceHotReloadFilenameSuffix(ref NewFileNameWithoutExtension, (ModuleName) => GetReplacementModuleSuffix(OnlyModulesLocal, ModuleName)))
+				FileReference HotReloadFile;
+				if(!OriginalFileToHotReloadFile.TryGetValue(Action.ProducedItems[0].Location, out HotReloadFile))
 				{
 					continue;
 				}
 
+				string OriginalFileNameWithoutExtension = Utils.GetFilenameWithoutAnyExtensions(Action.ProducedItems[0].AbsolutePath);
+				string NewFileNameWithoutExtension = Utils.GetFilenameWithoutAnyExtensions(HotReloadFile.FullName);
+
 				// Find the response file in the command line.  We'll need to make a copy of it with our new file name.
 				string ResponseFileExtension = ".response";
 				int ResponseExtensionIndex = Action.CommandArguments.IndexOf(ResponseFileExtension, StringComparison.InvariantCultureIgnoreCase);
@@ -1592,10 +1840,13 @@ namespace UnrealBuildTool
 					string OriginalResponseFilePathWithoutExtension = Action.CommandArguments.Substring(ResponseFilePathIndex + 2, (ResponseExtensionIndex - ResponseFilePathIndex) - 2);
 					string OriginalResponseFilePath = OriginalResponseFilePathWithoutExtension + ResponseFileExtension;
 
-					string NewResponseFilePath = OriginalResponseFilePath.Replace(OriginalFileNameWithoutExtension, NewFileNameWithoutExtension);
+					string NewResponseFilePath = ReplaceBaseFileName(OriginalResponseFilePath, OriginalFileNameWithoutExtension, NewFileNameWithoutExtension);
 
 					// Copy the old response file to the new path
-					File.Copy(OriginalResponseFilePath, NewResponseFilePath, overwrite: true);
+					if(String.Compare(OriginalResponseFilePath, NewResponseFilePath, StringComparison.OrdinalIgnoreCase) != 0)
+					{
+						File.Copy(OriginalResponseFilePath, NewResponseFilePath, overwrite: true);
+					}
 
 					// Keep track of the new response file name.  We'll have to do some edits afterwards.
 					ResponseFilePaths.Add(NewResponseFilePath);
@@ -1619,7 +1870,7 @@ namespace UnrealBuildTool
 						string OriginalLinkScriptFilePathWithoutExtension = Action.CommandArguments.Substring(LinkScriptFilePathIndex + 1, (LinkScriptExtensionIndex - LinkScriptFilePathIndex) - 1);
 						string OriginalLinkScriptFilePath = OriginalLinkScriptFilePathWithoutExtension + LinkScriptFileExtension;
 
-						string NewLinkScriptFilePath = OriginalLinkScriptFilePath.Replace(OriginalFileNameWithoutExtension, NewFileNameWithoutExtension);
+						string NewLinkScriptFilePath = ReplaceBaseFileName(OriginalLinkScriptFilePath, OriginalFileNameWithoutExtension, NewFileNameWithoutExtension);
 
 						// Copy the old response file to the new path
 						File.Copy(OriginalLinkScriptFilePath, NewLinkScriptFilePath, overwrite: true);
@@ -1632,7 +1883,7 @@ namespace UnrealBuildTool
 					for (int ItemIndex = 0; ItemIndex < Action.PrerequisiteItems.Count; ++ItemIndex)
 					{
 						FileItem OriginalPrerequisiteItem = Action.PrerequisiteItems[ItemIndex];
-						string NewPrerequisiteItemFilePath = OriginalPrerequisiteItem.AbsolutePath.Replace(OriginalFileNameWithoutExtension, NewFileNameWithoutExtension);
+						string NewPrerequisiteItemFilePath = ReplaceBaseFileName(OriginalPrerequisiteItem.AbsolutePath, OriginalFileNameWithoutExtension, NewFileNameWithoutExtension);
 
 						if (OriginalPrerequisiteItem.AbsolutePath != NewPrerequisiteItemFilePath)
 						{
@@ -1653,7 +1904,7 @@ namespace UnrealBuildTool
 								string OriginalResponseFilePathWithoutExtension = OriginalPrerequisiteItem.AbsolutePath.Substring(0, ResponseExtensionIndex);
 								string OriginalResponseFilePath = OriginalResponseFilePathWithoutExtension + ResponseFileExtension;
 
-								string NewResponseFilePath = OriginalResponseFilePath.Replace(OriginalFileNameWithoutExtension, NewFileNameWithoutExtension);
+								string NewResponseFilePath = ReplaceBaseFileName(OriginalResponseFilePath, OriginalFileNameWithoutExtension, NewFileNameWithoutExtension);
 
 								// Copy the old response file to the new path
 								File.Copy(OriginalResponseFilePath, NewResponseFilePath, overwrite: true);
@@ -1666,7 +1917,7 @@ namespace UnrealBuildTool
 				}
 
 				// Go ahead and replace all occurrences of our file name in the command-line (ignoring extensions)
-				Action.CommandArguments = Action.CommandArguments.Replace(OriginalFileNameWithoutExtension, NewFileNameWithoutExtension);
+				Action.CommandArguments = ReplaceBaseFileName(Action.CommandArguments, OriginalFileNameWithoutExtension, NewFileNameWithoutExtension);
 
 
 				// Update this action's list of produced items too
@@ -1674,7 +1925,7 @@ namespace UnrealBuildTool
 				{
 					FileItem OriginalProducedItem = Action.ProducedItems[ItemIndex];
 
-					string NewProducedItemFilePath = OriginalProducedItem.AbsolutePath.Replace(OriginalFileNameWithoutExtension, NewFileNameWithoutExtension);
+					string NewProducedItemFilePath = ReplaceBaseFileName(OriginalProducedItem.AbsolutePath, OriginalFileNameWithoutExtension, NewFileNameWithoutExtension);
 					if (OriginalProducedItem.AbsolutePath != NewProducedItemFilePath)
 					{
 						// OK, the produced item's file name changed so we'll update it to point to our new file
@@ -1691,16 +1942,23 @@ namespace UnrealBuildTool
 				}
 
 				// The status description of the item has the file name, so we'll update it too
-				Action.StatusDescription = Action.StatusDescription.Replace(OriginalFileNameWithoutExtension, NewFileNameWithoutExtension);
+				Action.StatusDescription = ReplaceBaseFileName(Action.StatusDescription, OriginalFileNameWithoutExtension, NewFileNameWithoutExtension);
 
 
 				// Keep track of the file names, so we can fix up response files afterwards.
-				OriginalFileNameAndNewFileNameList_NoExtensions.Add(Tuple.Create(OriginalFileNameWithoutExtension, NewFileNameWithoutExtension));
+				if(!OriginalFileNameAndNewFileNameList_NoExtensions.ContainsKey(OriginalFileNameWithoutExtension))
+				{
+					OriginalFileNameAndNewFileNameList_NoExtensions[OriginalFileNameWithoutExtension] = NewFileNameWithoutExtension;
+				}
+				else if(OriginalFileNameAndNewFileNameList_NoExtensions[OriginalFileNameWithoutExtension] != NewFileNameWithoutExtension)
+				{
+					throw new BuildException("Unexpected conflict in renaming files; {0} maps to {1} and {2}", OriginalFileNameWithoutExtension, OriginalFileNameAndNewFileNameList_NoExtensions[OriginalFileNameWithoutExtension], NewFileNameWithoutExtension);
+				}
 			}
 
 
 			// Do another pass and update any actions that depended on the original file names that we changed
-			foreach (Action Action in ActionGraph.AllActions)
+			foreach (Action Action in Actions)
 			{
 				for (int ItemIndex = 0; ItemIndex < Action.PrerequisiteItems.Count; ++ItemIndex)
 				{
@@ -1721,20 +1979,19 @@ namespace UnrealBuildTool
 				foreach (string ResponseFilePath in ResponseFilePaths)
 				{
 					// Load the file up
-					StringBuilder FileContents = new StringBuilder(Utils.ReadAllText(ResponseFilePath));
+					string FileContents = Utils.ReadAllText(ResponseFilePath);
 
 					// Replace all of the old file names with new ones
-					foreach (Tuple<string, string> FileNameTuple in OriginalFileNameAndNewFileNameList_NoExtensions)
+					foreach (KeyValuePair<string, string> FileNameTuple in OriginalFileNameAndNewFileNameList_NoExtensions)
 					{
-						string OriginalFileNameWithoutExtension = FileNameTuple.Item1;
-						string NewFileNameWithoutExtension = FileNameTuple.Item2;
+						string OriginalFileNameWithoutExtension = FileNameTuple.Key;
+						string NewFileNameWithoutExtension = FileNameTuple.Value;
 
-						FileContents.Replace(OriginalFileNameWithoutExtension, NewFileNameWithoutExtension);
+						FileContents = ReplaceBaseFileName(FileContents, OriginalFileNameWithoutExtension, NewFileNameWithoutExtension);
 					}
 
 					// Overwrite the original file
-					string FileContentsString = FileContents.ToString();
-					File.WriteAllText(ResponseFilePath, FileContentsString, new System.Text.UTF8Encoding(false));
+					File.WriteAllText(ResponseFilePath, FileContents, new System.Text.UTF8Encoding(false));
 				}
 
 				if (UEBuildPlatform.IsPlatformInGroup(BuildHostPlatform.Current.Platform, UnrealPlatformGroup.Unix))
@@ -1742,26 +1999,25 @@ namespace UnrealBuildTool
 					foreach (string LinkScriptFilePath in LinkScriptFilePaths)
 					{
 						// Load the file up
-						StringBuilder FileContents = new StringBuilder(Utils.ReadAllText(LinkScriptFilePath));
+						string FileContents = Utils.ReadAllText(LinkScriptFilePath);
 
 						// Replace all of the old file names with new ones
-						foreach (Tuple<string, string> FileNameTuple in OriginalFileNameAndNewFileNameList_NoExtensions)
+						foreach (KeyValuePair<string, string> FileNameTuple in OriginalFileNameAndNewFileNameList_NoExtensions)
 						{
-							string OriginalFileNameWithoutExtension = FileNameTuple.Item1;
-							string NewFileNameWithoutExtension = FileNameTuple.Item2;
+							string OriginalFileNameWithoutExtension = FileNameTuple.Key;
+							string NewFileNameWithoutExtension = FileNameTuple.Value;
 
-							FileContents.Replace(OriginalFileNameWithoutExtension, NewFileNameWithoutExtension);
+							FileContents = ReplaceBaseFileName(FileContents, OriginalFileNameWithoutExtension, NewFileNameWithoutExtension);
 						}
 
 						// Overwrite the original file
-						string FileContentsString = FileContents.ToString();
-						File.WriteAllText(LinkScriptFilePath, FileContentsString, new System.Text.UTF8Encoding(false));
+						File.WriteAllText(LinkScriptFilePath, FileContents, new System.Text.UTF8Encoding(false));
 					}
 				}
 			}
 
 			// Update the action that writes out the module manifests
-			foreach(Action Action in ActionGraph.AllActions)
+			foreach(Action Action in Actions)
 			{
 				if(Action.ActionType == ActionType.WriteMetadata)
 				{
@@ -1801,10 +2057,12 @@ namespace UnrealBuildTool
 					{
 						foreach (KeyValuePair<string, string> Manifest in FileNameToVersionManifest.Value.ModuleNameToFileName)
 						{
-							string ModuleFilename = Manifest.Value;
-							if (UnrealBuildTool.ReplaceHotReloadFilenameSuffix(ref ModuleFilename, (ModuleName) => UnrealBuildTool.GetReplacementModuleSuffix(OnlyModules, ModuleName)))
+							FileReference OriginalFile = FileReference.Combine(FileNameToVersionManifest.Key.Directory, Manifest.Key);
+
+							FileReference HotReloadFile;
+							if(OriginalFileToHotReloadFile.TryGetValue(OriginalFile, out HotReloadFile))
 							{
-								FileNameToVersionManifest.Value.ModuleNameToFileName[Manifest.Key] = ModuleFilename;
+								FileNameToVersionManifest.Value.ModuleNameToFileName[Manifest.Key] = HotReloadFile.GetFileName();
 								break;
 							}
 						}
