UBT: Move core hot reload functionality into its own file.

diff --git a/Engine/Source/Programs/UnrealBuildTool/System/ActionGraph.cs b/Engine/Source/Programs/UnrealBuildTool/System/ActionGraph.cs
--- a/Engine/Source/Programs/UnrealBuildTool/System/ActionGraph.cs
+++ b/Engine/Source/Programs/UnrealBuildTool/System/ActionGraph.cs
@@ -383,7 +383,7 @@ namespace UnrealBuildTool
 		/// <summary>
 		/// Executes a list of actions.
 		/// </summary>
-		public static bool ExecuteActions(BuildConfiguration BuildConfiguration, List<Action> ActionsToExecute, out string ExecutorName, string TargetInfoForTelemetry, EHotReload HotReload)
+		public static bool ExecuteActions(BuildConfiguration BuildConfiguration, List<Action> ActionsToExecute, out string ExecutorName, string TargetInfoForTelemetry, HotReloadMode HotReload)
 		{
 			bool Result = true;
 			ExecutorName = "";
@@ -427,7 +427,7 @@ namespace UnrealBuildTool
 							{
 								bool bExists = File.Exists(Item.AbsolutePath) || Directory.Exists(Item.AbsolutePath);
 
-								if (HotReload != EHotReload.Disabled)
+								if (HotReload != HotReloadMode.Disabled)
 								{
 									string FailedFilename = Path.Combine(Path.GetDirectoryName(Item.AbsolutePath), "failed.hotreload");
 									if (!bExists)
diff --git a/Engine/Source/Programs/UnrealBuildTool/System/HotReload.cs b/Engine/Source/Programs/UnrealBuildTool/System/HotReload.cs
new file mode 100644
--- /dev/null
+++ b/Engine/Source/Programs/UnrealBuildTool/System/HotReload.cs
@@ -0,0 +1,611 @@
+using System;
+using System.Collections.Generic;
+using System.IO;
+using System.Linq;
+using System.Text;
+using System.Threading.Tasks;
+using Tools.DotNETCommon;
+
+namespace UnrealBuildTool
+{
+	/// <summary>
+	/// The current hot reload mode
+	/// </summary>
+	enum HotReloadMode
+	{
+		Disabled,
+		FromIDE,
+		FromEditor
+	}
+
+	/// <summary>
+	/// Stores the current hot reload state, tracking temporary files created by previous invocations.
+	/// </summary>
+	[Serializable]
+	class HotReloadState
+	{
+		/// <summary>
+		/// Suffix to use for the next hot reload invocation
+		/// </summary>
+		public int NextSuffix = 1;
+
+		/// <summary>
+		/// Map from original filename in the action graph to hot reload file
+		/// </summary>
+		public Dictionary<FileReference, FileReference> OriginalFileToHotReloadFile = new Dictionary<FileReference, FileReference>();
+
+		/// <summary>
+		/// Set of all temporary files created for hot reload
+		/// </summary>
+		public HashSet<FileReference> TemporaryFiles = new HashSet<FileReference>();
+
+		/// <summary>
+		/// Adds all the actions into the hot reload state, so we can restore the action graph on next iteration
+		/// </summary>
+		/// <param name="ActionsToExecute">The actions being executed</param>
+		/// <param name="OldLocationToNewLocation">Mapping from file from their original location (either a previously hot-reloaded file, or an originally compiled file)</param>
+		public void CaptureActions(IEnumerable<Action> ActionsToExecute, Dictionary<FileReference, FileReference> OldLocationToNewLocation)
+		{
+			// Build a mapping of all file items to their original location
+			Dictionary<FileReference, FileReference> HotReloadFileToOriginalFile = new Dictionary<FileReference, FileReference>();
+			foreach(KeyValuePair<FileReference, FileReference> Pair in OriginalFileToHotReloadFile)
+			{
+				HotReloadFileToOriginalFile[Pair.Value] = Pair.Key;
+			}
+			foreach(KeyValuePair<FileReference, FileReference> Pair in OldLocationToNewLocation)
+			{
+				FileReference OriginalLocation;
+				if(!HotReloadFileToOriginalFile.TryGetValue(Pair.Key, out OriginalLocation))
+				{
+					OriginalLocation = Pair.Key;
+				}
+				HotReloadFileToOriginalFile[Pair.Value] = OriginalLocation;
+			}
+
+			// Now filter out all the hot reload files and update the state
+			foreach(Action Action in ActionsToExecute)
+			{
+				foreach(FileItem ProducedItem in Action.ProducedItems)
+				{
+					FileReference OriginalLocation;
+					if(HotReloadFileToOriginalFile.TryGetValue(ProducedItem.Location, out OriginalLocation))
+					{
+						OriginalFileToHotReloadFile[OriginalLocation] = ProducedItem.Location;
+						TemporaryFiles.Add(ProducedItem.Location);
+					}
+				}
+			}
+		}
+
+		/// <summary>
+		/// Gets the location of the hot-reload state file for a particular target
+		/// </summary>
+		/// <param name="ProjectFile">Project containing the target</param>
+		/// <param name="TargetName">Name of the target</param>
+		/// <param name="Platform">Platform being built</param>
+		/// <param name="Configuration">Configuration being built</param>
+		/// <param name="Architecture">Architecture being built</param>
+		/// <returns>Location of the hot reload state file</returns>
+		public static FileReference GetLocation(FileReference ProjectFile, string TargetName, UnrealTargetPlatform Platform, UnrealTargetConfiguration Configuration, string Architecture)
+		{
+			DirectoryReference BaseDir = DirectoryReference.FromFile(ProjectFile) ?? UnrealBuildTool.EngineDirectory;
+			return FileReference.Combine(BaseDir, UEBuildTarget.GetPlatformIntermediateFolder(Platform, Architecture), TargetName, Configuration.ToString(), "HotReload.state");
+		}
+
+		/// <summary>
+		/// Read the hot reload state from the given location
+		/// </summary>
+		/// <param name="Location">Location to read from</param>
+		/// <returns>New hot reload state instance</returns>
+		public static HotReloadState Load(FileReference Location)
+		{
+			return BinaryFormatterUtils.Load<HotReloadState>(Location);
+		}
+
+		/// <summary>
+		/// Writes the state to disk
+		/// </summary>
+		/// <param name="Location">Location to write to</param>
+		public void Save(FileReference Location)
+		{
+			DirectoryReference.CreateDirectory(Location.Directory);
+			BinaryFormatterUtils.Save(Location, this);
+		}
+	}
+
+	static class HotReload
+	{
+		/// <summary>
+		/// Checks if the editor is currently running and this is a hot-reload
+		/// </summary>
+		public static bool ShouldDoHotReloadFromIDE(BuildConfiguration BuildConfiguration, string[] Arguments, TargetDescriptor TargetDesc)
+		{
+			// Check if Hot-reload is disabled globally for this project
+			ConfigHierarchy Hierarchy = ConfigCache.ReadHierarchy(ConfigHierarchyType.Engine, DirectoryReference.FromFile(TargetDesc.ProjectFile), TargetDesc.Platform);
+			bool bAllowHotReloadFromIDE;
+			if (Hierarchy.TryGetValue("BuildConfiguration", "bAllowHotReloadFromIDE", out bAllowHotReloadFromIDE) && !bAllowHotReloadFromIDE)
+			{
+				return false;
+			}
+
+			if (Arguments.Any(x => x.Equals("-NoHotReloadFromIDE", StringComparison.InvariantCultureIgnoreCase)))
+			{
+				// Hot reload disabled through command line, possibly running from UAT
+				return false;
+			}
+
+			bool bIsRunning = false;
+
+			// @todo ubtmake: Kind of cheating here to figure out if an editor target.  At this point we don't have access to the actual target description, and
+			// this code must be able to execute before we create or load module rules DLLs so that hot reload can work with bUseUBTMakefiles
+			bool bIsEditorTarget = TargetDesc.Name.EndsWith("Editor", StringComparison.InvariantCultureIgnoreCase);
+
+			if (!BuildConfiguration.bGenerateManifest && bIsEditorTarget)
+			{
+				string EditorBaseFileName = "UE4Editor";
+				if (TargetDesc.Configuration != UnrealTargetConfiguration.Development && TargetDesc.Configuration != UnrealTargetConfiguration.DebugGame)
+				{
+					EditorBaseFileName = String.Format("{0}-{1}-{2}", EditorBaseFileName, TargetDesc.Platform, TargetDesc.Configuration);
+				}
+
+				FileReference EditorLocation;
+				if (TargetDesc.Platform == UnrealTargetPlatform.Win64)
+				{
+					EditorLocation = FileReference.Combine(UnrealBuildTool.EngineDirectory, "Binaries", "Win64", String.Format("{0}.exe", EditorBaseFileName));
+				}
+				else if (TargetDesc.Platform == UnrealTargetPlatform.Mac)
+				{
+					EditorLocation = FileReference.Combine(UnrealBuildTool.EngineDirectory, "Binaries", "Mac", String.Format("{0}.app/Contents/MacOS/{0}", EditorBaseFileName));
+				}
+				else if (TargetDesc.Platform == UnrealTargetPlatform.Linux)
+				{
+					EditorLocation = FileReference.Combine(UnrealBuildTool.EngineDirectory, "Binaries", "Linux", EditorBaseFileName);
+				}
+				else
+				{
+					throw new BuildException("Unknown editor filename for this platform");
+				}
+
+				BuildHostPlatform.ProcessInfo[] Processes = BuildHostPlatform.Current.GetProcesses();
+				string EditorRunsDir = Path.Combine(UnrealBuildTool.EngineDirectory.FullName, "Intermediate", "EditorRuns");
+
+				if (!Directory.Exists(EditorRunsDir))
+				{
+					return false;
+				}
+
+				FileInfo[] EditorRunsFiles = new DirectoryInfo(EditorRunsDir).GetFiles();
+
+				foreach (FileInfo File in EditorRunsFiles)
+				{
+					int PID;
+					BuildHostPlatform.ProcessInfo Proc = null;
+					if (!Int32.TryParse(File.Name, out PID) || (Proc = Processes.FirstOrDefault(P => P.PID == PID)) == default(BuildHostPlatform.ProcessInfo))
+					{
+						// Delete stale files (it may happen if editor crashes).
+						File.Delete();
+						continue;
+					}
+
+					// Don't break here to allow clean-up of other stale files.
+					if (!bIsRunning)
+					{
+						// Otherwise check if the path matches.
+						bIsRunning = new FileReference(Proc.Filename) == EditorLocation;
+					}
+				}
+			}
+			return bIsRunning;
+		}
+
+		/// <summary>
+		/// Delete all temporary files created by previous hot reload invocations
+		/// </summary>
+		/// <param name="HotReloadStateFile">Location of the state file</param>
+		public static void DeleteTemporaryFiles(FileReference HotReloadStateFile)
+		{
+			if(FileReference.Exists(HotReloadStateFile))
+			{
+				// Try to load the state file. If it fails, we'll just warn and continue.
+				HotReloadState State = null;
+				try
+				{
+					State = HotReloadState.Load(HotReloadStateFile);
+				}
+				catch(Exception Ex)
+				{
+					Log.TraceWarning("Unable to read hot reload state file: {0}", HotReloadStateFile);
+					Log.WriteException(Ex, null);
+					return;
+				}
+
+				// Delete all the output files
+				foreach(FileReference Location in State.TemporaryFiles.OrderBy(x => x.FullName, StringComparer.OrdinalIgnoreCase))
+				{
+					if(FileReference.Exists(Location))
+					{
+						try
+						{
+							FileReference.Delete(Location);
+						}
+						catch(Exception Ex)
+						{
+							throw new BuildException(Ex, "Unable to delete hot-reload file: {0}", Location);
+						}
+						Log.TraceInformation("Deleted hot-reload file: {0}", Location);
+					}
+				}
+
+				// Delete the state file itself
+				try
+				{
+					FileReference.Delete(HotReloadStateFile);
+				}
+				catch(Exception Ex)
+				{
+					throw new BuildException(Ex, "Unable to delete hot-reload state file: {0}", HotReloadStateFile);
+				}
+			}
+		}
+
+		/// <summary>
+		/// Replaces a hot reload suffix in a filename.
+		/// </summary>
+		public static FileReference ReplaceSuffix(FileReference File, int Suffix)
+		{
+			string FileName = File.GetFileName();
+
+			// Find the end of the target and module name
+			int HyphenIdx = FileName.IndexOf('-');
+			if (HyphenIdx == -1)
+			{
+				throw new BuildException("Hot-reloadable files are expected to contain a hyphen, eg. UE4Editor-Core");
+			}
+
+			int NameEndIdx = HyphenIdx + 1;
+			while(NameEndIdx < FileName.Length && FileName[NameEndIdx] != '.' && FileName[NameEndIdx] != '-')
+			{
+				NameEndIdx++;
+			}
+
+			// Strip any existing suffix
+			if(NameEndIdx + 1 < FileName.Length && Char.IsDigit(FileName[NameEndIdx + 1]))
+			{
+				int SuffixEndIdx = NameEndIdx + 2;
+				while(SuffixEndIdx < FileName.Length && Char.IsDigit(FileName[SuffixEndIdx]))
+				{
+					SuffixEndIdx++;
+				}
+				if(SuffixEndIdx == FileName.Length || FileName[SuffixEndIdx] == '-' || FileName[SuffixEndIdx] == '.')
+				{
+					FileName = FileName.Substring(0, NameEndIdx) + FileName.Substring(SuffixEndIdx);
+				}
+			}
+
+			string NewFileName = String.Format("{0}-{1:D4}{2}", FileName.Substring(0, NameEndIdx), Suffix, FileName.Substring(NameEndIdx));
+
+			return FileReference.Combine(File.Directory, NewFileName);
+		}
+
+		/// <summary>
+		/// Replaces a base filename within a string. Ensures that the filename is not a substring of another longer string (eg. replacing "Foo" will match "Foo.Bar" but not "FooBar" or "BarFooBar"). 
+		/// </summary>
+		/// <param name="Text">Text to replace within</param>
+		/// <param name="OldFileName">Old filename</param>
+		/// <param name="NewFileName">New filename</param>
+		/// <returns>Text with file names replaced</returns>
+		static string ReplaceBaseFileName(string Text, string OldFileName, string NewFileName)
+		{
+			int StartIdx = 0;
+			for(;;)
+			{
+				int Idx = Text.IndexOf(OldFileName, StartIdx, StringComparison.OrdinalIgnoreCase);
+				if(Idx == -1)
+				{
+					break;
+				}
+				else if((Idx == 0 || !IsBaseFileNameCharacter(Text[Idx - 1])) && (Idx + OldFileName.Length == Text.Length || !IsBaseFileNameCharacter(Text[Idx + OldFileName.Length])))
+				{
+					Text = Text.Substring(0, Idx) + NewFileName + Text.Substring(Idx + OldFileName.Length);
+					StartIdx = Idx + NewFileName.Length;
+				}
+				else
+				{
+					StartIdx = Idx + 1;
+				}
+			}
+			return Text;
+		}
+
+		/// <summary>
+		/// Determines if a character should be treated as part of a base filename, when updating strings for hot reload
+		/// </summary>
+		/// <param name="Character">The character to check</param>
+		/// <returns>True if the character is part of a base filename, false otherwise</returns>
+		static bool IsBaseFileNameCharacter(char Character)
+		{
+			return Char.IsLetterOrDigit(Character) || Character == '_';
+		}
+
+		/// <summary>
+		/// Patch the action graph for hot reloading, mapping files according to the given dictionary.
+		/// </summary>
+		public static void PatchActionGraph(IEnumerable<Action> Actions, Dictionary<FileReference, FileReference> OriginalFileToHotReloadFile)
+		{
+			// Gather all of the response files for link actions.  We're going to need to patch 'em up after we figure out new
+			// names for all of the output files and import libraries
+			List<string> ResponseFilePaths = new List<string>();
+
+			// Same as Response files but for all of the link.sh files for link actions.
+			// Only used on BuildHostPlatform Linux
+			List<string> LinkScriptFilePaths = new List<string>();
+
+			// Keep a map of the original file names and their new file names, so we can fix up response files after
+			Dictionary<string, string> OriginalFileNameAndNewFileNameList_NoExtensions = new Dictionary<string, string>();
+
+			// Finally, we'll keep track of any file items that we had to create counterparts for change file names, so we can fix those up too
+			Dictionary<FileItem, FileItem> AffectedOriginalFileItemAndNewFileItemMap = new Dictionary<FileItem, FileItem>();
+
+			foreach (Action Action in Actions.Where((Action) => Action.ActionType == ActionType.Link))
+			{
+				// Assume that the first produced item (with no extension) is our output file name
+				FileReference HotReloadFile;
+				if(!OriginalFileToHotReloadFile.TryGetValue(Action.ProducedItems[0].Location, out HotReloadFile))
+				{
+					continue;
+				}
+
+				string OriginalFileNameWithoutExtension = Utils.GetFilenameWithoutAnyExtensions(Action.ProducedItems[0].AbsolutePath);
+				string NewFileNameWithoutExtension = Utils.GetFilenameWithoutAnyExtensions(HotReloadFile.FullName);
+
+				// Find the response file in the command line.  We'll need to make a copy of it with our new file name.
+				string ResponseFileExtension = ".response";
+				int ResponseExtensionIndex = Action.CommandArguments.IndexOf(ResponseFileExtension, StringComparison.InvariantCultureIgnoreCase);
+				if (ResponseExtensionIndex != -1)
+				{
+					int ResponseFilePathIndex = Action.CommandArguments.LastIndexOf("@\"", ResponseExtensionIndex);
+					if (ResponseFilePathIndex == -1)
+					{
+						throw new BuildException("Couldn't find response file path in action's command arguments when hot reloading");
+					}
+
+					string OriginalResponseFilePathWithoutExtension = Action.CommandArguments.Substring(ResponseFilePathIndex + 2, (ResponseExtensionIndex - ResponseFilePathIndex) - 2);
+					string OriginalResponseFilePath = OriginalResponseFilePathWithoutExtension + ResponseFileExtension;
+
+					string NewResponseFilePath = ReplaceBaseFileName(OriginalResponseFilePath, OriginalFileNameWithoutExtension, NewFileNameWithoutExtension);
+
+					// Copy the old response file to the new path
+					if(String.Compare(OriginalResponseFilePath, NewResponseFilePath, StringComparison.OrdinalIgnoreCase) != 0)
+					{
+						File.Copy(OriginalResponseFilePath, NewResponseFilePath, overwrite: true);
+					}
+
+					// Keep track of the new response file name.  We'll have to do some edits afterwards.
+					ResponseFilePaths.Add(NewResponseFilePath);
+				}
+
+				// Find the *.link.sh file in the command line.  We'll need to make a copy of it with our new file name.
+				// Only currently used on Linux
+				if (UEBuildPlatform.IsPlatformInGroup(BuildHostPlatform.Current.Platform, UnrealPlatformGroup.Unix))
+				{
+					string LinkScriptFileExtension = ".link.sh";
+					int LinkScriptExtensionIndex = Action.CommandArguments.IndexOf(LinkScriptFileExtension, StringComparison.InvariantCultureIgnoreCase);
+					if (LinkScriptExtensionIndex != -1)
+					{
+						// We expect the script invocation to be quoted
+						int LinkScriptFilePathIndex = Action.CommandArguments.LastIndexOf("\"", LinkScriptExtensionIndex);
+						if (LinkScriptFilePathIndex == -1)
+						{
+							throw new BuildException("Couldn't find link script file path in action's command arguments when hot reloading. Is the path quoted?");
+						}
+
+						string OriginalLinkScriptFilePathWithoutExtension = Action.CommandArguments.Substring(LinkScriptFilePathIndex + 1, (LinkScriptExtensionIndex - LinkScriptFilePathIndex) - 1);
+						string OriginalLinkScriptFilePath = OriginalLinkScriptFilePathWithoutExtension + LinkScriptFileExtension;
+
+						string NewLinkScriptFilePath = ReplaceBaseFileName(OriginalLinkScriptFilePath, OriginalFileNameWithoutExtension, NewFileNameWithoutExtension);
+
+						// Copy the old response file to the new path
+						File.Copy(OriginalLinkScriptFilePath, NewLinkScriptFilePath, overwrite: true);
+
+						// Keep track of the new response file name.  We'll have to do some edits afterwards.
+						LinkScriptFilePaths.Add(NewLinkScriptFilePath);
+					}
+
+					// Update this action's list of prerequisite items too
+					for (int ItemIndex = 0; ItemIndex < Action.PrerequisiteItems.Count; ++ItemIndex)
+					{
+						FileItem OriginalPrerequisiteItem = Action.PrerequisiteItems[ItemIndex];
+						string NewPrerequisiteItemFilePath = ReplaceBaseFileName(OriginalPrerequisiteItem.AbsolutePath, OriginalFileNameWithoutExtension, NewFileNameWithoutExtension);
+
+						if (OriginalPrerequisiteItem.AbsolutePath != NewPrerequisiteItemFilePath)
+						{
+							// OK, the prerequisite item's file name changed so we'll update it to point to our new file
+							FileItem NewPrerequisiteItem = FileItem.GetItemByPath(NewPrerequisiteItemFilePath);
+							Action.PrerequisiteItems[ItemIndex] = NewPrerequisiteItem;
+
+							// Copy the other important settings from the original file item
+							NewPrerequisiteItem.bNeedsHotReloadNumbersDLLCleanUp = OriginalPrerequisiteItem.bNeedsHotReloadNumbersDLLCleanUp;
+							NewPrerequisiteItem.ProducingAction = OriginalPrerequisiteItem.ProducingAction;
+
+							// Keep track of it so we can fix up dependencies in a second pass afterwards
+							AffectedOriginalFileItemAndNewFileItemMap.Add(OriginalPrerequisiteItem, NewPrerequisiteItem);
+
+							ResponseExtensionIndex = OriginalPrerequisiteItem.AbsolutePath.IndexOf(ResponseFileExtension, StringComparison.InvariantCultureIgnoreCase);
+							if (ResponseExtensionIndex != -1)
+							{
+								string OriginalResponseFilePathWithoutExtension = OriginalPrerequisiteItem.AbsolutePath.Substring(0, ResponseExtensionIndex);
+								string OriginalResponseFilePath = OriginalResponseFilePathWithoutExtension + ResponseFileExtension;
+
+								string NewResponseFilePath = ReplaceBaseFileName(OriginalResponseFilePath, OriginalFileNameWithoutExtension, NewFileNameWithoutExtension);
+
+								// Copy the old response file to the new path
+								File.Copy(OriginalResponseFilePath, NewResponseFilePath, overwrite: true);
+
+								// Keep track of the new response file name.  We'll have to do some edits afterwards.
+								ResponseFilePaths.Add(NewResponseFilePath);
+							}
+						}
+					}
+				}
+
+				// Go ahead and replace all occurrences of our file name in the command-line (ignoring extensions)
+				Action.CommandArguments = ReplaceBaseFileName(Action.CommandArguments, OriginalFileNameWithoutExtension, NewFileNameWithoutExtension);
+
+
+				// Update this action's list of produced items too
+				for (int ItemIndex = 0; ItemIndex < Action.ProducedItems.Count; ++ItemIndex)
+				{
+					FileItem OriginalProducedItem = Action.ProducedItems[ItemIndex];
+
+					string NewProducedItemFilePath = ReplaceBaseFileName(OriginalProducedItem.AbsolutePath, OriginalFileNameWithoutExtension, NewFileNameWithoutExtension);
+					if (OriginalProducedItem.AbsolutePath != NewProducedItemFilePath)
+					{
+						// OK, the produced item's file name changed so we'll update it to point to our new file
+						FileItem NewProducedItem = FileItem.GetItemByPath(NewProducedItemFilePath);
+						Action.ProducedItems[ItemIndex] = NewProducedItem;
+
+						// Copy the other important settings from the original file item
+						NewProducedItem.bNeedsHotReloadNumbersDLLCleanUp = OriginalProducedItem.bNeedsHotReloadNumbersDLLCleanUp;
+						NewProducedItem.ProducingAction = OriginalProducedItem.ProducingAction;
+
+						// Keep track of it so we can fix up dependencies in a second pass afterwards
+						AffectedOriginalFileItemAndNewFileItemMap.Add(OriginalProducedItem, NewProducedItem);
+					}
+				}
+
+				// The status description of the item has the file name, so we'll update it too
+				Action.StatusDescription = ReplaceBaseFileName(Action.StatusDescription, OriginalFileNameWithoutExtension, NewFileNameWithoutExtension);
+
+
+				// Keep track of the file names, so we can fix up response files afterwards.
+				if(!OriginalFileNameAndNewFileNameList_NoExtensions.ContainsKey(OriginalFileNameWithoutExtension))
+				{
+					OriginalFileNameAndNewFileNameList_NoExtensions[OriginalFileNameWithoutExtension] = NewFileNameWithoutExtension;
+				}
+				else if(OriginalFileNameAndNewFileNameList_NoExtensions[OriginalFileNameWithoutExtension] != NewFileNameWithoutExtension)
+				{
+					throw new BuildException("Unexpected conflict in renaming files; {0} maps to {1} and {2}", OriginalFileNameWithoutExtension, OriginalFileNameAndNewFileNameList_NoExtensions[OriginalFileNameWithoutExtension], NewFileNameWithoutExtension);
+				}
+			}
+
+
+			// Do another pass and update any actions that depended on the original file names that we changed
+			foreach (Action Action in Actions)
+			{
+				for (int ItemIndex = 0; ItemIndex < Action.PrerequisiteItems.Count; ++ItemIndex)
+				{
+					FileItem OriginalFileItem = Action.PrerequisiteItems[ItemIndex];
+
+					FileItem NewFileItem;
+					if (AffectedOriginalFileItemAndNewFileItemMap.TryGetValue(OriginalFileItem, out NewFileItem))
+					{
+						// OK, looks like we need to replace this file item because we've renamed the file
+						Action.PrerequisiteItems[ItemIndex] = NewFileItem;
+					}
+				}
+			}
+
+
+			if (OriginalFileNameAndNewFileNameList_NoExtensions.Count > 0)
+			{
+				foreach (string ResponseFilePath in ResponseFilePaths)
+				{
+					// Load the file up
+					string FileContents = Utils.ReadAllText(ResponseFilePath);
+
+					// Replace all of the old file names with new ones
+					foreach (KeyValuePair<string, string> FileNameTuple in OriginalFileNameAndNewFileNameList_NoExtensions)
+					{
+						string OriginalFileNameWithoutExtension = FileNameTuple.Key;
+						string NewFileNameWithoutExtension = FileNameTuple.Value;
+
+						FileContents = ReplaceBaseFileName(FileContents, OriginalFileNameWithoutExtension, NewFileNameWithoutExtension);
+					}
+
+					// Overwrite the original file
+					File.WriteAllText(ResponseFilePath, FileContents, new System.Text.UTF8Encoding(false));
+				}
+
+				if (UEBuildPlatform.IsPlatformInGroup(BuildHostPlatform.Current.Platform, UnrealPlatformGroup.Unix))
+				{
+					foreach (string LinkScriptFilePath in LinkScriptFilePaths)
+					{
+						// Load the file up
+						string FileContents = Utils.ReadAllText(LinkScriptFilePath);
+
+						// Replace all of the old file names with new ones
+						foreach (KeyValuePair<string, string> FileNameTuple in OriginalFileNameAndNewFileNameList_NoExtensions)
+						{
+							string OriginalFileNameWithoutExtension = FileNameTuple.Key;
+							string NewFileNameWithoutExtension = FileNameTuple.Value;
+
+							FileContents = ReplaceBaseFileName(FileContents, OriginalFileNameWithoutExtension, NewFileNameWithoutExtension);
+						}
+
+						// Overwrite the original file
+						File.WriteAllText(LinkScriptFilePath, FileContents, new System.Text.UTF8Encoding(false));
+					}
+				}
+			}
+
+			// Update the action that writes out the module manifests
+			foreach(Action Action in Actions)
+			{
+				if(Action.ActionType == ActionType.WriteMetadata)
+				{
+					string Arguments = Action.CommandArguments;
+
+					// Find the argument for the metadata file
+					const string InputArgument = "-Input=";
+
+					int InputIdx = Arguments.IndexOf(InputArgument);
+					if(InputIdx == -1)
+					{
+						throw new Exception("Missing -Input= argument to WriteMetadata command when patching action graph.");
+					}
+
+					int FileNameIdx = InputIdx + InputArgument.Length;
+					if(Arguments[FileNameIdx] == '\"')
+					{
+						FileNameIdx++;
+					}
+
+					int FileNameEndIdx = FileNameIdx;
+					while(FileNameEndIdx < Arguments.Length && (Arguments[FileNameEndIdx] != ' ' || Arguments[FileNameIdx - 1] == '\"') && Arguments[FileNameEndIdx] != '\"')
+					{
+						FileNameEndIdx++;
+					}
+
+					// Read the metadata file
+					FileReference TargetInfoFile = new FileReference(Arguments.Substring(FileNameIdx, FileNameEndIdx - FileNameIdx));
+					if(!FileReference.Exists(TargetInfoFile))
+					{
+						throw new Exception(String.Format("Unable to find metadata file to patch action graph ({0})", TargetInfoFile));
+					}
+
+					// Update the module names
+					WriteMetadataTargetInfo TargetInfo = BinaryFormatterUtils.Load<WriteMetadataTargetInfo>(TargetInfoFile);
+					foreach (KeyValuePair<FileReference, ModuleManifest> FileNameToVersionManifest in TargetInfo.FileToManifest)
+					{
+						foreach (KeyValuePair<string, string> Manifest in FileNameToVersionManifest.Value.ModuleNameToFileName)
+						{
+							FileReference OriginalFile = FileReference.Combine(FileNameToVersionManifest.Key.Directory, Manifest.Key);
+
+							FileReference HotReloadFile;
+							if(OriginalFileToHotReloadFile.TryGetValue(OriginalFile, out HotReloadFile))
+							{
+								FileNameToVersionManifest.Value.ModuleNameToFileName[Manifest.Key] = HotReloadFile.GetFileName();
+								break;
+							}
+						}
+					}
+
+					// Write the hot-reload metadata file and update the argument list
+					FileReference HotReloadTargetInfoFile = FileReference.Combine(TargetInfoFile.Directory, "Metadata-HotReload.dat");
+					BinaryFormatterUtils.Save(HotReloadTargetInfoFile, TargetInfo);
+					Arguments = Arguments.Substring(0, FileNameIdx) + HotReloadTargetInfoFile + Arguments.Substring(FileNameEndIdx);
+				}
+			}
+		}
+	}
+}
diff --git a/Engine/Source/Programs/UnrealBuildTool/UnrealBuildTool.cs b/Engine/Source/Programs/UnrealBuildTool/UnrealBuildTool.cs
--- a/Engine/Source/Programs/UnrealBuildTool/UnrealBuildTool.cs
+++ b/Engine/Source/Programs/UnrealBuildTool/UnrealBuildTool.cs
@@ -12,21 +12,6 @@ using Tools.DotNETCommon;
 
 namespace UnrealBuildTool
 {
-	enum EHotReload
-	{
-		Disabled,
-		FromIDE,
-		FromEditor
-	}
-
-	[Serializable]
-	class HotReloadState
-	{
-		public int NextSuffix = 1;
-		public Dictionary<FileReference, FileReference> OriginalFileToHotReloadFile = new Dictionary<FileReference, FileReference>();
-		public HashSet<FileReference> AllHotReloadFiles = new HashSet<FileReference>();
-	}
-
 	static class UnrealBuildTool
 	{
 		/// <summary>
@@ -1091,28 +1076,28 @@ namespace UnrealBuildTool
 							}
 						}
 
-						EHotReload HotReloadMode = EHotReload.Disabled;
+						HotReloadMode HotReloadMode = HotReloadMode.Disabled;
 						if(Arguments.Any(x => x.Equals("-ForceHotReload", StringComparison.InvariantCultureIgnoreCase)))
 						{
-							HotReloadMode = EHotReload.FromIDE;
+							HotReloadMode = HotReloadMode.FromIDE;
 						}
 						else if (TargetDescs.Count == 1 && !Arguments.Any(x => x.Equals("-NoHotReload", StringComparison.InvariantCultureIgnoreCase)))
 						{
-							if (BuildConfiguration.bAllowHotReloadFromIDE && ShouldDoHotReloadFromIDE(BuildConfiguration, Arguments, TargetDescs[0]))
+							if (BuildConfiguration.bAllowHotReloadFromIDE && HotReload.ShouldDoHotReloadFromIDE(BuildConfiguration, Arguments, TargetDescs[0]))
 							{
-								HotReloadMode = EHotReload.FromIDE;
+								HotReloadMode = HotReloadMode.FromIDE;
 							}
 							else if (HotReloadModuleNameToSuffix.Count > 0 && TargetDescs[0].ForeignPlugin == null)
 							{
-								HotReloadMode = EHotReload.FromEditor;
+								HotReloadMode = HotReloadMode.FromEditor;
 							}
 
-							if (HotReloadMode != EHotReload.Disabled && BuildConfiguration.bCleanProject)
+							if (HotReloadMode != HotReloadMode.Disabled && BuildConfiguration.bCleanProject)
 							{
 								throw new BuildException("Unable to clean target while hot-reloading. Close the editor and try again.");
 							}
 						}
-						TargetDescriptor HotReloadTargetDesc = (HotReloadMode != EHotReload.Disabled) ? TargetDescs[0] : null;
+						TargetDescriptor HotReloadTargetDesc = (HotReloadMode != HotReloadMode.Disabled) ? TargetDescs[0] : null;
 
 						if (bIsAssemblingBuild)
 						{
@@ -1204,60 +1189,22 @@ namespace UnrealBuildTool
 
 								// Apply the previous hot reload state
 								HotReloadState HotReloadState = null;
-								if(HotReloadMode == EHotReload.Disabled)
+								if(HotReloadMode == HotReloadMode.Disabled)
 								{
 									// Delete any previous state files; they are no longer valid
 									foreach(TargetDescriptor TargetDesc in TargetDescs)
 									{
-										FileReference HotReloadStateFile = GetHotReloadStateFile(TargetDesc.ProjectFile, TargetDesc.Name, TargetDesc.Platform, TargetDesc.Configuration, TargetDesc.Architecture);
-										if(FileReference.Exists(HotReloadStateFile))
-										{
-											try
-											{
-												HotReloadState = ReadHotReloadState(HotReloadStateFile);
-											}
-											catch(Exception Ex)
-											{
-												Log.TraceWarning("Unable to read hot reload state file: {0}", HotReloadStateFile);
-												Log.WriteException(Ex, null);
-											}
-
-											if(HotReloadState != null)
-											{
-												foreach(FileReference Location in HotReloadState.AllHotReloadFiles.OrderBy(x => x.FullName, StringComparer.OrdinalIgnoreCase))
-												{
-													if(FileReference.Exists(Location))
-													{
-														try
-														{
-															FileReference.Delete(Location);
-														}
-														catch(Exception Ex)
-														{
-															throw new BuildException(Ex, "Unable to delete hot-reload file: {0}", Location);
-														}
-														Log.TraceInformation("Deleted hot-reload file: {0}", Location);
-													}
-												}
-												try
-												{
-													FileReference.Delete(HotReloadStateFile);
-												}
-												catch(Exception Ex)
-												{
-													throw new BuildException(Ex, "Unable to delete hot-reload state file: {0}", HotReloadStateFile);
-												}
-											}
-										}
+										FileReference HotReloadStateFile = HotReloadState.GetLocation(TargetDesc.ProjectFile, TargetDesc.Name, TargetDesc.Platform, TargetDesc.Configuration, TargetDesc.Architecture);
+										HotReload.DeleteTemporaryFiles(HotReloadStateFile);
 									}
 								}
 								else
 								{
 									// Read the previous state file and apply it to the action graph
-									FileReference HotReloadStateFile = GetHotReloadStateFile(TargetDescs[0].ProjectFile, TargetDescs[0].Name, TargetDescs[0].Platform, TargetDescs[0].Configuration, TargetDescs[0].Architecture);
+									FileReference HotReloadStateFile = HotReloadState.GetLocation(TargetDescs[0].ProjectFile, TargetDescs[0].Name, TargetDescs[0].Platform, TargetDescs[0].Configuration, TargetDescs[0].Architecture);
 									if(FileReference.Exists(HotReloadStateFile))
 									{
-										HotReloadState = ReadHotReloadState(HotReloadStateFile);
+										HotReloadState = HotReloadState.Load(HotReloadStateFile);
 									}
 									else
 									{
@@ -1265,7 +1212,7 @@ namespace UnrealBuildTool
 									}
 
 									// Update the action graph to produce these new files
-									PatchActionGraphForHotReload(PrerequisiteActions, HotReloadState.OriginalFileToHotReloadFile);
+									HotReload.PatchActionGraph(PrerequisiteActions, HotReloadState.OriginalFileToHotReloadFile);
 
 									// Update the module to output file mapping
 									foreach(string HotReloadModuleName in HotReloadModuleNamesForAllTargets)
@@ -1294,12 +1241,12 @@ namespace UnrealBuildTool
 												foreach(FileItem ModuleOutputItem in ModuleOutputItems)
 												{
 													FileReference OldLocation = ModuleOutputItem.Location;
-													FileReference NewLocation = ReplaceHotReloadSuffix(OldLocation, ModuleSuffix);
+													FileReference NewLocation = HotReload.ReplaceSuffix(OldLocation, ModuleSuffix);
 													OldLocationToNewLocation[OldLocation] = NewLocation;
 												}
 											}
 										}
-										PatchActionGraphForHotReload(PrerequisiteActions, OldLocationToNewLocation);
+										HotReload.PatchActionGraph(PrerequisiteActions, OldLocationToNewLocation);
 									}
 								}
 
@@ -1309,7 +1256,7 @@ namespace UnrealBuildTool
 
 								// Patch action history for hot reload when running in assembler mode.  In assembler mode, the suffix on the output file will be
 								// the same for every invocation on that makefile, but we need a new suffix each time.
-								if (HotReloadMode != EHotReload.Disabled)
+								if (HotReloadMode != HotReloadMode.Disabled)
 								{
 									// For all the hot-reloadable modules that may need a unique suffix appended, build a mapping from output item to all the output items in that module. We can't 
 									// apply a suffix to one without applying a suffix to all of them.
@@ -1358,12 +1305,12 @@ namespace UnrealBuildTool
 									foreach(FileItem FileRequiringSuffix in FilesRequiringSuffix)
 									{
 										FileReference OldLocation = FileRequiringSuffix.Location;
-										FileReference NewLocation = ReplaceHotReloadSuffix(OldLocation, HotReloadState.NextSuffix);
+										FileReference NewLocation = HotReload.ReplaceSuffix(OldLocation, HotReloadState.NextSuffix);
 										OldLocationToNewLocation[OldLocation] = NewLocation;
 									}
 
 									// Update the action graph with these new paths
-									PatchActionGraphForHotReload(PrerequisiteActions, OldLocationToNewLocation);
+									HotReload.PatchActionGraph(PrerequisiteActions, OldLocationToNewLocation);
 
 									// Get a new list of actions to execute now that the graph has been modified
 									ActionsToExecute = ActionGraph.GetActionsToExecute(BuildConfiguration, PrerequisiteActions, Targets, TargetToHeaders, bIsAssemblingBuild, bNeedsFullCPPIncludeRescan, out TargetToOutdatedPrerequisitesMap);
@@ -1393,7 +1340,7 @@ namespace UnrealBuildTool
 											if(HotReloadFileToOriginalFile.TryGetValue(ProducedItem.Location, out OriginalLocation))
 											{
 												HotReloadState.OriginalFileToHotReloadFile[OriginalLocation] = ProducedItem.Location;
-												HotReloadState.AllHotReloadFiles.Add(ProducedItem.Location);
+												HotReloadState.TemporaryFiles.Add(ProducedItem.Location);
 											}
 										}
 									}
@@ -1405,8 +1352,8 @@ namespace UnrealBuildTool
 									}
 
 									// Save the new state
-									FileReference HotReloadStateFile = GetHotReloadStateFile(TargetDescs[0].ProjectFile, TargetDescs[0].Name, TargetDescs[0].Platform, TargetDescs[0].Configuration, TargetDescs[0].Architecture);
-									WriteHotReloadState(HotReloadStateFile, HotReloadState);
+									FileReference HotReloadStateFile = HotReloadState.GetLocation(TargetDescs[0].ProjectFile, TargetDescs[0].Name, TargetDescs[0].Platform, TargetDescs[0].Configuration, TargetDescs[0].Architecture);
+									HotReloadState.Save(HotReloadStateFile);
 								}
 
 								// Display some stats to the user.
@@ -1530,89 +1477,6 @@ namespace UnrealBuildTool
 			}
 		}
 
-		/// <summary>
-		/// Checks if the editor is currently running and this is a hot-reload
-		/// </summary>
-		private static bool ShouldDoHotReloadFromIDE(BuildConfiguration BuildConfiguration, string[] Arguments, TargetDescriptor TargetDesc)
-		{
-			// Check if Hot-reload is disabled globally for this project
-			ConfigHierarchy Hierarchy = ConfigCache.ReadHierarchy(ConfigHierarchyType.Engine, DirectoryReference.FromFile(TargetDesc.ProjectFile), TargetDesc.Platform);
-			bool bAllowHotReloadFromIDE;
-			if (Hierarchy.TryGetValue("BuildConfiguration", "bAllowHotReloadFromIDE", out bAllowHotReloadFromIDE) && !bAllowHotReloadFromIDE)
-			{
-				return false;
-			}
-
-			if (Arguments.Any(x => x.Equals("-NoHotReloadFromIDE", StringComparison.InvariantCultureIgnoreCase)))
-			{
-				// Hot reload disabled through command line, possibly running from UAT
-				return false;
-			}
-
-			bool bIsRunning = false;
-
-			// @todo ubtmake: Kind of cheating here to figure out if an editor target.  At this point we don't have access to the actual target description, and
-			// this code must be able to execute before we create or load module rules DLLs so that hot reload can work with bUseUBTMakefiles
-			bool bIsEditorTarget = TargetDesc.Name.EndsWith("Editor", StringComparison.InvariantCultureIgnoreCase);
-
-			if (!BuildConfiguration.bGenerateManifest && bIsEditorTarget)
-			{
-				string EditorBaseFileName = "UE4Editor";
-				if (TargetDesc.Configuration != UnrealTargetConfiguration.Development && TargetDesc.Configuration != UnrealTargetConfiguration.DebugGame)
-				{
-					EditorBaseFileName = String.Format("{0}-{1}-{2}", EditorBaseFileName, TargetDesc.Platform, TargetDesc.Configuration);
-				}
-
-				FileReference EditorLocation;
-				if (TargetDesc.Platform == UnrealTargetPlatform.Win64)
-				{
-					EditorLocation = FileReference.Combine(UnrealBuildTool.EngineDirectory, "Binaries", "Win64", String.Format("{0}.exe", EditorBaseFileName));
-				}
-				else if (TargetDesc.Platform == UnrealTargetPlatform.Mac)
-				{
-					EditorLocation = FileReference.Combine(UnrealBuildTool.EngineDirectory, "Binaries", "Mac", String.Format("{0}.app/Contents/MacOS/{0}", EditorBaseFileName));
-				}
-				else if (TargetDesc.Platform == UnrealTargetPlatform.Linux)
-				{
-					EditorLocation = FileReference.Combine(UnrealBuildTool.EngineDirectory, "Binaries", "Linux", EditorBaseFileName);
-				}
-				else
-				{
-					throw new BuildException("Unknown editor filename for this platform");
-				}
-
-				BuildHostPlatform.ProcessInfo[] Processes = BuildHostPlatform.Current.GetProcesses();
-				string EditorRunsDir = Path.Combine(UnrealBuildTool.EngineDirectory.FullName, "Intermediate", "EditorRuns");
-
-				if (!Directory.Exists(EditorRunsDir))
-				{
-					return false;
-				}
-
-				FileInfo[] EditorRunsFiles = new DirectoryInfo(EditorRunsDir).GetFiles();
-
-				foreach (FileInfo File in EditorRunsFiles)
-				{
-					int PID;
-					BuildHostPlatform.ProcessInfo Proc = null;
-					if (!Int32.TryParse(File.Name, out PID) || (Proc = Processes.FirstOrDefault(P => P.PID == PID)) == default(BuildHostPlatform.ProcessInfo))
-					{
-						// Delete stale files (it may happen if editor crashes).
-						File.Delete();
-						continue;
-					}
-
-					// Don't break here to allow clean-up of other stale files.
-					if (!bIsRunning)
-					{
-						// Otherwise check if the path matches.
-						bIsRunning = new FileReference(Proc.Filename) == EditorLocation;
-					}
-				}
-			}
-			return bIsRunning;
-		}
-
 		/// <summary>
 		/// Parses the passed in command line for build configuration overrides.
 		/// </summary>
@@ -1697,384 +1561,6 @@ namespace UnrealBuildTool
 				}
 			}
 		}
-
-		public static FileReference GetHotReloadStateFile(FileReference ProjectFile, string TargetName, UnrealTargetPlatform Platform, UnrealTargetConfiguration Configuration, string Architecture)
-		{
-			DirectoryReference BaseDir = DirectoryReference.FromFile(ProjectFile) ?? UnrealBuildTool.EngineDirectory;
-			return FileReference.Combine(BaseDir, UEBuildTarget.GetPlatformIntermediateFolder(Platform, Architecture), TargetName, Configuration.ToString(), "HotReload.state");
-		}
-
-		public static HotReloadState ReadHotReloadState(FileReference Location)
-		{
-			return BinaryFormatterUtils.Load<HotReloadState>(Location);
-		}
-
-		public static void WriteHotReloadState(FileReference Location, HotReloadState State)
-		{
-			DirectoryReference.CreateDirectory(Location.Directory);
-			BinaryFormatterUtils.Save(Location, State);
-		}
-
-		/// <summary>
-		/// Replaces a hot reload suffix in a filename.
-		/// </summary>
-		public static FileReference ReplaceHotReloadSuffix(FileReference File, int Suffix)
-		{
-			string FileName = File.GetFileName();
-
-			// Find the end of the target and module name
-			int HyphenIdx = FileName.IndexOf('-');
-			if (HyphenIdx == -1)
-			{
-				throw new BuildException("Hot-reloadable files are expected to contain a hyphen, eg. UE4Editor-Core");
-			}
-
-			int NameEndIdx = HyphenIdx + 1;
-			while(NameEndIdx < FileName.Length && FileName[NameEndIdx] != '.' && FileName[NameEndIdx] != '-')
-			{
-				NameEndIdx++;
-			}
-
-			// Strip any existing suffix
-			if(NameEndIdx + 1 < FileName.Length && Char.IsDigit(FileName[NameEndIdx + 1]))
-			{
-				int SuffixEndIdx = NameEndIdx + 2;
-				while(SuffixEndIdx < FileName.Length && Char.IsDigit(FileName[SuffixEndIdx]))
-				{
-					SuffixEndIdx++;
-				}
-				if(SuffixEndIdx == FileName.Length || FileName[SuffixEndIdx] == '-' || FileName[SuffixEndIdx] == '.')
-				{
-					FileName = FileName.Substring(0, NameEndIdx) + FileName.Substring(SuffixEndIdx);
-				}
-			}
-
-			string NewFileName = String.Format("{0}-{1:D4}{2}", FileName.Substring(0, NameEndIdx), Suffix, FileName.Substring(NameEndIdx));
-
-			return FileReference.Combine(File.Directory, NewFileName);
-		}
-
-		/// <summary>
-		/// Replaces a base filename within a string. Ensures that the filename is not a substring of another longer string (eg. replacing "Foo" will match "Foo.Bar" but not "FooBar" or "BarFooBar"). 
-		/// </summary>
-		/// <param name="Text">Text to replace within</param>
-		/// <param name="OldFileName">Old filename</param>
-		/// <param name="NewFileName">New filename</param>
-		/// <returns>Text with file names replaced</returns>
-		static string ReplaceBaseFileName(string Text, string OldFileName, string NewFileName)
-		{
-			int StartIdx = 0;
-			for(;;)
-			{
-				int Idx = Text.IndexOf(OldFileName, StartIdx, StringComparison.OrdinalIgnoreCase);
-				if(Idx == -1)
-				{
-					break;
-				}
-				else if((Idx == 0 || !IsBaseFileNameCharacter(Text[Idx - 1])) && (Idx + OldFileName.Length == Text.Length || !IsBaseFileNameCharacter(Text[Idx + OldFileName.Length])))
-				{
-					Text = Text.Substring(0, Idx) + NewFileName + Text.Substring(Idx + OldFileName.Length);
-					StartIdx = Idx + NewFileName.Length;
-				}
-				else
-				{
-					StartIdx = Idx + 1;
-				}
-			}
-			return Text;
-		}
-
-		/// <summary>
-		/// Determines if a character should be treated as part of a base filename, when updating strings for hot reload
-		/// </summary>
-		/// <param name="Character">The character to check</param>
-		/// <returns>True if the character is part of a base filename, false otherwise</returns>
-		static bool IsBaseFileNameCharacter(char Character)
-		{
-			return Char.IsLetterOrDigit(Character) || Character == '_';
-		}
-
-		/// <summary>
-		/// Patch action history for hot reload when running in assembler mode.  In assembler mode, the suffix on the output file will be
-		/// the same for every invocation on that makefile, but we need a new suffix each time.
-		/// </summary>
-		static void PatchActionGraphForHotReload(IEnumerable<Action> Actions, Dictionary<FileReference, FileReference> OriginalFileToHotReloadFile)
-		{
-			// Gather all of the response files for link actions.  We're going to need to patch 'em up after we figure out new
-			// names for all of the output files and import libraries
-			List<string> ResponseFilePaths = new List<string>();
-
-			// Same as Response files but for all of the link.sh files for link actions.
-			// Only used on BuildHostPlatform Linux
-			List<string> LinkScriptFilePaths = new List<string>();
-
-			// Keep a map of the original file names and their new file names, so we can fix up response files after
-			Dictionary<string, string> OriginalFileNameAndNewFileNameList_NoExtensions = new Dictionary<string, string>();
-
-			// Finally, we'll keep track of any file items that we had to create counterparts for change file names, so we can fix those up too
-			Dictionary<FileItem, FileItem> AffectedOriginalFileItemAndNewFileItemMap = new Dictionary<FileItem, FileItem>();
-
-			foreach (Action Action in Actions.Where((Action) => Action.ActionType == ActionType.Link))
-			{
-				// Assume that the first produced item (with no extension) is our output file name
-				FileReference HotReloadFile;
-				if(!OriginalFileToHotReloadFile.TryGetValue(Action.ProducedItems[0].Location, out HotReloadFile))
-				{
-					continue;
-				}
-
-				string OriginalFileNameWithoutExtension = Utils.GetFilenameWithoutAnyExtensions(Action.ProducedItems[0].AbsolutePath);
-				string NewFileNameWithoutExtension = Utils.GetFilenameWithoutAnyExtensions(HotReloadFile.FullName);
-
-				// Find the response file in the command line.  We'll need to make a copy of it with our new file name.
-				string ResponseFileExtension = ".response";
-				int ResponseExtensionIndex = Action.CommandArguments.IndexOf(ResponseFileExtension, StringComparison.InvariantCultureIgnoreCase);
-				if (ResponseExtensionIndex != -1)
-				{
-					int ResponseFilePathIndex = Action.CommandArguments.LastIndexOf("@\"", ResponseExtensionIndex);
-					if (ResponseFilePathIndex == -1)
-					{
-						throw new BuildException("Couldn't find response file path in action's command arguments when hot reloading");
-					}
-
-					string OriginalResponseFilePathWithoutExtension = Action.CommandArguments.Substring(ResponseFilePathIndex + 2, (ResponseExtensionIndex - ResponseFilePathIndex) - 2);
-					string OriginalResponseFilePath = OriginalResponseFilePathWithoutExtension + ResponseFileExtension;
-
-					string NewResponseFilePath = ReplaceBaseFileName(OriginalResponseFilePath, OriginalFileNameWithoutExtension, NewFileNameWithoutExtension);
-
-					// Copy the old response file to the new path
-					if(String.Compare(OriginalResponseFilePath, NewResponseFilePath, StringComparison.OrdinalIgnoreCase) != 0)
-					{
-						File.Copy(OriginalResponseFilePath, NewResponseFilePath, overwrite: true);
-					}
-
-					// Keep track of the new response file name.  We'll have to do some edits afterwards.
-					ResponseFilePaths.Add(NewResponseFilePath);
-				}
-
-				// Find the *.link.sh file in the command line.  We'll need to make a copy of it with our new file name.
-				// Only currently used on Linux
-				if (UEBuildPlatform.IsPlatformInGroup(BuildHostPlatform.Current.Platform, UnrealPlatformGroup.Unix))
-				{
-					string LinkScriptFileExtension = ".link.sh";
-					int LinkScriptExtensionIndex = Action.CommandArguments.IndexOf(LinkScriptFileExtension, StringComparison.InvariantCultureIgnoreCase);
-					if (LinkScriptExtensionIndex != -1)
-					{
-						// We expect the script invocation to be quoted
-						int LinkScriptFilePathIndex = Action.CommandArguments.LastIndexOf("\"", LinkScriptExtensionIndex);
-						if (LinkScriptFilePathIndex == -1)
-						{
-							throw new BuildException("Couldn't find link script file path in action's command arguments when hot reloading. Is the path quoted?");
-						}
-
-						string OriginalLinkScriptFilePathWithoutExtension = Action.CommandArguments.Substring(LinkScriptFilePathIndex + 1, (LinkScriptExtensionIndex - LinkScriptFilePathIndex) - 1);
-						string OriginalLinkScriptFilePath = OriginalLinkScriptFilePathWithoutExtension + LinkScriptFileExtension;
-
-						string NewLinkScriptFilePath = ReplaceBaseFileName(OriginalLinkScriptFilePath, OriginalFileNameWithoutExtension, NewFileNameWithoutExtension);
-
-						// Copy the old response file to the new path
-						File.Copy(OriginalLinkScriptFilePath, NewLinkScriptFilePath, overwrite: true);
-
-						// Keep track of the new response file name.  We'll have to do some edits afterwards.
-						LinkScriptFilePaths.Add(NewLinkScriptFilePath);
-					}
-
-					// Update this action's list of prerequisite items too
-					for (int ItemIndex = 0; ItemIndex < Action.PrerequisiteItems.Count; ++ItemIndex)
-					{
-						FileItem OriginalPrerequisiteItem = Action.PrerequisiteItems[ItemIndex];
-						string NewPrerequisiteItemFilePath = ReplaceBaseFileName(OriginalPrerequisiteItem.AbsolutePath, OriginalFileNameWithoutExtension, NewFileNameWithoutExtension);
-
-						if (OriginalPrerequisiteItem.AbsolutePath != NewPrerequisiteItemFilePath)
-						{
-							// OK, the prerequisite item's file name changed so we'll update it to point to our new file
-							FileItem NewPrerequisiteItem = FileItem.GetItemByPath(NewPrerequisiteItemFilePath);
-							Action.PrerequisiteItems[ItemIndex] = NewPrerequisiteItem;
-
-							// Copy the other important settings from the original file item
-							NewPrerequisiteItem.bNeedsHotReloadNumbersDLLCleanUp = OriginalPrerequisiteItem.bNeedsHotReloadNumbersDLLCleanUp;
-							NewPrerequisiteItem.ProducingAction = OriginalPrerequisiteItem.ProducingAction;
-
-							// Keep track of it so we can fix up dependencies in a second pass afterwards
-							AffectedOriginalFileItemAndNewFileItemMap.Add(OriginalPrerequisiteItem, NewPrerequisiteItem);
-
-							ResponseExtensionIndex = OriginalPrerequisiteItem.AbsolutePath.IndexOf(ResponseFileExtension, StringComparison.InvariantCultureIgnoreCase);
-							if (ResponseExtensionIndex != -1)
-							{
-								string OriginalResponseFilePathWithoutExtension = OriginalPrerequisiteItem.AbsolutePath.Substring(0, ResponseExtensionIndex);
-								string OriginalResponseFilePath = OriginalResponseFilePathWithoutExtension + ResponseFileExtension;
-
-								string NewResponseFilePath = ReplaceBaseFileName(OriginalResponseFilePath, OriginalFileNameWithoutExtension, NewFileNameWithoutExtension);
-
-								// Copy the old response file to the new path
-								File.Copy(OriginalResponseFilePath, NewResponseFilePath, overwrite: true);
-
-								// Keep track of the new response file name.  We'll have to do some edits afterwards.
-								ResponseFilePaths.Add(NewResponseFilePath);
-							}
-						}
-					}
-				}
-
-				// Go ahead and replace all occurrences of our file name in the command-line (ignoring extensions)
-				Action.CommandArguments = ReplaceBaseFileName(Action.CommandArguments, OriginalFileNameWithoutExtension, NewFileNameWithoutExtension);
-
-
-				// Update this action's list of produced items too
-				for (int ItemIndex = 0; ItemIndex < Action.ProducedItems.Count; ++ItemIndex)
-				{
-					FileItem OriginalProducedItem = Action.ProducedItems[ItemIndex];
-
-					string NewProducedItemFilePath = ReplaceBaseFileName(OriginalProducedItem.AbsolutePath, OriginalFileNameWithoutExtension, NewFileNameWithoutExtension);
-					if (OriginalProducedItem.AbsolutePath != NewProducedItemFilePath)
-					{
-						// OK, the produced item's file name changed so we'll update it to point to our new file
-						FileItem NewProducedItem = FileItem.GetItemByPath(NewProducedItemFilePath);
-						Action.ProducedItems[ItemIndex] = NewProducedItem;
-
-						// Copy the other important settings from the original file item
-						NewProducedItem.bNeedsHotReloadNumbersDLLCleanUp = OriginalProducedItem.bNeedsHotReloadNumbersDLLCleanUp;
-						NewProducedItem.ProducingAction = OriginalProducedItem.ProducingAction;
-
-						// Keep track of it so we can fix up dependencies in a second pass afterwards
-						AffectedOriginalFileItemAndNewFileItemMap.Add(OriginalProducedItem, NewProducedItem);
-					}
-				}
-
-				// The status description of the item has the file name, so we'll update it too
-				Action.StatusDescription = ReplaceBaseFileName(Action.StatusDescription, OriginalFileNameWithoutExtension, NewFileNameWithoutExtension);
-
-
-				// Keep track of the file names, so we can fix up response files afterwards.
-				if(!OriginalFileNameAndNewFileNameList_NoExtensions.ContainsKey(OriginalFileNameWithoutExtension))
-				{
-					OriginalFileNameAndNewFileNameList_NoExtensions[OriginalFileNameWithoutExtension] = NewFileNameWithoutExtension;
-				}
-				else if(OriginalFileNameAndNewFileNameList_NoExtensions[OriginalFileNameWithoutExtension] != NewFileNameWithoutExtension)
-				{
-					throw new BuildException("Unexpected conflict in renaming files; {0} maps to {1} and {2}", OriginalFileNameWithoutExtension, OriginalFileNameAndNewFileNameList_NoExtensions[OriginalFileNameWithoutExtension], NewFileNameWithoutExtension);
-				}
-			}
-
-
-			// Do another pass and update any actions that depended on the original file names that we changed
-			foreach (Action Action in Actions)
-			{
-				for (int ItemIndex = 0; ItemIndex < Action.PrerequisiteItems.Count; ++ItemIndex)
-				{
-					FileItem OriginalFileItem = Action.PrerequisiteItems[ItemIndex];
-
-					FileItem NewFileItem;
-					if (AffectedOriginalFileItemAndNewFileItemMap.TryGetValue(OriginalFileItem, out NewFileItem))
-					{
-						// OK, looks like we need to replace this file item because we've renamed the file
-						Action.PrerequisiteItems[ItemIndex] = NewFileItem;
-					}
-				}
-			}
-
-
-			if (OriginalFileNameAndNewFileNameList_NoExtensions.Count > 0)
-			{
-				foreach (string ResponseFilePath in ResponseFilePaths)
-				{
-					// Load the file up
-					string FileContents = Utils.ReadAllText(ResponseFilePath);
-
-					// Replace all of the old file names with new ones
-					foreach (KeyValuePair<string, string> FileNameTuple in OriginalFileNameAndNewFileNameList_NoExtensions)
-					{
-						string OriginalFileNameWithoutExtension = FileNameTuple.Key;
-						string NewFileNameWithoutExtension = FileNameTuple.Value;
-
-						FileContents = ReplaceBaseFileName(FileContents, OriginalFileNameWithoutExtension, NewFileNameWithoutExtension);
-					}
-
-					// Overwrite the original file
-					File.WriteAllText(ResponseFilePath, FileContents, new System.Text.UTF8Encoding(false));
-				}
-
-				if (UEBuildPlatform.IsPlatformInGroup(BuildHostPlatform.Current.Platform, UnrealPlatformGroup.Unix))
-				{
-					foreach (string LinkScriptFilePath in LinkScriptFilePaths)
-					{
-						// Load the file up
-						string FileContents = Utils.ReadAllText(LinkScriptFilePath);
-
-						// Replace all of the old file names with new ones
-						foreach (KeyValuePair<string, string> FileNameTuple in OriginalFileNameAndNewFileNameList_NoExtensions)
-						{
-							string OriginalFileNameWithoutExtension = FileNameTuple.Key;
-							string NewFileNameWithoutExtension = FileNameTuple.Value;
-
-							FileContents = ReplaceBaseFileName(FileContents, OriginalFileNameWithoutExtension, NewFileNameWithoutExtension);
-						}
-
-						// Overwrite the original file
-						File.WriteAllText(LinkScriptFilePath, FileContents, new System.Text.UTF8Encoding(false));
-					}
-				}
-			}
-
-			// Update the action that writes out the module manifests
-			foreach(Action Action in Actions)
-			{
-				if(Action.ActionType == ActionType.WriteMetadata)
-				{
-					string Arguments = Action.CommandArguments;
-
-					// Find the argument for the metadata file
-					const string InputArgument = "-Input=";
-
-					int InputIdx = Arguments.IndexOf(InputArgument);
-					if(InputIdx == -1)
-					{
-						throw new Exception("Missing -Input= argument to WriteMetadata command when patching action graph.");
-					}
-
-					int FileNameIdx = InputIdx + InputArgument.Length;
-					if(Arguments[FileNameIdx] == '\"')
-					{
-						FileNameIdx++;
-					}
-
-					int FileNameEndIdx = FileNameIdx;
-					while(FileNameEndIdx < Arguments.Length && (Arguments[FileNameEndIdx] != ' ' || Arguments[FileNameIdx - 1] == '\"') && Arguments[FileNameEndIdx] != '\"')
-					{
-						FileNameEndIdx++;
-					}
-
-					// Read the metadata file
-					FileReference TargetInfoFile = new FileReference(Arguments.Substring(FileNameIdx, FileNameEndIdx - FileNameIdx));
-					if(!FileReference.Exists(TargetInfoFile))
-					{
-						throw new Exception(String.Format("Unable to find metadata file to patch action graph ({0})", TargetInfoFile));
-					}
-
-					// Update the module names
-					WriteMetadataTargetInfo TargetInfo = BinaryFormatterUtils.Load<WriteMetadataTargetInfo>(TargetInfoFile);
-					foreach (KeyValuePair<FileReference, ModuleManifest> FileNameToVersionManifest in TargetInfo.FileToManifest)
-					{
-						foreach (KeyValuePair<string, string> Manifest in FileNameToVersionManifest.Value.ModuleNameToFileName)
-						{
-							FileReference OriginalFile = FileReference.Combine(FileNameToVersionManifest.Key.Directory, Manifest.Key);
-
-							FileReference HotReloadFile;
-							if(OriginalFileToHotReloadFile.TryGetValue(OriginalFile, out HotReloadFile))
-							{
-								FileNameToVersionManifest.Value.ModuleNameToFileName[Manifest.Key] = HotReloadFile.GetFileName();
-								break;
-							}
-						}
-					}
-
-					// Write the hot-reload metadata file and update the argument list
-					FileReference HotReloadTargetInfoFile = FileReference.Combine(TargetInfoFile.Directory, "Metadata-HotReload.dat");
-					BinaryFormatterUtils.Save(HotReloadTargetInfoFile, TargetInfo);
-					Arguments = Arguments.Substring(0, FileNameIdx) + HotReloadTargetInfoFile + Arguments.Substring(FileNameEndIdx);
-				}
-			}
-		}
 	}
 
 	[ToolMode("Default")]
diff --git a/Engine/Source/Programs/UnrealBuildTool/UnrealBuildTool.csproj b/Engine/Source/Programs/UnrealBuildTool/UnrealBuildTool.csproj
--- a/Engine/Source/Programs/UnrealBuildTool/UnrealBuildTool.csproj
+++ b/Engine/Source/Programs/UnrealBuildTool/UnrealBuildTool.csproj
@@ -186,6 +186,7 @@
     <Compile Include="System\CppIncludeParser.cs" />
     <Compile Include="System\DataDrivenPlatformInfo.cs" />
     <Compile Include="System\EncryptionAndSigning.cs" />
+    <Compile Include="System\HotReload.cs" />
     <Compile Include="System\ModuleManifest.cs" />
     <Compile Include="System\PluginReferenceDescriptor.cs" />
     <Compile Include="System\PrecompiledManifest.cs" />
