UBT: Remove hack for tracking failed links in hot reload. This should no longer be necessary; since the action graph is restored to the previous hot-reload state, we will be incrementally rebuild the same files if no source files have changed.

diff --git a/Engine/Source/Programs/UnrealBuildTool/System/ActionGraph.cs b/Engine/Source/Programs/UnrealBuildTool/System/ActionGraph.cs
--- a/Engine/Source/Programs/UnrealBuildTool/System/ActionGraph.cs
+++ b/Engine/Source/Programs/UnrealBuildTool/System/ActionGraph.cs
@@ -383,7 +383,7 @@ namespace UnrealBuildTool
 		/// <summary>
 		/// Executes a list of actions.
 		/// </summary>
-		public static bool ExecuteActions(BuildConfiguration BuildConfiguration, List<Action> ActionsToExecute, out string ExecutorName, string TargetInfoForTelemetry, HotReloadMode HotReload)
+		public static bool ExecuteActions(BuildConfiguration BuildConfiguration, List<Action> ActionsToExecute, out string ExecutorName)
 		{
 			bool Result = true;
 			ExecutorName = "";
@@ -426,32 +426,6 @@ namespace UnrealBuildTool
 							foreach (FileItem Item in BuildAction.ProducedItems)
 							{
 								bool bExists = File.Exists(Item.AbsolutePath) || Directory.Exists(Item.AbsolutePath);
-
-								if (HotReload != HotReloadMode.Disabled)
-								{
-									string FailedFilename = Path.Combine(Path.GetDirectoryName(Item.AbsolutePath), "failed.hotreload");
-									if (!bExists)
-									{
-										// Create a failed.hotreload file here to indicate that we need to attempt another hotreload link
-										// step in future, even though no source files have changed.
-										// This is necessary because we also don't want to link a brand new instance of a module every time
-										// a user hits the Compile button when nothing has changed.
-										FileItem.CreateIntermediateTextFile(new FileReference(FailedFilename), "");
-									}
-									else
-									{
-										try
-										{
-											File.Delete(FailedFilename);
-										}
-										catch
-										{
-											// Ignore but log failed deletions
-											Log.TraceVerbose("Failed to delete failed.hotreload file \"{0}\" - this may cause redundant hotreloads", FailedFilename);
-										}
-									}
-								}
-
 								if (!bExists)
 								{
 									throw new BuildException("UBT ERROR: Failed to produce item: " + Item.AbsolutePath);
diff --git a/Engine/Source/Programs/UnrealBuildTool/UnrealBuildTool.cs b/Engine/Source/Programs/UnrealBuildTool/UnrealBuildTool.cs
--- a/Engine/Source/Programs/UnrealBuildTool/UnrealBuildTool.cs
+++ b/Engine/Source/Programs/UnrealBuildTool/UnrealBuildTool.cs
@@ -1372,7 +1372,6 @@ namespace UnrealBuildTool
 								}
 
 								// Execute the actions.
-								string TargetInfoForTelemetry = String.Join("|", Targets.Select(x => String.Format("{0} {1} {2}{3}", x.TargetName, x.Platform, x.Configuration, "")));
 								DateTime ExecutorStartTime = DateTime.UtcNow;
 								if (BuildConfiguration.bXGEExport)
 								{
@@ -1381,7 +1380,7 @@ namespace UnrealBuildTool
 								}
 								else
 								{
-									bSuccess = ActionGraph.ExecuteActions(BuildConfiguration, ActionsToExecute.ToList(), out ExecutorName, TargetInfoForTelemetry, HotReloadMode);
+									bSuccess = ActionGraph.ExecuteActions(BuildConfiguration, ActionsToExecute.ToList(), out ExecutorName);
 								}
 								TotalExecutorTime += (DateTime.UtcNow - ExecutorStartTime).TotalSeconds;
 
