UBT: Remove unused variable.

#rnx

diff --git a/Engine/Source/Programs/UnrealBuildTool/UnrealBuildTool.cs b/Engine/Source/Programs/UnrealBuildTool/UnrealBuildTool.cs
--- a/Engine/Source/Programs/UnrealBuildTool/UnrealBuildTool.cs
+++ b/Engine/Source/Programs/UnrealBuildTool/UnrealBuildTool.cs
@@ -441,7 +441,7 @@ namespace UnrealBuildTool
 					// Build our project
 					if (Result == ECompilationResult.Succeeded)
 					{
-						Result = RunUBT(BuildConfiguration, Arguments, ProjectFile, true);
+						Result = RunUBT(BuildConfiguration, Arguments, ProjectFile);
 					}
 
 					// Print some performance info
@@ -615,7 +615,7 @@ namespace UnrealBuildTool
 			}
 		}
 
-		internal static ECompilationResult RunUBT(BuildConfiguration BuildConfiguration, string[] Arguments, FileReference ProjectFile, bool bCatchExceptions)
+		internal static ECompilationResult RunUBT(BuildConfiguration BuildConfiguration, string[] Arguments, FileReference ProjectFile)
 		{
 			bool bSuccess = true;
 
@@ -1414,10 +1414,6 @@ namespace UnrealBuildTool
 			}
 			catch (Exception Ex)
 			{
-				if (!bCatchExceptions)
-				{
-					throw;
-				}
 				Log.WriteException(Ex, String.IsNullOrEmpty(BuildConfiguration.LogFileName) ? null : BuildConfiguration.LogFileName);
 				BuildResult = ECompilationResult.OtherCompilationError;
 			}
