UBT: Move the default build mode handler into its own file.

diff --git a/Engine/Source/Programs/UnrealBuildTool/Modes/BuildMode.cs b/Engine/Source/Programs/UnrealBuildTool/Modes/BuildMode.cs
new file mode 100644
--- /dev/null
+++ b/Engine/Source/Programs/UnrealBuildTool/Modes/BuildMode.cs
@@ -0,0 +1,28 @@
+using System;
+using System.Collections.Generic;
+using System.Diagnostics;
+using System.Linq;
+using System.Text;
+using System.Threading.Tasks;
+using Tools.DotNETCommon;
+
+namespace UnrealBuildTool
+{
+	/// <summary>
+	/// Builds a target
+	/// </summary>
+	[ToolMode("Build")]
+	class BuildMode : ToolMode
+	{
+		/// <summary>
+		/// Main entry point
+		/// </summary>
+		/// <param name="CmdLine">Command-line arguments</param>
+		/// <returns>One of the values of ECompilationResult</returns>
+		public override int Execute(CommandLineArguments CmdLine)
+		{
+			return UnrealBuildTool.GuardedMain(CmdLine);
+		}
+	}
+}
+
diff --git a/Engine/Source/Programs/UnrealBuildTool/UnrealBuildTool.cs b/Engine/Source/Programs/UnrealBuildTool/UnrealBuildTool.cs
--- a/Engine/Source/Programs/UnrealBuildTool/UnrealBuildTool.cs
+++ b/Engine/Source/Programs/UnrealBuildTool/UnrealBuildTool.cs
@@ -523,7 +523,7 @@ namespace UnrealBuildTool
 			[CommandLine("-VSCode", Value="GenerateProjectFiles")]
 			[CommandLine("-VSMac", Value="GenerateProjectFiles")]
 			[CommandLine("-CLion", Value="GenerateProjectFiles")]
-			public string Mode = "Default";
+			public string Mode = "Build";
 
 			/// <summary>
 			/// Initialize the options with the given commnad line arguments
@@ -1557,13 +1557,4 @@ namespace UnrealBuildTool
 			}
 		}
 	}
-
-	[ToolMode("Default")]
-	class DefaultToolMode : ToolMode
-	{
-		public override int Execute(CommandLineArguments Arguments)
-		{
-			return UnrealBuildTool.GuardedMain(Arguments);
-		}
-	}
 }
diff --git a/Engine/Source/Programs/UnrealBuildTool/UnrealBuildTool.csproj b/Engine/Source/Programs/UnrealBuildTool/UnrealBuildTool.csproj
--- a/Engine/Source/Programs/UnrealBuildTool/UnrealBuildTool.csproj
+++ b/Engine/Source/Programs/UnrealBuildTool/UnrealBuildTool.csproj
@@ -140,6 +140,7 @@
     </When>
   </Choose>
   <ItemGroup>
+    <Compile Include="Modes\BuildMode.cs" />
     <Compile Include="Modes\DeployMode.cs" />
     <Compile Include="Modes\GenerateProjectFilesMode.cs" />
     <Compile Include="Modes\SetupPlatformsMode.cs" />
