UBT: Remove the bDebugInfo global flag. Most of this stuff can be switched from other flags.

#rnx

diff --git a/Engine/Source/Programs/UnrealBuildTool/Executors/Distcc.cs b/Engine/Source/Programs/UnrealBuildTool/Executors/Distcc.cs
--- a/Engine/Source/Programs/UnrealBuildTool/Executors/Distcc.cs
+++ b/Engine/Source/Programs/UnrealBuildTool/Executors/Distcc.cs
@@ -372,12 +372,12 @@ namespace UnrealBuildTool
 					TotalThreadSeconds += ThreadSeconds;
 				}
 
-				Log.WriteLineIf(bLogDetailedActionStats || UnrealBuildTool.bPrintDebugInfo,
+				Log.WriteLineIf(bLogDetailedActionStats,
 					LogEventType.Console,
 					"-------- End Detailed Actions Stats -----------------------------------------------------------");
 
 				// Log total CPU seconds and numbers of processors involved in tasks.
-				Log.WriteLineIf(bLogDetailedActionStats || UnrealBuildTool.bPrintDebugInfo,
+				Log.WriteLineIf(bLogDetailedActionStats,
 					LogEventType.Console, "Cumulative thread seconds ({0} processors): {1:0.00}", System.Environment.ProcessorCount, TotalThreadSeconds);
 			}
 			return bDistccResult;
diff --git a/Engine/Source/Programs/UnrealBuildTool/Executors/LocalExecutor.cs b/Engine/Source/Programs/UnrealBuildTool/Executors/LocalExecutor.cs
--- a/Engine/Source/Programs/UnrealBuildTool/Executors/LocalExecutor.cs
+++ b/Engine/Source/Programs/UnrealBuildTool/Executors/LocalExecutor.cs
@@ -121,12 +121,10 @@ namespace UnrealBuildTool
 			ActionStartInfo.RedirectStandardError = false;
 
 			// Log command-line used to execute task if debug info printing is enabled.
-			if (UnrealBuildTool.bPrintDebugInfo)
-			{
-				Log.TraceVerbose("Executing: {0} {1}", ExpandedCommandPath, ActionStartInfo.Arguments);
-			}
+			Log.TraceVerbose("Executing: {0} {1}", ExpandedCommandPath, ActionStartInfo.Arguments);
+
 			// Log summary if wanted.
-			else if (Action.bShouldOutputStatusDescription)
+			if (Action.bShouldOutputStatusDescription)
 			{
 				string CommandDescription = Action.CommandDescription != null ? Action.CommandDescription : Path.GetFileName(ExpandedCommandPath);
 				if (string.IsNullOrEmpty(CommandDescription))
@@ -517,11 +515,11 @@ namespace UnrealBuildTool
 			Log.WriteLineIf(bLogDetailedActionStats, LogEventType.Console, "-------- End Detailed Actions Stats -----------------------------------------------------------");
 
 			// Log total CPU seconds and numbers of processors involved in tasks.
-			Log.WriteLineIf(bLogDetailedActionStats || UnrealBuildTool.bPrintDebugInfo,
+			Log.WriteLineIf(bLogDetailedActionStats,
 				LogEventType.Console, "Cumulative thread seconds ({0} processors): {1:0.00}", System.Environment.ProcessorCount, TotalThreadSeconds);
 
 			// Log detailed stats
-			Log.WriteLineIf(bLogDetailedActionStats || UnrealBuildTool.bPrintDebugInfo,
+			Log.WriteLineIf(bLogDetailedActionStats,
 				LogEventType.Console,
 				"Cumulative action seconds ({0} processors): {1:0.00} building projects, {2:0.00} compiling, {3:0.00} creating app bundles, {4:0.00} generating debug info, {5:0.00} linking, {6:0.00} other",
 				System.Environment.ProcessorCount,
diff --git a/Engine/Source/Programs/UnrealBuildTool/Executors/SNDBS.cs b/Engine/Source/Programs/UnrealBuildTool/Executors/SNDBS.cs
--- a/Engine/Source/Programs/UnrealBuildTool/Executors/SNDBS.cs
+++ b/Engine/Source/Programs/UnrealBuildTool/Executors/SNDBS.cs
@@ -1,4 +1,4 @@
-﻿// Copyright 1998-2018 Epic Games, Inc. All Rights Reserved.
+// Copyright 1998-2018 Epic Games, Inc. All Rights Reserved.
 
 using System;
 using System.Collections.Generic;
@@ -420,7 +420,7 @@ namespace UnrealBuildTool
 				Log.TraceInformation("-------- End Detailed Actions Stats -----------------------------------------------------------");
 
 				// Log total CPU seconds and numbers of processors involved in tasks.
-				Log.WriteLineIf(bLogDetailedActionStats || UnrealBuildTool.bPrintDebugInfo,
+				Log.WriteLineIf(bLogDetailedActionStats,
 					LogEventType.Console, "Cumulative thread seconds ({0} processors): {1:0.00}", System.Environment.ProcessorCount, TotalThreadSeconds);
 			}
 			return SNDBSResult;
diff --git a/Engine/Source/Programs/UnrealBuildTool/UnrealBuildTool.cs b/Engine/Source/Programs/UnrealBuildTool/UnrealBuildTool.cs
--- a/Engine/Source/Programs/UnrealBuildTool/UnrealBuildTool.cs
+++ b/Engine/Source/Programs/UnrealBuildTool/UnrealBuildTool.cs
@@ -109,15 +109,6 @@ namespace UnrealBuildTool
 		/// </summary>
 		static string RemoteIniPath = null;
 
-		/// <summary>
-		/// Whether to print debug information out to the log
-		/// </summary>
-		static public bool bPrintDebugInfo
-		{
-			get;
-			private set;
-		}
-
 		/// <summary>
 		/// Whether to print performance information to the log
 		/// </summary>
@@ -346,7 +337,6 @@ namespace UnrealBuildTool
 					CommandLine.ParseArguments(Arguments, BuildConfiguration);
 
 					// Copy some of the static settings that are being deprecated from BuildConfiguration
-					bPrintDebugInfo = BuildConfiguration.bPrintDebugInfo || Log.OutputLevel == LogEventType.Verbose || Log.OutputLevel == LogEventType.VeryVerbose;
 					bPrintPerformanceInfo = BuildConfiguration.bPrintPerformanceInfo;
 
 					// Then let the command lines override any configs necessary.
