Add missing copyright declarations.

#rnx

diff --git a/Engine/Source/Programs/UnrealBuildTool/Modes/BuildMode.cs b/Engine/Source/Programs/UnrealBuildTool/Modes/BuildMode.cs
--- a/Engine/Source/Programs/UnrealBuildTool/Modes/BuildMode.cs
+++ b/Engine/Source/Programs/UnrealBuildTool/Modes/BuildMode.cs
@@ -1,3 +1,5 @@
+// Copyright 1998-2018 Epic Games, Inc. All Rights Reserved.
+
 using System;
 using System.Collections.Generic;
 using System.Diagnostics;
diff --git a/Engine/Source/Programs/UnrealBuildTool/System/HotReload.cs b/Engine/Source/Programs/UnrealBuildTool/System/HotReload.cs
--- a/Engine/Source/Programs/UnrealBuildTool/System/HotReload.cs
+++ b/Engine/Source/Programs/UnrealBuildTool/System/HotReload.cs
@@ -1,3 +1,5 @@
+// Copyright 1998-2018 Epic Games, Inc. All Rights Reserved.
+
 using System;
 using System.Collections.Generic;
 using System.IO;
