UBT: Postpone discovery of source files for modules until they are actually being compiled. Also move predicates for the derived action graph to remain valid into a separate class.

#rnx

diff --git a/Engine/Source/Programs/UnrealBuildTool/Configuration/BuildConfiguration.cs b/Engine/Source/Programs/UnrealBuildTool/Configuration/BuildConfiguration.cs
--- a/Engine/Source/Programs/UnrealBuildTool/Configuration/BuildConfiguration.cs
+++ b/Engine/Source/Programs/UnrealBuildTool/Configuration/BuildConfiguration.cs
@@ -143,7 +143,7 @@ namespace UnrealBuildTool
 		/// <summary>
 		/// If specified, we will only build this particular source file, ignore all other outputs.  Useful for testing non-Unity builds.
 		/// </summary>
-		public string SingleFileToCompile = null;
+		public FileReference SingleFileToCompile = null;
 
 		/// <summary>
 		/// Whether to skip checking for files identified by the junk manifest
diff --git a/Engine/Source/Programs/UnrealBuildTool/Configuration/UEBuildBinary.cs b/Engine/Source/Programs/UnrealBuildTool/Configuration/UEBuildBinary.cs
--- a/Engine/Source/Programs/UnrealBuildTool/Configuration/UEBuildBinary.cs
+++ b/Engine/Source/Programs/UnrealBuildTool/Configuration/UEBuildBinary.cs
@@ -158,10 +158,11 @@ namespace UnrealBuildTool
 		/// <param name="LinkEnvironment">The environment to link the binary in</param>
 		/// <param name="SharedPCHs">List of templates for shared PCHs</param>
 		/// <param name="WorkingSet">The working set of source files</param>
+		/// <param name="Predicates">Predicates for the generated actions to be valid</param>
 		/// <param name="ExeDir">Directory containing the output executable</param>
 		/// <param name="ActionGraph">Graph to add build actions to</param>
 		/// <returns>Set of built products</returns>
-		public List<FileItem> Build(ReadOnlyTargetRules Target, UEToolChain ToolChain, CppCompileEnvironment CompileEnvironment, LinkEnvironment LinkEnvironment, List<PrecompiledHeaderTemplate> SharedPCHs, ISourceFileWorkingSet WorkingSet, DirectoryReference ExeDir, ActionGraph ActionGraph)
+		public List<FileItem> Build(ReadOnlyTargetRules Target, UEToolChain ToolChain, CppCompileEnvironment CompileEnvironment, LinkEnvironment LinkEnvironment, List<PrecompiledHeaderTemplate> SharedPCHs, ISourceFileWorkingSet WorkingSet, BuildPredicateStore Predicates, DirectoryReference ExeDir, ActionGraph ActionGraph)
 		{
 			// Return nothing if we're using precompiled binaries. If we're not linking, we might want just one module to be compiled (eg. a foreign plugin), so allow any actions to run.
 			if (bUsePrecompiled && !Target.bDisableLinking)
@@ -170,7 +171,7 @@ namespace UnrealBuildTool
 			}
 
 			// Setup linking environment.
-			LinkEnvironment BinaryLinkEnvironment = SetupBinaryLinkEnvironment(Target, ToolChain, LinkEnvironment, CompileEnvironment, SharedPCHs, WorkingSet, ExeDir, ActionGraph);
+			LinkEnvironment BinaryLinkEnvironment = SetupBinaryLinkEnvironment(Target, ToolChain, LinkEnvironment, CompileEnvironment, SharedPCHs, WorkingSet, Predicates, ExeDir, ActionGraph);
 
 			// If we're generating projects, we only need include paths and definitions, there is no need to run the linking logic.
 			if (ProjectFileGenerator.bGenerateProjectFiles)
@@ -649,7 +650,7 @@ namespace UnrealBuildTool
 			return BinaryCompileEnvironment;
 		}
 
-		private LinkEnvironment SetupBinaryLinkEnvironment(ReadOnlyTargetRules Target, UEToolChain ToolChain, LinkEnvironment LinkEnvironment, CppCompileEnvironment CompileEnvironment, List<PrecompiledHeaderTemplate> SharedPCHs, ISourceFileWorkingSet WorkingSet, DirectoryReference ExeDir, ActionGraph ActionGraph)
+		private LinkEnvironment SetupBinaryLinkEnvironment(ReadOnlyTargetRules Target, UEToolChain ToolChain, LinkEnvironment LinkEnvironment, CppCompileEnvironment CompileEnvironment, List<PrecompiledHeaderTemplate> SharedPCHs, ISourceFileWorkingSet WorkingSet, BuildPredicateStore Predicates, DirectoryReference ExeDir, ActionGraph ActionGraph)
 		{
 			LinkEnvironment BinaryLinkEnvironment = new LinkEnvironment(LinkEnvironment);
 			HashSet<UEBuildModule> LinkEnvironmentVisitedModules = new HashSet<UEBuildModule>();
@@ -664,7 +665,7 @@ namespace UnrealBuildTool
 				{
 					// Compile each module.
 					Log.TraceVerbose("Compile module: " + Module.Name);
-					LinkInputFiles = Module.Compile(Target, ToolChain, BinaryCompileEnvironment, SharedPCHs, WorkingSet, ActionGraph);
+					LinkInputFiles = Module.Compile(Target, ToolChain, BinaryCompileEnvironment, SharedPCHs, WorkingSet, Predicates, ActionGraph);
 
 					// NOTE: Because of 'Shared PCHs', in monolithic builds the same PCH file may appear as a link input
 					// multiple times for a single binary.  We'll check for that here, and only add it once.  This avoids
diff --git a/Engine/Source/Programs/UnrealBuildTool/Configuration/UEBuildModule.cs b/Engine/Source/Programs/UnrealBuildTool/Configuration/UEBuildModule.cs
--- a/Engine/Source/Programs/UnrealBuildTool/Configuration/UEBuildModule.cs
+++ b/Engine/Source/Programs/UnrealBuildTool/Configuration/UEBuildModule.cs
@@ -733,7 +733,7 @@ namespace UnrealBuildTool
 		/// <summary>
 		/// Compiles the module, and returns a list of files output by the compiler.
 		/// </summary>
-		public abstract List<FileItem> Compile(ReadOnlyTargetRules Target, UEToolChain ToolChain, CppCompileEnvironment CompileEnvironment, List<PrecompiledHeaderTemplate> SharedPCHModules, ISourceFileWorkingSet WorkingSet, ActionGraph ActionGraph);
+		public abstract List<FileItem> Compile(ReadOnlyTargetRules Target, UEToolChain ToolChain, CppCompileEnvironment CompileEnvironment, List<PrecompiledHeaderTemplate> SharedPCHModules, ISourceFileWorkingSet WorkingSet, BuildPredicateStore Predicates, ActionGraph ActionGraph);
 
 		// Object interface.
 		public override string ToString()
diff --git a/Engine/Source/Programs/UnrealBuildTool/Configuration/UEBuildModuleCPP.cs b/Engine/Source/Programs/UnrealBuildTool/Configuration/UEBuildModuleCPP.cs
--- a/Engine/Source/Programs/UnrealBuildTool/Configuration/UEBuildModuleCPP.cs
+++ b/Engine/Source/Programs/UnrealBuildTool/Configuration/UEBuildModuleCPP.cs
@@ -61,21 +61,6 @@ namespace UnrealBuildTool
 			}
 		}
 
-		/// <summary>
-		/// All the source files for this module
-		/// </summary>
-		public readonly List<FileItem> SourceFiles = new List<FileItem>();
-
-		/// <summary>
-		/// A list of the absolute paths of source files to be built in this module.
-		/// </summary>
-		public readonly SourceFilesClass SourceFilesToBuild = new SourceFilesClass();
-
-		/// <summary>
-		/// A list of the source files that were found for the module.
-		/// </summary>
-		public readonly SourceFilesClass SourceFilesFound = new SourceFilesClass();
-
 		/// <summary>
 		/// The directory for this module's object files
 		/// </summary>
@@ -101,14 +86,21 @@ namespace UnrealBuildTool
 		/// </summary>
 		public List<string> InvalidIncludeDirectiveMessages;
 
+		/// <summary>
+		/// Set of source directories referenced during a build
+		/// </summary>
+		HashSet<DirectoryReference> SourceDirectories;
+
 		protected override void GetReferencedDirectories(HashSet<DirectoryReference> Directories)
 		{
 			base.GetReferencedDirectories(Directories);
 
-			foreach(FileItem SourceFile in SourceFiles)
+			if(SourceDirectories == null)
 			{
-				Directories.Add(SourceFile.Location.Directory);
+				throw new BuildException("GetReferencedDirectories() should not be called before building.");
 			}
+
+			Directories.UnionWith(SourceDirectories);
 		}
 
 		/// <summary>
@@ -202,20 +194,12 @@ namespace UnrealBuildTool
         };
 
 
-		public UEBuildModuleCPP(ModuleRules Rules, DirectoryReference IntermediateDirectory, DirectoryReference GeneratedCodeDirectory, IEnumerable<FileItem> SourceFiles, bool bBuildSourceFiles)
+		public UEBuildModuleCPP(ModuleRules Rules, DirectoryReference IntermediateDirectory, DirectoryReference GeneratedCodeDirectory)
 			: base(Rules)
 		{
 			this.IntermediateDirectory = IntermediateDirectory;
 			this.GeneratedCodeDirectory = GeneratedCodeDirectory;
 
-			this.SourceFiles = SourceFiles.ToList();
-
-			CategorizeSourceFiles(SourceFiles, SourceFilesFound);
-			if (bBuildSourceFiles)
-			{
-				SourceFilesToBuild.CopyFrom(SourceFilesFound);
-			}
-
 			foreach (string Def in PublicDefinitions)
 			{
 				Log.TraceVerbose("Compile Env {0}: {1}", Name, Def);
@@ -310,7 +294,7 @@ namespace UnrealBuildTool
 		}
 
 		// UEBuildModule interface.
-		public override List<FileItem> Compile(ReadOnlyTargetRules Target, UEToolChain ToolChain, CppCompileEnvironment BinaryCompileEnvironment, List<PrecompiledHeaderTemplate> SharedPCHs, ISourceFileWorkingSet WorkingSet, ActionGraph ActionGraph)
+		public override List<FileItem> Compile(ReadOnlyTargetRules Target, UEToolChain ToolChain, CppCompileEnvironment BinaryCompileEnvironment, List<PrecompiledHeaderTemplate> SharedPCHs, ISourceFileWorkingSet WorkingSet, BuildPredicateStore Predicates, ActionGraph ActionGraph)
 		{
 			UEBuildPlatform BuildPlatform = UEBuildPlatform.GetBuildPlatformForCPPTargetPlatform(BinaryCompileEnvironment.Platform);
 
@@ -330,8 +314,19 @@ namespace UnrealBuildTool
 				return LinkInputFiles;
 			}
 
+			List<FileReference> SourceFilePaths = SourceFileSearch.FindModuleSourceFiles(ModuleRulesFile: RulesFile, SearchedDirectories: Predicates.SourceDirectories);
+			List<FileItem> SourceFiles = GetCPlusPlusFilesToBuild(SourceFilePaths, ModuleDirectory, Target.Platform);
+
+			SourceFilesClass SourceFilesFound = new SourceFilesClass();
+			CategorizeSourceFiles(SourceFiles, SourceFilesFound);
+
+			SourceFilesClass SourceFilesToBuild = new SourceFilesClass();
+			SourceFilesToBuild.CopyFrom(SourceFilesFound);
+
+			SourceDirectories = new HashSet<DirectoryReference>(SourceFilePaths.Select(x => x.Directory));
+
 			// Process all of the header file dependencies for this module
-			CheckFirstIncludeMatchesEachCppFile(Target, ModuleCompileEnvironment);
+			CheckFirstIncludeMatchesEachCppFile(Target, ModuleCompileEnvironment, SourceFilesToBuild.CPPFiles);
 
 			// Make sure our RC files have cached includes.  
 			foreach (FileItem RCFile in SourceFilesToBuild.RCFiles)
@@ -477,7 +472,7 @@ namespace UnrealBuildTool
 			List<FileItem> CPPFilesToCompile = SourceFilesToBuild.CPPFiles;
 			if (bModuleUsesUnityBuild)
 			{
-				CPPFilesToCompile = Unity.GenerateUnityCPPs(Target, CPPFilesToCompile, CompileEnvironment, WorkingSet, Rules.ShortName ?? Name, IntermediateDirectory);
+				CPPFilesToCompile = Unity.GenerateUnityCPPs(Target, CPPFilesToCompile, CompileEnvironment, WorkingSet, Predicates, Rules.ShortName ?? Name, IntermediateDirectory);
 				LinkInputFiles.AddRange(CompileUnityFilesWithToolChain(Target, ToolChain, CompileEnvironment, ModuleCompileEnvironment, CPPFilesToCompile, ActionGraph).ObjectFiles);
 			}
 			else
@@ -519,7 +514,7 @@ namespace UnrealBuildTool
 
 					if (bModuleUsesUnityBuild)
 					{
-						GeneratedFileItems = Unity.GenerateUnityCPPs(Target, GeneratedFileItems, GeneratedCPPCompileEnvironment, WorkingSet, (Rules.ShortName ?? Name) + ".gen", IntermediateDirectory);
+						GeneratedFileItems = Unity.GenerateUnityCPPs(Target, GeneratedFileItems, GeneratedCPPCompileEnvironment, WorkingSet, Predicates, (Rules.ShortName ?? Name) + ".gen", IntermediateDirectory);
 						LinkInputFiles.AddRange(CompileUnityFilesWithToolChain(Target, ToolChain, GeneratedCPPCompileEnvironment, ModuleCompileEnvironment, GeneratedFileItems, ActionGraph).ObjectFiles);
 					}
 					else
@@ -934,7 +929,8 @@ namespace UnrealBuildTool
 		/// </summary>
 		/// <param name="Target">The target being compiled</param>
 		/// <param name="ModuleCompileEnvironment">Compile environment for the module</param>
-		private void CheckFirstIncludeMatchesEachCppFile(ReadOnlyTargetRules Target, CppCompileEnvironment ModuleCompileEnvironment)
+		/// <param name="CppFiles">List of C++ source files</param>
+		private void CheckFirstIncludeMatchesEachCppFile(ReadOnlyTargetRules Target, CppCompileEnvironment ModuleCompileEnvironment, List<FileItem> CppFiles)
 		{
 			if(Rules.PCHUsage == ModuleRules.PCHUsageMode.UseExplicitOrSharedPCHs)
 			{
@@ -954,7 +950,7 @@ namespace UnrealBuildTool
 					}
 
 					// Store the module compile environment along with the .cpp file.  This is so that we can use it later on when looking for header dependencies
-					foreach (FileItem CPPFile in SourceFilesFound.CPPFiles)
+					foreach (FileItem CPPFile in CppFiles)
 					{
 						CPPFile.CachedIncludePaths = ModuleCompileEnvironment.IncludePaths;
 					}
@@ -963,7 +959,7 @@ namespace UnrealBuildTool
 					InvalidIncludeDirectiveMessages = new List<string>();
 					if (Rules != null && Rules.bEnforceIWYU && Target.bEnforceIWYU)
 					{
-						foreach (FileItem CPPFile in SourceFilesFound.CPPFiles)
+						foreach (FileItem CPPFile in CppFiles)
 						{
 							List<DependencyInclude> DirectIncludeFilenames = ModuleCompileEnvironment.Headers.GetDirectIncludeDependencies(CPPFile, bOnlyCachedDependencies: false);
 							if (DirectIncludeFilenames.Count > 0)
@@ -1173,5 +1169,102 @@ namespace UnrealBuildTool
 				}
 			}
 		}
+
+		/// <summary>
+		/// Given a list of source files for a module, filters them into a list of files that should actually be included in a build
+		/// </summary>
+		/// <param name="SourceFiles">Original list of files, which may contain non-source</param>
+		/// <param name="SourceFilesBaseDirectory">Directory that the source files are in</param>
+		/// <param name="TargetPlatform">The platform we're going to compile for</param>
+		/// <returns>The list of source files to actually compile</returns>
+		static List<FileItem> GetCPlusPlusFilesToBuild(List<FileReference> SourceFiles, DirectoryReference SourceFilesBaseDirectory, UnrealTargetPlatform TargetPlatform)
+		{
+			// Make a list of all platform name strings that we're *not* currently compiling, to speed
+			// up file path comparisons later on
+			List<UnrealTargetPlatform> SupportedPlatforms = new List<UnrealTargetPlatform>();
+			SupportedPlatforms.Add(TargetPlatform);
+			List<string> OtherPlatformNameStrings = Utils.MakeListOfUnsupportedPlatforms(SupportedPlatforms);
+
+
+			// @todo projectfiles: Consider saving out cached list of source files for modules so we don't need to harvest these each time
+
+			List<FileItem> FilteredFileItems = new List<FileItem>();
+			FilteredFileItems.Capacity = SourceFiles.Count;
+
+			// @todo projectfiles: hard-coded source file set.  Should be made extensible by platform tool chains.
+			string[] CompilableSourceFileTypes = new string[]
+				{
+					".cpp",
+					".c",
+					".cc",
+					".mm",
+					".m",
+					".rc",
+					".manifest"
+				};
+
+			// When generating project files, we have no file to extract source from, so we'll locate the code files manually
+			foreach (FileReference SourceFilePath in SourceFiles)
+			{
+				// We're only able to compile certain types of files
+				bool IsCompilableSourceFile = false;
+				foreach (string CurExtension in CompilableSourceFileTypes)
+				{
+					if (SourceFilePath.HasExtension(CurExtension))
+					{
+						IsCompilableSourceFile = true;
+						break;
+					}
+				}
+
+				if (IsCompilableSourceFile)
+				{
+					if (SourceFilePath.IsUnderDirectory(SourceFilesBaseDirectory))
+					{
+						// Store the path as relative to the project file
+						string RelativeFilePath = SourceFilePath.MakeRelativeTo(SourceFilesBaseDirectory);
+
+						// All compiled files should always be in a sub-directory under the project file directory.  We enforce this here.
+						if (Path.IsPathRooted(RelativeFilePath) || RelativeFilePath.StartsWith(".."))
+						{
+							throw new BuildException("Error: Found source file {0} in project whose path was not relative to the base directory of the source files", RelativeFilePath);
+						}
+
+						// Check for source files that don't belong to the platform we're currently compiling.  We'll filter
+						// those source files out
+						bool IncludeThisFile = true;
+						foreach (string CurPlatformName in OtherPlatformNameStrings)
+						{
+							if (RelativeFilePath.IndexOf(Path.DirectorySeparatorChar + CurPlatformName + Path.DirectorySeparatorChar, StringComparison.InvariantCultureIgnoreCase) != -1
+								|| RelativeFilePath.StartsWith(CurPlatformName + Path.DirectorySeparatorChar))
+							{
+								IncludeThisFile = false;
+								break;
+							}
+						}
+
+						if (IncludeThisFile)
+						{
+							FilteredFileItems.Add(FileItem.GetItemByFileReference(SourceFilePath));
+						}
+					}
+				}
+			}
+
+			// @todo projectfiles: Consider enabling this error but changing it to a warning instead.  It can fire for
+			//    projects that are being digested for IntelliSense (because the module was set as a cross-
+			//	  platform dependency), but all of their source files were filtered out due to platform masking
+			//    in the project generator
+			bool AllowEmptyProjects = true;
+			if (!AllowEmptyProjects)
+			{
+				if (FilteredFileItems.Count == 0)
+				{
+					throw new BuildException("Could not find any valid source files for base directory {0}.  Project has {1} files in it", SourceFilesBaseDirectory, SourceFiles.Count);
+				}
+			}
+
+			return FilteredFileItems;
+		}
 	}
 }
diff --git a/Engine/Source/Programs/UnrealBuildTool/Configuration/UEBuildModuleExternal.cs b/Engine/Source/Programs/UnrealBuildTool/Configuration/UEBuildModuleExternal.cs
--- a/Engine/Source/Programs/UnrealBuildTool/Configuration/UEBuildModuleExternal.cs
+++ b/Engine/Source/Programs/UnrealBuildTool/Configuration/UEBuildModuleExternal.cs
@@ -20,7 +20,7 @@ namespace UnrealBuildTool
 		}
 
 		// UEBuildModule interface.
-		public override List<FileItem> Compile(ReadOnlyTargetRules Target, UEToolChain ToolChain, CppCompileEnvironment CompileEnvironment, List<PrecompiledHeaderTemplate> SharedPCHs, ISourceFileWorkingSet WorkingSet, ActionGraph ActionGraph)
+		public override List<FileItem> Compile(ReadOnlyTargetRules Target, UEToolChain ToolChain, CppCompileEnvironment CompileEnvironment, List<PrecompiledHeaderTemplate> SharedPCHs, ISourceFileWorkingSet WorkingSet, BuildPredicateStore Predicates, ActionGraph ActionGraph)
 		{
 			return new List<FileItem>();
 		}
diff --git a/Engine/Source/Programs/UnrealBuildTool/Configuration/UEBuildTarget.cs b/Engine/Source/Programs/UnrealBuildTool/Configuration/UEBuildTarget.cs
--- a/Engine/Source/Programs/UnrealBuildTool/Configuration/UEBuildTarget.cs
+++ b/Engine/Source/Programs/UnrealBuildTool/Configuration/UEBuildTarget.cs
@@ -734,11 +734,6 @@ namespace UnrealBuildTool
 		/// </summary>
 		public bool bDeployAfterCompile;
 
-		/// <summary>
-		/// Directories which are scanned for source files. The UBT makefile uses this to check for files being added/removed.
-		/// </summary>
-		public HashSet<DirectoryReference> SourceDirectories = new HashSet<DirectoryReference>();
-
 		/// <summary>
 		/// A list of the module filenames which were used to build this target.
 		/// </summary>
@@ -782,7 +777,6 @@ namespace UnrealBuildTool
 			PostBuildStepScripts = (FileReference[])Info.GetValue("po", typeof(FileReference[]));
 			bDeployAfterCompile = Info.GetBoolean("dt");
 			bHasProjectScriptPlugin = Info.GetBoolean("sp");
-			SourceDirectories = (HashSet<DirectoryReference>)Info.GetValue("sd", typeof(HashSet<DirectoryReference>));
 		}
 
 		public void GetObjectData(SerializationInfo Info, StreamingContext Context)
@@ -810,7 +804,6 @@ namespace UnrealBuildTool
 			Info.AddValue("po", PostBuildStepScripts);
 			Info.AddValue("dt", bDeployAfterCompile);
 			Info.AddValue("sp", bHasProjectScriptPlugin);
-			Info.AddValue("sd", SourceDirectories);
 		}
 
 		/// <summary>
@@ -1479,7 +1472,7 @@ namespace UnrealBuildTool
 		/// <summary>
 		/// Builds the target, appending list of output files and returns building result.
 		/// </summary>
-		public ECompilationResult Build(BuildConfiguration BuildConfiguration, CPPHeaders Headers, List<FileItem> OutputItems, Dictionary<string, FileItem[]> ModuleNameToOutputItems, List<UHTModuleInfo> UObjectModules, ISourceFileWorkingSet WorkingSet, ActionGraph ActionGraph, bool bIsAssemblingBuild)
+		public ECompilationResult Build(BuildConfiguration BuildConfiguration, CPPHeaders Headers, List<FileItem> OutputItems, Dictionary<string, FileItem[]> ModuleNameToOutputItems, List<UHTModuleInfo> UObjectModules, ISourceFileWorkingSet WorkingSet, ActionGraph ActionGraph, BuildPredicateStore Predicates, bool bIsAssemblingBuild)
 		{
 			CppPlatform CppPlatform = UEBuildPlatform.GetBuildPlatform(Platform).DefaultCppPlatform;
 			CppConfiguration CppConfiguration = GetCppConfiguration(Configuration);
@@ -1517,32 +1510,12 @@ namespace UnrealBuildTool
 			}
 
 			// If we're just compiling a single file, filter the list of binaries to only include the file we're interested in.
-			if (!String.IsNullOrEmpty(BuildConfiguration.SingleFileToCompile))
+			if (BuildConfiguration.SingleFileToCompile != null)
 			{
-				FileItem SingleFileItem = FileItem.GetItemByPath(BuildConfiguration.SingleFileToCompile);
-
-				HashSet<UEBuildModuleCPP> Dependencies = GatherDependencyModules(Binaries);
-
-				// We only want to build the binaries for this single file
-				List<UEBuildBinary> FilteredBinaries = new List<UEBuildBinary>();
-				foreach (UEBuildModuleCPP Dependency in Dependencies)
-				{
-					bool bFileExistsInDependency = Dependency.SourceFilesFound.CPPFiles.Exists(x => x.AbsolutePath == SingleFileItem.AbsolutePath);
-					if (bFileExistsInDependency)
-					{
-						FilteredBinaries.Add(Dependency.Binary);
-
-						UEBuildModuleCPP.SourceFilesClass EmptySourceFileList = new UEBuildModuleCPP.SourceFilesClass();
-						Dependency.SourceFilesToBuild.CopyFrom(EmptySourceFileList);
-						Dependency.SourceFilesToBuild.CPPFiles.Add(SingleFileItem);
-					}
-				}
-				Binaries = FilteredBinaries;
-
-				// Check we have at least one match
+				Binaries.RemoveAll(x => !x.Modules.Any(y => BuildConfiguration.SingleFileToCompile.IsUnderDirectory(y.ModuleDirectory)));
 				if(Binaries.Count == 0)
 				{
-					throw new BuildException("Couldn't find any module containing {0} in {1}.", SingleFileItem.Location, TargetName);
+					throw new BuildException("Couldn't find any module containing {0} in {1}.", BuildConfiguration.SingleFileToCompile, TargetName);
 				}
 			}
 
@@ -1757,7 +1730,7 @@ namespace UnrealBuildTool
 			DirectoryReference ExeDir = OutputPaths[0].Directory;
 			foreach (UEBuildBinary Binary in Binaries)
 			{
-				List<FileItem> BinaryOutputItems = Binary.Build(Rules, TargetToolChain, GlobalCompileEnvironment, GlobalLinkEnvironment, SharedPCHs, WorkingSet, ExeDir, ActionGraph);
+				List<FileItem> BinaryOutputItems = Binary.Build(Rules, TargetToolChain, GlobalCompileEnvironment, GlobalLinkEnvironment, SharedPCHs, WorkingSet, Predicates, ExeDir, ActionGraph);
 				if(!bCompileMonolithic)
 				{
 					ModuleNameToOutputItems[Binary.PrimaryModule.Name] = BinaryOutputItems.ToArray();
@@ -3729,23 +3702,6 @@ namespace UnrealBuildTool
 					}
 				}
 
-				// Figure out whether we need to build this module
-				// We don't care about actual source files when generating projects, as these are discovered separately
-				bool bDiscoverFiles = !ProjectFileGenerator.bGenerateProjectFiles;
-				bool bBuildFiles = bDiscoverFiles;
-
-				List<FileItem> FoundSourceFiles = new List<FileItem>();
-				if (RulesObject.Type == ModuleRules.ModuleType.CPlusPlus)
-				{
-					// So all we care about are the game module and/or plugins.
-					if (bDiscoverFiles && !RulesObject.bUsePrecompiled)
-					{
-						List<FileReference> SourceFilePaths = new List<FileReference>();
-						SourceFilePaths = SourceFileSearch.FindModuleSourceFiles(ModuleRulesFile: RulesObject.File, SearchedDirectories: SourceDirectories);
-						FoundSourceFiles = GetCPlusPlusFilesToBuild(SourceFilePaths, ModuleDirectory, Platform);
-					}
-				}
-
 				// Allow the current platform to modify the module rules
 				UEBuildPlatform.GetBuildPlatform(Platform).ModifyModuleRulesForActivePlatform(ModuleName, RulesObject, Rules);
 
@@ -3756,7 +3712,7 @@ namespace UnrealBuildTool
 				UEBuildPlatform.PlatformModifyHostModuleRules(ModuleName, RulesObject, Rules);
 
 				// Now, go ahead and create the module builder instance
-				Module = InstantiateModule(RulesObject, GeneratedCodeDirectory, FoundSourceFiles, bBuildFiles);
+				Module = InstantiateModule(RulesObject, GeneratedCodeDirectory);
 				Modules.Add(Module.Name, Module);
 				FlatModuleCsData.Add(new FlatModuleCsDataType(Module.Name, (Module.RulesFile == null) ? null : Module.RulesFile.FullName, RulesObject.ExternalDependencies));
 			}
@@ -3781,9 +3737,7 @@ namespace UnrealBuildTool
 
 		protected UEBuildModule InstantiateModule(
 			ModuleRules RulesObject,
-			DirectoryReference GeneratedCodeDirectory,
-			List<FileItem> ModuleSourceFiles,
-			bool bBuildSourceFiles)
+			DirectoryReference GeneratedCodeDirectory)
 		{
 			switch (RulesObject.Type)
 			{
@@ -3791,9 +3745,7 @@ namespace UnrealBuildTool
 					return new UEBuildModuleCPP(
 							Rules: RulesObject,
 							IntermediateDirectory: GetModuleIntermediateDirectory(RulesObject),
-							GeneratedCodeDirectory: GeneratedCodeDirectory,
-							SourceFiles: ModuleSourceFiles,
-							bBuildSourceFiles: bBuildSourceFiles
+							GeneratedCodeDirectory: GeneratedCodeDirectory
 						);
 
 				case ModuleRules.ModuleType.External:
@@ -3844,103 +3796,5 @@ namespace UnrealBuildTool
 			}
 			return NewPathList;
 		}
-
-
-		/// <summary>
-		/// Given a list of source files for a module, filters them into a list of files that should actually be included in a build
-		/// </summary>
-		/// <param name="SourceFiles">Original list of files, which may contain non-source</param>
-		/// <param name="SourceFilesBaseDirectory">Directory that the source files are in</param>
-		/// <param name="TargetPlatform">The platform we're going to compile for</param>
-		/// <returns>The list of source files to actually compile</returns>
-		static List<FileItem> GetCPlusPlusFilesToBuild(List<FileReference> SourceFiles, DirectoryReference SourceFilesBaseDirectory, UnrealTargetPlatform TargetPlatform)
-		{
-			// Make a list of all platform name strings that we're *not* currently compiling, to speed
-			// up file path comparisons later on
-			List<UnrealTargetPlatform> SupportedPlatforms = new List<UnrealTargetPlatform>();
-			SupportedPlatforms.Add(TargetPlatform);
-			List<string> OtherPlatformNameStrings = Utils.MakeListOfUnsupportedPlatforms(SupportedPlatforms);
-
-
-			// @todo projectfiles: Consider saving out cached list of source files for modules so we don't need to harvest these each time
-
-			List<FileItem> FilteredFileItems = new List<FileItem>();
-			FilteredFileItems.Capacity = SourceFiles.Count;
-
-			// @todo projectfiles: hard-coded source file set.  Should be made extensible by platform tool chains.
-			string[] CompilableSourceFileTypes = new string[]
-				{
-					".cpp",
-					".c",
-					".cc",
-					".mm",
-					".m",
-					".rc",
-					".manifest"
-				};
-
-			// When generating project files, we have no file to extract source from, so we'll locate the code files manually
-			foreach (FileReference SourceFilePath in SourceFiles)
-			{
-				// We're only able to compile certain types of files
-				bool IsCompilableSourceFile = false;
-				foreach (string CurExtension in CompilableSourceFileTypes)
-				{
-					if (SourceFilePath.HasExtension(CurExtension))
-					{
-						IsCompilableSourceFile = true;
-						break;
-					}
-				}
-
-				if (IsCompilableSourceFile)
-				{
-					if (SourceFilePath.IsUnderDirectory(SourceFilesBaseDirectory))
-					{
-						// Store the path as relative to the project file
-						string RelativeFilePath = SourceFilePath.MakeRelativeTo(SourceFilesBaseDirectory);
-
-						// All compiled files should always be in a sub-directory under the project file directory.  We enforce this here.
-						if (Path.IsPathRooted(RelativeFilePath) || RelativeFilePath.StartsWith(".."))
-						{
-							throw new BuildException("Error: Found source file {0} in project whose path was not relative to the base directory of the source files", RelativeFilePath);
-						}
-
-						// Check for source files that don't belong to the platform we're currently compiling.  We'll filter
-						// those source files out
-						bool IncludeThisFile = true;
-						foreach (string CurPlatformName in OtherPlatformNameStrings)
-						{
-							if (RelativeFilePath.IndexOf(Path.DirectorySeparatorChar + CurPlatformName + Path.DirectorySeparatorChar, StringComparison.InvariantCultureIgnoreCase) != -1
-								|| RelativeFilePath.StartsWith(CurPlatformName + Path.DirectorySeparatorChar))
-							{
-								IncludeThisFile = false;
-								break;
-							}
-						}
-
-						if (IncludeThisFile)
-						{
-							FilteredFileItems.Add(FileItem.GetItemByFileReference(SourceFilePath));
-						}
-					}
-				}
-			}
-
-			// @todo projectfiles: Consider enabling this error but changing it to a warning instead.  It can fire for
-			//    projects that are being digested for IntelliSense (because the module was set as a cross-
-			//	  platform dependency), but all of their source files were filtered out due to platform masking
-			//    in the project generator
-			bool AllowEmptyProjects = true;
-			if (!AllowEmptyProjects)
-			{
-				if (FilteredFileItems.Count == 0)
-				{
-					throw new BuildException("Could not find any valid source files for base directory {0}.  Project has {1} files in it", SourceFilesBaseDirectory, SourceFiles.Count);
-				}
-			}
-
-			return FilteredFileItems;
-		}
 	}
 }
diff --git a/Engine/Source/Programs/UnrealBuildTool/System/ActionGraph.cs b/Engine/Source/Programs/UnrealBuildTool/System/ActionGraph.cs
--- a/Engine/Source/Programs/UnrealBuildTool/System/ActionGraph.cs
+++ b/Engine/Source/Programs/UnrealBuildTool/System/ActionGraph.cs
@@ -337,13 +337,13 @@ namespace UnrealBuildTool
 			}
 
 			// If we're only compiling a single file, we should always compile and should never link.
-			if (!string.IsNullOrEmpty(BuildConfiguration.SingleFileToCompile))
+			if (BuildConfiguration.SingleFileToCompile != null)
 			{
 				// Never do anything but compile the target file
 				AllActions.RemoveAll(x => x.ActionType != ActionType.Compile);
 
 				// Check all of the leftover compilation actions for the one we want... that one is always outdated.
-				FileItem SingleFileToCompile = FileItem.GetExistingItemByPath(BuildConfiguration.SingleFileToCompile);
+				FileItem SingleFileToCompile = FileItem.GetItemByFileReference(BuildConfiguration.SingleFileToCompile);
 				foreach (Action Action in AllActions)
 				{
 					bool bIsSingleFileAction = Action.PrerequisiteItems.Contains(SingleFileToCompile);
diff --git a/Engine/Source/Programs/UnrealBuildTool/System/BuildPredicateStore.cs b/Engine/Source/Programs/UnrealBuildTool/System/BuildPredicateStore.cs
new file mode 100644
--- /dev/null
+++ b/Engine/Source/Programs/UnrealBuildTool/System/BuildPredicateStore.cs
@@ -0,0 +1,31 @@
+using System;
+using System.Collections.Generic;
+using System.Linq;
+using System.Text;
+using System.Threading.Tasks;
+using Tools.DotNETCommon;
+
+namespace UnrealBuildTool
+{
+	/// <summary>
+	/// Contains predicates required for the action graph generated for a build to be valid. Unlike individual actions within the graph, these conditions must hold for the graph itself to be valid.
+	/// </summary>
+	[Serializable]
+	class BuildPredicateStore
+	{
+		/// <summary>
+		/// Set of all source directories. Any files being added or removed from these directories will invalidate the makefile.
+		/// </summary>
+		public HashSet<DirectoryReference> SourceDirectories = new HashSet<DirectoryReference>();
+
+		/// <summary>
+		/// The set of source files that UnrealBuildTool determined to be part of the programmer's "working set". Used for adaptive non-unity builds.
+		/// </summary>
+		public HashSet<FileItem> WorkingSet = new HashSet<FileItem>();
+
+		/// <summary>
+		/// Set of files which are currently not part of the working set, but could be.
+		/// </summary>
+		public HashSet<FileItem> CandidatesForWorkingSet = new HashSet<FileItem>();
+	}
+}
diff --git a/Engine/Source/Programs/UnrealBuildTool/System/ExternalExecution.cs b/Engine/Source/Programs/UnrealBuildTool/System/ExternalExecution.cs
--- a/Engine/Source/Programs/UnrealBuildTool/System/ExternalExecution.cs
+++ b/Engine/Source/Programs/UnrealBuildTool/System/ExternalExecution.cs
@@ -617,7 +617,7 @@ namespace UnrealBuildTool
 					// If we've got this far and there are no source files then it's likely we're installed and ignoring
 					// engine files, so we don't need a .gen.cpp either
 					Info.GeneratedCPPFilenameBase = Path.Combine(Module.GeneratedCodeDirectory.FullName, Info.ModuleName) + ".gen";
-					if (Module.SourceFilesToBuild.Count != 0)
+					if (!Module.Rules.bUsePrecompiled)
 					{
 						Module.GeneratedCodeWildcard = Path.Combine(Module.GeneratedCodeDirectory.FullName, "*.gen.cpp");
 					}
diff --git a/Engine/Source/Programs/UnrealBuildTool/System/UBTMakefile.cs b/Engine/Source/Programs/UnrealBuildTool/System/UBTMakefile.cs
--- a/Engine/Source/Programs/UnrealBuildTool/System/UBTMakefile.cs
+++ b/Engine/Source/Programs/UnrealBuildTool/System/UBTMakefile.cs
@@ -70,14 +70,9 @@ namespace UnrealBuildTool
 		public bool bUseAdaptiveUnityBuild;
 
 		/// <summary>
-		/// Current working set of source files, for when bUseAdaptiveUnityBuild is enabled
+		/// List of predicates for the action graph to be valid.
 		/// </summary>
-		public HashSet<FileItem> SourceFileWorkingSet = new HashSet<FileItem>();
-
-		/// <summary>
-		/// Set of source files which are included in unity files, but which should invalidate the makefile if modified (for when bUseAdaptiveUnityBuild is enabled)
-		/// </summary>
-		public HashSet<FileItem> CandidateSourceFilesForWorkingSet = new HashSet<FileItem>();
+		public List<BuildPredicateStore> TargetBuildPredicates;
 
 		public UBTMakefile()
 		{
@@ -99,9 +94,8 @@ namespace UnrealBuildTool
 			ModuleNameToOutputItems = (Dictionary<string, FileItem[]>)Info.GetValue("mn", typeof(Dictionary<string, FileItem[]>));
 			HotReloadModuleNamesForAllTargets = (HashSet<string>)Info.GetValue("gm", typeof(HashSet<string>));
 			Targets = (List<UEBuildTarget>)Info.GetValue("ta", typeof(List<UEBuildTarget>));
+			TargetBuildPredicates = (List<BuildPredicateStore>)Info.GetValue("tp", typeof(List<BuildPredicateStore>));
 			bUseAdaptiveUnityBuild = Info.GetBoolean("ua");
-			SourceFileWorkingSet = (HashSet<FileItem>)Info.GetValue("ws", typeof(HashSet<FileItem>));
-			CandidateSourceFilesForWorkingSet = (HashSet<FileItem>)Info.GetValue("wc", typeof(HashSet<FileItem>));
 		}
 
 		public void GetObjectData(SerializationInfo Info, StreamingContext Context)
@@ -115,9 +109,8 @@ namespace UnrealBuildTool
 			Info.AddValue("mn", ModuleNameToOutputItems);
 			Info.AddValue("gm", HotReloadModuleNamesForAllTargets);
 			Info.AddValue("ta", Targets);
+			Info.AddValue("tp", TargetBuildPredicates);
 			Info.AddValue("ua", bUseAdaptiveUnityBuild);
-			Info.AddValue("ws", SourceFileWorkingSet);
-			Info.AddValue("wc", CandidateSourceFilesForWorkingSet);
 		}
 
 
@@ -132,8 +125,7 @@ namespace UnrealBuildTool
 				ModuleNameToOutputItems != null && 
 				HotReloadModuleNamesForAllTargets != null &&
 				Targets != null && Targets.Count > 0 &&
-				SourceFileWorkingSet != null &&
-				CandidateSourceFilesForWorkingSet != null;
+				TargetBuildPredicates != null;
 		}
 
 
@@ -362,8 +354,11 @@ namespace UnrealBuildTool
 						}
 					}
 				}
+			}
 
-				foreach(DirectoryReference SourceDirectory in Target.SourceDirectories)
+			foreach(BuildPredicateStore Predicates in LoadedUBTMakefile.TargetBuildPredicates)
+			{
+				foreach(DirectoryReference SourceDirectory in Predicates.SourceDirectories)
 				{
 					DirectoryInfo SourceDirectoryInfo = new DirectoryInfo(SourceDirectory.FullName);
 					if(!SourceDirectoryInfo.Exists || SourceDirectoryInfo.LastWriteTimeUtc > UBTMakefileInfo.LastWriteTimeUtc)
@@ -427,25 +422,28 @@ namespace UnrealBuildTool
 			// iteration times.)
 			if (LoadedUBTMakefile.bUseAdaptiveUnityBuild)
 			{
-				// Check if any source files in the working set no longer belong in it
-				foreach (FileItem SourceFile in LoadedUBTMakefile.SourceFileWorkingSet)
+				foreach(BuildPredicateStore Predicates in LoadedUBTMakefile.TargetBuildPredicates)
 				{
-					if (!WorkingSet.Contains(SourceFile.Location) && File.GetLastWriteTimeUtc(SourceFile.AbsolutePath) > UBTMakefileInfo.LastWriteTimeUtc)
+					// Check if any source files in the working set no longer belong in it
+					foreach (FileItem SourceFile in Predicates.WorkingSet)
 					{
-						Log.TraceLog("{0} was part of source working set and now is not; invalidating makefile ({1})", SourceFile.AbsolutePath, UBTMakefileInfo.FullName);
-						ReasonNotLoaded = string.Format("working set of source files changed");
-						return null;
+						if (!WorkingSet.Contains(SourceFile.Location) && File.GetLastWriteTimeUtc(SourceFile.AbsolutePath) > UBTMakefileInfo.LastWriteTimeUtc)
+						{
+							Log.TraceLog("{0} was part of source working set and now is not; invalidating makefile ({1})", SourceFile.AbsolutePath, UBTMakefileInfo.FullName);
+							ReasonNotLoaded = string.Format("working set of source files changed");
+							return null;
+						}
 					}
-				}
 
-				// Check if any source files that are eligible for being in the working set have been modified
-				foreach (FileItem SourceFile in LoadedUBTMakefile.CandidateSourceFilesForWorkingSet)
-				{
-					if (WorkingSet.Contains(SourceFile.Location) && File.GetLastWriteTimeUtc(SourceFile.AbsolutePath) > UBTMakefileInfo.LastWriteTimeUtc)
+					// Check if any source files that are eligible for being in the working set have been modified
+					foreach (FileItem SourceFile in Predicates.CandidatesForWorkingSet)
 					{
-						Log.TraceLog("{0} was part of source working set and now is not; invalidating makefile ({1})", SourceFile.AbsolutePath, UBTMakefileInfo.FullName);
-						ReasonNotLoaded = string.Format("working set of source files changed");
-						return null;
+						if (WorkingSet.Contains(SourceFile.Location) && File.GetLastWriteTimeUtc(SourceFile.AbsolutePath) > UBTMakefileInfo.LastWriteTimeUtc)
+						{
+							Log.TraceLog("{0} was part of source working set and now is not; invalidating makefile ({1})", SourceFile.AbsolutePath, UBTMakefileInfo.FullName);
+							ReasonNotLoaded = string.Format("working set of source files changed");
+							return null;
+						}
 					}
 				}
 			}
diff --git a/Engine/Source/Programs/UnrealBuildTool/System/Unity.cs b/Engine/Source/Programs/UnrealBuildTool/System/Unity.cs
--- a/Engine/Source/Programs/UnrealBuildTool/System/Unity.cs
+++ b/Engine/Source/Programs/UnrealBuildTool/System/Unity.cs
@@ -12,17 +12,6 @@ namespace UnrealBuildTool
 {
 	class Unity
 	{
-		/// <summary>
-		/// The set of source files that UnrealBuildTool determined to be part of the programmer's "working set".
-		/// This feature is only used when bUseAdaptiveUnityBuild is enabled
-		/// </summary>
-		public static HashSet<FileItem> SourceFileWorkingSet = new HashSet<FileItem>();
-
-		/// <summary>
-		/// Set of files which are candidates for being part of the working set.
-		/// </summary>
-		public static HashSet<FileItem> CandidateSourceFilesForWorkingSet = new HashSet<FileItem>();
-
 		/// <summary>
 		/// Set of target names we've printed out adaptive non-unity messages for.
 		/// </summary>
@@ -163,6 +152,7 @@ namespace UnrealBuildTool
 		/// <param name="CPPFiles">The C++ files to #include.</param>
 		/// <param name="CompileEnvironment">The environment that is used to compile the C++ files.</param>
 		/// <param name="WorkingSet">Interface to query files which belong to the working set</param>
+		/// <param name="Predicates">Dependencies for this target</param>
 		/// <param name="BaseName">Base name to use for the Unity files</param>
 		/// <param name="IntermediateDirectory">Intermediate directory for unity cpp files</param>
 		/// <returns>The "unity" C++ files.</returns>
@@ -171,6 +161,7 @@ namespace UnrealBuildTool
 			List<FileItem> CPPFiles,
 			CppCompileEnvironment CompileEnvironment,
 			ISourceFileWorkingSet WorkingSet,
+			BuildPredicateStore Predicates,
 			string BaseName,
 			DirectoryReference IntermediateDirectory
 			)
@@ -227,7 +218,7 @@ namespace UnrealBuildTool
 							// Mark this file as part of the working set.  This will be saved into the UBT Makefile so that
 							// the assembler can automatically invalidate the Makefile when the working set changes (allowing this
 							// code to run again, to build up new unity blobs.)
-							SourceFileWorkingSet.Add(CPPFile);
+							Predicates.WorkingSet.Add(CPPFile);
 						}
 					}
 
@@ -249,7 +240,7 @@ namespace UnrealBuildTool
 					}
 
 					// When adaptive unity is enabled, go ahead and exclude any source files that we're actively working with
-					if (bUseAdaptiveUnityBuild && SourceFileWorkingSet.Contains(CPPFile))
+					if (bUseAdaptiveUnityBuild && Predicates.WorkingSet.Contains(CPPFile))
 					{
 						// Just compile this file normally, not as part of the unity blob
 						NewCPPFiles.Add(CPPFile);
@@ -275,7 +266,7 @@ namespace UnrealBuildTool
 						// If adaptive unity build is enabled for this module, add this source file to the set that will invalidate the makefile
 						if(bUseAdaptiveUnityBuild)
 						{
-							CandidateSourceFilesForWorkingSet.Add(CPPFile);
+							Predicates.CandidatesForWorkingSet.Add(CPPFile);
 						}
 
 						// Compile this file as part of the unity blob
diff --git a/Engine/Source/Programs/UnrealBuildTool/UnrealBuildTool.cs b/Engine/Source/Programs/UnrealBuildTool/UnrealBuildTool.cs
--- a/Engine/Source/Programs/UnrealBuildTool/UnrealBuildTool.cs
+++ b/Engine/Source/Programs/UnrealBuildTool/UnrealBuildTool.cs
@@ -360,7 +360,7 @@ namespace UnrealBuildTool
 						else if (LowercaseArg.StartsWith("-singlefile="))
 						{
 							BuildConfiguration.bUseUBTMakefiles = false;
-							BuildConfiguration.SingleFileToCompile = LowercaseArg.Replace("-singlefile=", "");
+							BuildConfiguration.SingleFileToCompile = new FileReference(LowercaseArg.Replace("-singlefile=", ""));
 						}
 					}
 
@@ -904,6 +904,7 @@ namespace UnrealBuildTool
 					Dictionary<string, FileItem[]> ModuleNameToOutputItems = new Dictionary<string, FileItem[]>(StringComparer.OrdinalIgnoreCase);
 					Dictionary<string, List<UHTModuleInfo>> TargetNameToUObjectModules = new Dictionary<string, List<UHTModuleInfo>>(StringComparer.InvariantCultureIgnoreCase);
 					HashSet<string> HotReloadModuleNamesForAllTargets = new HashSet<string>(StringComparer.OrdinalIgnoreCase);
+					List<BuildPredicateStore> TargetBuildPredicates = new List<BuildPredicateStore>();
 					foreach (UEBuildTarget Target in Targets)
 					{
 						// Create the header cache for this target
@@ -951,13 +952,14 @@ namespace UnrealBuildTool
 							List<FileItem> TargetOutputItems = new List<FileItem>();
 							List<UHTModuleInfo> TargetUObjectModules = new List<UHTModuleInfo>();
 							Dictionary<string, FileItem[]> TargetModuleNameToOutputItems = new Dictionary<string, FileItem[]>(StringComparer.OrdinalIgnoreCase);
+							BuildPredicateStore Predicates = new BuildPredicateStore();
 							if (BuildConfiguration.bCleanProject)
 							{
 								BuildResult = Target.Clean(!BuildConfiguration.bDoNotBuildUHT);
 							}
 							else
 							{
-								BuildResult = Target.Build(BuildConfiguration, TargetToHeaders[Target], TargetOutputItems, TargetModuleNameToOutputItems, TargetUObjectModules, WorkingSet, ActionGraph, bIsAssemblingBuild);
+								BuildResult = Target.Build(BuildConfiguration, TargetToHeaders[Target], TargetOutputItems, TargetModuleNameToOutputItems, TargetUObjectModules, WorkingSet, ActionGraph, Predicates, bIsAssemblingBuild);
 							}
 							if (BuildResult != ECompilationResult.Succeeded)
 							{
@@ -973,6 +975,8 @@ namespace UnrealBuildTool
 
 							OutputItemsForAllTargets.AddRange(TargetOutputItems);
 
+							TargetBuildPredicates.Add(Predicates);
+
 							// Update mapping of the target name to the list of UObject modules in this target
 							TargetNameToUObjectModules[Target.GetTargetName()] = TargetUObjectModules;
 
@@ -1017,8 +1021,7 @@ namespace UnrealBuildTool
 							UBTMakefile.HotReloadModuleNamesForAllTargets = HotReloadModuleNamesForAllTargets;
 							UBTMakefile.Targets = Targets;
 							UBTMakefile.bUseAdaptiveUnityBuild = Targets.Any(x => x.Rules.bUseAdaptiveUnityBuild);
-							UBTMakefile.SourceFileWorkingSet = Unity.SourceFileWorkingSet;
-							UBTMakefile.CandidateSourceFilesForWorkingSet = Unity.CandidateSourceFilesForWorkingSet;
+							UBTMakefile.TargetBuildPredicates = TargetBuildPredicates;
 
 							if (BuildConfiguration.bUseUBTMakefiles)
 							{
@@ -1383,10 +1386,10 @@ namespace UnrealBuildTool
 
 							// Run the deployment steps
 							if (BuildResult.Succeeded()
-								&& String.IsNullOrEmpty(BuildConfiguration.SingleFileToCompile)
+								&& BuildConfiguration.SingleFileToCompile == null
 								&& !BuildConfiguration.bGenerateManifest
 								&& !BuildConfiguration.bCleanProject
-							&& !BuildConfiguration.bXGEExport)
+								&& !BuildConfiguration.bXGEExport)
 							{
 								foreach (UEBuildTarget Target in Targets)
 								{
diff --git a/Engine/Source/Programs/UnrealBuildTool/UnrealBuildTool.csproj b/Engine/Source/Programs/UnrealBuildTool/UnrealBuildTool.csproj
--- a/Engine/Source/Programs/UnrealBuildTool/UnrealBuildTool.csproj
+++ b/Engine/Source/Programs/UnrealBuildTool/UnrealBuildTool.csproj
@@ -213,6 +213,7 @@
     <Compile Include="System\RulesAssembly.cs" />
     <Compile Include="System\ScopedTimer.cs" />
     <Compile Include="System\SingleInstanceMutex.cs" />
+    <Compile Include="System\BuildPredicateStore.cs" />
     <Compile Include="System\SourceFileWorkingSet.cs" />
     <Compile Include="System\StartupTraceListener.cs" />
     <Compile Include="System\SupportedConfigurationsAttribute.cs" />
