UBT: Fix exception checking for license violations when building engine targets.

#rnx

diff --git a/Engine/Source/Programs/UnrealBuildTool/Configuration/UEBuildModuleCPP.cs b/Engine/Source/Programs/UnrealBuildTool/Configuration/UEBuildModuleCPP.cs
--- a/Engine/Source/Programs/UnrealBuildTool/Configuration/UEBuildModuleCPP.cs
+++ b/Engine/Source/Programs/UnrealBuildTool/Configuration/UEBuildModuleCPP.cs
@@ -323,7 +323,7 @@ namespace UnrealBuildTool
 			SourceFilesClass SourceFilesToBuild = new SourceFilesClass();
 			SourceFilesToBuild.CopyFrom(SourceFilesFound);
 
-			SourceDirectories = new HashSet<DirectoryReference>(SourceFilePaths.Select(x => x.Directory));
+			SourceDirectories = new HashSet<DirectoryReference>(SourceFiles.Select(x => x.Location.Directory));
 
 			// Process all of the header file dependencies for this module
 			CheckFirstIncludeMatchesEachCppFile(Target, ModuleCompileEnvironment, SourceFilesToBuild.CPPFiles);
diff --git a/Engine/Source/Programs/UnrealBuildTool/Configuration/UEBuildTarget.cs b/Engine/Source/Programs/UnrealBuildTool/Configuration/UEBuildTarget.cs
--- a/Engine/Source/Programs/UnrealBuildTool/Configuration/UEBuildTarget.cs
+++ b/Engine/Source/Programs/UnrealBuildTool/Configuration/UEBuildTarget.cs
@@ -1519,29 +1519,6 @@ namespace UnrealBuildTool
 				}
 			}
 
-			if(!Rules.bDisableLinking)
-			{
-				// Check the distribution level of all binaries based on the dependencies they have
-				if(ProjectFile == null && !Rules.bOutputPubliclyDistributable)
-				{
-					Dictionary<UEBuildModule, Dictionary<RestrictedFolder, DirectoryReference>> ModuleRestrictedFolderCache = new Dictionary<UEBuildModule, Dictionary<RestrictedFolder, DirectoryReference>>();
-
-					bool bResult = true;
-					foreach (UEBuildBinary Binary in Binaries)
-					{
-						bResult &= Binary.CheckRestrictedFolders(DirectoryReference.FromFile(ProjectFile), ModuleRestrictedFolderCache);
-					}
-
-					if(!bResult)
-					{
-						throw new BuildException("Unable to create binaries in less restricted locations than their input files.");
-					}
-				}
-
-				// Check for linking against modules prohibited by the EULA
-				CheckForEULAViolation();
-			}
-
 			// Build a mapping from module to its plugin
 			Dictionary<UEBuildModule, UEBuildPlugin> ModuleToPlugin = new Dictionary<UEBuildModule, UEBuildPlugin>();
 			foreach(UEBuildPlugin Plugin in BuildPlugins)
@@ -1881,6 +1858,30 @@ namespace UnrealBuildTool
 				GenerateManifest(ManifestFileName, BuildProducts);
 			}
 
+			// Check there are no EULA or restricted folder violations
+			if(!Rules.bDisableLinking)
+			{
+				// Check the distribution level of all binaries based on the dependencies they have
+				if(ProjectFile == null && !Rules.bOutputPubliclyDistributable)
+				{
+					Dictionary<UEBuildModule, Dictionary<RestrictedFolder, DirectoryReference>> ModuleRestrictedFolderCache = new Dictionary<UEBuildModule, Dictionary<RestrictedFolder, DirectoryReference>>();
+
+					bool bResult = true;
+					foreach (UEBuildBinary Binary in Binaries)
+					{
+						bResult &= Binary.CheckRestrictedFolders(DirectoryReference.FromFile(ProjectFile), ModuleRestrictedFolderCache);
+					}
+
+					if(!bResult)
+					{
+						throw new BuildException("Unable to create binaries in less restricted locations than their input files.");
+					}
+				}
+
+				// Check for linking against modules prohibited by the EULA
+				CheckForEULAViolation();
+			}
+
 			// Clean any stale modules which exist in multiple output directories. This can lead to the wrong DLL being loaded on Windows.
 			CleanStaleModules();
 			return ECompilationResult.Succeeded;
