UBT: Move UnrealBuildTool.GuardedMain() into the BuildMode class.

#rnx

diff --git a/Engine/Source/Programs/UnrealBuildTool/Modes/BuildMode.cs b/Engine/Source/Programs/UnrealBuildTool/Modes/BuildMode.cs
--- a/Engine/Source/Programs/UnrealBuildTool/Modes/BuildMode.cs
+++ b/Engine/Source/Programs/UnrealBuildTool/Modes/BuildMode.cs
@@ -23,7 +23,166 @@ namespace UnrealBuildTool
 		/// <returns>One of the values of ECompilationResult</returns>
 		public override int Execute(CommandLineArguments CmdLine)
 		{
-			return UnrealBuildTool.GuardedMain(CmdLine);
+			DateTime StartTime = DateTime.UtcNow;
+			ECompilationResult Result = ECompilationResult.Succeeded;
+
+			// Initialize the log system, buffering the output until we can create the log file
+			StartupTraceListener StartupListener = new StartupTraceListener();
+			Trace.Listeners.Add(StartupListener);
+
+			// Write the command line
+			Log.TraceLog("Command line: {0}", Environment.CommandLine);
+
+			// Grab the environment.
+			UnrealBuildTool.InitialEnvironment = Environment.GetEnvironmentVariables();
+			if (UnrealBuildTool.InitialEnvironment.Count < 1)
+			{
+				throw new BuildException("Environment could not be read");
+			}
+
+			using(SingleInstanceMutex.Acquire(SingleInstanceMutexType.PerBranch, CmdLine))
+			{
+				string[] Arguments = CmdLine.GetRawArray();
+				try
+				{
+					// Change the working directory to be the Engine/Source folder. We are likely running from Engine/Binaries/DotNET
+					// This is critical to be done early so any code that relies on the current directory being Engine/Source will work.
+					DirectoryReference.SetCurrentDirectory(UnrealBuildTool.EngineSourceDirectory);
+
+					// Read the XML configuration files
+					FileReference XmlConfigCache = CmdLine.GetFileReferenceOrDefault("-XmlConfigCache=", null);
+					XmlConfig.ReadConfigFiles(XmlConfigCache);
+
+					// Create the build configuration object, and read the settings
+					BuildConfiguration BuildConfiguration = new BuildConfiguration();
+					XmlConfig.ApplyTo(BuildConfiguration);
+					CommandLine.ParseArguments(Arguments, BuildConfiguration);
+
+					// Copy some of the static settings that are being deprecated from BuildConfiguration
+					UnrealBuildTool.bPrintPerformanceInfo = BuildConfiguration.bPrintPerformanceInfo;
+
+					// Then let the command lines override any configs necessary.
+					string LogSuffix = "";
+					foreach (string Argument in Arguments)
+					{
+						string LowercaseArg = Argument.ToLowerInvariant();
+						if (LowercaseArg.StartsWith("-LogSuffix=", StringComparison.OrdinalIgnoreCase))
+						{
+							LogSuffix += "_" + Argument.Substring("-LogSuffix=".Length);
+						}
+						else if (LowercaseArg == "-nolog")
+						{
+							BuildConfiguration.LogFileName = null;
+						}
+						else if (LowercaseArg == "-xgeexport")
+						{
+							BuildConfiguration.bXGEExport = true;
+							BuildConfiguration.bAllowXGE = true;
+						}
+						else if (LowercaseArg.StartsWith("-singlefile="))
+						{
+							BuildConfiguration.bUseUBTMakefiles = false;
+							BuildConfiguration.SingleFileToCompile = new FileReference(LowercaseArg.Replace("-singlefile=", ""));
+						}
+					}
+
+					// Parse the remote INI setting
+					string RemoteIniPath;
+					CmdLine.TryGetValue("-RemoteIni=", out RemoteIniPath);
+					UnrealBuildTool.SetRemoteIniPath(RemoteIniPath);
+
+					// Create the log file, and flush the startup listener to it
+					if (!String.IsNullOrEmpty(BuildConfiguration.LogFileName) && !Log.HasFileWriter())
+					{
+						FileReference LogLocation = new FileReference(BuildConfiguration.LogFileName);
+						if(LogSuffix.Length > 0)
+						{
+							LogLocation = LogLocation.ChangeExtension(null) + LogSuffix + LogLocation.GetExtension();
+						}
+
+						TextWriterTraceListener LogTraceListener = Log.AddFileWriter("DefaultLogTraceListener", LogLocation);
+						StartupListener.CopyTo(LogTraceListener);
+					}
+					Trace.Listeners.Remove(StartupListener);
+
+					// Build the list of game projects that we know about. When building from the editor (for hot-reload) or for projects from installed builds, we require the 
+					// project file to be passed in. Otherwise we scan for projects in directories named in UE4Games.uprojectdirs.
+					FileReference ProjectFile;
+					if (UnrealBuildTool.TryParseProjectFileArgument(CmdLine, out ProjectFile))
+					{
+						UProjectInfo.AddProject(ProjectFile);
+					}
+					else
+					{
+						UProjectInfo.FillProjectInfo();
+					}
+
+					// Read the project-specific build configuration settings
+					ConfigCache.ReadSettings(DirectoryReference.FromFile(ProjectFile), UnrealTargetPlatform.Unknown, BuildConfiguration);
+
+					DateTime BasicInitStartTime = DateTime.UtcNow;
+
+					bool bSpecificModulesOnly = false;
+					foreach (string Arg in Arguments)
+					{
+						string LowercaseArg = Arg.ToLowerInvariant();
+						if (LowercaseArg.StartsWith("-modulewithsuffix="))
+						{
+							bSpecificModulesOnly = true;
+							continue;
+						}
+						else if (LowercaseArg == "-vsdebugandroid")
+						{
+							AndroidProjectGenerator.VSDebugCommandLineOptionPresent = true;
+						}
+					}
+
+					// Find and register all tool chains, build platforms, etc. that are present
+					UnrealBuildTool.RegisterAllUBTClasses(false);
+
+					if (UnrealBuildTool.bPrintPerformanceInfo)
+					{
+						double BasicInitTime = (DateTime.UtcNow - BasicInitStartTime).TotalSeconds;
+						Log.TraceInformation("Basic UBT initialization took " + BasicInitTime + "s");
+					}
+
+					// now that we know the available platforms, we can delete other platforms' junk. if we're only building specific modules from the editor, don't touch anything else (it may be in use).
+					if (!bSpecificModulesOnly && !BuildConfiguration.bIgnoreJunk)
+					{
+						JunkDeleter.DeleteJunk();
+					}
+
+					// Build our project
+					if (Result == ECompilationResult.Succeeded)
+					{
+						Result = UnrealBuildTool.RunUBT(BuildConfiguration, Arguments, ProjectFile);
+					}
+
+					// Print some performance info
+					double BuildDuration = (DateTime.UtcNow - StartTime).TotalSeconds;
+					if (UnrealBuildTool.bPrintPerformanceInfo)
+					{
+						Log.TraceInformation("GetIncludes time: " + CPPHeaders.TotalTimeSpentGettingIncludes + "s (" + CPPHeaders.TotalIncludesRequested + " includes)");
+						Log.TraceInformation("DirectIncludes cache miss time: " + CPPHeaders.DirectIncludeCacheMissesTotalTime + "s (" + CPPHeaders.TotalDirectIncludeCacheMisses + " misses)");
+						Log.TraceInformation("FindIncludePaths calls: " + CPPHeaders.TotalFindIncludedFileCalls + " (" + CPPHeaders.IncludePathSearchAttempts + " searches)");
+						Log.TraceInformation("Deep C++ include scan time: " + UnrealBuildTool.TotalDeepIncludeScanTime + "s");
+						Log.TraceInformation("Include Resolves: {0} ({1} misses, {2:0.00}%)", CPPHeaders.TotalDirectIncludeResolves, CPPHeaders.TotalDirectIncludeResolveCacheMisses, (float)CPPHeaders.TotalDirectIncludeResolveCacheMisses / (float)CPPHeaders.TotalDirectIncludeResolves * 100);
+						Log.TraceInformation("Total FileItems: {0} ({1} missing)", FileItem.TotalFileItemCount, FileItem.MissingFileItemCount);
+
+						Log.TraceInformation("Execution time: {0}s", BuildDuration);
+					}
+				}
+				catch (Exception Exception)
+				{
+					Log.TraceError("UnrealBuildTool Exception: " + Exception);
+					Result = ECompilationResult.OtherCompilationError;
+				}
+			}
+
+			// Print some performance info
+			Log.TraceVerbose("Execution time: {0}", (DateTime.UtcNow - StartTime).TotalSeconds);
+
+			return (int)Result;
 		}
 	}
 }
diff --git a/Engine/Source/Programs/UnrealBuildTool/UnrealBuildTool.cs b/Engine/Source/Programs/UnrealBuildTool/UnrealBuildTool.cs
--- a/Engine/Source/Programs/UnrealBuildTool/UnrealBuildTool.cs
+++ b/Engine/Source/Programs/UnrealBuildTool/UnrealBuildTool.cs
@@ -115,7 +115,7 @@ namespace UnrealBuildTool
 		static public bool bPrintPerformanceInfo
 		{
 			get;
-			private set;
+			set;
 		}
 
 		/// <summary>
@@ -292,175 +292,6 @@ namespace UnrealBuildTool
 			return false;
 		}
 
-		/// <summary>
-		/// UBT startup order is fairly fragile, and relies on globals that may or may not be safe to use yet.
-		/// This function is for super early startup stuff that should not access Configuration classes (anything loaded by XmlConfg).
-		/// This should be very minimal startup code.
-		/// </summary>
-		/// <param name="CmdLine">Cmdline arguments</param>
-		internal static int GuardedMain(CommandLineArguments CmdLine)
-		{
-			DateTime StartTime = DateTime.UtcNow;
-			ECompilationResult Result = ECompilationResult.Succeeded;
-
-			// Initialize the log system, buffering the output until we can create the log file
-			StartupTraceListener StartupListener = new StartupTraceListener();
-			Trace.Listeners.Add(StartupListener);
-
-			// Write the command line
-			Log.TraceLog("Command line: {0}", Environment.CommandLine);
-
-			// Grab the environment.
-			InitialEnvironment = Environment.GetEnvironmentVariables();
-			if (InitialEnvironment.Count < 1)
-			{
-				throw new BuildException("Environment could not be read");
-			}
-
-			using(SingleInstanceMutex.Acquire(SingleInstanceMutexType.PerBranch, CmdLine))
-			{
-				string[] Arguments = CmdLine.GetRawArray();
-				try
-				{
-
-					// Change the working directory to be the Engine/Source folder. We are likely running from Engine/Binaries/DotNET
-					// This is critical to be done early so any code that relies on the current directory being Engine/Source will work.
-					DirectoryReference.SetCurrentDirectory(EngineSourceDirectory);
-
-					// Read the XML configuration files
-					FileReference XmlConfigCache = CmdLine.GetFileReferenceOrDefault("-XmlConfigCache=", null);
-					XmlConfig.ReadConfigFiles(XmlConfigCache);
-
-					// Create the build configuration object, and read the settings
-					BuildConfiguration BuildConfiguration = new BuildConfiguration();
-					XmlConfig.ApplyTo(BuildConfiguration);
-					CommandLine.ParseArguments(Arguments, BuildConfiguration);
-
-					// Copy some of the static settings that are being deprecated from BuildConfiguration
-					bPrintPerformanceInfo = BuildConfiguration.bPrintPerformanceInfo;
-
-					// Then let the command lines override any configs necessary.
-					string LogSuffix = "";
-					foreach (string Argument in Arguments)
-					{
-						string LowercaseArg = Argument.ToLowerInvariant();
-						if (LowercaseArg.StartsWith("-LogSuffix=", StringComparison.OrdinalIgnoreCase))
-						{
-							LogSuffix += "_" + Argument.Substring("-LogSuffix=".Length);
-						}
-						else if (LowercaseArg == "-nolog")
-						{
-							BuildConfiguration.LogFileName = null;
-						}
-						else if (LowercaseArg == "-xgeexport")
-						{
-							BuildConfiguration.bXGEExport = true;
-							BuildConfiguration.bAllowXGE = true;
-						}
-						else if (LowercaseArg.StartsWith("-singlefile="))
-						{
-							BuildConfiguration.bUseUBTMakefiles = false;
-							BuildConfiguration.SingleFileToCompile = new FileReference(LowercaseArg.Replace("-singlefile=", ""));
-						}
-					}
-
-					// Parse the remote INI setting
-					CmdLine.TryGetValue("-RemoteIni=", out RemoteIniPath);
-
-					// Create the log file, and flush the startup listener to it
-					if (!String.IsNullOrEmpty(BuildConfiguration.LogFileName) && !Log.HasFileWriter())
-					{
-						FileReference LogLocation = new FileReference(BuildConfiguration.LogFileName);
-						if(LogSuffix.Length > 0)
-						{
-							LogLocation = LogLocation.ChangeExtension(null) + LogSuffix + LogLocation.GetExtension();
-						}
-
-						TextWriterTraceListener LogTraceListener = Log.AddFileWriter("DefaultLogTraceListener", LogLocation);
-						StartupListener.CopyTo(LogTraceListener);
-					}
-					Trace.Listeners.Remove(StartupListener);
-
-					// Build the list of game projects that we know about. When building from the editor (for hot-reload) or for projects from installed builds, we require the 
-					// project file to be passed in. Otherwise we scan for projects in directories named in UE4Games.uprojectdirs.
-					FileReference ProjectFile;
-					if (TryParseProjectFileArgument(CmdLine, out ProjectFile))
-					{
-						UProjectInfo.AddProject(ProjectFile);
-					}
-					else
-					{
-						UProjectInfo.FillProjectInfo();
-					}
-
-					// Read the project-specific build configuration settings
-					ConfigCache.ReadSettings(DirectoryReference.FromFile(ProjectFile), UnrealTargetPlatform.Unknown, BuildConfiguration);
-
-					DateTime BasicInitStartTime = DateTime.UtcNow;
-
-					bool bSpecificModulesOnly = false;
-					foreach (string Arg in Arguments)
-					{
-						string LowercaseArg = Arg.ToLowerInvariant();
-						if (LowercaseArg.StartsWith("-modulewithsuffix="))
-						{
-							bSpecificModulesOnly = true;
-							continue;
-						}
-						else if (LowercaseArg == "-vsdebugandroid")
-						{
-							AndroidProjectGenerator.VSDebugCommandLineOptionPresent = true;
-						}
-					}
-
-					// Find and register all tool chains, build platforms, etc. that are present
-					RegisterAllUBTClasses(false);
-
-					if (UnrealBuildTool.bPrintPerformanceInfo)
-					{
-						double BasicInitTime = (DateTime.UtcNow - BasicInitStartTime).TotalSeconds;
-						Log.TraceInformation("Basic UBT initialization took " + BasicInitTime + "s");
-					}
-
-					// now that we know the available platforms, we can delete other platforms' junk. if we're only building specific modules from the editor, don't touch anything else (it may be in use).
-					if (!bSpecificModulesOnly && !BuildConfiguration.bIgnoreJunk)
-					{
-						JunkDeleter.DeleteJunk();
-					}
-
-					// Build our project
-					if (Result == ECompilationResult.Succeeded)
-					{
-						Result = RunUBT(BuildConfiguration, Arguments, ProjectFile);
-					}
-
-					// Print some performance info
-					double BuildDuration = (DateTime.UtcNow - StartTime).TotalSeconds;
-					if (UnrealBuildTool.bPrintPerformanceInfo)
-					{
-						Log.TraceInformation("GetIncludes time: " + CPPHeaders.TotalTimeSpentGettingIncludes + "s (" + CPPHeaders.TotalIncludesRequested + " includes)");
-						Log.TraceInformation("DirectIncludes cache miss time: " + CPPHeaders.DirectIncludeCacheMissesTotalTime + "s (" + CPPHeaders.TotalDirectIncludeCacheMisses + " misses)");
-						Log.TraceInformation("FindIncludePaths calls: " + CPPHeaders.TotalFindIncludedFileCalls + " (" + CPPHeaders.IncludePathSearchAttempts + " searches)");
-						Log.TraceInformation("Deep C++ include scan time: " + UnrealBuildTool.TotalDeepIncludeScanTime + "s");
-						Log.TraceInformation("Include Resolves: {0} ({1} misses, {2:0.00}%)", CPPHeaders.TotalDirectIncludeResolves, CPPHeaders.TotalDirectIncludeResolveCacheMisses, (float)CPPHeaders.TotalDirectIncludeResolveCacheMisses / (float)CPPHeaders.TotalDirectIncludeResolves * 100);
-						Log.TraceInformation("Total FileItems: {0} ({1} missing)", FileItem.TotalFileItemCount, FileItem.MissingFileItemCount);
-
-						Log.TraceInformation("Execution time: {0}s", BuildDuration);
-					}
-				}
-				catch (Exception Exception)
-				{
-					Log.TraceError("UnrealBuildTool Exception: " + Exception);
-					Result = ECompilationResult.OtherCompilationError;
-				}
-			}
-
-			// Print some performance info
-			Log.TraceVerbose("Execution time: {0}", (DateTime.UtcNow - StartTime).TotalSeconds);
-
-			return (int)Result;
-		}
-
 		/// <summary>
 		/// Global options for UBT (any modes)
 		/// </summary>
