UBT: Platform project file generators now receive the command line arguments as a parameter to their constructor. Removes several hacks to configure options globally.

#rnx

diff --git a/Engine/Source/Programs/UnrealBuildTool/Modes/BuildMode.cs b/Engine/Source/Programs/UnrealBuildTool/Modes/BuildMode.cs
--- a/Engine/Source/Programs/UnrealBuildTool/Modes/BuildMode.cs
+++ b/Engine/Source/Programs/UnrealBuildTool/Modes/BuildMode.cs
@@ -131,10 +131,6 @@ namespace UnrealBuildTool
 							bSpecificModulesOnly = true;
 							continue;
 						}
-						else if (LowercaseArg == "-vsdebugandroid")
-						{
-							AndroidProjectGenerator.VSDebugCommandLineOptionPresent = true;
-						}
 					}
 
 					// Find and register all tool chains, build platforms, etc. that are present
diff --git a/Engine/Source/Programs/UnrealBuildTool/Modes/GenerateProjectFilesMode.cs b/Engine/Source/Programs/UnrealBuildTool/Modes/GenerateProjectFilesMode.cs
--- a/Engine/Source/Programs/UnrealBuildTool/Modes/GenerateProjectFilesMode.cs
+++ b/Engine/Source/Programs/UnrealBuildTool/Modes/GenerateProjectFilesMode.cs
@@ -63,12 +63,6 @@ namespace UnrealBuildTool
 				// Find and register all tool chains, build platforms, etc. that are present
 				UnrealBuildTool.RegisterAllUBTClasses(false);
 
-				// HACK
-				if (Arguments.HasOption("-vsdebugandroid"))
-				{
-					AndroidProjectGenerator.VSDebugCommandLineOptionPresent = true;
-				}
-
 				// If there aren't any formats set, read the default project file format from the config file
 				if (ProjectFileFormats.Count == 0)
 				{
@@ -98,7 +92,7 @@ namespace UnrealBuildTool
 				{
 					if (CheckType.IsClass && !CheckType.IsAbstract && CheckType.IsSubclassOf(typeof(PlatformProjectGenerator)))
 					{
-						PlatformProjectGenerator Generator = (PlatformProjectGenerator)Activator.CreateInstance(CheckType);
+						PlatformProjectGenerator Generator = (PlatformProjectGenerator)Activator.CreateInstance(CheckType, Arguments);
 						foreach(UnrealTargetPlatform Platform in Generator.GetPlatforms())
 						{
 							Log.TraceVerbose("Registering project generator {0} for {1}", CheckType, Platform);
diff --git a/Engine/Source/Programs/UnrealBuildTool/Platform/Android/AndroidProjectGenerator.cs b/Engine/Source/Programs/UnrealBuildTool/Platform/Android/AndroidProjectGenerator.cs
--- a/Engine/Source/Programs/UnrealBuildTool/Platform/Android/AndroidProjectGenerator.cs
+++ b/Engine/Source/Programs/UnrealBuildTool/Platform/Android/AndroidProjectGenerator.cs
@@ -14,21 +14,30 @@ namespace UnrealBuildTool
 	/// </summary>
 	class AndroidProjectGenerator : PlatformProjectGenerator
 	{
-		static bool CheckedForNsight = false;		// whether we have checked for a recent enough version of Nsight yet
-		static bool NsightInstalled = false;		// true if a recent enough version of Nsight is installed
-		static int NsightVersionCode = 0;           // version code matching detected Nsight
+		bool CheckedForNsight = false;		// whether we have checked for a recent enough version of Nsight yet
+		bool NsightInstalled = false;		// true if a recent enough version of Nsight is installed
+		int NsightVersionCode = 0;           // version code matching detected Nsight
 
-		public static bool VSDebugCommandLineOptionPresent = false;    //User must put -vsdebugandroid commandline option to build the debug projects
-		static bool VSDebuggingEnabled = false;      // When set to true, allows debugging with built in MS Cross Platform Android tools.  
+		bool VSDebugCommandLineOptionPresent = false;    //User must put -vsdebugandroid commandline option to build the debug projects
+		bool VSDebuggingEnabled = false;      // When set to true, allows debugging with built in MS Cross Platform Android tools.  
 													 //  It adds separate projects ending in .androidproj and a file VSAndroidUnreal.props for the engine and all game projects
 
-		static bool VSSupportChecked = false;       // Don't want to check multiple times
+		bool VSSupportChecked = false;       // Don't want to check multiple times
 
-		static bool VSPropsFileWritten = false;     // This is for the file VSAndroidUnreal.props which only needs to be written once
+		bool VSPropsFileWritten = false;     // This is for the file VSAndroidUnreal.props which only needs to be written once
 
-		static bool VSSandboxedSDK = false;         // Checks if VS has installed the sandboxed SDK version to support Unreal
+		bool VSSandboxedSDK = false;         // Checks if VS has installed the sandboxed SDK version to support Unreal
 													// If this sandboxed SDK is present change the config to consume it instead of the main VS Android SDK
 
+		public AndroidProjectGenerator(CommandLineArguments Arguments)
+			: base(Arguments)
+		{
+			if (Arguments.HasOption("-vsdebugandroid"))
+			{
+				VSDebugCommandLineOptionPresent = true;
+			}
+		}
+
 		bool IsVSAndroidSupportInstalled()
 		{
 			if (VSSupportChecked)
diff --git a/Engine/Source/Programs/UnrealBuildTool/Platform/HTML5/HTML5ProjectGenerator.cs b/Engine/Source/Programs/UnrealBuildTool/Platform/HTML5/HTML5ProjectGenerator.cs
--- a/Engine/Source/Programs/UnrealBuildTool/Platform/HTML5/HTML5ProjectGenerator.cs
+++ b/Engine/Source/Programs/UnrealBuildTool/Platform/HTML5/HTML5ProjectGenerator.cs
@@ -14,6 +14,15 @@ namespace UnrealBuildTool
 	/// </summary>
 	class HTML5ProjectGenerator : PlatformProjectGenerator
 	{
+		/// <summary>
+		/// Constructor
+		/// </summary>
+		/// <param name="Arguments">Command line arguments passed to the project generator</param>
+		public HTML5ProjectGenerator(CommandLineArguments Arguments)
+			: base(Arguments)
+		{
+		}
+
 		/// <summary>
 		/// Enumerate all the platforms that this generator supports
 		/// </summary>
diff --git a/Engine/Source/Programs/UnrealBuildTool/Platform/IOS/IOSProjectGenerator.cs b/Engine/Source/Programs/UnrealBuildTool/Platform/IOS/IOSProjectGenerator.cs
--- a/Engine/Source/Programs/UnrealBuildTool/Platform/IOS/IOSProjectGenerator.cs
+++ b/Engine/Source/Programs/UnrealBuildTool/Platform/IOS/IOSProjectGenerator.cs
@@ -14,6 +14,15 @@ namespace UnrealBuildTool
 	/// </summary>
 	class IOSProjectGenerator : PlatformProjectGenerator
 	{
+		/// <summary>
+		/// Constructor
+		/// </summary>
+		/// <param name="Arguments">Command line arguments passed to the project generator</param>
+		public IOSProjectGenerator(CommandLineArguments Arguments)
+			: base(Arguments)
+		{
+		}
+
 		/// <summary>
 		/// Enumerate all the platforms that this generator supports
 		/// </summary>
diff --git a/Engine/Source/Programs/UnrealBuildTool/Platform/IOS/IOSToolChain.cs b/Engine/Source/Programs/UnrealBuildTool/Platform/IOS/IOSToolChain.cs
--- a/Engine/Source/Programs/UnrealBuildTool/Platform/IOS/IOSToolChain.cs
+++ b/Engine/Source/Programs/UnrealBuildTool/Platform/IOS/IOSToolChain.cs
@@ -1464,9 +1464,11 @@ namespace UnrealBuildTool
 			ProjectFileGenerator.bGenerateProjectFiles = true;
 			try
 			{
+				CommandLineArguments CmdLine = new CommandLineArguments(Arguments);
+
 				PlatformProjectGeneratorCollection PlatformProjectGenerators = new PlatformProjectGeneratorCollection();
-				PlatformProjectGenerators.RegisterPlatformProjectGenerator(UnrealTargetPlatform.IOS, new IOSProjectGenerator());
-				PlatformProjectGenerators.RegisterPlatformProjectGenerator(UnrealTargetPlatform.TVOS, new TVOSProjectGenerator());
+				PlatformProjectGenerators.RegisterPlatformProjectGenerator(UnrealTargetPlatform.IOS, new IOSProjectGenerator(CmdLine));
+				PlatformProjectGenerators.RegisterPlatformProjectGenerator(UnrealTargetPlatform.TVOS, new TVOSProjectGenerator(CmdLine));
 
 				XcodeProjectFileGenerator Generator = new XcodeProjectFileGenerator(ProjectFile);
 				return Generator.GenerateProjectFiles(PlatformProjectGenerators, Arguments);
diff --git a/Engine/Source/Programs/UnrealBuildTool/Platform/Linux/LinuxProjectGenerator.cs b/Engine/Source/Programs/UnrealBuildTool/Platform/Linux/LinuxProjectGenerator.cs
--- a/Engine/Source/Programs/UnrealBuildTool/Platform/Linux/LinuxProjectGenerator.cs
+++ b/Engine/Source/Programs/UnrealBuildTool/Platform/Linux/LinuxProjectGenerator.cs
@@ -14,6 +14,15 @@ namespace UnrealBuildTool
 	/// </summary>
 	class LinuxProjectGenerator : PlatformProjectGenerator
 	{
+		/// <summary>
+		/// Constructor
+		/// </summary>
+		/// <param name="Arguments">Command line arguments passed to the project generator</param>
+		public LinuxProjectGenerator(CommandLineArguments Arguments)
+			: base(Arguments)
+		{
+		}
+
 		/// <summary>
 		/// Enumerate all the platforms that this generator supports
 		/// </summary>
diff --git a/Engine/Source/Programs/UnrealBuildTool/Platform/Lumin/LuminProjectGenerator.cs b/Engine/Source/Programs/UnrealBuildTool/Platform/Lumin/LuminProjectGenerator.cs
--- a/Engine/Source/Programs/UnrealBuildTool/Platform/Lumin/LuminProjectGenerator.cs
+++ b/Engine/Source/Programs/UnrealBuildTool/Platform/Lumin/LuminProjectGenerator.cs
@@ -17,6 +17,15 @@ namespace UnrealBuildTool
 		static bool VSSupportChecked = false;       // Don't want to check multiple times
 		static bool VSDebuggingEnabled = false;
 
+		/// <summary>
+		/// Constructor
+		/// </summary>
+		/// <param name="Arguments">Command line arguments passed to the project generator</param>
+		public LuminProjectGenerator(CommandLineArguments Arguments)
+			: base(Arguments)
+		{
+		}
+
 		private bool IsVSLuminSupportInstalled(VCProjectFileFormat ProjectFileFormat)
 		{
 			if (!VSSupportChecked)
diff --git a/Engine/Source/Programs/UnrealBuildTool/Platform/Mac/MacProjectGenerator.cs b/Engine/Source/Programs/UnrealBuildTool/Platform/Mac/MacProjectGenerator.cs
--- a/Engine/Source/Programs/UnrealBuildTool/Platform/Mac/MacProjectGenerator.cs
+++ b/Engine/Source/Programs/UnrealBuildTool/Platform/Mac/MacProjectGenerator.cs
@@ -14,6 +14,15 @@ namespace UnrealBuildTool
 	/// </summary>
 	class MacProjectGenerator : PlatformProjectGenerator
 	{
+		/// <summary>
+		/// Constructor
+		/// </summary>
+		/// <param name="Arguments">Command line arguments passed to the project generator</param>
+		public MacProjectGenerator(CommandLineArguments Arguments)
+			: base(Arguments)
+		{
+		}
+
 		/// <summary>
 		/// Register the platform with the UEPlatformProjectGenerator class
 		/// </summary>
diff --git a/Engine/Source/Programs/UnrealBuildTool/Platform/TVOS/TVOSProjectGenerator.cs b/Engine/Source/Programs/UnrealBuildTool/Platform/TVOS/TVOSProjectGenerator.cs
--- a/Engine/Source/Programs/UnrealBuildTool/Platform/TVOS/TVOSProjectGenerator.cs
+++ b/Engine/Source/Programs/UnrealBuildTool/Platform/TVOS/TVOSProjectGenerator.cs
@@ -11,6 +11,15 @@ namespace UnrealBuildTool
 {
 	class TVOSProjectGenerator : IOSProjectGenerator
     {
+		/// <summary>
+		/// Constructor
+		/// </summary>
+		/// <param name="Arguments">Command line arguments passed to the project generator</param>
+		public TVOSProjectGenerator(CommandLineArguments Arguments)
+			: base(Arguments)
+		{
+		}
+
 		/// <summary>
 		/// Enumerate all the platforms that this generator supports
 		/// </summary>
diff --git a/Engine/Source/Programs/UnrealBuildTool/Platform/Windows/WindowsProjectGenerator.cs b/Engine/Source/Programs/UnrealBuildTool/Platform/Windows/WindowsProjectGenerator.cs
--- a/Engine/Source/Programs/UnrealBuildTool/Platform/Windows/WindowsProjectGenerator.cs
+++ b/Engine/Source/Programs/UnrealBuildTool/Platform/Windows/WindowsProjectGenerator.cs
@@ -14,6 +14,15 @@ namespace UnrealBuildTool
 	/// </summary>
 	class WindowsProjectGenerator : PlatformProjectGenerator
 	{
+		/// <summary>
+		/// Constructor
+		/// </summary>
+		/// <param name="Arguments">Command line arguments passed to the project generator</param>
+		public WindowsProjectGenerator(CommandLineArguments Arguments)
+			: base(Arguments)
+		{
+		}
+
 		/// <summary>
 		/// Enumerate all the platforms that this generator supports
 		/// </summary>
diff --git a/Engine/Source/Programs/UnrealBuildTool/ProjectFiles/PlatformProjectGenerator.cs b/Engine/Source/Programs/UnrealBuildTool/ProjectFiles/PlatformProjectGenerator.cs
--- a/Engine/Source/Programs/UnrealBuildTool/ProjectFiles/PlatformProjectGenerator.cs
+++ b/Engine/Source/Programs/UnrealBuildTool/ProjectFiles/PlatformProjectGenerator.cs
@@ -14,6 +14,14 @@ namespace UnrealBuildTool
 	/// </summary>
 	abstract class PlatformProjectGenerator
 	{
+		/// <summary>
+		/// Constructor
+		/// </summary>
+		/// <param name="Arguments">Command line arguments passed to the project generator</param>
+		public PlatformProjectGenerator(CommandLineArguments Arguments)
+		{
+		}
+
 		/// <summary>
 		/// Register the platform with the UEPlatformProjectGenerator class
 		/// </summary>
