Pass -IgnoreJunk when hot reloading from the editor rather than special casing it from inside UBT.

#rnx

diff --git a/Engine/Source/Developer/HotReload/Private/HotReload.cpp b/Engine/Source/Developer/HotReload/Private/HotReload.cpp
--- a/Engine/Source/Developer/HotReload/Private/HotReload.cpp
+++ b/Engine/Source/Developer/HotReload/Private/HotReload.cpp
@@ -1638,7 +1638,7 @@ bool FHotReloadModule::StartCompilingModuleDLLs(const TArray< FModuleToRecompile
 		ExtraArg += TEXT("-nosharedpch ");
 	}
 	
-	FString CmdLineParams = FString::Printf( TEXT( "%s %s %s %s%s" ), 
+	FString CmdLineParams = FString::Printf( TEXT( "%s %s %s %s%s -IgnoreJunk" ), 
 		*ModuleArg, 
 		BuildPlatformName, BuildConfigurationName, 
 		*ExtraArg, *InAdditionalCmdLineArgs );
diff --git a/Engine/Source/Programs/UnrealBuildTool/Modes/BuildMode.cs b/Engine/Source/Programs/UnrealBuildTool/Modes/BuildMode.cs
--- a/Engine/Source/Programs/UnrealBuildTool/Modes/BuildMode.cs
+++ b/Engine/Source/Programs/UnrealBuildTool/Modes/BuildMode.cs
@@ -122,17 +122,6 @@ namespace UnrealBuildTool
 
 					DateTime BasicInitStartTime = DateTime.UtcNow;
 
-					bool bSpecificModulesOnly = false;
-					foreach (string Arg in Arguments)
-					{
-						string LowercaseArg = Arg.ToLowerInvariant();
-						if (LowercaseArg.StartsWith("-modulewithsuffix="))
-						{
-							bSpecificModulesOnly = true;
-							continue;
-						}
-					}
-
 					// Find and register all tool chains, build platforms, etc. that are present
 					UnrealBuildTool.RegisterAllUBTClasses(false);
 
@@ -143,7 +132,7 @@ namespace UnrealBuildTool
 					}
 
 					// now that we know the available platforms, we can delete other platforms' junk. if we're only building specific modules from the editor, don't touch anything else (it may be in use).
-					if (!bSpecificModulesOnly && !BuildConfiguration.bIgnoreJunk)
+					if (!BuildConfiguration.bIgnoreJunk)
 					{
 						JunkDeleter.DeleteJunk();
 					}
