UBT: Move log file configuration outside of the BuildConfiguration object, so it can be initialized sooner.

#rnx

diff --git a/Engine/Source/Programs/UnrealBuildTool/Configuration/BuildConfiguration.cs b/Engine/Source/Programs/UnrealBuildTool/Configuration/BuildConfiguration.cs
--- a/Engine/Source/Programs/UnrealBuildTool/Configuration/BuildConfiguration.cs
+++ b/Engine/Source/Programs/UnrealBuildTool/Configuration/BuildConfiguration.cs
@@ -39,12 +39,6 @@ namespace UnrealBuildTool
 		[XmlConfigFile]
 		public bool bPrintDebugInfo = false;
 
-		/// <summary>
-		/// Specifies the file to use for logging
-		/// </summary>
-		[XmlConfigFile]
-		public string LogFileName = "../Programs/UnrealBuildTool/Log.txt";
-
 		/// <summary>
 		/// Prints performance diagnostics about include dependencies and other bits
 		/// </summary>
@@ -67,7 +61,8 @@ namespace UnrealBuildTool
 		/// <summary>
 		/// Whether we should just export the XGE XML and pretend it succeeded
 		/// </summary>
-		public bool bXGEExport;
+		[CommandLine("-XGEExport")]
+		public bool bXGEExport = false;
 
 		/// <summary>
 		/// Whether SN-DBS may be used.
@@ -143,6 +138,7 @@ namespace UnrealBuildTool
 		/// <summary>
 		/// If specified, we will only build this particular source file, ignore all other outputs.  Useful for testing non-Unity builds.
 		/// </summary>
+		[CommandLine("-SingleFile=")]
 		public FileReference SingleFileToCompile = null;
 
 		/// <summary>
diff --git a/Engine/Source/Programs/UnrealBuildTool/Modes/BuildMode.cs b/Engine/Source/Programs/UnrealBuildTool/Modes/BuildMode.cs
--- a/Engine/Source/Programs/UnrealBuildTool/Modes/BuildMode.cs
+++ b/Engine/Source/Programs/UnrealBuildTool/Modes/BuildMode.cs
@@ -16,15 +16,22 @@ namespace UnrealBuildTool
 	[ToolMode("Build")]
 	class BuildMode : ToolMode
 	{
+		/// <summary>
+		/// Specifies the file to use for logging
+		/// </summary>
+		[XmlConfigFile(Category = "BuildConfiguration")]
+		public string BaseLogFileName = "../Programs/UnrealBuildTool/Log.txt";
+
 		/// <summary>
 		/// Main entry point
 		/// </summary>
-		/// <param name="CmdLine">Command-line arguments</param>
+		/// <param name="Arguments">Command-line arguments</param>
 		/// <returns>One of the values of ECompilationResult</returns>
-		public override int Execute(CommandLineArguments CmdLine)
+		public override int Execute(CommandLineArguments Arguments)
 		{
+			Arguments.ApplyTo(this);
+
 			DateTime StartTime = DateTime.UtcNow;
-			ECompilationResult Result = ECompilationResult.Succeeded;
 
 			// Initialize the log system, buffering the output until we can create the log file
 			StartupTraceListener StartupListener = new StartupTraceListener();
@@ -40,134 +47,110 @@ namespace UnrealBuildTool
 				throw new BuildException("Environment could not be read");
 			}
 
-			using(SingleInstanceMutex.Acquire(SingleInstanceMutexType.PerBranch, CmdLine))
+			// Get a mutex for building in this branch
+			using(SingleInstanceMutex.Acquire(SingleInstanceMutexType.PerBranch, Arguments))
 			{
-				string[] Arguments = CmdLine.GetRawArray();
-				try
+				// Change the working directory to be the Engine/Source folder. We are likely running from Engine/Binaries/DotNET
+				// This is critical to be done early so any code that relies on the current directory being Engine/Source will work.
+				DirectoryReference.SetCurrentDirectory(UnrealBuildTool.EngineSourceDirectory);
+
+				// Read the XML configuration files
+				FileReference XmlConfigCache = Arguments.GetFileReferenceOrDefault("-XmlConfigCache=", null);
+				XmlConfig.ReadConfigFiles(XmlConfigCache);
+				XmlConfig.ApplyTo(this);
+
+				// Create the log file, and flush the startup listener to it
+				FileReference LogFile = null;
+				if(!Arguments.HasOption("-NoLog") && !Log.HasFileWriter())
 				{
-					// Change the working directory to be the Engine/Source folder. We are likely running from Engine/Binaries/DotNET
-					// This is critical to be done early so any code that relies on the current directory being Engine/Source will work.
-					DirectoryReference.SetCurrentDirectory(UnrealBuildTool.EngineSourceDirectory);
-
-					// Read the XML configuration files
-					FileReference XmlConfigCache = CmdLine.GetFileReferenceOrDefault("-XmlConfigCache=", null);
-					XmlConfig.ReadConfigFiles(XmlConfigCache);
-
-					// Create the build configuration object, and read the settings
-					BuildConfiguration BuildConfiguration = new BuildConfiguration();
-					XmlConfig.ApplyTo(BuildConfiguration);
-					CommandLine.ParseArguments(Arguments, BuildConfiguration);
-
-					// Copy some of the static settings that are being deprecated from BuildConfiguration
-					UnrealBuildTool.bPrintPerformanceInfo = BuildConfiguration.bPrintPerformanceInfo;
-
-					// Then let the command lines override any configs necessary.
-					string LogSuffix = "";
-					foreach (string Argument in Arguments)
+					LogFile = new FileReference(BaseLogFileName);
+					foreach(string LogSuffix in Arguments.GetValues("-LogSuffix="))
 					{
-						string LowercaseArg = Argument.ToLowerInvariant();
-						if (LowercaseArg.StartsWith("-LogSuffix=", StringComparison.OrdinalIgnoreCase))
-						{
-							LogSuffix += "_" + Argument.Substring("-LogSuffix=".Length);
-						}
-						else if (LowercaseArg == "-nolog")
-						{
-							BuildConfiguration.LogFileName = null;
-						}
-						else if (LowercaseArg == "-xgeexport")
-						{
-							BuildConfiguration.bXGEExport = true;
-							BuildConfiguration.bAllowXGE = true;
-						}
-						else if (LowercaseArg.StartsWith("-singlefile="))
-						{
-							BuildConfiguration.bUseUBTMakefiles = false;
-							BuildConfiguration.SingleFileToCompile = new FileReference(LowercaseArg.Replace("-singlefile=", ""));
-						}
+						LogFile = LogFile.ChangeExtension(null) + "_" + LogSuffix + LogFile.GetExtension();
 					}
 
-					// Parse the remote INI setting
-					string RemoteIniPath;
-					CmdLine.TryGetValue("-RemoteIni=", out RemoteIniPath);
-					UnrealBuildTool.SetRemoteIniPath(RemoteIniPath);
+					TextWriterTraceListener LogTraceListener = Log.AddFileWriter("DefaultLogTraceListener", LogFile);
+					StartupListener.CopyTo(LogTraceListener);
+				}
+				Trace.Listeners.Remove(StartupListener);
 
-					// Create the log file, and flush the startup listener to it
-					if (!String.IsNullOrEmpty(BuildConfiguration.LogFileName) && !Log.HasFileWriter())
-					{
-						FileReference LogLocation = new FileReference(BuildConfiguration.LogFileName);
-						if(LogSuffix.Length > 0)
-						{
-							LogLocation = LogLocation.ChangeExtension(null) + LogSuffix + LogLocation.GetExtension();
-						}
-
-						TextWriterTraceListener LogTraceListener = Log.AddFileWriter("DefaultLogTraceListener", LogLocation);
-						StartupListener.CopyTo(LogTraceListener);
-					}
-					Trace.Listeners.Remove(StartupListener);
+				// Create the build configuration object, and read the settings
+				BuildConfiguration BuildConfiguration = new BuildConfiguration();
+				XmlConfig.ApplyTo(BuildConfiguration);
+				Arguments.ApplyTo(BuildConfiguration);
 
-					// Build the list of game projects that we know about. When building from the editor (for hot-reload) or for projects from installed builds, we require the 
-					// project file to be passed in. Otherwise we scan for projects in directories named in UE4Games.uprojectdirs.
-					FileReference ProjectFile;
-					if (UnrealBuildTool.TryParseProjectFileArgument(CmdLine, out ProjectFile))
-					{
-						UProjectInfo.AddProject(ProjectFile);
-					}
-					else
-					{
-						UProjectInfo.FillProjectInfo();
-					}
+				// Copy some of the static settings that are being deprecated from BuildConfiguration
+				UnrealBuildTool.bPrintPerformanceInfo = BuildConfiguration.bPrintPerformanceInfo;
 
-					// Read the project-specific build configuration settings
-					ConfigCache.ReadSettings(DirectoryReference.FromFile(ProjectFile), UnrealTargetPlatform.Unknown, BuildConfiguration);
+				// Then let the command lines override any configs necessary.
+				if(BuildConfiguration.bXGEExport)
+				{
+					BuildConfiguration.bAllowXGE = true;
+				}
+				if(BuildConfiguration.SingleFileToCompile != null)
+				{
+					BuildConfiguration.bUseUBTMakefiles = false;
+				}
 
-					DateTime BasicInitStartTime = DateTime.UtcNow;
+				// Parse the remote INI setting
+				string RemoteIniPath;
+				Arguments.TryGetValue("-RemoteIni=", out RemoteIniPath);
+				UnrealBuildTool.SetRemoteIniPath(RemoteIniPath);
 
-					// Find and register all tool chains, build platforms, etc. that are present
-					UnrealBuildTool.RegisterAllUBTClasses(false);
+				// Build the list of game projects that we know about. When building from the editor (for hot-reload) or for projects from installed builds, we require the 
+				// project file to be passed in. Otherwise we scan for projects in directories named in UE4Games.uprojectdirs.
+				FileReference ProjectFile;
+				if (UnrealBuildTool.TryParseProjectFileArgument(Arguments, out ProjectFile))
+				{
+					UProjectInfo.AddProject(ProjectFile);
+				}
+				else
+				{
+					UProjectInfo.FillProjectInfo();
+				}
 
-					if (UnrealBuildTool.bPrintPerformanceInfo)
-					{
-						double BasicInitTime = (DateTime.UtcNow - BasicInitStartTime).TotalSeconds;
-						Log.TraceInformation("Basic UBT initialization took " + BasicInitTime + "s");
-					}
+				// Read the project-specific build configuration settings
+				ConfigCache.ReadSettings(DirectoryReference.FromFile(ProjectFile), UnrealTargetPlatform.Unknown, BuildConfiguration);
 
-					// now that we know the available platforms, we can delete other platforms' junk. if we're only building specific modules from the editor, don't touch anything else (it may be in use).
-					if (!BuildConfiguration.bIgnoreJunk)
-					{
-						JunkDeleter.DeleteJunk();
-					}
+				DateTime BasicInitStartTime = DateTime.UtcNow;
 
-					// Build our project
-					if (Result == ECompilationResult.Succeeded)
-					{
-						Result = UnrealBuildTool.RunUBT(BuildConfiguration, Arguments, ProjectFile);
-					}
+				// Find and register all tool chains, build platforms, etc. that are present
+				UnrealBuildTool.RegisterAllUBTClasses(false);
 
-					// Print some performance info
-					double BuildDuration = (DateTime.UtcNow - StartTime).TotalSeconds;
-					if (UnrealBuildTool.bPrintPerformanceInfo)
-					{
-						Log.TraceInformation("GetIncludes time: " + CPPHeaders.TotalTimeSpentGettingIncludes + "s (" + CPPHeaders.TotalIncludesRequested + " includes)");
-						Log.TraceInformation("DirectIncludes cache miss time: " + CPPHeaders.DirectIncludeCacheMissesTotalTime + "s (" + CPPHeaders.TotalDirectIncludeCacheMisses + " misses)");
-						Log.TraceInformation("FindIncludePaths calls: " + CPPHeaders.TotalFindIncludedFileCalls + " (" + CPPHeaders.IncludePathSearchAttempts + " searches)");
-						Log.TraceInformation("Deep C++ include scan time: " + UnrealBuildTool.TotalDeepIncludeScanTime + "s");
-						Log.TraceInformation("Include Resolves: {0} ({1} misses, {2:0.00}%)", CPPHeaders.TotalDirectIncludeResolves, CPPHeaders.TotalDirectIncludeResolveCacheMisses, (float)CPPHeaders.TotalDirectIncludeResolveCacheMisses / (float)CPPHeaders.TotalDirectIncludeResolves * 100);
-						Log.TraceInformation("Total FileItems: {0} ({1} missing)", FileItem.TotalFileItemCount, FileItem.MissingFileItemCount);
-
-						Log.TraceInformation("Execution time: {0}s", BuildDuration);
-					}
+				if (UnrealBuildTool.bPrintPerformanceInfo)
+				{
+					double BasicInitTime = (DateTime.UtcNow - BasicInitStartTime).TotalSeconds;
+					Log.TraceInformation("Basic UBT initialization took " + BasicInitTime + "s");
 				}
-				catch (Exception Exception)
+
+				// now that we know the available platforms, we can delete other platforms' junk. if we're only building specific modules from the editor, don't touch anything else (it may be in use).
+				if (!BuildConfiguration.bIgnoreJunk)
 				{
-					Log.TraceError("UnrealBuildTool Exception: " + Exception);
-					Result = ECompilationResult.OtherCompilationError;
+					JunkDeleter.DeleteJunk();
 				}
-			}
 
-			// Print some performance info
-			Log.TraceVerbose("Execution time: {0}", (DateTime.UtcNow - StartTime).TotalSeconds);
+				// Build our project
+				ECompilationResult Result = UnrealBuildTool.RunUBT(BuildConfiguration, Arguments.GetRawArray(), ProjectFile, LogFile);
 
-			return (int)Result;
+				// Print some performance info
+				double BuildDuration = (DateTime.UtcNow - StartTime).TotalSeconds;
+				if (UnrealBuildTool.bPrintPerformanceInfo)
+				{
+					Log.TraceInformation("GetIncludes time: " + CPPHeaders.TotalTimeSpentGettingIncludes + "s (" + CPPHeaders.TotalIncludesRequested + " includes)");
+					Log.TraceInformation("DirectIncludes cache miss time: " + CPPHeaders.DirectIncludeCacheMissesTotalTime + "s (" + CPPHeaders.TotalDirectIncludeCacheMisses + " misses)");
+					Log.TraceInformation("FindIncludePaths calls: " + CPPHeaders.TotalFindIncludedFileCalls + " (" + CPPHeaders.IncludePathSearchAttempts + " searches)");
+					Log.TraceInformation("Deep C++ include scan time: " + UnrealBuildTool.TotalDeepIncludeScanTime + "s");
+					Log.TraceInformation("Include Resolves: {0} ({1} misses, {2:0.00}%)", CPPHeaders.TotalDirectIncludeResolves, CPPHeaders.TotalDirectIncludeResolveCacheMisses, (float)CPPHeaders.TotalDirectIncludeResolveCacheMisses / (float)CPPHeaders.TotalDirectIncludeResolves * 100);
+					Log.TraceInformation("Total FileItems: {0} ({1} missing)", FileItem.TotalFileItemCount, FileItem.MissingFileItemCount);
+
+					Log.TraceInformation("Execution time: {0}s", BuildDuration);
+				}
+
+				// Print some performance info
+				Log.TraceVerbose("Execution time: {0}", (DateTime.UtcNow - StartTime).TotalSeconds);
+
+				return (int)Result;
+			}
 		}
 	}
 }
diff --git a/Engine/Source/Programs/UnrealBuildTool/System/ExternalExecution.cs b/Engine/Source/Programs/UnrealBuildTool/System/ExternalExecution.cs
--- a/Engine/Source/Programs/UnrealBuildTool/System/ExternalExecution.cs
+++ b/Engine/Source/Programs/UnrealBuildTool/System/ExternalExecution.cs
@@ -1180,7 +1180,7 @@ namespace UnrealBuildTool
 						}
 
 						// Output the log next to the current log
-						if(String.IsNullOrEmpty(BuildConfiguration.LogFileName))
+						if(Environment.CommandLine.IndexOf("-NoLog", StringComparison.OrdinalIgnoreCase) != -1)
 						{
 							UBTArguments.Append(" -nolog");
 						}
diff --git a/Engine/Source/Programs/UnrealBuildTool/UnrealBuildTool.cs b/Engine/Source/Programs/UnrealBuildTool/UnrealBuildTool.cs
--- a/Engine/Source/Programs/UnrealBuildTool/UnrealBuildTool.cs
+++ b/Engine/Source/Programs/UnrealBuildTool/UnrealBuildTool.cs
@@ -436,7 +436,7 @@ namespace UnrealBuildTool
 			}
 		}
 
-		internal static ECompilationResult RunUBT(BuildConfiguration BuildConfiguration, string[] Arguments, FileReference ProjectFile)
+		internal static ECompilationResult RunUBT(BuildConfiguration BuildConfiguration, string[] Arguments, FileReference ProjectFile, FileReference LogFile)
 		{
 			bool bSuccess = true;
 
@@ -500,10 +500,10 @@ namespace UnrealBuildTool
 						TargetDescriptor TargetDesc = TargetDescs[Idx];
 						if(TargetDesc.Platform == UnrealTargetPlatform.Mac || TargetDesc.Platform == UnrealTargetPlatform.IOS || TargetDesc.Platform == UnrealTargetPlatform.TVOS)
 						{
-							FileReference LogFile = null;
-							if(!String.IsNullOrEmpty(BuildConfiguration.LogFileName))
+							FileReference RemoteLogFile = null;
+							if(LogFile != null)
 							{
-								LogFile = new FileReference(Path.Combine(Path.GetDirectoryName(BuildConfiguration.LogFileName), Path.GetFileNameWithoutExtension(BuildConfiguration.LogFileName) + "_Remote.txt"));
+								RemoteLogFile = FileReference.Combine(LogFile.Directory, LogFile.GetFileNameWithoutExtension() + "_Remote.txt");
 							}
 
 							RemoteMac RemoteMac = new RemoteMac(TargetDesc.ProjectFile);
@@ -1238,7 +1238,7 @@ namespace UnrealBuildTool
 			}
 			catch (Exception Ex)
 			{
-				Log.WriteException(Ex, String.IsNullOrEmpty(BuildConfiguration.LogFileName) ? null : BuildConfiguration.LogFileName);
+				Log.WriteException(Ex, (LogFile == null) ? null : LogFile.FullName);
 				BuildResult = ECompilationResult.OtherCompilationError;
 			}
 
