UBT: Remove project-specific setting for BuildConfiguration.bAllowHotReloadFromIDE; this value is already checked in ShouldDoHotReloadFromIDE(), and is target-specific.

#rnx

diff --git a/Engine/Source/Programs/UnrealBuildTool/Configuration/BuildConfiguration.cs b/Engine/Source/Programs/UnrealBuildTool/Configuration/BuildConfiguration.cs
--- a/Engine/Source/Programs/UnrealBuildTool/Configuration/BuildConfiguration.cs
+++ b/Engine/Source/Programs/UnrealBuildTool/Configuration/BuildConfiguration.cs
@@ -184,8 +184,8 @@ namespace UnrealBuildTool
 		/// <summary>
 		/// True if hot-reload from IDE is allowed
 		/// </summary>
+		[CommandLine("-NoHotReloadFromIDE", Value="false")]
 		[XmlConfigFile(Category = "UEBuildConfiguration")]
-		[ConfigFile(ConfigHierarchyType.Engine, "BuildConfiguration")]
 		public bool bAllowHotReloadFromIDE = true;
 
 		/// <summary>
diff --git a/Engine/Source/Programs/UnrealBuildTool/Modes/BuildMode.cs b/Engine/Source/Programs/UnrealBuildTool/Modes/BuildMode.cs
--- a/Engine/Source/Programs/UnrealBuildTool/Modes/BuildMode.cs
+++ b/Engine/Source/Programs/UnrealBuildTool/Modes/BuildMode.cs
@@ -109,9 +109,6 @@ namespace UnrealBuildTool
 					UProjectInfo.FillProjectInfo();
 				}
 
-				// Read the project-specific build configuration settings
-				ConfigCache.ReadSettings(DirectoryReference.FromFile(ProjectFile), UnrealTargetPlatform.Unknown, BuildConfiguration);
-
 				DateTime BasicInitStartTime = DateTime.UtcNow;
 
 				// Find and register all tool chains, build platforms, etc. that are present
diff --git a/Engine/Source/Programs/UnrealBuildTool/System/HotReload.cs b/Engine/Source/Programs/UnrealBuildTool/System/HotReload.cs
--- a/Engine/Source/Programs/UnrealBuildTool/System/HotReload.cs
+++ b/Engine/Source/Programs/UnrealBuildTool/System/HotReload.cs
@@ -120,7 +120,7 @@ namespace UnrealBuildTool
 		/// <summary>
 		/// Checks if the editor is currently running and this is a hot-reload
 		/// </summary>
-		public static bool ShouldDoHotReloadFromIDE(BuildConfiguration BuildConfiguration, string[] Arguments, TargetDescriptor TargetDesc)
+		public static bool ShouldDoHotReloadFromIDE(BuildConfiguration BuildConfiguration, TargetDescriptor TargetDesc)
 		{
 			// Check if Hot-reload is disabled globally for this project
 			ConfigHierarchy Hierarchy = ConfigCache.ReadHierarchy(ConfigHierarchyType.Engine, DirectoryReference.FromFile(TargetDesc.ProjectFile), TargetDesc.Platform);
@@ -130,9 +130,8 @@ namespace UnrealBuildTool
 				return false;
 			}
 
-			if (Arguments.Any(x => x.Equals("-NoHotReloadFromIDE", StringComparison.InvariantCultureIgnoreCase)))
+			if(!BuildConfiguration.bAllowHotReloadFromIDE)
 			{
-				// Hot reload disabled through command line, possibly running from UAT
 				return false;
 			}
 
diff --git a/Engine/Source/Programs/UnrealBuildTool/UnrealBuildTool.cs b/Engine/Source/Programs/UnrealBuildTool/UnrealBuildTool.cs
--- a/Engine/Source/Programs/UnrealBuildTool/UnrealBuildTool.cs
+++ b/Engine/Source/Programs/UnrealBuildTool/UnrealBuildTool.cs
@@ -907,7 +907,7 @@ namespace UnrealBuildTool
 						}
 						else if (TargetDescs.Count == 1 && !Arguments.Any(x => x.Equals("-NoHotReload", StringComparison.InvariantCultureIgnoreCase)))
 						{
-							if (BuildConfiguration.bAllowHotReloadFromIDE && HotReload.ShouldDoHotReloadFromIDE(BuildConfiguration, Arguments, TargetDescs[0]))
+							if (BuildConfiguration.bAllowHotReloadFromIDE && HotReload.ShouldDoHotReloadFromIDE(BuildConfiguration, TargetDescs[0]))
 							{
 								HotReloadMode = HotReloadMode.FromIDE;
 							}
