UBT: Replace UProjectInfo with the NativeProjects class, which caches the list of native projects to the engine. Foreign projects do not need to be registered with it, and the list of native projects is found on demand.

#rnx

diff --git a/Engine/Source/Programs/AutomationTool/AutomationUtils/Automation.cs b/Engine/Source/Programs/AutomationTool/AutomationUtils/Automation.cs
--- a/Engine/Source/Programs/AutomationTool/AutomationUtils/Automation.cs
+++ b/Engine/Source/Programs/AutomationTool/AutomationUtils/Automation.cs
@@ -470,9 +470,6 @@ AutomationTool.exe [-verbose] [-compileonly] [-p4] Command0 [-Arg0 -Arg1 -Arg2 .
 				return ExitCode.Error_Unknown;
 			}
 
-			// Fill in the project info
-			UnrealBuildTool.UProjectInfo.FillProjectInfo();
-
 			// Clean rules folders up
 			ProjectUtils.CleanupFolders();
 
diff --git a/Engine/Source/Programs/AutomationTool/AutomationUtils/ProjectUtils.cs b/Engine/Source/Programs/AutomationTool/AutomationUtils/ProjectUtils.cs
--- a/Engine/Source/Programs/AutomationTool/AutomationUtils/ProjectUtils.cs
+++ b/Engine/Source/Programs/AutomationTool/AutomationUtils/ProjectUtils.cs
@@ -726,7 +726,7 @@ namespace AutomationTool
         {
             BaseEngineProject = new BranchUProject();
 
-            var AllProjects = UnrealBuildTool.UProjectInfo.AllProjectFiles;
+            IEnumerable<FileReference> AllProjects = UnrealBuildTool.NativeProjects.EnumerateProjectFiles();
 			using(TelemetryStopwatch SortProjectsStopwatch = new TelemetryStopwatch("SortProjects"))
 			{
 				foreach (var InfoEntry in AllProjects)
diff --git a/Engine/Source/Programs/AutomationTool/AutomationUtils/ScriptCompiler.cs b/Engine/Source/Programs/AutomationTool/AutomationUtils/ScriptCompiler.cs
--- a/Engine/Source/Programs/AutomationTool/AutomationUtils/ScriptCompiler.cs
+++ b/Engine/Source/Programs/AutomationTool/AutomationUtils/ScriptCompiler.cs
@@ -125,7 +125,7 @@ namespace AutomationTool
 			List<DirectoryReference> AllGameFolders;
 			if(ScriptsForProjectFileName == null)
 			{
-				AllGameFolders = UProjectInfo.AllProjectFiles.Select(x => x.Directory).ToList();
+				AllGameFolders = NativeProjects.EnumerateProjectFiles().Select(x => x.Directory).ToList();
 			}
 			else
 			{
diff --git a/Engine/Source/Programs/UnrealBuildTool/Configuration/TargetDescriptor.cs b/Engine/Source/Programs/UnrealBuildTool/Configuration/TargetDescriptor.cs
--- a/Engine/Source/Programs/UnrealBuildTool/Configuration/TargetDescriptor.cs
+++ b/Engine/Source/Programs/UnrealBuildTool/Configuration/TargetDescriptor.cs
@@ -165,9 +165,9 @@ namespace UnrealBuildTool
 			foreach(string TargetName in TargetNames)
 			{
 				// If a project file was not specified see if we can find one
-				if (ProjectFile == null && UProjectInfo.TryGetProjectForTarget(TargetName, out ProjectFile))
+				if (ProjectFile == null && NativeProjects.TryGetProjectForTarget(TargetName, out ProjectFile))
 				{
-					Log.TraceVerbose("Found project file for {0} - {1}", TargetName, ProjectFile);
+					Log.TraceLog("Found project file for {0}: {1}", TargetName, ProjectFile);
 				}
 
 				// Create the target descriptor
diff --git a/Engine/Source/Programs/UnrealBuildTool/Modes/BuildMode.cs b/Engine/Source/Programs/UnrealBuildTool/Modes/BuildMode.cs
--- a/Engine/Source/Programs/UnrealBuildTool/Modes/BuildMode.cs
+++ b/Engine/Source/Programs/UnrealBuildTool/Modes/BuildMode.cs
@@ -100,15 +100,8 @@ namespace UnrealBuildTool
 				// Build the list of game projects that we know about. When building from the editor (for hot-reload) or for projects from installed builds, we require the 
 				// project file to be passed in. Otherwise we scan for projects in directories named in UE4Games.uprojectdirs.
 				FileReference ProjectFile;
-				if (UnrealBuildTool.TryParseProjectFileArgument(Arguments, out ProjectFile))
-				{
-					UProjectInfo.AddProject(ProjectFile);
-				}
-				else
-				{
-					UProjectInfo.FillProjectInfo();
-				}
-
+				UnrealBuildTool.TryParseProjectFileArgument(Arguments, out ProjectFile);
+				
 				DateTime BasicInitStartTime = DateTime.UtcNow;
 
 				// Find and register all tool chains, build platforms, etc. that are present
diff --git a/Engine/Source/Programs/UnrealBuildTool/Modes/GenerateProjectFilesMode.cs b/Engine/Source/Programs/UnrealBuildTool/Modes/GenerateProjectFilesMode.cs
--- a/Engine/Source/Programs/UnrealBuildTool/Modes/GenerateProjectFilesMode.cs
+++ b/Engine/Source/Programs/UnrealBuildTool/Modes/GenerateProjectFilesMode.cs
@@ -162,17 +162,6 @@ namespace UnrealBuildTool
 				// Check there are no superfluous command line arguments
 				Arguments.CheckAllArgumentsUsed();
 
-				// Build the list of game projects that we know about. When building from the editor (for hot-reload) or for projects from installed builds, we require the 
-				// project file to be passed in. Otherwise we scan for projects in directories named in UE4Games.uprojectdirs.
-				if (ProjectFile != null)
-				{
-					UProjectInfo.AddProject(ProjectFile);
-				}
-				else
-				{
-					UProjectInfo.FillProjectInfo();
-				}
-
 				// Now generate project files
 				ProjectFileGenerator.bGenerateProjectFiles = true;
 				foreach(ProjectFileGenerator Generator in Generators)
diff --git a/Engine/Source/Programs/UnrealBuildTool/ProjectFiles/ProjectFileGenerator.cs b/Engine/Source/Programs/UnrealBuildTool/ProjectFiles/ProjectFileGenerator.cs
--- a/Engine/Source/Programs/UnrealBuildTool/ProjectFiles/ProjectFileGenerator.cs
+++ b/Engine/Source/Programs/UnrealBuildTool/ProjectFiles/ProjectFileGenerator.cs
@@ -473,7 +473,7 @@ namespace UnrealBuildTool
 			}
 			else
 			{
-				ProjectFiles.AddRange(UProjectInfo.AllProjectFiles.Where(x => ShouldGenerateProject(x)));
+				ProjectFiles.AddRange(NativeProjects.EnumerateProjectFiles().Where(x => ShouldGenerateProject(x)));
 			}
 			return ProjectFiles;
 		}
@@ -2077,8 +2077,8 @@ namespace UnrealBuildTool
 				{
 					RulesAssembly RulesAssembly;
 
-					FileReference CheckProjectFile;
-					if(!UProjectInfo.TryGetProjectForTarget(TargetName, out CheckProjectFile))
+					FileReference CheckProjectFile = AllGames.FirstOrDefault(x => TargetFilePath.IsUnderDirectory(x.Directory));
+					if(CheckProjectFile == null)
 					{
 						if(TargetFilePath.IsUnderDirectory(UnrealBuildTool.EnterpriseDirectory))
 						{
@@ -2176,7 +2176,7 @@ namespace UnrealBuildTool
 
 					bool bProjectAlreadyExisted;
 					ProjectFile ProjectFile = FindOrAddProject(ProjectFilePath, BaseFolder, IncludeInGeneratedProjects: true, bAlreadyExisted: out bProjectAlreadyExisted);
-					ProjectFile.IsForeignProject = CheckProjectFile != null && UProjectInfo.IsForeignProject(CheckProjectFile);
+					ProjectFile.IsForeignProject = CheckProjectFile != null && !NativeProjects.IsNativeProject(CheckProjectFile);
 					ProjectFile.IsGeneratedProject = true;
 					ProjectFile.IsStubProject = UnrealBuildTool.IsProjectInstalled();
 					if (TargetRulesObject.bBuildInSolutionByDefault.HasValue)
diff --git a/Engine/Source/Programs/UnrealBuildTool/System/NativeProjects.cs b/Engine/Source/Programs/UnrealBuildTool/System/NativeProjects.cs
new file mode 100644
--- /dev/null
+++ b/Engine/Source/Programs/UnrealBuildTool/System/NativeProjects.cs
@@ -0,0 +1,192 @@
+// Copyright 1998-2018 Epic Games, Inc. All Rights Reserved.
+
+using System;
+using System.Collections.Generic;
+using System.IO;
+using System.Linq;
+using System.Text;
+using System.Threading.Tasks;
+using Tools.DotNETCommon;
+
+namespace UnrealBuildTool
+{
+	/// <summary>
+	/// Utility functions for querying native projects (ie. those found via a .uprojectdirs query)
+	/// </summary>
+	public static class NativeProjects
+	{
+		/// <summary>
+		/// Object used for synchronizing access to static fields
+		/// </summary>
+		static object LockObject = new object();
+
+		/// <summary>
+		/// The native project base directories
+		/// </summary>
+		static HashSet<DirectoryReference> CachedBaseDirectories;
+
+		/// <summary>
+		/// Cached list of project files within all the base directories
+		/// </summary>
+		static HashSet<FileReference> CachedProjectFiles;
+
+		/// <summary>
+		/// Cached map of target names to the project file they belong to
+		/// </summary>
+		static Dictionary<string, FileReference> CachedTargetNameToProjectFile;
+
+		/// <summary>
+		/// Retrieve the list of base directories for native projects
+		/// </summary>
+		public static IEnumerable<DirectoryReference> EnumerateBaseDirectories()
+		{
+			if(CachedBaseDirectories == null)
+			{
+				lock(LockObject)
+				{
+					if(CachedBaseDirectories == null)
+					{
+						HashSet<DirectoryReference> BaseDirs = new HashSet<DirectoryReference>();
+						foreach (FileReference RootFile in DirectoryLookupCache.EnumerateFiles(UnrealBuildTool.RootDirectory))
+						{
+							if(RootFile.HasExtension(".uprojectdirs"))
+							{
+								foreach(string Line in File.ReadAllLines(RootFile.FullName))
+								{
+									string TrimLine = Line.Trim();
+									if(!TrimLine.StartsWith(";"))
+									{
+										DirectoryReference BaseProjectDir = DirectoryReference.Combine(UnrealBuildTool.RootDirectory, TrimLine);
+										if(BaseProjectDir.IsUnderDirectory(UnrealBuildTool.RootDirectory))
+										{
+											BaseDirs.Add(BaseProjectDir);
+										}
+										else
+										{
+											Log.TraceWarning("Project search path '{0}' referenced by '{1}' is not under '{2}', ignoring.", TrimLine, RootFile, UnrealBuildTool.RootDirectory);
+										}
+									}
+								}
+							}
+						}
+						CachedBaseDirectories = BaseDirs;
+					}
+				}
+			}
+			return CachedBaseDirectories;
+		}
+
+		/// <summary>
+		/// Returns a list of all the projects
+		/// </summary>
+		/// <returns>List of projects</returns>
+		public static IEnumerable<FileReference> EnumerateProjectFiles()
+		{
+			if(CachedProjectFiles == null)
+			{
+				lock(LockObject)
+				{
+					if(CachedProjectFiles == null)
+					{
+						HashSet<FileReference> ProjectFiles = new HashSet<FileReference>();
+						foreach(DirectoryReference BaseDirectory in EnumerateBaseDirectories())
+						{
+							if(DirectoryLookupCache.DirectoryExists(BaseDirectory))
+							{
+								foreach(DirectoryReference SubDirectory in DirectoryLookupCache.EnumerateDirectories(BaseDirectory))
+								{
+									foreach(FileReference File in DirectoryLookupCache.EnumerateFiles(SubDirectory))
+									{
+										if(File.HasExtension(".uproject"))
+										{
+											ProjectFiles.Add(File);
+										}
+									}
+								}
+							}
+						}
+						CachedProjectFiles = ProjectFiles;
+					}
+				}
+			}
+			return CachedProjectFiles;
+		}
+
+		/// <summary>
+		/// Get the project folder for the given target name
+		/// </summary>
+		/// <param name="InTargetName">Name of the target of interest</param>
+		/// <param name="OutProjectFileName">The project filename</param>
+		/// <returns>True if the target was found</returns>
+		public static bool TryGetProjectForTarget(string InTargetName, out FileReference OutProjectFileName)
+		{
+			if(CachedTargetNameToProjectFile == null)
+			{
+				lock(LockObject)
+				{
+					if(CachedTargetNameToProjectFile == null)
+					{
+						Dictionary<string, FileReference> TargetNameToProjectFile = new Dictionary<string, FileReference>();
+						foreach(FileReference ProjectFile in EnumerateProjectFiles())
+						{
+							DirectoryReference SourceDirectory = DirectoryReference.Combine(ProjectFile.Directory, "Source");
+							if(DirectoryLookupCache.DirectoryExists(SourceDirectory))
+							{
+								FindTargetFiles(SourceDirectory, TargetNameToProjectFile, ProjectFile);
+							}
+
+							DirectoryReference IntermediateSourceDirectory = DirectoryReference.Combine(ProjectFile.Directory, "Intermediate", "Source");
+							if(DirectoryLookupCache.DirectoryExists(IntermediateSourceDirectory))
+							{
+								FindTargetFiles(IntermediateSourceDirectory, TargetNameToProjectFile, ProjectFile);
+							}
+						}
+						CachedTargetNameToProjectFile = TargetNameToProjectFile;
+					}
+				}
+			}
+			return CachedTargetNameToProjectFile.TryGetValue(InTargetName, out OutProjectFileName);
+		}
+
+		/// <summary>
+		/// Finds all target files under a given folder, and add them to the target name to project file map
+		/// </summary>
+		/// <param name="Directory">Directory to search</param>
+		/// <param name="TargetNameToProjectFile">Map from target name to project file</param>
+		/// <param name="ProjectFile">The project file for this directory</param>
+		private static void FindTargetFiles(DirectoryReference Directory, Dictionary<string, FileReference> TargetNameToProjectFile, FileReference ProjectFile)
+		{
+			// Search for all target files within this directory
+			bool bSearchSubFolders = true;
+			foreach (FileReference File in DirectoryLookupCache.EnumerateFiles(Directory))
+			{
+				if (File.HasExtension(".target.cs"))
+				{
+					string TargetName = Path.GetFileNameWithoutExtension(File.GetFileNameWithoutExtension());
+					TargetNameToProjectFile[TargetName] = ProjectFile;
+					bSearchSubFolders = false;
+				}
+			}
+
+			// If we didn't find anything, recurse through the subfolders
+			if(bSearchSubFolders)
+			{
+				foreach(DirectoryReference SubDirectory in DirectoryLookupCache.EnumerateDirectories(Directory))
+				{
+					FindTargetFiles(SubDirectory, TargetNameToProjectFile, ProjectFile);
+				}
+			}
+		}
+
+		/// <summary>
+		/// Checks if a given project is a native project
+		/// </summary>
+		/// <param name="ProjectFile">The project file to check</param>
+		/// <returns>True if the given project is a native project</returns>
+		public static bool IsNativeProject(FileReference ProjectFile)
+		{
+			EnumerateProjectFiles();
+			return CachedProjectFiles.Contains(ProjectFile);
+		}
+	}
+}
diff --git a/Engine/Source/Programs/UnrealBuildTool/System/UProjectInfo.cs b/Engine/Source/Programs/UnrealBuildTool/System/UProjectInfo.cs
deleted file mode 100644
diff --git a/Engine/Source/Programs/UnrealBuildTool/UnrealBuildTool.csproj b/Engine/Source/Programs/UnrealBuildTool/UnrealBuildTool.csproj
--- a/Engine/Source/Programs/UnrealBuildTool/UnrealBuildTool.csproj
+++ b/Engine/Source/Programs/UnrealBuildTool/UnrealBuildTool.csproj
@@ -188,6 +188,7 @@
     <Compile Include="System\DataDrivenPlatformInfo.cs" />
     <Compile Include="System\EncryptionAndSigning.cs" />
     <Compile Include="System\HotReload.cs" />
+    <Compile Include="System\NativeProjects.cs" />
     <Compile Include="System\ModuleManifest.cs" />
     <Compile Include="System\PluginReferenceDescriptor.cs" />
     <Compile Include="System\PrecompiledManifest.cs" />
@@ -347,7 +348,6 @@
     <Compile Include="System\RulesCompiler.cs" />
     <Compile Include="System\SourceFileSearch.cs" />
     <Compile Include="System\Unity.cs" />
-    <Compile Include="System\UProjectInfo.cs" />
     <Compile Include="System\Utils.cs" />
     <Compile Include="ToolChain\UEToolChain.cs" />
     <Compile Include="Platform\Windows\UEBuildWindows.cs" />
