Missing file.

#rnx

diff --git a/Engine/Source/Programs/UnrealBuildTool/System/ActionGraph.cs b/Engine/Source/Programs/UnrealBuildTool/System/ActionGraph.cs
--- a/Engine/Source/Programs/UnrealBuildTool/System/ActionGraph.cs
+++ b/Engine/Source/Programs/UnrealBuildTool/System/ActionGraph.cs
@@ -1051,7 +1051,7 @@ namespace UnrealBuildTool
 		/// <returns>True if this is a C++ source file</returns>
 		private static bool IsCPPResourceFile(FileItem FileItem)
 		{
-			return (FileItem.AbsolutePath.EndsWith(".rc", StringComparison.InvariantCultureIgnoreCase));
+			return (FileItem.AbsolutePath.EndsWith(".rc", StringComparison.InvariantCultureIgnoreCase) || FileItem.AbsolutePath.EndsWith(".rc2", StringComparison.InvariantCultureIgnoreCase));
 		}
 
 
