UBT: Store the project file on each target descriptor, and remove the need for a single "global" active project at any one time. Add support for specifying multiple targets on the command line via -Target arguments (eg. -Target="UnrealHeaderTool Win64 Development" -Target="UnrealPak Win64 Development"), and multiple platforms and configurations via + characters (eg. "UE4Editor Win64 Debug+Development")

diff --git a/Engine/Source/Programs/DotNETCommon/DotNETUtilities/CommandLineArguments.cs b/Engine/Source/Programs/DotNETCommon/DotNETUtilities/CommandLineArguments.cs
--- a/Engine/Source/Programs/DotNETCommon/DotNETUtilities/CommandLineArguments.cs
+++ b/Engine/Source/Programs/DotNETCommon/DotNETUtilities/CommandLineArguments.cs
@@ -678,6 +678,100 @@ namespace Tools.DotNETCommon
 			}
 		}
 
+		/// <summary>
+		/// Splits a command line into individual arguments
+		/// </summary>
+		/// <param name="CommandLine">The command line text</param>
+		/// <returns>Array of arguments</returns>
+		public static string[] Split(string CommandLine)
+		{
+			StringBuilder Argument = new StringBuilder();
+
+			List<string> Arguments = new List<string>();
+			for(int Idx = 0; Idx < CommandLine.Length; Idx++)
+			{
+				if(!Char.IsWhiteSpace(CommandLine[Idx]))
+				{
+					Argument.Clear();
+					for(bool bInQuotes = false; Idx < CommandLine.Length; Idx++)
+					{
+						if(CommandLine[Idx] == '\"')
+						{
+							bInQuotes ^= true;
+						}
+						else if(!bInQuotes && Char.IsWhiteSpace(CommandLine[Idx]))
+						{
+							break;
+						}
+						else
+						{
+							Argument.Append(CommandLine[Idx]);
+						}
+					}
+					Arguments.Add(Argument.ToString());
+				}
+			}
+			return Arguments.ToArray();
+		}
+
+		/// <summary>
+		/// Appends the given arguments to the current argument list
+		/// </summary>
+		/// <param name="Arguments">The arguments to add</param>
+		/// <returns>New argument list</returns>
+		public CommandLineArguments Append(IEnumerable<string> AppendArguments)
+		{
+			CommandLineArguments NewArguments = new CommandLineArguments(Enumerable.Concat(Arguments, AppendArguments).ToArray());
+			for(int Idx = 0; Idx < Arguments.Length; Idx++)
+			{
+				if(HasBeenUsed(Idx))
+				{
+					NewArguments.MarkAsUsed(Idx);
+				}
+			}
+			return NewArguments;
+		}
+
+		/// <summary>
+		/// Retrieves all arguments with the given prefix, and returns the remaining a list of strings
+		/// </summary>
+		/// <param name="Prefix">Prefix for the arguments to remove</param>
+		/// <param name="Values">Receives a list of values with the given prefix</param>
+		/// <returns>New argument list</returns>
+		public CommandLineArguments Remove(string Prefix, out List<string> Values)
+		{
+			Values = new List<string>();
+
+			// Split the arguments into the values array and an array of new arguments
+			int[] NewArgumentIndex = new int[Arguments.Length];
+			List<string> NewArgumentList = new List<string>(Arguments.Length);
+			for(int Idx = 0; Idx < Arguments.Length; Idx++)
+			{
+				string Argument = Arguments[Idx];
+				if(Argument.StartsWith(Prefix, StringComparison.OrdinalIgnoreCase))
+				{
+					NewArgumentIndex[Idx] = -1;
+					Values.Add(Argument.Substring(Prefix.Length));
+				}
+				else
+				{
+					NewArgumentIndex[Idx] = NewArgumentList.Count;
+					NewArgumentList.Add(Argument);
+				}
+			}
+
+			// Create the new argument list, and mark the same arguments as used
+			CommandLineArguments NewArguments = new CommandLineArguments(NewArgumentList.ToArray());
+			for(int Idx = 0; Idx < Arguments.Length; Idx++)
+			{
+				if(HasBeenUsed(Idx) && NewArgumentIndex[Idx] != -1)
+				{
+					NewArguments.MarkAsUsed(NewArgumentIndex[Idx]);
+				}
+			}
+			return NewArguments;
+		}
+
 		/// <summary>
 		/// Checks that there are no unused arguments (and warns if there are)
 		/// </summary>
@@ -880,5 +974,50 @@ namespace Tools.DotNETCommon
 		{
 			return Arguments;
 		}
+
+		/// <summary>
+		/// Takes a command line argument and adds quotes if necessary
+		/// </summary>
+		/// <param name="Argument">The command line argument</param>
+		/// <returns>The command line argument with quotes inserted to escape it if necessary</returns>
+		public static void Append(StringBuilder CommandLine, string Argument)
+		{
+			if(CommandLine.Length > 0)
+			{
+				CommandLine.Append(' ');
+			}
+
+			int SpaceIdx = Argument.IndexOf(' ');
+			if(SpaceIdx == -1)
+			{
+				CommandLine.Append(Argument);
+			}
+			else
+			{
+				int EqualsIdx = Argument.IndexOf('=');
+				if(EqualsIdx == -1)
+				{
+					CommandLine.AppendFormat("\"{0}\"", Argument);
+				}
+				else
+				{
+					CommandLine.AppendFormat("{0}\"{1}\"", Argument.Substring(0, EqualsIdx + 1), Argument.Substring(EqualsIdx + 1));
+				}
+			}
+		}
+
+		/// <summary>
+		/// Converts this string to 
+		/// </summary>
+		/// <returns></returns>
+		public override string ToString()
+		{
+			StringBuilder Result = new StringBuilder();
+			foreach(string Argument in Arguments)
+			{
+				Append(Result, Argument);
+			}
+			return Result.ToString();
+		}
 	}
 }
diff --git a/Engine/Source/Programs/UnrealBuildTool/Configuration/TargetDescriptor.cs b/Engine/Source/Programs/UnrealBuildTool/Configuration/TargetDescriptor.cs
--- a/Engine/Source/Programs/UnrealBuildTool/Configuration/TargetDescriptor.cs
+++ b/Engine/Source/Programs/UnrealBuildTool/Configuration/TargetDescriptor.cs
@@ -20,7 +20,7 @@ namespace UnrealBuildTool
 		public UnrealTargetConfiguration Configuration;
 		public string Architecture;
 		public FileReference ForeignPlugin;
-		public string[] AdditionalArguments;
+		public CommandLineArguments AdditionalArguments;
 
 		/// <summary>
 		/// Constructor
@@ -30,15 +30,17 @@ namespace UnrealBuildTool
 		/// <param name="Platform">Platform to build for</param>
 		/// <param name="Configuration">Configuration to build</param>
 		/// <param name="Architecture">Architecture to build for</param>
-		/// <param name="Arguments">Additional arguments for the target</param>
-		public TargetDescriptor(FileReference ProjectFile, string TargetName, UnrealTargetPlatform Platform, UnrealTargetConfiguration Configuration, string Architecture, string[] Arguments)
+		/// <param name="ForeignPlugin">Path to a plugin to be compiled against this target. Usually null.</param>
+		/// <param name="AdditionalArguments">Additional arguments for the target</param>
+		public TargetDescriptor(FileReference ProjectFile, string TargetName, UnrealTargetPlatform Platform, UnrealTargetConfiguration Configuration, string Architecture, FileReference ForeignPlugin, CommandLineArguments AdditionalArguments)
 		{
 			this.ProjectFile = ProjectFile;
 			this.Name = TargetName;
 			this.Platform = Platform;
 			this.Configuration = Configuration;
 			this.Architecture = Architecture;
-			this.AdditionalArguments = Arguments;
+			this.ForeignPlugin = ForeignPlugin;
+			this.AdditionalArguments = AdditionalArguments;
 		}
 
 		/// <summary>
@@ -47,141 +49,245 @@ namespace UnrealBuildTool
 		/// <param name="Arguments">Command-line arguments</param>
 		/// <param name="bUsePrecompiled">Whether to use a precompiled engine distribution</param>
 		/// <param name="bSkipRulesCompile">Whether to skip compiling rules assemblies</param>
-		/// <param name="ProjectFile">The project file, if already set. May be updated if not.</param>
 		/// <returns>List of target descriptors</returns>
-		public static List<TargetDescriptor> ParseCommandLine(string[] Arguments, bool bUsePrecompiled, bool bSkipRulesCompile, ref FileReference ProjectFile)
+		public static List<TargetDescriptor> ParseCommandLine(CommandLineArguments Arguments, bool bUsePrecompiled, bool bSkipRulesCompile)
 		{
-			UnrealTargetPlatform Platform = UnrealTargetPlatform.Unknown;
-			UnrealTargetConfiguration Configuration = UnrealTargetConfiguration.Unknown;
-			List<string> TargetNames = new List<string>();
-			List<TargetType> TargetTypes = new List<TargetType>();
-			string Architecture = null;
-			FileReference ForeignPlugin = null;
-			List<string> AdditionalArguments = new List<string>();
+			List<TargetDescriptor> TargetDescriptors = new List<TargetDescriptor>();
+			ParseCommandLine(Arguments, bUsePrecompiled, bSkipRulesCompile, TargetDescriptors);
+			return TargetDescriptors;
+		}
 
-			// Settings for creating/using static libraries for the engine
-			for (int ArgumentIndex = 0; ArgumentIndex < Arguments.Length; ArgumentIndex++)
+		/// <summary>
+		/// Parse a list of target descriptors from the command line
+		/// </summary>
+		/// <param name="Arguments">Command-line arguments</param>
+		/// <param name="bUsePrecompiled">Whether to use a precompiled engine distribution</param>
+		/// <param name="bSkipRulesCompile">Whether to skip compiling rules assemblies</param>
+		/// <param name="TargetDescriptors">Receives the list of parsed target descriptors</param>
+		public static void ParseCommandLine(CommandLineArguments Arguments, bool bUsePrecompiled, bool bSkipRulesCompile, List<TargetDescriptor> TargetDescriptors)
+		{
+			List<string> TargetLists;
+			Arguments = Arguments.Remove("-TargetList=", out TargetLists);
+
+			List<string> Targets;
+			Arguments = Arguments.Remove("-Target=", out Targets);
+
+			if(TargetLists.Count > 0 || Targets.Count > 0)
 			{
-				string Argument = Arguments[ArgumentIndex];
-				if(!Argument.StartsWith("-"))
+				// Try to parse multiple arguments from a single command line
+				foreach(string TargetList in TargetLists)
 				{
-					UnrealTargetPlatform ParsedPlatform;
-					if(Enum.TryParse(Argument, true, out ParsedPlatform) && ParsedPlatform != UnrealTargetPlatform.Unknown)
+					string[] Lines = File.ReadAllLines(TargetList);
+					foreach(string Line in Lines)
 					{
-						if(Platform != UnrealTargetPlatform.Unknown)
+						string TrimLine = Line.Trim();
+						if(TrimLine.Length > 0 && TrimLine[0] != ';')
 						{
-							throw new BuildException("Multiple platforms specified on command line (first {0}, then {1})", Platform, ParsedPlatform);
+							CommandLineArguments NewArguments = Arguments.Append(CommandLineArguments.Split(TrimLine));
+							ParseCommandLine(NewArguments, bUsePrecompiled, bSkipRulesCompile, TargetDescriptors);
 						}
-						Platform = ParsedPlatform;
-						continue;
-					}
-
-					UnrealTargetConfiguration ParsedConfiguration;
-					if(Enum.TryParse(Argument, true, out ParsedConfiguration) && ParsedConfiguration != UnrealTargetConfiguration.Unknown)
-					{
-						if(Configuration != UnrealTargetConfiguration.Unknown)
-						{
-							throw new BuildException("Multiple configurations specified on command line (first {0}, then {1})", Configuration, ParsedConfiguration);
-						}
-						Configuration = ParsedConfiguration;
-						continue;
 					}
+				}
 
-					// Make sure the target name is valid. It may be the path to a project file.
-					if(Argument.IndexOfAny(new char[]{ Path.DirectorySeparatorChar, Path.AltDirectorySeparatorChar, '.' }) == -1)
-					{
-						TargetNames.Add(Argument);
-					}
+				foreach(string Target in Targets)
+				{
+					CommandLineArguments NewArguments = Arguments.Append(CommandLineArguments.Split(Target));
+					ParseCommandLine(NewArguments, bUsePrecompiled, bSkipRulesCompile, TargetDescriptors);
 				}
-				else
+			}
+			else
+			{
+				// Otherwise just process the whole command line together
+				ParseSingleCommandLine(Arguments, bUsePrecompiled, bSkipRulesCompile, TargetDescriptors);
+			}
+		}
+
+		/// <summary>
+		/// Parse a list of target descriptors from the command line
+		/// </summary>
+		/// <param name="Arguments">Command-line arguments</param>
+		/// <param name="bUsePrecompiled">Whether to use a precompiled engine distribution</param>
+		/// <param name="bSkipRulesCompile">Whether to skip compiling rules assemblies</param>
+		/// <param name="TargetDescriptors">List of target descriptors</param>
+		public static void ParseSingleCommandLine(CommandLineArguments Arguments, bool bUsePrecompiled, bool bSkipRulesCompile, List<TargetDescriptor> TargetDescriptors)
+		{
+			List<UnrealTargetPlatform> Platforms = new List<UnrealTargetPlatform>();
+			List<UnrealTargetConfiguration> Configurations = new List<UnrealTargetConfiguration>();
+			List<string> TargetNames = new List<string>();
+			FileReference ProjectFile = Arguments.GetFileReferenceOrDefault("-Project=", null);
+
+			// Settings for creating/using static libraries for the engine
+			for (int ArgumentIndex = 0; ArgumentIndex < Arguments.Count; ArgumentIndex++)
+			{
+				string Argument = Arguments[ArgumentIndex];
+				if(Argument.Length > 0 && Argument[0] != '-')
 				{
-					string Value;
-					if(ParseArgumentValue(Argument, "-TargetType=", out Value))
+					// Mark this argument as used. We'll interpret it as one thing or another.
+					Arguments.MarkAsUsed(ArgumentIndex);
+
+					// Check if it's a project file argument
+					if(Argument.EndsWith(".uproject", StringComparison.OrdinalIgnoreCase))
 					{
-						TargetType Type;
-						if(!Enum.TryParse(Value, true, out Type))
+						FileReference NewProjectFile = new FileReference(Argument);
+						if(ProjectFile != null && ProjectFile != NewProjectFile)
 						{
-							throw new BuildException("Invalid target type: '{0}'", Value);
+							throw new BuildException("Multiple project files specified on command line (first {0}, then {1})", ProjectFile, NewProjectFile);
 						}
-						TargetTypes.Add(Type);
+						ProjectFile = new FileReference(Argument);
+						continue;
 					}
-					else if(ParseArgumentValue(Argument, "-Plugin=", out Value))
+
+					// Split it into separate arguments
+					string[] InlineArguments = Argument.Split('+');
+
+					// Try to parse them as platforms
+					UnrealTargetPlatform ParsedPlatform;
+					if(Enum.TryParse(InlineArguments[0], true, out ParsedPlatform) && ParsedPlatform != UnrealTargetPlatform.Unknown)
 					{
-						if(ForeignPlugin != null)
+						Platforms.Add(ParsedPlatform);
+						for(int InlineArgumentIdx = 1; InlineArgumentIdx < InlineArguments.Length; InlineArgumentIdx++)
 						{
-							throw new BuildException("Only one foreign plugin to compile may be specified per invocation");
+							string InlineArgument = InlineArguments[InlineArgumentIdx];
+							if(!Enum.TryParse(InlineArgument, true, out ParsedPlatform) || ParsedPlatform == UnrealTargetPlatform.Unknown)
+							{
+								throw new BuildException("Invalid platform '{0}'", InlineArgument);
+							}
+							Platforms.Add(ParsedPlatform);
 						}
-						ForeignPlugin = new FileReference(Value);
+						continue;
 					}
-					else
+
+					// Try to parse them as configurations
+					UnrealTargetConfiguration ParsedConfiguration;
+					if(Enum.TryParse(InlineArguments[0], true, out ParsedConfiguration))
 					{
-						switch (Arguments[ArgumentIndex].ToUpperInvariant())
+						Configurations.Add(ParsedConfiguration);
+						for(int InlineArgumentIdx = 1; InlineArgumentIdx < InlineArguments.Length; InlineArgumentIdx++)
 						{
-							case "-MODULE":
-								throw new BuildException("'-Module <Name>' syntax is no longer supported on the command line. Use '-Module=<Name>' instead.");
-							case "-MODULEWITHSUFFIX":
-								throw new BuildException("'-ModuleWithSuffix <Name> <Suffix>' syntax is no longer supported on the command line. Use '-Module=<Name>,<Suffix>' instead.");
-							case "-PLUGIN":
-								throw new BuildException("'-Plugin <Path>' syntax is no longer supported on the command line. Use '-Plugin=<Path>' instead.");
-							case "-RECEIPT":
-								throw new BuildException("'-Receipt <Path>' syntax is no longer supported on the command line. Use '-Receipt=<Path>' instead.");
-							default:
-								AdditionalArguments.Add(Arguments[ArgumentIndex]);
-								break;
+							string InlineArgument = InlineArguments[InlineArgumentIdx];
+							if(!Enum.TryParse(InlineArgument, true, out ParsedConfiguration))
+							{
+								throw new BuildException("Invalid configuration '{0}'", InlineArgument);
+							}
+							Configurations.Add(ParsedConfiguration);
 						}
+						continue;
 					}
+
+					// Otherwise assume they are target names
+					TargetNames.AddRange(InlineArguments);
 				}
 			}
 
-			if (Platform == UnrealTargetPlatform.Unknown)
+			if (Platforms.Count == 0)
 			{
-				throw new BuildException("Couldn't find platform name.");
+				throw new BuildException("No platforms specified for target");
 			}
-			if (Configuration == UnrealTargetConfiguration.Unknown)
+			if (Configurations.Count == 0)
 			{
-				throw new BuildException("Couldn't determine configuration name.");
+				throw new BuildException("No configurations specified for target");
 			}
 
-			if(Architecture == null)
-			{
-				Architecture = UEBuildPlatform.GetBuildPlatform(Platform).GetDefaultArchitecture(ProjectFile);
-			}
+			FileReference ForeignPlugin = Arguments.GetFileReferenceOrDefault("-Plugin=", null);
 
-			// Create all the target descriptors for targets specified by type
-			foreach(TargetType Type in TargetTypes)
+			foreach(UnrealTargetPlatform Platform in Platforms)
 			{
-				if (ProjectFile == null)
+				// Parse the architecture parameter, or get the default for the platform
+				List<string> Architectures = new List<string>(Arguments.GetValues("-Architecture=", '+'));
+				if(Architectures.Count == 0)
 				{
-					throw new BuildException("-TargetType=... requires a project file to be specified");
+					Architectures.Add(UEBuildPlatform.GetBuildPlatform(Platform).GetDefaultArchitecture(ProjectFile));
 				}
-				else
+
+				foreach(string Architecture in Architectures)
 				{
-					TargetNames.Add(RulesCompiler.CreateProjectRulesAssembly(ProjectFile, bUsePrecompiled, bSkipRulesCompile).GetTargetNameByType(Type, Platform, Configuration, Architecture, ProjectFile));
+					foreach(UnrealTargetConfiguration Configuration in Configurations)
+					{
+						// Create all the target descriptors for targets specified by type
+						foreach(string TargetTypeString in Arguments.GetValues("-TargetType="))
+						{
+							TargetType TargetType;
+							if(!Enum.TryParse(TargetTypeString, out TargetType))
+							{
+								throw new BuildException("Invalid target type '{0}'", TargetTypeString);
+							}
+
+							if (ProjectFile == null)
+							{
+								throw new BuildException("-TargetType=... requires a project file to be specified");
+							}
+							else
+							{
+								TargetNames.Add(RulesCompiler.CreateProjectRulesAssembly(ProjectFile, bUsePrecompiled, bSkipRulesCompile).GetTargetNameByType(TargetType, Platform, Configuration, Architecture, ProjectFile));
+							}
+						}
+
+						// Make sure we could parse something
+						if (TargetNames.Count == 0)
+						{
+							throw new BuildException("No target name was specified on the command-line.");
+						}
+
+						// Create all the target descriptors
+						foreach(string TargetName in TargetNames)
+						{
+							// If a project file was not specified see if we can find one
+							if (ProjectFile == null && NativeProjects.TryGetProjectForTarget(TargetName, out ProjectFile))
+							{
+								Log.TraceVerbose("Found project file for {0} - {1}", TargetName, ProjectFile);
+							}
+
+							// Pull out all the arguments that haven't been used so far
+							List<string> AdditionalArguments = new List<string>();
+							for(int Idx = 0; Idx < Arguments.Count; Idx++)
+							{
+								if(!Arguments.HasBeenUsed(Idx))
+								{
+									AdditionalArguments.Add(Arguments[Idx]);
+								}
+							}
+
+							// Create the target descriptor
+							TargetDescriptors.Add(new TargetDescriptor(ProjectFile, TargetName, Platform, Configuration, Architecture, ForeignPlugin, new CommandLineArguments(AdditionalArguments.ToArray())));
+						}
+					}
 				}
 			}
+		}
+
+		/// <summary>
+		/// Try to parse the project file from the command line
+		/// </summary>
+		/// <param name="Arguments">The command line arguments</param>
+		/// <param name="ProjectFile">The project file that was parsed</param>
+		/// <returns>True if the project file was parsed, false otherwise</returns>
+		public static bool TryParseProjectFileArgument(CommandLineArguments Arguments, out FileReference ProjectFile)
+		{
+			FileReference ExplicitProjectFile;
+			if(Arguments.TryGetValue("-Project=", out ExplicitProjectFile))
+			{
+				ProjectFile = ExplicitProjectFile;
+				return true;
+			}
 
-			// Create all the target descriptor
-			List<TargetDescriptor> Targets = new List<TargetDescriptor>();
-			foreach(string TargetName in TargetNames)
+			for(int Idx = 0; Idx < Arguments.Count; Idx++)
 			{
-				// If a project file was not specified see if we can find one
-				if (ProjectFile == null && NativeProjects.TryGetProjectForTarget(TargetName, out ProjectFile))
+				if(Arguments[Idx][0] != '-' && Arguments[Idx].EndsWith(".uproject", StringComparison.OrdinalIgnoreCase))
 				{
-					Log.TraceLog("Found project file for {0}: {1}", TargetName, ProjectFile);
+					Arguments.MarkAsUsed(Idx);
+					ProjectFile = new FileReference(Arguments[Idx]);
+					return true;
 				}
-
-				// Create the target descriptor
-				TargetDescriptor Target = new TargetDescriptor(ProjectFile, TargetName, Platform, Configuration, Architecture, AdditionalArguments.ToArray());
-				Target.ForeignPlugin = ForeignPlugin;
-				Targets.Add(Target);
 			}
 
-			// Make sure we could parse something
-			if (Targets.Count == 0)
+			if(UnrealBuildTool.IsProjectInstalled())
 			{
-				throw new BuildException("No target name was specified on the command-line.");
+				ProjectFile = UnrealBuildTool.GetInstalledProjectFile();
+				return true;
 			}
-			return Targets;
+
+			ProjectFile = null;
+			return false;
 		}
 
 		/// <summary>
@@ -204,5 +310,28 @@ namespace UnrealBuildTool
 				return false;
 			}
 		}
+
+		/// <summary>
+		/// Format this object for the debugger
+		/// </summary>
+		/// <returns>String representation of this target descriptor</returns>
+		public override string ToString()
+		{
+			StringBuilder Result = new StringBuilder();
+			Result.AppendFormat("{0} {1} {2}", Name, Platform, Configuration);
+			if(!String.IsNullOrEmpty(Architecture))
+			{
+				Result.AppendFormat(" -Architecture={0}", Architecture);
+			}
+			if(ProjectFile != null)
+			{
+				Result.AppendFormat(" -Project={0}", Utils.MakePathSafeToUseWithCommandLine(ProjectFile));
+			}
+			if(AdditionalArguments.Count > 0)
+			{
+				Result.AppendFormat(" {0}", AdditionalArguments);
+			}
+			return Result.ToString();
+		}
 	}
 }
diff --git a/Engine/Source/Programs/UnrealBuildTool/Modes/BuildMode.cs b/Engine/Source/Programs/UnrealBuildTool/Modes/BuildMode.cs
--- a/Engine/Source/Programs/UnrealBuildTool/Modes/BuildMode.cs
+++ b/Engine/Source/Programs/UnrealBuildTool/Modes/BuildMode.cs
@@ -97,16 +97,14 @@ namespace UnrealBuildTool
 				Arguments.TryGetValue("-RemoteIni=", out RemoteIniPath);
 				UnrealBuildTool.SetRemoteIniPath(RemoteIniPath);
 
-				// Build the list of game projects that we know about. When building from the editor (for hot-reload) or for projects from installed builds, we require the 
-				// project file to be passed in. Otherwise we scan for projects in directories named in UE4Games.uprojectdirs.
-				FileReference ProjectFile;
-				UnrealBuildTool.TryParseProjectFileArgument(Arguments, out ProjectFile);
-				
 				DateTime BasicInitStartTime = DateTime.UtcNow;
 
 				// Find and register all tool chains, build platforms, etc. that are present
 				UnrealBuildTool.RegisterAllUBTClasses(false);
 
+				// Parse all the targets to build
+				List<TargetDescriptor> Targets = TargetDescriptor.ParseCommandLine(Arguments, false, false);
+
 				if (UnrealBuildTool.bPrintPerformanceInfo)
 				{
 					double BasicInitTime = (DateTime.UtcNow - BasicInitStartTime).TotalSeconds;
@@ -120,7 +118,7 @@ namespace UnrealBuildTool
 				}
 
 				// Build our project
-				ECompilationResult Result = UnrealBuildTool.RunUBT(BuildConfiguration, Arguments.GetRawArray(), ProjectFile, LogFile);
+				ECompilationResult Result = UnrealBuildTool.RunUBT(BuildConfiguration, Arguments.GetRawArray(), LogFile);
 
 				// Print some performance info
 				double BuildDuration = (DateTime.UtcNow - StartTime).TotalSeconds;
diff --git a/Engine/Source/Programs/UnrealBuildTool/ProjectFiles/ProjectFileGenerator.cs b/Engine/Source/Programs/UnrealBuildTool/ProjectFiles/ProjectFileGenerator.cs
--- a/Engine/Source/Programs/UnrealBuildTool/ProjectFiles/ProjectFileGenerator.cs
+++ b/Engine/Source/Programs/UnrealBuildTool/ProjectFiles/ProjectFileGenerator.cs
@@ -1704,7 +1704,7 @@ namespace UnrealBuildTool
 							const string DefaultArchitecture = "";
 
 							// Create the target descriptor
-							TargetDescriptor TargetDesc = new TargetDescriptor(CurTarget.UnrealProjectFilePath, CurTarget.Name, BuildHostPlatform.Current.Platform, UnrealTargetConfiguration.Development, DefaultArchitecture, NewArguments.ToArray());
+							TargetDescriptor TargetDesc = new TargetDescriptor(CurTarget.UnrealProjectFilePath, CurTarget.Name, BuildHostPlatform.Current.Platform, UnrealTargetConfiguration.Development, DefaultArchitecture, null, new CommandLineArguments(NewArguments.ToArray()));
 
 							// Create the target
 							UEBuildTarget Target = UEBuildTarget.CreateTarget(TargetDesc, NewArguments.ToArray(), false, false, bUsePrecompiled);
diff --git a/Engine/Source/Programs/UnrealBuildTool/System/SourceFileWorkingSet.cs b/Engine/Source/Programs/UnrealBuildTool/System/SourceFileWorkingSet.cs
--- a/Engine/Source/Programs/UnrealBuildTool/System/SourceFileWorkingSet.cs
+++ b/Engine/Source/Programs/UnrealBuildTool/System/SourceFileWorkingSet.cs
@@ -312,9 +312,9 @@ namespace UnrealBuildTool
 		/// Create an ISourceFileWorkingSet instance suitable for the given project or root directory
 		/// </summary>
 		/// <param name="RootDir">The root directory</param>
-		/// <param name="ProjectDir">The current project</param>
+		/// <param name="ProjectDirs">The project directories</param>
 		/// <returns>Working set instance for the given directory</returns>
-		public static ISourceFileWorkingSet Create(DirectoryReference RootDir, DirectoryReference ProjectDir)
+		public static ISourceFileWorkingSet Create(DirectoryReference RootDir, IEnumerable<DirectoryReference> ProjectDirs)
 		{
 			if (Provider == ProviderType.None || ProjectFileGenerator.bGenerateProjectFiles)
 			{
@@ -327,7 +327,7 @@ namespace UnrealBuildTool
 				{
 					WorkingSet = new GitSourceFileWorkingSet(GitPath, DirectoryReference.Combine(RootDir, RepositoryPath), null);
 				}
-				else if(!TryCreateGitWorkingSet(RootDir, ProjectDir, out WorkingSet))
+				else if(!TryCreateGitWorkingSet(RootDir, ProjectDirs, out WorkingSet))
 				{
 					WorkingSet = new GitSourceFileWorkingSet(GitPath, RootDir, null);
 				}
@@ -340,7 +340,7 @@ namespace UnrealBuildTool
 			else
 			{
 				GitSourceFileWorkingSet WorkingSet;
-				if(TryCreateGitWorkingSet(RootDir, ProjectDir, out WorkingSet))
+				if(TryCreateGitWorkingSet(RootDir, ProjectDirs, out WorkingSet))
 				{
 					return WorkingSet;
 				}
@@ -351,7 +351,7 @@ namespace UnrealBuildTool
 			}
 		}
 
-		static bool TryCreateGitWorkingSet(DirectoryReference RootDir, DirectoryReference ProjectDir, out GitSourceFileWorkingSet OutWorkingSet)
+		static bool TryCreateGitWorkingSet(DirectoryReference RootDir, IEnumerable<DirectoryReference> ProjectDirs, out GitSourceFileWorkingSet OutWorkingSet)
 		{
 			GitSourceFileWorkingSet WorkingSet  = null;
 
@@ -362,7 +362,7 @@ namespace UnrealBuildTool
 			}
 
 			// Try to create a working set for the project directory
-			if(ProjectDir != null)
+			foreach(DirectoryReference ProjectDir in ProjectDirs)
 			{
 				if(WorkingSet == null || !ProjectDir.IsUnderDirectory(RootDir))
 				{
diff --git a/Engine/Source/Programs/UnrealBuildTool/UnrealBuildTool.cs b/Engine/Source/Programs/UnrealBuildTool/UnrealBuildTool.cs
--- a/Engine/Source/Programs/UnrealBuildTool/UnrealBuildTool.cs
+++ b/Engine/Source/Programs/UnrealBuildTool/UnrealBuildTool.cs
@@ -167,6 +167,22 @@ namespace UnrealBuildTool
 			return bIsProjectInstalled.Value;
 		}
 
+		/// <summary>
+		/// Gets the installed project file
+		/// </summary>
+		/// <returns>Location of the installed project file</returns>
+		static public FileReference GetInstalledProjectFile()
+		{
+			if(IsProjectInstalled())
+			{
+				return InstalledProjectFile;
+			}
+			else
+			{
+				return null;
+			}
+		}
+
 		/// <summary>
 		/// Checks whether the given file is under an installed directory, and should not be overridden
 		/// </summary>
@@ -436,7 +452,7 @@ namespace UnrealBuildTool
 			}
 		}
 
-		internal static ECompilationResult RunUBT(BuildConfiguration BuildConfiguration, string[] Arguments, FileReference ProjectFile, FileReference LogFile)
+		internal static ECompilationResult RunUBT(BuildConfiguration BuildConfiguration, string[] Arguments, FileReference LogFile)
 		{
 			bool bSuccess = true;
 
@@ -458,8 +474,6 @@ namespace UnrealBuildTool
 
 			try
 			{
-				List<string[]> TargetSettings = ParseCommandLineFlags(Arguments);
-
 				int ArgumentIndex;
 				// action graph implies using the dependency resolve cache
 				bool GeneratingActionGraph = Utils.ParseCommandLineFlag(Arguments, "-graph", out ArgumentIndex);
@@ -480,10 +494,7 @@ namespace UnrealBuildTool
 				{
 					DateTime TargetDescstStartTime = DateTime.UtcNow;
 
-					foreach (string[] TargetSetting in TargetSettings)
-					{
-						TargetDescs.AddRange(TargetDescriptor.ParseCommandLine(TargetSetting, BuildConfiguration.bUsePrecompiled, bSkipRulesCompile, ref ProjectFile));
-					}
+					TargetDescs.AddRange(TargetDescriptor.ParseCommandLine(new CommandLineArguments(Arguments), BuildConfiguration.bUsePrecompiled, bSkipRulesCompile));
 
 					if (UnrealBuildTool.bPrintPerformanceInfo)
 					{
@@ -545,7 +556,7 @@ namespace UnrealBuildTool
 					// at this point whether we're dealing with modular or monolithic game binaries, so we opt to always invalidate
 					// cached includes if the target we're switching to is either a game target (has project file) or "UE4Editor".
 					bool bMightHaveSharedBuildProducts =
-						ProjectFile != null ||  // Is this a game? (has a .uproject file for the target)
+						TargetDescs[0].ProjectFile != null ||  // Is this a game? (has a .uproject file for the target)
 						TargetDescs[0].Name.Equals("UE4Editor", StringComparison.InvariantCultureIgnoreCase); // Is the engine?
 					if (bMightHaveSharedBuildProducts)
 					{
@@ -630,7 +641,18 @@ namespace UnrealBuildTool
 					throw new BuildException("UnrealBuildTool: At least one of either IsGatheringBuild or IsAssemblingBuild must be true.  Did you pass '-NoGather' with '-NoAssemble'?");
 				}
 
-				using (ISourceFileWorkingSet WorkingSet = SourceFileWorkingSet.Create(UnrealBuildTool.RootDirectory, DirectoryReference.FromFile(ProjectFile)))
+				// Get a set of all the project directories.
+				HashSet<DirectoryReference> ProjectDirs = new HashSet<DirectoryReference>();
+				foreach(TargetDescriptor TargetDesc in TargetDescs)
+				{
+					if(TargetDesc.ProjectFile != null)
+					{
+						ProjectDirs.Add(TargetDesc.ProjectFile.Directory);
+					}
+				}
+
+				// Create the working set provider
+				using (ISourceFileWorkingSet WorkingSet = SourceFileWorkingSet.Create(UnrealBuildTool.RootDirectory, ProjectDirs))
 				{
 					UBTMakefile UBTMakefile = null;
 					{
@@ -651,7 +673,7 @@ namespace UnrealBuildTool
 
 							// Try to load the UBTMakefile.  It will only be loaded if it has valid content and is not determined to be out of date.    
 							string ReasonNotLoaded;
-							UBTMakefile = UBTMakefile.LoadUBTMakefile(UBTMakefilePath, ProjectFile, WorkingSet, out ReasonNotLoaded);
+							UBTMakefile = UBTMakefile.LoadUBTMakefile(UBTMakefilePath, TargetDescs[0].ProjectFile, WorkingSet, out ReasonNotLoaded);
 
 							if (UBTMakefile != null)
 							{
@@ -659,7 +681,7 @@ namespace UnrealBuildTool
 								FileInfo UBTMakefileInfo = new FileInfo(UBTMakefilePath.FullName);
 								foreach (TargetDescriptor Desc in TargetDescs)
 								{
-									DirectoryReference ProjectDirectory = DirectoryReference.FromFile(ProjectFile);
+									DirectoryReference ProjectDirectory = DirectoryReference.FromFile(Desc.ProjectFile);
 									foreach (ConfigHierarchyType IniType in (ConfigHierarchyType[])Enum.GetValues(typeof(ConfigHierarchyType)))
 									{
 										foreach (FileReference IniFilename in ConfigHierarchy.EnumerateConfigFileLocations(IniType, ProjectDirectory, Desc.Platform))
@@ -739,9 +761,9 @@ namespace UnrealBuildTool
 					foreach (UEBuildTarget Target in Targets)
 					{
 						// Create the header cache for this target
-						FileReference DependencyCacheFile = DependencyCache.GetDependencyCachePathForTarget(ProjectFile, Target.Platform, Target.TargetName);
+						FileReference DependencyCacheFile = DependencyCache.GetDependencyCachePathForTarget(Target.ProjectFile, Target.Platform, Target.TargetName);
 						bool bUseFlatCPPIncludeDependencyCache = BuildConfiguration.bUseUBTMakefiles && bIsAssemblingBuild;
-						CPPHeaders Headers = new CPPHeaders(ProjectFile, DependencyCacheFile, bUseFlatCPPIncludeDependencyCache, BuildConfiguration.bUseUBTMakefiles, BuildConfiguration.bUseIncludeDependencyResolveCache, BuildConfiguration.bTestIncludeDependencyResolveCache);
+						CPPHeaders Headers = new CPPHeaders(Target.ProjectFile, DependencyCacheFile, bUseFlatCPPIncludeDependencyCache, BuildConfiguration.bUseUBTMakefiles, BuildConfiguration.bUseIncludeDependencyResolveCache, BuildConfiguration.bTestIncludeDependencyResolveCache);
 						TargetToHeaders[Target] = Headers;
 
 						// Make sure the appropriate executor is selected
@@ -1296,41 +1318,6 @@ namespace UnrealBuildTool
 			}
 		}
 
-		/// <summary>
-		/// Parses the passed in command line for build configuration overrides.
-		/// </summary>
-		/// <param name="Arguments">List of arguments to parse</param>
-		/// <returns>List of build target settings</returns>
-		private static List<string[]> ParseCommandLineFlags(string[] Arguments)
-		{
-			List<string[]> TargetSettings = new List<string[]>();
-			int ArgumentIndex = 0;
-
-			if (Utils.ParseCommandLineFlag(Arguments, "-targets", out ArgumentIndex))
-			{
-				if (ArgumentIndex + 1 >= Arguments.Length)
-				{
-					throw new BuildException("Expected filename after -targets argument, but found nothing.");
-				}
-				// Parse lines from the referenced file into target settings.
-				string[] Lines = File.ReadAllLines(Arguments[ArgumentIndex + 1]);
-				foreach (string Line in Lines)
-				{
-					if (Line != "" && Line[0] != ';')
-					{
-						TargetSettings.Add(Line.Split(' '));
-					}
-				}
-			}
-			// Simply use full command line arguments as target setting if not otherwise overridden.
-			else
-			{
-				TargetSettings.Add(Arguments);
-			}
-
-			return TargetSettings;
-		}
-
 		/// <summary>
 		/// Helper class to update the C++ dependency cache on a background thread. Captures exceptions and re-throws on the main thread when joined.
 		/// </summary>
