Remove debug code to generate targets.

#rnx

diff --git a/Engine/Source/Programs/UnrealBuildTool/Modes/BuildMode.cs b/Engine/Source/Programs/UnrealBuildTool/Modes/BuildMode.cs
--- a/Engine/Source/Programs/UnrealBuildTool/Modes/BuildMode.cs
+++ b/Engine/Source/Programs/UnrealBuildTool/Modes/BuildMode.cs
@@ -102,9 +102,6 @@ namespace UnrealBuildTool
 				// Find and register all tool chains, build platforms, etc. that are present
 				UnrealBuildTool.RegisterAllUBTClasses(false);
 
-				// Parse all the targets to build
-				List<TargetDescriptor> Targets = TargetDescriptor.ParseCommandLine(Arguments, false, false);
-
 				if (UnrealBuildTool.bPrintPerformanceInfo)
 				{
 					double BasicInitTime = (DateTime.UtcNow - BasicInitStartTime).TotalSeconds;
