UBT: Initial support for Visual Studio 2019.

diff --git a/Engine/Source/Programs/UnrealBuildTool/Modes/GenerateProjectFilesMode.cs b/Engine/Source/Programs/UnrealBuildTool/Modes/GenerateProjectFilesMode.cs
--- a/Engine/Source/Programs/UnrealBuildTool/Modes/GenerateProjectFilesMode.cs
+++ b/Engine/Source/Programs/UnrealBuildTool/Modes/GenerateProjectFilesMode.cs
@@ -25,6 +25,7 @@ namespace UnrealBuildTool
 		[CommandLine("-2013unsupported", Value = nameof(ProjectFileFormat.VisualStudio2013))]
 		[CommandLine("-2015", Value = nameof(ProjectFileFormat.VisualStudio2015))] // + override compiler
 		[CommandLine("-2017", Value = nameof(ProjectFileFormat.VisualStudio2017))] // + override compiler
+		[CommandLine("-2019", Value = nameof(ProjectFileFormat.VisualStudio2019))] // + override compiler
 		[CommandLine("-Makefile", Value = nameof(ProjectFileFormat.Make))]
 		[CommandLine("-CMakefile", Value = nameof(ProjectFileFormat.CMake))]
 		[CommandLine("-QMakefile", Value = nameof(ProjectFileFormat.QMake))]
@@ -138,6 +139,9 @@ namespace UnrealBuildTool
 						case ProjectFileFormat.VisualStudio2017:
 							Generator = new VCProjectFileGenerator(ProjectFile, VCProjectFileFormat.VisualStudio2017, Arguments);
 							break;
+						case ProjectFileFormat.VisualStudio2019:
+							Generator = new VCProjectFileGenerator(ProjectFile, VCProjectFileFormat.VisualStudio2019, Arguments);
+							break;
 						case ProjectFileFormat.XCode:
 							Generator = new XcodeProjectFileGenerator(ProjectFile);
 							break;
diff --git a/Engine/Source/Programs/UnrealBuildTool/Platform/Windows/UEBuildWindows.cs b/Engine/Source/Programs/UnrealBuildTool/Platform/Windows/UEBuildWindows.cs
--- a/Engine/Source/Programs/UnrealBuildTool/Platform/Windows/UEBuildWindows.cs
+++ b/Engine/Source/Programs/UnrealBuildTool/Platform/Windows/UEBuildWindows.cs
@@ -48,6 +48,11 @@ namespace UnrealBuildTool
 		/// Visual Studio 2017 (Visual C++ 15.0)
 		/// </summary>
 		VisualStudio2017,
+
+		/// <summary>
+		/// Visual Studio 2019 (Visual C++ 16.0)
+		/// </summary>
+		VisualStudio2019,
 	}
 
 	/// <summary>
@@ -83,6 +88,7 @@ namespace UnrealBuildTool
 		[XmlConfigFile(Category = "WindowsPlatform")]
 		[CommandLine("-2015", Value = "VisualStudio2015")]
 		[CommandLine("-2017", Value = "VisualStudio2017")]
+		[CommandLine("-2019", Value = "VisualStudio2019")]
 		public WindowsCompiler Compiler = WindowsCompiler.Default;
 
 		/// <summary>
@@ -159,7 +165,7 @@ namespace UnrealBuildTool
 		/// Microsoft provides legacy_stdio_definitions library to enable building with VS2015 until they fix everything up.
 		public bool bNeedsLegacyStdioDefinitionsLib
 		{
-			get { return Compiler == WindowsCompiler.VisualStudio2015 || Compiler == WindowsCompiler.VisualStudio2017 || Compiler == WindowsCompiler.Clang; }
+			get { return Compiler == WindowsCompiler.VisualStudio2015 || Compiler == WindowsCompiler.VisualStudio2017 || Compiler == WindowsCompiler.VisualStudio2019 || Compiler == WindowsCompiler.Clang; }
 		}
 
 		/// <summary>
@@ -233,6 +239,7 @@ namespace UnrealBuildTool
 				case WindowsCompiler.Intel:
 				case WindowsCompiler.VisualStudio2015:
 				case WindowsCompiler.VisualStudio2017:
+				case WindowsCompiler.VisualStudio2019:
 					return "2015"; // VS2017 is backwards compatible with VS2015 compiler
 
 				default:
@@ -559,11 +566,32 @@ namespace UnrealBuildTool
 		internal static WindowsCompiler GetDefaultCompiler(FileReference ProjectFile)
 		{
 			// If there's no specific compiler set, try to pick the matching compiler for the selected IDE
+			foreach(ProjectFileFormat Format in ProjectFileGeneratorSettings.ParseFormatList(ProjectFileGeneratorSettings.Format))
+			{
+				if (Format == ProjectFileFormat.VisualStudio2019)
+				{
+					return WindowsCompiler.VisualStudio2019;
+				}
+				else if (Format == ProjectFileFormat.VisualStudio2017)
+				{
+					return WindowsCompiler.VisualStudio2017;
+				}
+				else if (Format == ProjectFileFormat.VisualStudio2015)
+				{
+					return WindowsCompiler.VisualStudio2015;
+				}
+			}
+
+			// Also check the default format for the Visual Studio project generator
 			object ProjectFormatObject;
 			if (XmlConfig.TryGetValue(typeof(VCProjectFileGenerator), "Version", out ProjectFormatObject))
 			{
 				VCProjectFileFormat ProjectFormat = (VCProjectFileFormat)ProjectFormatObject;
-				if (ProjectFormat == VCProjectFileFormat.VisualStudio2017)
+				if (ProjectFormat == VCProjectFileFormat.VisualStudio2019)
+				{
+					return WindowsCompiler.VisualStudio2019;
+				}
+				else if (ProjectFormat == VCProjectFileFormat.VisualStudio2017)
 				{
 					return WindowsCompiler.VisualStudio2017;
 				}
@@ -577,7 +605,11 @@ namespace UnrealBuildTool
 			ProjectFileFormat PreferredAccessor;
 			if(ProjectFileGenerator.GetPreferredSourceCodeAccessor(ProjectFile, out PreferredAccessor))
 			{
-				if(PreferredAccessor == ProjectFileFormat.VisualStudio2017)
+				if(PreferredAccessor == ProjectFileFormat.VisualStudio2019)
+			    {
+				    return WindowsCompiler.VisualStudio2019;
+			    }
+				else if(PreferredAccessor == ProjectFileFormat.VisualStudio2017)
 			    {
 				    return WindowsCompiler.VisualStudio2017;
 			    }
@@ -596,6 +628,10 @@ namespace UnrealBuildTool
 			{
 				return WindowsCompiler.VisualStudio2015;
 			}
+			if (HasCompiler(WindowsCompiler.VisualStudio2019))
+			{
+				return WindowsCompiler.VisualStudio2019;
+			}
 
 			// If we do have a Visual Studio installation, but we're missing just the C++ parts, warn about that.
 			DirectoryReference VSInstallDir;
@@ -607,6 +643,10 @@ namespace UnrealBuildTool
 			{
 				Log.TraceWarning("Visual Studio 2015 is installed, but is missing the C++ toolchain. Please verify that \"Common Tools for Visual C++ 2015\" are selected from the Visual Studio 2015 installation options.");
 			}
+			else if (TryGetVSInstallDir(WindowsCompiler.VisualStudio2019, out VSInstallDir))
+			{
+				Log.TraceWarning("Visual Studio 2019 is installed, but is missing the C++ toolchain. Please verify that the \"VC++ 2019 toolset\" component is selected in the Visual Studio 2019 installation options.");
+			}
 			else
 			{
 				Log.TraceWarning("No Visual C++ installation was found. Please download and install Visual Studio 2015 with C++ components.");
@@ -629,6 +669,8 @@ namespace UnrealBuildTool
 					return "Visual Studio 2015";
 				case WindowsCompiler.VisualStudio2017:
 					return "Visual Studio 2017";
+				case WindowsCompiler.VisualStudio2019:
+					return "Visual Studio 2019";
 				default:
 					return Compiler.ToString();
 			}
@@ -685,29 +727,44 @@ namespace UnrealBuildTool
 						    InstallDirs.Add(InstallDir);
 					    }
 				    }
-				    else if(Compiler == WindowsCompiler.VisualStudio2017)
+				    else if(Compiler == WindowsCompiler.VisualStudio2017 || Compiler == WindowsCompiler.VisualStudio2019)
 				    {
-					    // Enumerate all the installed Visual Studio instances using the interop SDK. There may be several installed side by side on a single machine (preview, community, enterprise, etc...).
 						List<DirectoryReference> PreReleaseInstallDirs = new List<DirectoryReference>();
-					    try
-					    {
-						    SetupConfiguration Setup = new SetupConfiguration();
-						    IEnumSetupInstances Enumerator = Setup.EnumAllInstances(); 
-    
-						    ISetupInstance[] Instances = new ISetupInstance[1];
-						    for(;;)
-						    { 
-							    int NumFetched; 
-							    Enumerator.Next(1, Instances, out NumFetched);
-    
-							    if(NumFetched == 0)
-							    {
-								    break;
-							    }
-    
-							    ISetupInstance2 Instance = (ISetupInstance2)Instances[0];
-							    if((Instance.GetState() & InstanceState.Local) == InstanceState.Local)
-							    {
+						try
+						{
+							SetupConfiguration Setup = new SetupConfiguration();
+							IEnumSetupInstances Enumerator = Setup.EnumAllInstances();
+
+							ISetupInstance[] Instances = new ISetupInstance[1];
+							for(;;)
+							{
+								int NumFetched;
+								Enumerator.Next(1, Instances, out NumFetched);
+
+								if(NumFetched == 0)
+								{
+									break;
+								}
+
+								ISetupInstance2 Instance = (ISetupInstance2)Instances[0];
+								if((Instance.GetState() & InstanceState.Local) == InstanceState.Local)
+								{
+									string VersionString = Instance.GetInstallationVersion();
+
+									VersionNumber Version;
+									if (VersionNumber.TryParse(VersionString, out Version))
+									{
+										VersionNumber Version2019 = new VersionNumber(16);
+										if(Compiler == WindowsCompiler.VisualStudio2019 && Version < Version2019)
+										{
+											continue;
+										}
+										else if(Compiler == WindowsCompiler.VisualStudio2017 && Version >= Version2019)
+										{
+											continue;
+										}
+									}
+
 									ISetupInstanceCatalog Catalog = (ISetupInstanceCatalog)Instance as ISetupInstanceCatalog;
 									if (Catalog != null && Catalog.IsPrerelease())
 									{
@@ -717,15 +774,15 @@ namespace UnrealBuildTool
 									{
 										InstallDirs.Add(new DirectoryReference(Instance.GetInstallationPath()));
 									}
-							    }
-						    }
-					    }
-					    catch
-					    {
-					    }
+								}
+							}
+						}
+						catch
+						{
+						}
 						InstallDirs.AddRange(PreReleaseInstallDirs);
-				    }
-				    else
+					}
+					else
 				    {
 					    throw new BuildException("Unsupported compiler version ({0})", Compiler);
 				    }
@@ -733,7 +790,7 @@ namespace UnrealBuildTool
 				CachedVSInstallDirs.Add(Compiler, InstallDirs);
 			}
 			return InstallDirs;
-		}	 
+		}
 
 		/// <summary>
 		/// Determines the directory containing the MSVC toolchain
@@ -809,7 +866,7 @@ namespace UnrealBuildTool
 							}
 					    }
 				    }
-				    else if(Compiler == WindowsCompiler.VisualStudio2017)
+				    else if(Compiler == WindowsCompiler.VisualStudio2017 || Compiler == WindowsCompiler.VisualStudio2019)
 				    {
 						// Enumerate all the manually installed toolchains
 						List<DirectoryReference> InstallDirs = FindVSInstallDirs(Compiler);
@@ -821,7 +878,7 @@ namespace UnrealBuildTool
 							    foreach(DirectoryReference ToolChainDir in DirectoryReference.EnumerateDirectories(ToolChainBaseDir))
 							    {
 								    VersionNumber Version;
-								    if(VersionNumber.TryParse(ToolChainDir.GetDirectoryName(), out Version) && IsValidToolChainDir2017(ToolChainDir) && !ToolChainVersionToDir.ContainsKey(Version))
+								    if(VersionNumber.TryParse(ToolChainDir.GetDirectoryName(), out Version) && IsValidToolChainDir2017or2019(ToolChainDir) && !ToolChainVersionToDir.ContainsKey(Version))
 								    {
 									    ToolChainVersionToDir[Version] = ToolChainDir;
 								    }
@@ -833,13 +890,13 @@ namespace UnrealBuildTool
 						DirectoryReference PlatformDir;
 						if(UEBuildPlatformSDK.TryGetHostPlatformAutoSDKDir(out PlatformDir))
 						{
-							DirectoryReference ToolChainBaseDir = DirectoryReference.Combine(PlatformDir, "Win64", "VS2017");
+							DirectoryReference ToolChainBaseDir = DirectoryReference.Combine(PlatformDir, "Win64", (Compiler == WindowsCompiler.VisualStudio2019)? "VS2019" : "VS2017");
 							if(DirectoryReference.Exists(ToolChainBaseDir))
 							{
 								foreach(DirectoryReference ToolChainDir in DirectoryReference.EnumerateDirectories(ToolChainBaseDir))
 								{
 									VersionNumber Version;
-									if(VersionNumber.TryParse(ToolChainDir.GetDirectoryName(), out Version) && IsValidToolChainDir2017(ToolChainDir) && !ToolChainVersionToDir.ContainsKey(Version))
+									if(VersionNumber.TryParse(ToolChainDir.GetDirectoryName(), out Version) && IsValidToolChainDir2017or2019(ToolChainDir) && !ToolChainVersionToDir.ContainsKey(Version))
 									{
 										ToolChainVersionToDir[Version] = ToolChainDir;
 									}
@@ -882,7 +939,7 @@ namespace UnrealBuildTool
 		/// </summary>
 		/// <param name="ToolChainDir">Directory to check</param>
 		/// <returns>True if the given directory is valid</returns>
-		static bool IsValidToolChainDir2017(DirectoryReference ToolChainDir)
+		static bool IsValidToolChainDir2017or2019(DirectoryReference ToolChainDir)
 		{
 			return FileReference.Exists(FileReference.Combine(ToolChainDir, "bin", "Hostx86", "x64", "cl.exe")) || FileReference.Exists(FileReference.Combine(ToolChainDir, "bin", "Hostx64", "x64", "cl.exe"));
 		}
diff --git a/Engine/Source/Programs/UnrealBuildTool/Platform/Windows/VCEnvironment.cs b/Engine/Source/Programs/UnrealBuildTool/Platform/Windows/VCEnvironment.cs
--- a/Engine/Source/Programs/UnrealBuildTool/Platform/Windows/VCEnvironment.cs
+++ b/Engine/Source/Programs/UnrealBuildTool/Platform/Windows/VCEnvironment.cs
@@ -164,7 +164,7 @@ namespace UnrealBuildTool
 		/// <returns>Directory containing the 32-bit toolchain binaries</returns>
 		static DirectoryReference GetVCToolPath32(WindowsCompiler Compiler, DirectoryReference VCToolChainDir)
 		{
-			if (Compiler == WindowsCompiler.VisualStudio2017)
+			if (Compiler >= WindowsCompiler.VisualStudio2017)
 			{
 				FileReference NativeCompilerPath = FileReference.Combine(VCToolChainDir, "bin", "HostX64", "x86", "cl.exe");
 				if (FileReference.Exists(NativeCompilerPath))
@@ -199,7 +199,7 @@ namespace UnrealBuildTool
 		/// <returns>Directory containing the 64-bit toolchain binaries</returns>
 		static DirectoryReference GetVCToolPath64(WindowsCompiler Compiler, DirectoryReference VCToolChainDir)
 		{
-			if (Compiler == WindowsCompiler.VisualStudio2017)
+			if (Compiler >= WindowsCompiler.VisualStudio2017)
 			{
 				// Use the native 64-bit compiler if present
 				FileReference NativeCompilerPath = FileReference.Combine(VCToolChainDir, "bin", "HostX64", "x64", "cl.exe");
diff --git a/Engine/Source/Programs/UnrealBuildTool/ProjectFiles/ProjectFileGenerator.cs b/Engine/Source/Programs/UnrealBuildTool/ProjectFiles/ProjectFileGenerator.cs
--- a/Engine/Source/Programs/UnrealBuildTool/ProjectFiles/ProjectFileGenerator.cs
+++ b/Engine/Source/Programs/UnrealBuildTool/ProjectFiles/ProjectFileGenerator.cs
@@ -133,6 +133,7 @@ namespace UnrealBuildTool
 		VisualStudio2013,
 		VisualStudio2015,
 		VisualStudio2017,
+		VisualStudio2019,
 		XCode,
 		Eddie,
 		VisualStudioCode,
diff --git a/Engine/Source/Programs/UnrealBuildTool/ProjectFiles/VisualStudio/VCProjectFileGenerator.cs b/Engine/Source/Programs/UnrealBuildTool/ProjectFiles/VisualStudio/VCProjectFileGenerator.cs
--- a/Engine/Source/Programs/UnrealBuildTool/ProjectFiles/VisualStudio/VCProjectFileGenerator.cs
+++ b/Engine/Source/Programs/UnrealBuildTool/ProjectFiles/VisualStudio/VCProjectFileGenerator.cs
@@ -43,6 +43,7 @@ namespace UnrealBuildTool
 		VisualStudio2013, // Unsupported
 		VisualStudio2015,
 		VisualStudio2017,
+		VisualStudio2019,
 	}
 
 	/// <summary>
@@ -110,6 +111,10 @@ namespace UnrealBuildTool
 			{
 				BuildToolOverride = "-2017";
 			}
+			else if(InArguments.HasOption("-2019"))
+			{
+				BuildToolOverride = "-2019";
+			}
 		}
 
 		/// File extension for project files we'll be generating (e.g. ".vcxproj")
@@ -197,6 +202,8 @@ namespace UnrealBuildTool
 					return "14.0";
 				case VCProjectFileFormat.VisualStudio2017:
 					return "15.0";
+				case VCProjectFileFormat.VisualStudio2019:
+					return "15.0"; // Correct as of VS2019 Preview 1
 			}
 			return string.Empty;
 		}
@@ -214,6 +221,9 @@ namespace UnrealBuildTool
 					return "v140";
 				case VCProjectFileFormat.VisualStudio2017:
                     return "v141";
+				case VCProjectFileFormat.VisualStudio2019:
+					return "v141"; // Correct as of VS2019 Preview 1
+
             }
 			return string.Empty;
 		}
@@ -256,6 +266,14 @@ namespace UnrealBuildTool
 					ProjectFileFormat = VCProjectFileFormat.Default;
 				}
 			}
+			else if(ProjectFileFormat == VCProjectFileFormat.VisualStudio2019)
+			{
+				if (!WindowsPlatform.HasCompiler(WindowsCompiler.VisualStudio2019))
+				{
+					Log.TraceWarning("Visual Studio C++ 2019 installation not found - ignoring preferred project file format.");
+					ProjectFileFormat = VCProjectFileFormat.Default;
+				}
+			}
 
 			// Certain platforms override the project file format because their debugger add-ins may not yet support the latest
 			// version of Visual Studio.  This is their chance to override that.
@@ -271,6 +289,10 @@ namespace UnrealBuildTool
 				{
 					ProjectFileFormat = VCProjectFileFormat.VisualStudio2015;
 				}
+				else if (WindowsPlatform.HasCompiler(WindowsCompiler.VisualStudio2019) && WindowsPlatform.HasIDE(WindowsCompiler.VisualStudio2019))
+				{
+					ProjectFileFormat = VCProjectFileFormat.VisualStudio2019;
+				}
 
 				// Allow the SDKs to override
 				foreach (UnrealTargetPlatform SupportedPlatform in SupportedPlatforms)
@@ -398,7 +420,15 @@ namespace UnrealBuildTool
 			StringBuilder VCSolutionFileContent = new StringBuilder();
 
 			// Solution file header. Note that a leading newline is required for file type detection to work correclty in the shell.
-			if (ProjectFileFormat == VCProjectFileFormat.VisualStudio2017)
+			if (ProjectFileFormat == VCProjectFileFormat.VisualStudio2019)
+			{
+				VCSolutionFileContent.AppendLine();
+				VCSolutionFileContent.AppendLine("Microsoft Visual Studio Solution File, Format Version 12.00");
+				VCSolutionFileContent.AppendLine("# Visual Studio Version 16");
+				VCSolutionFileContent.AppendLine("VisualStudioVersion = 16.0.28315.86");
+				VCSolutionFileContent.AppendLine("MinimumVisualStudioVersion = 10.0.40219.1");
+			}
+			else if (ProjectFileFormat == VCProjectFileFormat.VisualStudio2017)
 			{
 				VCSolutionFileContent.AppendLine();
 				VCSolutionFileContent.AppendLine("Microsoft Visual Studio Solution File, Format Version 12.00");
@@ -772,6 +802,9 @@ namespace UnrealBuildTool
 					case VCProjectFileFormat.VisualStudio2017:
 						SolutionOptionsFileName = FileReference.Combine(MasterProjectPath, ".vs", Path.GetFileNameWithoutExtension(SolutionFileName), "v15", ".suo");
 						break;
+					case VCProjectFileFormat.VisualStudio2019:
+						SolutionOptionsFileName = FileReference.Combine(MasterProjectPath, ".vs", Path.GetFileNameWithoutExtension(SolutionFileName), "v16", ".suo");
+						break;
 					default:
 						throw new BuildException("Unsupported Visual Studio version");
 				}
diff --git a/Engine/Source/Runtime/Core/Public/Windows/WindowsPlatformCompilerSetup.h b/Engine/Source/Runtime/Core/Public/Windows/WindowsPlatformCompilerSetup.h
--- a/Engine/Source/Runtime/Core/Public/Windows/WindowsPlatformCompilerSetup.h
+++ b/Engine/Source/Runtime/Core/Public/Windows/WindowsPlatformCompilerSetup.h
@@ -9,8 +9,8 @@
 /**
 * Future-proofing the min version check so we keep bumping it whenever we upgrade.
 */
-#if defined(_MSC_VER) && _MSC_VER > 1919 
-	#pragma message("Detected compiler newer than Visual Studio 2017, please update min version checking in WindowsPlatformCompilerSetup.h")
+#if defined(_MSC_VER) && _MSC_VER > 1929 
+	#pragma message("Detected compiler newer than Visual Studio 2019, please update min version checking in WindowsPlatformCompilerSetup.h")
 #endif
 
 /**
