UBT: Fix null reference exception when generating project files.

#rnx

diff --git a/Engine/Source/Programs/UnrealBuildTool/Platform/Windows/UEBuildWindows.cs b/Engine/Source/Programs/UnrealBuildTool/Platform/Windows/UEBuildWindows.cs
--- a/Engine/Source/Programs/UnrealBuildTool/Platform/Windows/UEBuildWindows.cs
+++ b/Engine/Source/Programs/UnrealBuildTool/Platform/Windows/UEBuildWindows.cs
@@ -566,20 +566,23 @@ namespace UnrealBuildTool
 		internal static WindowsCompiler GetDefaultCompiler(FileReference ProjectFile)
 		{
 			// If there's no specific compiler set, try to pick the matching compiler for the selected IDE
-			foreach(ProjectFileFormat Format in ProjectFileGeneratorSettings.ParseFormatList(ProjectFileGeneratorSettings.Format))
+			if(ProjectFileGeneratorSettings.Format != null)
 			{
-				if (Format == ProjectFileFormat.VisualStudio2019)
+				foreach(ProjectFileFormat Format in ProjectFileGeneratorSettings.ParseFormatList(ProjectFileGeneratorSettings.Format))
 				{
-					return WindowsCompiler.VisualStudio2019;
-				}
-				else if (Format == ProjectFileFormat.VisualStudio2017)
-				{
-					return WindowsCompiler.VisualStudio2017;
-				}
-				else if (Format == ProjectFileFormat.VisualStudio2015)
-				{
-					return WindowsCompiler.VisualStudio2015;
-				}
+					if (Format == ProjectFileFormat.VisualStudio2019)
+					{
+						return WindowsCompiler.VisualStudio2019;
+					}
+					else if (Format == ProjectFileFormat.VisualStudio2017)
+					{
+						return WindowsCompiler.VisualStudio2017;
+					}
+					else if (Format == ProjectFileFormat.VisualStudio2015)
+					{
+						return WindowsCompiler.VisualStudio2015;
+					}
+				} 
 			}
 
 			// Also check the default format for the Visual Studio project generator
