UBT: Allow using standard variables for additional library paths (eg. $(ModuleDir), $(PluginDir), etc...)

diff --git a/Engine/Source/Programs/UnrealBuildTool/Configuration/UEBuildModule.cs b/Engine/Source/Programs/UnrealBuildTool/Configuration/UEBuildModule.cs
--- a/Engine/Source/Programs/UnrealBuildTool/Configuration/UEBuildModule.cs
+++ b/Engine/Source/Programs/UnrealBuildTool/Configuration/UEBuildModule.cs
@@ -274,9 +274,9 @@ namespace UnrealBuildTool
 		/// <summary>
 		/// Converts an optional string list parameter to a well-defined hash set.
 		/// </summary>
-		protected static HashSet<string> HashSetFromOptionalEnumerableStringParameter(IEnumerable<string> InEnumerableStrings)
+		protected HashSet<string> HashSetFromOptionalEnumerableStringParameter(IEnumerable<string> InEnumerableStrings)
 		{
-			return InEnumerableStrings == null ? new HashSet<string>() : new HashSet<string>(InEnumerableStrings);
+			return InEnumerableStrings == null ? new HashSet<string>() : new HashSet<string>(InEnumerableStrings.Select(x => ExpandPathVariables(x, null, null)));
 		}
 
 		/// <summary>
