UBT: Add a simple single-threaded timeline capture class to UBT, for tracking the progress of each run.

diff --git a/Engine/Source/Programs/UnrealBuildTool/System/Timeline.cs b/Engine/Source/Programs/UnrealBuildTool/System/Timeline.cs
new file mode 100644
--- /dev/null
+++ b/Engine/Source/Programs/UnrealBuildTool/System/Timeline.cs
@@ -0,0 +1,228 @@
+using System;
+using System.Collections.Generic;
+using System.Diagnostics;
+using System.Linq;
+using System.Text;
+using System.Threading.Tasks;
+using Tools.DotNETCommon;
+
+namespace UnrealBuildTool
+{
+	/// <summary>
+	/// Public interface for a timeline scope. Should be disposed to exit the scope.
+	/// </summary>
+	interface ITimelineEvent : IDisposable
+	{
+		void Finish();
+	}
+
+	/// <summary>
+	/// Tracks simple high-level timing data
+	/// </summary>
+	static class Timeline
+	{
+		/// <summary>
+		/// A marker in the timeline
+		/// </summary>
+		[DebuggerDisplay("{Name}")]
+		class Event : ITimelineEvent
+		{
+			/// <summary>
+			/// Name of the marker
+			/// </summary>
+			public string Name;
+
+			/// <summary>
+			/// Time at which the event ocurred
+			/// </summary>
+			public TimeSpan StartTime;
+
+			/// <summary>
+			/// Time at which the event ended
+			/// </summary>
+			public TimeSpan? FinishTime;
+
+			/// <summary>
+			/// All the child events within this scope
+			/// </summary>
+			public List<Event> Children;
+
+			/// <summary>
+			/// Constructor
+			/// </summary>
+			/// <param name="Name">Event name</param>
+			/// <param name="StartTime">Time of the event</param>
+			public Event(string Name, TimeSpan StartTime)
+			{
+				this.Name = Name;
+				this.StartTime = StartTime;
+			}
+
+			/// <summary>
+			/// Finishes the current event
+			/// </summary>
+			public void Finish()
+			{
+				if(!FinishTime.HasValue)
+				{
+					FinishTime = Stopwatch.Elapsed;
+					if(Children != null)
+					{
+						EndEvent(this);
+					}
+				}
+			}
+
+			/// <summary>
+			/// Disposes of the current event
+			/// </summary>
+			public void Dispose()
+			{
+				Finish();
+			}
+		}
+
+		/// <summary>
+		/// The stopwatch used for timing
+		/// </summary>
+		static Stopwatch Stopwatch = Stopwatch.StartNew();
+
+		/// <summary>
+		/// The parent event
+		/// </summary>
+		static Event RootEvent = new Event("<Root>", TimeSpan.Zero) { Children = new List<Event>() };
+
+		/// <summary>
+		/// The currently active marker
+		/// </summary>
+		static List<Event> ScopedEvents = new List<Event> { RootEvent };
+
+		/// <summary>
+		/// Records a timeline marker with the given name
+		/// </summary>
+		/// <param name="Name">The marker name</param>
+		public static void AddEvent(string Name)
+		{
+			Event Event = new Event(Name, Stopwatch.Elapsed);
+			Event.FinishTime = Event.StartTime;
+			ScopedEvents[ScopedEvents.Count - 1].Children.Add(Event);
+		}
+
+		public static ITimelineEvent ScopeEvent(string Name)
+		{
+			Event Event = new Event(Name, Stopwatch.Elapsed);
+			Event.Children = new List<Event>();
+			ScopedEvents[ScopedEvents.Count - 1].Children.Add(Event);
+			ScopedEvents.Add(Event);
+			return Event;
+		}
+
+		static void EndEvent(Event Event)
+		{
+			int Idx = ScopedEvents.IndexOf(Event);
+			if(Idx != -1)
+			{
+				ScopedEvents.RemoveRange(Idx, ScopedEvents.Count - Idx);
+			}
+		}
+
+		/// <summary>
+		/// Prints this information to the log
+		/// </summary>
+		public static void Print(TimeSpan MaxUnknownTime, LogEventType Verbosity)
+		{
+			TimeSpan CurrentTime = Stopwatch.Elapsed;
+			Log.WriteLine(Verbosity, "[{0,6}]", FormatTime(TimeSpan.Zero));
+
+			List<TimeSpan> OuterStartTimes = new List<TimeSpan>();
+			OuterStartTimes.Add(TimeSpan.Zero);
+			
+			Print(RootEvent.Children, OuterStartTimes, CurrentTime, MaxUnknownTime, Verbosity);
+			Log.WriteLine(Verbosity, "[{0,6}]", FormatTime(CurrentTime));
+		}
+
+		/// <summary>
+		/// Prints a list of events
+		/// </summary>
+		/// <param name="Events">The events to print</param>
+		/// <param name="OuterStartTimes">List of all the start times for parent events</param>
+		/// <param name="OuterFinishTime">Finish time for the parent node</param>
+		/// <param name="MaxUnknownTime">Maximum unnacounted-for gap to allow in the timeline before displaying a dummy event</param>
+		/// <param name="Verbosity">Verbosity for log output</param>
+		static void Print(List<Event> Events, List<TimeSpan> OuterStartTimes, TimeSpan? OuterFinishTime, TimeSpan MaxUnknownTime, LogEventType Verbosity)
+		{
+			const string UnknownEvent = "<unknown>";
+
+			TimeSpan? LastTime = OuterStartTimes[OuterStartTimes.Count - 1];
+			foreach(Event Event in Events)
+			{
+				if(LastTime.HasValue && Event.StartTime - LastTime.Value > MaxUnknownTime)
+				{
+					PrintEvent(LastTime.Value, Event.StartTime, UnknownEvent, OuterStartTimes, Verbosity);
+				}
+
+				PrintEvent(Event.StartTime, Event.FinishTime, Event.Name, OuterStartTimes, Verbosity);
+
+				if(Event.Children != null)
+				{
+					OuterStartTimes.Add(Event.StartTime);
+					Print(Event.Children, OuterStartTimes, Event.FinishTime, MaxUnknownTime, Verbosity);
+					OuterStartTimes.RemoveAt(OuterStartTimes.Count - 1);
+				}
+
+				LastTime = Event.FinishTime;
+			}
+
+			if(LastTime.HasValue && OuterFinishTime.HasValue && OuterFinishTime.Value - LastTime.Value > MaxUnknownTime)
+			{
+				PrintEvent(LastTime.Value, OuterFinishTime.Value, UnknownEvent, OuterStartTimes, Verbosity);
+			}
+		}
+
+		/// <summary>
+		/// Prints an individual event to the log
+		/// </summary>
+		/// <param name="StartTime">Start time for the event</param>
+		/// <param name="FinishTime">Finish time for the event. May be null.</param>
+		/// <param name="Label">Event name</param>
+		/// <param name="OuterStartTimes">List of all the start times for parent events</param>
+		/// <param name="Verbosity">Verbosity for the output</param>
+		static void PrintEvent(TimeSpan StartTime, TimeSpan? FinishTime, string Label, List<TimeSpan> OuterStartTimes, LogEventType Verbosity)
+		{
+			StringBuilder Prefix = new StringBuilder();
+
+			for(int Idx = 0; Idx < OuterStartTimes.Count - 1; Idx++)
+			{
+				Prefix.AppendFormat(" {0,6}          ", FormatTime(StartTime - OuterStartTimes[Idx]));
+			}
+
+			Prefix.AppendFormat("[{0,6}]", FormatTime(StartTime - OuterStartTimes[OuterStartTimes.Count - 1]));
+
+			if (!FinishTime.HasValue)
+			{
+				Prefix.AppendFormat("({0,6})", "???");
+			}
+			else if(FinishTime.Value == StartTime)
+			{
+				Prefix.Append(" ------ ");
+			}
+			else
+			{
+				Prefix.AppendFormat("({0,6})", "+" + FormatTime(FinishTime.Value - StartTime));
+			}
+
+			Log.WriteLine(Verbosity, "{0} {1}", Prefix.ToString(), Label);
+		}
+
+		/// <summary>
+		/// Formats a timespan in milliseconds
+		/// </summary>
+		/// <param name="Time">The time to format</param>
+		/// <returns>Formatted timespan</returns>
+		static string FormatTime(TimeSpan Time)
+		{
+			int TotalMilliseconds = (int)Time.TotalMilliseconds;
+			return String.Format("{0}.{1:000}", TotalMilliseconds / 1000, TotalMilliseconds % 1000);
+		}
+	}
+}
diff --git a/Engine/Source/Programs/UnrealBuildTool/UnrealBuildTool.csproj b/Engine/Source/Programs/UnrealBuildTool/UnrealBuildTool.csproj
--- a/Engine/Source/Programs/UnrealBuildTool/UnrealBuildTool.csproj
+++ b/Engine/Source/Programs/UnrealBuildTool/UnrealBuildTool.csproj
@@ -220,6 +220,7 @@
     <Compile Include="System\SupportedConfigurationsAttribute.cs" />
     <Compile Include="System\SupportedPlatformsAttribute.cs" />
     <Compile Include="System\TargetReceipt.cs" />
+    <Compile Include="System\Timeline.cs" />
     <Compile Include="System\ToolMode.cs" />
     <Compile Include="System\UBTMakefile.cs" />
     <Compile Include="System\VersionNumber.cs" />
