UBT: Simplify timeline events to use a flat array of events, so that un-terminated scopes do not corrupt the output format.

#rnx

diff --git a/Engine/Source/Programs/UnrealBuildTool/System/Timeline.cs b/Engine/Source/Programs/UnrealBuildTool/System/Timeline.cs
--- a/Engine/Source/Programs/UnrealBuildTool/System/Timeline.cs
+++ b/Engine/Source/Programs/UnrealBuildTool/System/Timeline.cs
@@ -2,6 +2,8 @@ using System;
 using System.Collections.Generic;
 using System.Diagnostics;
 using System.Linq;
+using System.Reflection;
+using System.Runtime.CompilerServices;
 using System.Text;
 using System.Threading.Tasks;
 using Tools.DotNETCommon;
@@ -42,20 +44,17 @@ namespace UnrealBuildTool
 			/// </summary>
 			public TimeSpan? FinishTime;
 
-			/// <summary>
-			/// All the child events within this scope
-			/// </summary>
-			public List<Event> Children;
-
 			/// <summary>
 			/// Constructor
 			/// </summary>
 			/// <param name="Name">Event name</param>
 			/// <param name="StartTime">Time of the event</param>
-			public Event(string Name, TimeSpan StartTime)
+			/// <param name="FinishTime">Finish time for the event. May be null.</param>
+			public Event(string Name, TimeSpan StartTime, TimeSpan? FinishTime)
 			{
 				this.Name = Name;
 				this.StartTime = StartTime;
+				this.FinishTime = FinishTime;
 			}
 
 			/// <summary>
@@ -66,10 +65,6 @@ namespace UnrealBuildTool
 				if(!FinishTime.HasValue)
 				{
 					FinishTime = Stopwatch.Elapsed;
-					if(Children != null)
-					{
-						EndEvent(this);
-					}
 				}
 			}
 
@@ -85,17 +80,20 @@ namespace UnrealBuildTool
 		/// <summary>
 		/// The stopwatch used for timing
 		/// </summary>
-		static Stopwatch Stopwatch = Stopwatch.StartNew();
+		static Stopwatch Stopwatch = new Stopwatch();
 
 		/// <summary>
-		/// The parent event
+		/// The recorded events
 		/// </summary>
-		static Event RootEvent = new Event("<Root>", TimeSpan.Zero) { Children = new List<Event>() };
+		static List<Event> Events = new List<Event>();
 
 		/// <summary>
-		/// The currently active marker
+		/// Start the stopwatch
 		/// </summary>
-		static List<Event> ScopedEvents = new List<Event> { RootEvent };
+		public static void Start()
+		{
+			Stopwatch.Restart();
+		}
 
 		/// <summary>
 		/// Records a timeline marker with the given name
@@ -103,27 +101,32 @@ namespace UnrealBuildTool
 		/// <param name="Name">The marker name</param>
 		public static void AddEvent(string Name)
 		{
-			Event Event = new Event(Name, Stopwatch.Elapsed);
-			Event.FinishTime = Event.StartTime;
-			ScopedEvents[ScopedEvents.Count - 1].Children.Add(Event);
+			TimeSpan Time = Stopwatch.Elapsed;
+			Events.Add(new Event(Name, Time, Time));
 		}
 
+		/// <summary>
+		/// Enters a scope event with the given name. Should be disposed to terminate the scope.
+		/// </summary>
+		/// <param name="Name">Name of the event</param>
+		/// <returns>Event to track the length of the event</returns>
 		public static ITimelineEvent ScopeEvent(string Name)
 		{
-			Event Event = new Event(Name, Stopwatch.Elapsed);
-			Event.Children = new List<Event>();
-			ScopedEvents[ScopedEvents.Count - 1].Children.Add(Event);
-			ScopedEvents.Add(Event);
+			Event Event = new Event(Name, Stopwatch.Elapsed, null);
+			Events.Add(Event);
 			return Event;
 		}
 
-		static void EndEvent(Event Event)
+		/// <summary>
+		/// Enters a scope event with the calling method name as the event name
+		/// </summary>
+		/// <returns>Event to track the lifetime of the method. Dispose it before returning.</returns>
+		[MethodImplAttribute(MethodImplOptions.NoInlining)]
+		public static ITimelineEvent ScopeMethod()
 		{
-			int Idx = ScopedEvents.IndexOf(Event);
-			if(Idx != -1)
-			{
-				ScopedEvents.RemoveRange(Idx, ScopedEvents.Count - Idx);
-			}
+			StackFrame Frame = new StackFrame(1);
+			MethodBase Method = Frame.GetMethod();
+			return ScopeEvent(String.Format("{0}.{1}", Method.DeclaringType.Name, Method.Name));
 		}
 
 		/// <summary>
@@ -131,52 +134,73 @@ namespace UnrealBuildTool
 		/// </summary>
 		public static void Print(TimeSpan MaxUnknownTime, LogEventType Verbosity)
 		{
-			TimeSpan CurrentTime = Stopwatch.Elapsed;
+			// Print the start time
+			Log.WriteLine(Verbosity, "Timeline:");
+			Log.WriteLine(Verbosity, "");
 			Log.WriteLine(Verbosity, "[{0,6}]", FormatTime(TimeSpan.Zero));
 
-			List<TimeSpan> OuterStartTimes = new List<TimeSpan>();
-			OuterStartTimes.Add(TimeSpan.Zero);
-			
-			Print(RootEvent.Children, OuterStartTimes, CurrentTime, MaxUnknownTime, Verbosity);
-			Log.WriteLine(Verbosity, "[{0,6}]", FormatTime(CurrentTime));
-		}
+			// Create the root event
+			TimeSpan FinishTime = Stopwatch.Elapsed;
 
-		/// <summary>
-		/// Prints a list of events
-		/// </summary>
-		/// <param name="Events">The events to print</param>
-		/// <param name="OuterStartTimes">List of all the start times for parent events</param>
-		/// <param name="OuterFinishTime">Finish time for the parent node</param>
-		/// <param name="MaxUnknownTime">Maximum unnacounted-for gap to allow in the timeline before displaying a dummy event</param>
-		/// <param name="Verbosity">Verbosity for log output</param>
-		static void Print(List<Event> Events, List<TimeSpan> OuterStartTimes, TimeSpan? OuterFinishTime, TimeSpan MaxUnknownTime, LogEventType Verbosity)
-		{
-			const string UnknownEvent = "<unknown>";
+			List<Event> OuterEvents = new List<Event>();
+			OuterEvents.Add(new Event("<Root>", TimeSpan.Zero, FinishTime));
 
-			TimeSpan? LastTime = OuterStartTimes[OuterStartTimes.Count - 1];
-			foreach(Event Event in Events)
+			// Print out all the child events
+			TimeSpan LastTime = TimeSpan.Zero;
+			for(int EventIdx = 0; EventIdx < Events.Count; EventIdx++)
 			{
-				if(LastTime.HasValue && Event.StartTime - LastTime.Value > MaxUnknownTime)
+				Event Event = Events[EventIdx];
+
+				// Pop events off the stack
+				for (; OuterEvents.Count > 0; OuterEvents.RemoveAt(OuterEvents.Count - 1))
 				{
-					PrintEvent(LastTime.Value, Event.StartTime, UnknownEvent, OuterStartTimes, Verbosity);
+					Event OuterEvent = OuterEvents.Last();
+					if (Event.StartTime < OuterEvent.FinishTime.Value)
+					{
+						break;
+					}
+					UpdateLastEventTime(ref LastTime, OuterEvent.FinishTime.Value, MaxUnknownTime, OuterEvents, Verbosity);
 				}
 
-				PrintEvent(Event.StartTime, Event.FinishTime, Event.Name, OuterStartTimes, Verbosity);
+				// If there's a gap since the last event, print an unknown marker
+				UpdateLastEventTime(ref LastTime, Event.StartTime, MaxUnknownTime, OuterEvents, Verbosity);
 
-				if(Event.Children != null)
+				// Print this event
+				Print(Event.StartTime, Event.FinishTime, Event.Name, OuterEvents, Verbosity);
+
+				// Push it onto the stack
+				if(Event.FinishTime.HasValue && EventIdx + 1 < Events.Count && Events[EventIdx + 1].StartTime < Event.FinishTime.Value)
 				{
-					OuterStartTimes.Add(Event.StartTime);
-					Print(Event.Children, OuterStartTimes, Event.FinishTime, MaxUnknownTime, Verbosity);
-					OuterStartTimes.RemoveAt(OuterStartTimes.Count - 1);
+					OuterEvents.Add(Event);
 				}
+			}
 
-				LastTime = Event.FinishTime;
+			// Remove everything from the stack
+			for(; OuterEvents.Count > 1; OuterEvents.RemoveAt(OuterEvents.Count - 1))
+			{
+				UpdateLastEventTime(ref LastTime, OuterEvents.Last().FinishTime.Value, MaxUnknownTime, OuterEvents, Verbosity);
 			}
 
-			if(LastTime.HasValue && OuterFinishTime.HasValue && OuterFinishTime.Value - LastTime.Value > MaxUnknownTime)
+			// Print the finish time
+			Log.WriteLine(Verbosity, "[{0,6}]", FormatTime(FinishTime));
+		}
+
+		/// <summary>
+		/// Updates the last event time
+		/// </summary>
+		/// <param name="LastTime"></param>
+		/// <param name="NewTime"></param>
+		/// <param name="MaxUnknownTime"></param>
+		/// <param name="OuterEvents"></param>
+		/// <param name="Verbosity"></param>
+		static void UpdateLastEventTime(ref TimeSpan LastTime, TimeSpan NewTime, TimeSpan MaxUnknownTime, List<Event> OuterEvents, LogEventType Verbosity)
+		{
+			const string UnknownEvent = "<unknown>";
+			if (NewTime - LastTime > MaxUnknownTime)
 			{
-				PrintEvent(LastTime.Value, OuterFinishTime.Value, UnknownEvent, OuterStartTimes, Verbosity);
+				Print(LastTime, NewTime, UnknownEvent, OuterEvents, Verbosity);
 			}
+			LastTime = NewTime;
 		}
 
 		/// <summary>
@@ -185,18 +209,18 @@ namespace UnrealBuildTool
 		/// <param name="StartTime">Start time for the event</param>
 		/// <param name="FinishTime">Finish time for the event. May be null.</param>
 		/// <param name="Label">Event name</param>
-		/// <param name="OuterStartTimes">List of all the start times for parent events</param>
+		/// <param name="OuterEvents">List of all the start times for parent events</param>
 		/// <param name="Verbosity">Verbosity for the output</param>
-		static void PrintEvent(TimeSpan StartTime, TimeSpan? FinishTime, string Label, List<TimeSpan> OuterStartTimes, LogEventType Verbosity)
+		static void Print(TimeSpan StartTime, TimeSpan? FinishTime, string Label, List<Event> OuterEvents, LogEventType Verbosity)
 		{
 			StringBuilder Prefix = new StringBuilder();
 
-			for(int Idx = 0; Idx < OuterStartTimes.Count - 1; Idx++)
+			for(int Idx = 0; Idx < OuterEvents.Count - 1; Idx++)
 			{
-				Prefix.AppendFormat(" {0,6}          ", FormatTime(StartTime - OuterStartTimes[Idx]));
+				Prefix.AppendFormat(" {0,6}          ", FormatTime(StartTime - OuterEvents[Idx].StartTime));
 			}
 
-			Prefix.AppendFormat("[{0,6}]", FormatTime(StartTime - OuterStartTimes[OuterStartTimes.Count - 1]));
+			Prefix.AppendFormat("[{0,6}]", FormatTime(StartTime - OuterEvents[OuterEvents.Count - 1].StartTime));
 
 			if (!FinishTime.HasValue)
 			{
