UBT: Start the timeline as soon as the program starts, and flush it before quitting.

#rnx

diff --git a/Engine/Source/Programs/UnrealBuildTool/System/Timeline.cs b/Engine/Source/Programs/UnrealBuildTool/System/Timeline.cs
--- a/Engine/Source/Programs/UnrealBuildTool/System/Timeline.cs
+++ b/Engine/Source/Programs/UnrealBuildTool/System/Timeline.cs
@@ -176,7 +176,7 @@ namespace UnrealBuildTool
 			}
 
 			// Remove everything from the stack
-			for(; OuterEvents.Count > 1; OuterEvents.RemoveAt(OuterEvents.Count - 1))
+			for(; OuterEvents.Count > 0; OuterEvents.RemoveAt(OuterEvents.Count - 1))
 			{
 				UpdateLastEventTime(ref LastTime, OuterEvents.Last().FinishTime.Value, MaxUnknownTime, OuterEvents, Verbosity);
 			}
diff --git a/Engine/Source/Programs/UnrealBuildTool/UnrealBuildTool.cs b/Engine/Source/Programs/UnrealBuildTool/UnrealBuildTool.cs
--- a/Engine/Source/Programs/UnrealBuildTool/UnrealBuildTool.cs
+++ b/Engine/Source/Programs/UnrealBuildTool/UnrealBuildTool.cs
@@ -381,6 +381,9 @@ namespace UnrealBuildTool
 		{
 			try
 			{
+				// Start capturing performance info
+				Timeline.Start();
+
 				// Parse the command line arguments
 				CommandLineArguments Arguments = new CommandLineArguments(ArgumentsArray);
 
@@ -447,6 +450,9 @@ namespace UnrealBuildTool
 			}
 			finally
 			{
+				// Print out all the performance info
+				Timeline.Print(TimeSpan.FromMilliseconds(20.0), LogEventType.Log);
+
 				// Make sure we flush the logs however we exit
 				Trace.Close();
 			}
