Copy of CL 4599554

Fixed Unix version of UBT's UEBuildTarget.CreateCopyAction

diff --git a/Engine/Source/Programs/UnrealBuildTool/Configuration/UEBuildTarget.cs b/Engine/Source/Programs/UnrealBuildTool/Configuration/UEBuildTarget.cs
--- a/Engine/Source/Programs/UnrealBuildTool/Configuration/UEBuildTarget.cs
+++ b/Engine/Source/Programs/UnrealBuildTool/Configuration/UEBuildTarget.cs
@@ -1909,7 +1909,7 @@ namespace UnrealBuildTool
 			else
 			{
 				CopyAction.CommandPath = "/bin/sh";
-				CopyAction.CommandArguments = String.Format("cp -f {0} {1}", Utils.EscapeShellArgument(SourceFile.FullName), Utils.EscapeShellArgument(TargetFile.FullName));
+				CopyAction.CommandArguments = String.Format("-c 'cp -f {0} {1}'", Utils.EscapeShellArgument(SourceFile.FullName), Utils.EscapeShellArgument(TargetFile.FullName));
 			}
 			CopyAction.WorkingDirectory = Environment.CurrentDirectory;
 			CopyAction.PrerequisiteItems.Add(SourceFileItem);
