UBT: Make non-monolithic targets on Mac store dylibs in Engine/Binaries/Mac instead of inside the app bundle (heavily based on BenM's work)

diff --git a/Engine/Source/Programs/UnrealBuildTool/Configuration/UEBuildTarget.cs b/Engine/Source/Programs/UnrealBuildTool/Configuration/UEBuildTarget.cs
--- a/Engine/Source/Programs/UnrealBuildTool/Configuration/UEBuildTarget.cs
+++ b/Engine/Source/Programs/UnrealBuildTool/Configuration/UEBuildTarget.cs
@@ -1705,6 +1705,10 @@ namespace UnrealBuildTool
 
 			// Build the target's binaries.
 			DirectoryReference ExeDir = OutputPaths[0].Directory;
+			if (CppPlatform == CppPlatform.Mac && ExeDir.FullName.EndsWith(".app/Contents/MacOS"))
+			{
+				ExeDir = ExeDir.ParentDirectory.ParentDirectory.ParentDirectory;
+			}
 			foreach (UEBuildBinary Binary in Binaries)
 			{
 				List<FileItem> BinaryOutputItems = Binary.Build(Rules, TargetToolChain, GlobalCompileEnvironment, GlobalLinkEnvironment, SharedPCHs, WorkingSet, Predicates, ExeDir, ActionGraph);
@@ -2204,12 +2208,6 @@ namespace UnrealBuildTool
 					}
 				}
 			}
-
-			// On Mac AppBinaries paths for non-console targets need to be adjusted to be inside the app bundle
-			if (Platform == UnrealTargetPlatform.Mac && !Rules.bIsBuildingConsoleApplication)
-			{
-				MacToolChain.FixBundleBinariesPaths(OutputPath, Binaries);
-			}
 		}
 
 		/// <summary>
diff --git a/Engine/Source/Programs/UnrealBuildTool/Platform/Mac/MacToolChain.cs b/Engine/Source/Programs/UnrealBuildTool/Platform/Mac/MacToolChain.cs
--- a/Engine/Source/Programs/UnrealBuildTool/Platform/Mac/MacToolChain.cs
+++ b/Engine/Source/Programs/UnrealBuildTool/Platform/Mac/MacToolChain.cs
@@ -580,8 +580,8 @@ namespace UnrealBuildTool
 			string ExeDir = Path.GetDirectoryName(ExeAbsolutePath);
 
 			// Only dylibs and frameworks, and only those that are outside of Engine/Binaries/Mac and Engine/Source/ThirdParty, and outside of the folder where the executable is need an additional RPATH entry
-			if ((Library.EndsWith("dylib") || Library.EndsWith(".framework")) && !LibraryFullPath.Contains("/Engine/Binaries/Mac/")
-			    && !LibraryFullPath.Contains("/Engine/Source/ThirdParty/") && LibraryDir != ExeDir && !RPaths.Contains(LibraryDir))
+			if ((Library.EndsWith("dylib") || Library.EndsWith(".framework"))
+				&& !LibraryFullPath.Contains("/Engine/Source/ThirdParty/") && LibraryDir != ExeDir && !RPaths.Contains(LibraryDir))
 			{
 				// macOS gatekeeper erroneously complains about not seeing the CEF3 framework in the codesigned Launcher because it's only present in one of the folders specified in RPATHs.
 				// To work around this we will only add a single RPATH entry for it, for the framework stored in .app/Contents/UE4/ subfolder of the packaged app bundle
@@ -770,19 +770,7 @@ namespace UnrealBuildTool
 
 			if (LinkEnvironment.bIsBuildingDLL)
 			{
-				// Add the output file to the command-line.
-				string Filename = "";
-				int Index = OutputFile.AbsolutePath.LastIndexOf(".app/Contents/MacOS/");
-				if (Index > -1)
-				{
-					Index += ".app/Contents/MacOS/".Length;
-					Filename = OutputFile.AbsolutePath.Substring(Index);
-				}
-				else
-				{
-					Filename = Path.GetFileName(OutputFile.AbsolutePath);
-				}
-				LinkCommand += string.Format(" -install_name {0}/{1}", DylibsPath, Filename);
+				LinkCommand += string.Format(" -install_name {0}/{1}", DylibsPath, Path.GetFileName(OutputFile.AbsolutePath));
 			}
 
 			if (!bIsBuildingLibrary)
@@ -1194,36 +1182,8 @@ namespace UnrealBuildTool
 			}
 		}
 
-		public static void FixBundleBinariesPaths(FileReference OutputPath, List<UEBuildBinary> Binaries)
-		{
-			string BundleContentsPath = OutputPath.FullName + ".app/Contents/";
-			foreach (UEBuildBinary Binary in Binaries)
-			{
-				string BinaryFileName = Path.GetFileName(Binary.OutputFilePath.FullName);
-				if (BinaryFileName.EndsWith(".dylib"))
-				{
-					// Only dylibs from the same folder as the executable should be moved to the bundle. UE4Editor-*Game* dylibs and plugins will be loaded
-					// from their Binaries/Mac folders.
-					string DylibDir = Path.GetDirectoryName(Path.GetFullPath(Binary.OutputFilePath.FullName));
-					string ExeDir = Path.GetDirectoryName(Path.GetFullPath(OutputPath.FullName));
-					if (DylibDir.StartsWith(ExeDir))
-					{
-						// get the subdir, which is the DylibDir - ExeDir
-						string SubDir = DylibDir.Replace(ExeDir, "");
-						Binary.OutputFilePaths[0] = new FileReference(BundleContentsPath + "MacOS" + SubDir + "/" + BinaryFileName);
-					}
-				}
-				else if (!BinaryFileName.EndsWith(".a") && !Binary.OutputFilePath.FullName.Contains(".app/Contents/MacOS/")) // Binaries can contain duplicates
-				{
-					Binary.OutputFilePaths[0] += ".app/Contents/MacOS/" + BinaryFileName;
-				}
-			}
-		}
-
 		static private DirectoryReference BundleContentsDirectory;
 
-		public static Dictionary<UEBuildBinary, Dictionary<FileReference, BuildProductType>> AllBuildProducts = new Dictionary<UEBuildBinary, Dictionary<FileReference, BuildProductType>>();
-
 		public override void ModifyBuildProducts(ReadOnlyTargetRules Target, UEBuildBinary Binary, List<string> Libraries, List<UEBuildBundleResource> BundleResources, Dictionary<FileReference, BuildProductType> BuildProducts)
 		{
 			if (Target.bUsePDBFiles == true)
@@ -1319,8 +1279,6 @@ namespace UnrealBuildTool
 					BuildProducts.Add(FileReference.Combine(BundleContentsDirectory, "Resources/" + IconName + ".icns"), BuildProductType.RequiredResource);
 				}
 			}
-
-			AllBuildProducts.Add(Binary, new Dictionary<FileReference, BuildProductType>(BuildProducts));
 		}
 
 		public override ICollection<FileItem> PostBuild(FileItem Executable, LinkEnvironment BinaryLinkEnvironment, ActionGraph ActionGraph)
diff --git a/Engine/Source/Programs/UnrealBuildTool/Platform/Mac/UEBuildMac.cs b/Engine/Source/Programs/UnrealBuildTool/Platform/Mac/UEBuildMac.cs
--- a/Engine/Source/Programs/UnrealBuildTool/Platform/Mac/UEBuildMac.cs
+++ b/Engine/Source/Programs/UnrealBuildTool/Platform/Mac/UEBuildMac.cs
@@ -228,16 +228,34 @@ namespace UnrealBuildTool
 
 		public override DirectoryReference GetBundleDirectory(ReadOnlyTargetRules Rules, List<FileReference> OutputFiles)
 		{
-			if(Rules.bIsBuildingConsoleApplication)
+			if (Rules.bIsBuildingConsoleApplication)
 			{
 				return null;
 			}
 			else
 			{
-				return new DirectoryReference(OutputFiles[0].FullName + ".app");
+				return OutputFiles[0].Directory.ParentDirectory.ParentDirectory;
 			}
 		}
 
+		/// <summary>
+		/// For platforms that need to output multiple files per binary (ie Android "fat" binaries)
+		/// this will emit multiple paths. By default, it simply makes an array from the input
+		/// </summary>
+		public override List<FileReference> FinalizeBinaryPaths(FileReference BinaryName, FileReference ProjectFile, ReadOnlyTargetRules Target)
+		{
+			List<FileReference> BinaryPaths = new List<FileReference>();
+			if (Target.bIsBuildingConsoleApplication || !String.IsNullOrEmpty(BinaryName.GetExtension()))
+			{
+				BinaryPaths.Add(BinaryName);
+			}
+			else
+			{
+				BinaryPaths.Add(new FileReference(BinaryName.FullName + ".app/Contents/MacOS/" + BinaryName.GetFileName()));
+			}
+			return BinaryPaths;
+		}
+
 		/// <summary>
 		/// Modify the rules for a newly created module, in a target that's being built for this platform.
 		/// This is not required - but allows for hiding details of a particular platform.
diff --git a/Engine/Source/Runtime/Core/Private/Mac/MacPlatformProcess.cpp b/Engine/Source/Runtime/Core/Private/Mac/MacPlatformProcess.cpp
--- a/Engine/Source/Runtime/Core/Private/Mac/MacPlatformProcess.cpp
+++ b/Engine/Source/Runtime/Core/Private/Mac/MacPlatformProcess.cpp
@@ -933,19 +933,6 @@ const TCHAR* FMacPlatformProcess::GetBinariesSubdirectory()
 	return TEXT("Mac");
 }
 
-const FString FMacPlatformProcess::GetModulesDirectory()
-{
-	if ([[[[NSBundle mainBundle] bundlePath] pathExtension] isEqual: @"app"])
-	{
-		// If we're an app bundle, modules dylibs are stored in .app/Contents/MacOS
-		return [[[NSBundle mainBundle] executablePath] stringByDeletingLastPathComponent];
-	}
-	else
-	{
-		return FGenericPlatformProcess::GetModulesDirectory();
-	}
-}
-
 void FMacPlatformProcess::LaunchFileInDefaultExternalApplication( const TCHAR* FileName, const TCHAR* Parms /*= NULL*/, ELaunchVerb::Type Verb /*= ELaunchVerb::Open*/ )
 {
 	SCOPED_AUTORELEASE_POOL;
diff --git a/Engine/Source/Runtime/Core/Public/Mac/MacPlatformProcess.h b/Engine/Source/Runtime/Core/Public/Mac/MacPlatformProcess.h
--- a/Engine/Source/Runtime/Core/Public/Mac/MacPlatformProcess.h
+++ b/Engine/Source/Runtime/Core/Public/Mac/MacPlatformProcess.h
@@ -114,7 +114,6 @@ struct CORE_API FMacPlatformProcess : public FGenericPlatformProcess
 	static FString GenerateApplicationPath( const FString& AppName, EBuildConfigurations::Type BuildConfiguration);
 	static const TCHAR* GetModuleExtension();
 	static const TCHAR* GetBinariesSubdirectory();
-	static const FString GetModulesDirectory();
 	static bool CanLaunchURL(const TCHAR* URL);
 	static void LaunchURL(const TCHAR* URL, const TCHAR* Parms, FString* Error);
 	static FString GetGameBundleId();
