UBT: Use the new timeline class for performance tracing.

#rnx

diff --git a/Engine/Source/Programs/UnrealBuildTool/Configuration/UEBuildTarget.cs b/Engine/Source/Programs/UnrealBuildTool/Configuration/UEBuildTarget.cs
--- a/Engine/Source/Programs/UnrealBuildTool/Configuration/UEBuildTarget.cs
+++ b/Engine/Source/Programs/UnrealBuildTool/Configuration/UEBuildTarget.cs
@@ -342,8 +342,6 @@ namespace UnrealBuildTool
 		/// <returns>The build target object for the specified build rules source file</returns>
 		public static UEBuildTarget CreateTarget(TargetDescriptor Desc, string[] Arguments, bool bSkipRulesCompile, bool bCompilingSingleFile, bool bUsePrecompiled)
 		{
-			DateTime CreateTargetStartTime = DateTime.UtcNow;
-
 			RulesAssembly RulesAssembly = RulesCompiler.CreateTargetRulesAssembly(Desc.ProjectFile, Desc.Name, bSkipRulesCompile, bUsePrecompiled, Desc.ForeignPlugin);
 
 			TargetRules RulesObject = RulesAssembly.CreateTargetRules(Desc.Name, Desc.Platform, Desc.Configuration, Desc.Architecture, Desc.ProjectFile, Arguments);
@@ -443,13 +441,6 @@ namespace UnrealBuildTool
 
 			// Generate a build target from this rules module
 			UEBuildTarget BuildTarget = new UEBuildTarget(Desc, new ReadOnlyTargetRules(RulesObject), RulesAssembly);
-
-			if (UnrealBuildTool.bPrintPerformanceInfo)
-			{
-				double CreateTargetTime = (DateTime.UtcNow - CreateTargetStartTime).TotalSeconds;
-				Log.TraceInformation("CreateTarget for " + Desc.Name + " took " + CreateTargetTime + "s");
-			}
-
 			return BuildTarget;
 		}
 
@@ -1657,7 +1648,10 @@ namespace UnrealBuildTool
 				NameToFlatModuleData[FlatModuleData.ModuleName] = FlatModuleData;
 			}
 
-			ExternalExecution.SetupUObjectModules(ModulesToGenerateHeadersFor, Rules.Platform, ProjectDescriptor, UObjectModules, NameToFlatModuleData, Rules.GeneratedCodeVersion, bIsAssemblingBuild);
+			using(Timeline.ScopeEvent("ExternalExecution.SetupUObjectModules()"))
+			{
+				ExternalExecution.SetupUObjectModules(ModulesToGenerateHeadersFor, Rules.Platform, ProjectDescriptor, UObjectModules, NameToFlatModuleData, Rules.GeneratedCodeVersion, bIsAssemblingBuild);
+			}
 
 			// NOTE: Even in Gather mode, we need to run UHT to make sure the files exist for the static action graph to be setup correctly.  This is because UHT generates .cpp
 			// files that are injected as top level prerequisites.  If UHT only emitted included header files, we wouldn't need to run it during the Gather phase at all.
diff --git a/Engine/Source/Programs/UnrealBuildTool/Modes/BuildMode.cs b/Engine/Source/Programs/UnrealBuildTool/Modes/BuildMode.cs
--- a/Engine/Source/Programs/UnrealBuildTool/Modes/BuildMode.cs
+++ b/Engine/Source/Programs/UnrealBuildTool/Modes/BuildMode.cs
@@ -31,8 +31,6 @@ namespace UnrealBuildTool
 		{
 			Arguments.ApplyTo(this);
 
-			DateTime StartTime = DateTime.UtcNow;
-
 			// Initialize the log system, buffering the output until we can create the log file
 			StartupTraceListener StartupListener = new StartupTraceListener();
 			Trace.Listeners.Add(StartupListener);
@@ -79,9 +77,6 @@ namespace UnrealBuildTool
 				XmlConfig.ApplyTo(BuildConfiguration);
 				Arguments.ApplyTo(BuildConfiguration);
 
-				// Copy some of the static settings that are being deprecated from BuildConfiguration
-				UnrealBuildTool.bPrintPerformanceInfo = BuildConfiguration.bPrintPerformanceInfo;
-
 				// Then let the command lines override any configs necessary.
 				if(BuildConfiguration.bXGEExport)
 				{
@@ -102,11 +97,7 @@ namespace UnrealBuildTool
 				// Find and register all tool chains, build platforms, etc. that are present
 				UnrealBuildTool.RegisterAllUBTClasses(false);
 
-				if (UnrealBuildTool.bPrintPerformanceInfo)
-				{
-					double BasicInitTime = (DateTime.UtcNow - BasicInitStartTime).TotalSeconds;
-					Log.TraceInformation("Basic UBT initialization took " + BasicInitTime + "s");
-				}
+				Timeline.AddEvent("Basic UBT initialization");
 
 				// now that we know the available platforms, we can delete other platforms' junk. if we're only building specific modules from the editor, don't touch anything else (it may be in use).
 				if (!BuildConfiguration.bIgnoreJunk)
@@ -115,26 +106,19 @@ namespace UnrealBuildTool
 				}
 
 				// Build our project
-				ECompilationResult Result = UnrealBuildTool.RunUBT(BuildConfiguration, Arguments.GetRawArray(), LogFile);
-
-				// Print some performance info
-				double BuildDuration = (DateTime.UtcNow - StartTime).TotalSeconds;
-				if (UnrealBuildTool.bPrintPerformanceInfo)
+				using(Timeline.ScopeEvent("Calling RunUBT"))
 				{
-					Log.TraceInformation("GetIncludes time: " + CPPHeaders.TotalTimeSpentGettingIncludes + "s (" + CPPHeaders.TotalIncludesRequested + " includes)");
-					Log.TraceInformation("DirectIncludes cache miss time: " + CPPHeaders.DirectIncludeCacheMissesTotalTime + "s (" + CPPHeaders.TotalDirectIncludeCacheMisses + " misses)");
-					Log.TraceInformation("FindIncludePaths calls: " + CPPHeaders.TotalFindIncludedFileCalls + " (" + CPPHeaders.IncludePathSearchAttempts + " searches)");
-					Log.TraceInformation("Deep C++ include scan time: " + UnrealBuildTool.TotalDeepIncludeScanTime + "s");
-					Log.TraceInformation("Include Resolves: {0} ({1} misses, {2:0.00}%)", CPPHeaders.TotalDirectIncludeResolves, CPPHeaders.TotalDirectIncludeResolveCacheMisses, (float)CPPHeaders.TotalDirectIncludeResolveCacheMisses / (float)CPPHeaders.TotalDirectIncludeResolves * 100);
-					Log.TraceInformation("Total FileItems: {0} ({1} missing)", FileItem.TotalFileItemCount, FileItem.MissingFileItemCount);
-
-					Log.TraceInformation("Execution time: {0}s", BuildDuration);
-				}
+					ECompilationResult Result = UnrealBuildTool.RunUBT(BuildConfiguration, Arguments.GetRawArray(), LogFile);
 
-				// Print some performance info
-				Log.TraceVerbose("Execution time: {0}", (DateTime.UtcNow - StartTime).TotalSeconds);
+					// Print some performance info
+					Log.TraceLog("DirectIncludes cache miss time: {0}s ({1} misses)", CPPHeaders.DirectIncludeCacheMissesTotalTime, CPPHeaders.TotalDirectIncludeCacheMisses);
+					Log.TraceLog("FindIncludePaths calls: {0} ({1} searches)", CPPHeaders.TotalFindIncludedFileCalls, CPPHeaders.IncludePathSearchAttempts);
+					Log.TraceLog("Deep C++ include scan time: {0}s", UnrealBuildTool.TotalDeepIncludeScanTime);
+					Log.TraceLog("Include Resolves: {0} ({1} misses, {2:0.00}%)", CPPHeaders.TotalDirectIncludeResolves, CPPHeaders.TotalDirectIncludeResolveCacheMisses, (float)CPPHeaders.TotalDirectIncludeResolveCacheMisses / (float)CPPHeaders.TotalDirectIncludeResolves * 100);
+					Log.TraceLog("Total FileItems: {0} ({1} missing)", FileItem.TotalFileItemCount, FileItem.MissingFileItemCount);
 
-				return (int)Result;
+					return (int)Result;
+				}
 			}
 		}
 	}
diff --git a/Engine/Source/Programs/UnrealBuildTool/System/ActionGraph.cs b/Engine/Source/Programs/UnrealBuildTool/System/ActionGraph.cs
--- a/Engine/Source/Programs/UnrealBuildTool/System/ActionGraph.cs
+++ b/Engine/Source/Programs/UnrealBuildTool/System/ActionGraph.cs
@@ -310,7 +310,7 @@ namespace UnrealBuildTool
 		/// </summary>
 		public HashSet<Action> GetActionsToExecute(BuildConfiguration BuildConfiguration, Action[] PrerequisiteActions, List<UEBuildTarget> Targets, Dictionary<UEBuildTarget, CPPHeaders> TargetToHeaders, bool bIsAssemblingBuild, bool bNeedsFullCPPIncludeRescan, out Dictionary<UEBuildTarget, List<FileItem>> TargetToOutdatedPrerequisitesMap)
 		{
-			DateTime CheckOutdatednessStartTime = DateTime.UtcNow;
+			ITimelineEvent GetActionsToExecuteTimer = Timeline.ScopeEvent("ActionGraph.GetActionsToExecute()");
 
 			// Build a set of all actions needed for this target.
 			Dictionary<Action, bool> IsActionOutdatedMap = new Dictionary<Action, bool>();
@@ -371,11 +371,7 @@ namespace UnrealBuildTool
 			// Build a list of actions that are both needed for this target and outdated.
 			HashSet<Action> ActionsToExecute = new HashSet<Action>(AllActions.Where(Action => Action.CommandPath != null && IsActionOutdatedMap.ContainsKey(Action) && OutdatedActionDictionary[Action]));
 
-			if (UnrealBuildTool.bPrintPerformanceInfo)
-			{
-				double CheckOutdatednessTime = (DateTime.UtcNow - CheckOutdatednessStartTime).TotalSeconds;
-				Log.TraceInformation("Checking outdatedness took " + CheckOutdatednessTime + "s");
-			}
+			GetActionsToExecuteTimer.Finish();
 
 			return ActionsToExecute;
 		}
@@ -903,7 +899,7 @@ namespace UnrealBuildTool
 				// includes for this file, so that we'll be able to determine whether it is out of date next time very quickly.
 				if (BuildConfiguration.bUseUBTMakefiles)
 				{
-					DateTime DeepIncludeScanStartTime = DateTime.UtcNow;
+					Stopwatch DeepIncludeScanTimer = Stopwatch.StartNew();
 
 					// @todo ubtmake: we may be scanning more files than we need to here -- indirectly outdated files are bIsOutdated=true by this point (for example basemost includes when deeper includes are dirty)
 					if (bIsOutdated && RootAction.ActionType == ActionType.Compile)	// @todo ubtmake: Does this work with RC files?  See above too.
@@ -935,11 +931,7 @@ namespace UnrealBuildTool
 						}
 					}
 
-					if (UnrealBuildTool.bPrintPerformanceInfo)
-					{
-						double DeepIncludeScanTime = (DateTime.UtcNow - DeepIncludeScanStartTime).TotalSeconds;
-						UnrealBuildTool.TotalDeepIncludeScanTime += DeepIncludeScanTime;
-					}
+					UnrealBuildTool.TotalDeepIncludeScanTime += DeepIncludeScanTimer.Elapsed;
 				}
 
 				// Cache the outdated-ness of this action.
@@ -956,18 +948,10 @@ namespace UnrealBuildTool
 		/// </summary>
 		void GatherAllOutdatedActions(BuildConfiguration BuildConfiguration, UEBuildTarget Target, CPPHeaders Headers, bool bIsAssemblingBuild, bool bNeedsFullCPPIncludeRescan, ActionHistory ActionHistory, ref Dictionary<Action, bool> OutdatedActions, Dictionary<UEBuildTarget, List<FileItem>> TargetToOutdatedPrerequisitesMap)
 		{
-			DateTime CheckOutdatednessStartTime = DateTime.UtcNow;
-
 			foreach (Action Action in AllActions)
 			{
 				IsActionOutdated(BuildConfiguration, Target, Headers, Action, bIsAssemblingBuild, bNeedsFullCPPIncludeRescan, OutdatedActions, ActionHistory, TargetToOutdatedPrerequisitesMap);
 			}
-
-			if (UnrealBuildTool.bPrintPerformanceInfo)
-			{
-				double CheckOutdatednessTime = (DateTime.UtcNow - CheckOutdatednessStartTime).TotalSeconds;
-				Log.TraceInformation("Checking actions for " + Target.GetTargetName() + " took " + CheckOutdatednessTime + "s");
-			}
 		}
 
 		/// <summary>
diff --git a/Engine/Source/Programs/UnrealBuildTool/System/CPPHeaders.cs b/Engine/Source/Programs/UnrealBuildTool/System/CPPHeaders.cs
--- a/Engine/Source/Programs/UnrealBuildTool/System/CPPHeaders.cs
+++ b/Engine/Source/Programs/UnrealBuildTool/System/CPPHeaders.cs
@@ -133,6 +133,11 @@ namespace UnrealBuildTool
 		/// </summary>
 		public static int IncludePathSearchAttempts = 0;
 
+		/// <summary>
+		/// Whether to print performance info on include file resolution
+		/// </summary>
+		public static bool bPrintPerformanceInfo = false;
+
 		/// <summary>
 		/// A cache of the list of other files that are directly or indirectly included by a C++ file.
 		/// </summary>
@@ -225,7 +230,7 @@ namespace UnrealBuildTool
 
 				IncludePathSearchAttempts += SearchAttempts;
 
-				if (UnrealBuildTool.bPrintPerformanceInfo)
+				if (bPrintPerformanceInfo)
 				{
 					// More than two search attempts indicates:
 					//		- Include path was not relative to the directory that the including file was in
@@ -420,9 +425,7 @@ namespace UnrealBuildTool
 			return DirectlyIncludedFiles;
 		}
 
-		public static double TotalTimeSpentGettingIncludes = 0.0;
-		public static int TotalIncludesRequested = 0;
-		public static double DirectIncludeCacheMissesTotalTime = 0.0;
+		public static TimeSpan DirectIncludeCacheMissesTotalTime = TimeSpan.Zero;
 		public static int TotalDirectIncludeCacheMisses = 0;
 		public static int TotalDirectIncludeResolveCacheMisses = 0;
 		public static int TotalDirectIncludeResolves = 0;
@@ -487,7 +490,8 @@ namespace UnrealBuildTool
 				return Result;
 			}
 
-			DateTime TimerStartTime = DateTime.UtcNow;
+			Stopwatch Timer = Stopwatch.StartNew();
+
 			++CPPHeaders.TotalDirectIncludeCacheMisses;
 
 			Result = GetUncachedDirectIncludeDependencies(CPPFile.Location, ProjectFile);
@@ -495,7 +499,7 @@ namespace UnrealBuildTool
 			// Populate cache with results.
 			IncludeDependencyCache.SetDependencyInfo(CPPFile, Result);
 
-			CPPHeaders.DirectIncludeCacheMissesTotalTime += (DateTime.UtcNow - TimerStartTime).TotalSeconds;
+			CPPHeaders.DirectIncludeCacheMissesTotalTime += Timer.Elapsed;
 
 			return Result;
 		}
diff --git a/Engine/Source/Programs/UnrealBuildTool/System/DependencyCache.cs b/Engine/Source/Programs/UnrealBuildTool/System/DependencyCache.cs
--- a/Engine/Source/Programs/UnrealBuildTool/System/DependencyCache.cs
+++ b/Engine/Source/Programs/UnrealBuildTool/System/DependencyCache.cs
@@ -92,25 +92,12 @@ namespace UnrealBuildTool
 			// See whether the cache file exists.
 			if (FileReference.Exists(CacheFile))
 			{
-				if (UnrealBuildTool.bPrintPerformanceInfo)
-				{
-					Log.TraceInformation("Loading existing IncludeFileCache: " + CacheFile.FullName);
-				}
-
-				DateTime TimerStartTime = DateTime.UtcNow;
-
 				// Deserialize cache from disk if there is one.
 				DependencyCache Result = Load(CacheFile);
 				if (Result != null)
 				{
 					// Successfully serialize, create the transient variables and return cache.
 					Result.UpdateTimeUtc = DateTime.UtcNow;
-
-					TimeSpan TimerDuration = DateTime.UtcNow - TimerStartTime;
-					if (UnrealBuildTool.bPrintPerformanceInfo)
-					{
-						Log.TraceInformation("Loading IncludeFileCache took " + TimerDuration.TotalSeconds + "s");
-					}
 					return Result;
 				}
 			}
@@ -239,8 +226,6 @@ namespace UnrealBuildTool
 			// Only save if we've made changes to it since load.
 			if (bIsDirty)
 			{
-				DateTime TimerStartTime = DateTime.UtcNow;
-
 				// Save update date as new creation date.
 				CreationTimeUtc = UpdateTimeUtc;
 
@@ -258,19 +243,6 @@ namespace UnrealBuildTool
 				{
 					Log.TraceError("Failed to write dependency cache: {0}", Ex.Message);
 				}
-
-				if (UnrealBuildTool.bPrintPerformanceInfo)
-				{
-					TimeSpan TimerDuration = DateTime.UtcNow - TimerStartTime;
-					Log.TraceInformation("Saving IncludeFileCache took " + TimerDuration.TotalSeconds + "s");
-				}
-			}
-			else
-			{
-				if (UnrealBuildTool.bPrintPerformanceInfo)
-				{
-					Log.TraceInformation("IncludeFileCache did not need to be saved (bIsDirty=false)");
-				}
 			}
 
 			FileReference MutexPath = BackingFile + ".buildmutex";
diff --git a/Engine/Source/Programs/UnrealBuildTool/System/ExternalExecution.cs b/Engine/Source/Programs/UnrealBuildTool/System/ExternalExecution.cs
--- a/Engine/Source/Programs/UnrealBuildTool/System/ExternalExecution.cs
+++ b/Engine/Source/Programs/UnrealBuildTool/System/ExternalExecution.cs
@@ -589,8 +589,6 @@ namespace UnrealBuildTool
 
 		public static void SetupUObjectModules(IEnumerable<UEBuildModuleCPP> ModulesToGenerateHeadersFor, UnrealTargetPlatform Platform, ProjectDescriptor ProjectDescriptor, List<UHTModuleInfo> UObjectModules, Dictionary<string, FlatModuleCsDataType> FlatModuleCsData, EGeneratedCodeVersion GeneratedCodeVersion, bool bIsAssemblingBuild)
 		{
-			DateTime UObjectDiscoveryStartTime = DateTime.UtcNow;
-
 			// Find the type of each module
 			Dictionary<UEBuildModuleCPP, UHTModuleType> ModuleToType = new Dictionary<UEBuildModuleCPP, UHTModuleType>();
 			foreach(UEBuildModuleCPP Module in ModulesToGenerateHeadersFor)
@@ -640,12 +638,6 @@ namespace UnrealBuildTool
 					}
 				}
 			}
-
-			if (UnrealBuildTool.bPrintPerformanceInfo)
-			{
-				double UObjectDiscoveryTime = (DateTime.UtcNow - UObjectDiscoveryStartTime).TotalSeconds;
-				Log.TraceInformation("UObject discovery time: " + UObjectDiscoveryTime + "s");
-			}
 		}
 
 		/// <summary>
@@ -1226,7 +1218,9 @@ namespace UnrealBuildTool
 
 					Stopwatch s = new Stopwatch();
 					s.Start();
+					ITimelineEvent Timer = Timeline.ScopeEvent("Executing UnrealHeaderTool");
 					ECompilationResult UHTResult = (ECompilationResult)RunExternalNativeExecutable(ExternalExecution.GetHeaderToolPath(HeaderToolReceipt), CmdLine);
+					Timer.Finish();
 					s.Stop();
 
 					if (UHTResult != ECompilationResult.Succeeded)
@@ -1251,10 +1245,6 @@ namespace UnrealBuildTool
 					}
 
 					Log.TraceInformation("Reflection code generated for {0} in {1} seconds", ActualTargetName, s.Elapsed.TotalSeconds);
-					if (UnrealBuildTool.bPrintPerformanceInfo)
-					{
-						Log.TraceInformation("UnrealHeaderTool took {1}", ActualTargetName, (double)s.ElapsedMilliseconds / 1000.0);
-					}
 
 					// Now that UHT has successfully finished generating code, we need to update all cached FileItems in case their last write time has changed.
 					// Otherwise UBT might not detect changes UHT made.
diff --git a/Engine/Source/Programs/UnrealBuildTool/System/FlatCPPIncludeDepencencyCache.cs b/Engine/Source/Programs/UnrealBuildTool/System/FlatCPPIncludeDepencencyCache.cs
--- a/Engine/Source/Programs/UnrealBuildTool/System/FlatCPPIncludeDepencencyCache.cs
+++ b/Engine/Source/Programs/UnrealBuildTool/System/FlatCPPIncludeDepencencyCache.cs
@@ -76,13 +76,6 @@ namespace UnrealBuildTool
 				return false;
 			}
 
-			if (UnrealBuildTool.bPrintPerformanceInfo)
-			{
-				Log.TraceInformation("Loading existing FlatCPPIncludeDependencyCache: " + BackingFile.FullName);
-			}
-
-			DateTime TimerStartTime = DateTime.UtcNow;
-
 			try
 			{
 				// If the .buildmutex file for the cache is present, it means that something went wrong between loading
@@ -107,13 +100,6 @@ namespace UnrealBuildTool
 					}
 
 					Cache = Deserialize(Reader);
-
-					if (UnrealBuildTool.bPrintPerformanceInfo)
-					{
-						TimeSpan TimerDuration = DateTime.UtcNow - TimerStartTime;
-						Log.TraceInformation("Loading FlatCPPIncludeDependencyCache took " + TimerDuration.TotalSeconds + "s");
-					}
-
 					return true;
 				}
 			}
@@ -134,8 +120,6 @@ namespace UnrealBuildTool
 			// Only save if we've made changes to it since load.
 			if (bIsDirty)
 			{
-				DateTime TimerStartTime = DateTime.UtcNow;
-
 				// Serialize the cache to disk.
 				try
 				{
@@ -150,19 +134,6 @@ namespace UnrealBuildTool
 				{
 					Log.TraceError("Failed to write FlatCPPIncludeDependencyCache: {0}", Ex.Message);
 				}
-
-				if (UnrealBuildTool.bPrintPerformanceInfo)
-				{
-					TimeSpan TimerDuration = DateTime.UtcNow - TimerStartTime;
-					Log.TraceInformation("Saving FlatCPPIncludeDependencyCache took " + TimerDuration.TotalSeconds + "s");
-				}
-			}
-			else
-			{
-				if (UnrealBuildTool.bPrintPerformanceInfo)
-				{
-					Log.TraceInformation("FlatCPPIncludeDependencyCache did not need to be saved (bIsDirty=false)");
-				}
 			}
 
 			if (File.Exists(BackingFile.FullName + ".buildmutex"))
diff --git a/Engine/Source/Programs/UnrealBuildTool/System/JunkDeleter.cs b/Engine/Source/Programs/UnrealBuildTool/System/JunkDeleter.cs
--- a/Engine/Source/Programs/UnrealBuildTool/System/JunkDeleter.cs
+++ b/Engine/Source/Programs/UnrealBuildTool/System/JunkDeleter.cs
@@ -16,18 +16,13 @@ namespace UnrealBuildTool
 		/// </summary>
 		static public void DeleteJunk()
 		{
-			DateTime JunkStartTime = DateTime.UtcNow;
-
 			if (UnrealBuildTool.IsEngineInstalled() == false)
 			{
-				List<string> JunkManifest = LoadJunkManifest();
-				DeleteAllJunk(JunkManifest);
-			}
-
-			if (UnrealBuildTool.bPrintPerformanceInfo)
-			{
-				double JunkTime = (DateTime.UtcNow - JunkStartTime).TotalSeconds;
-				Log.TraceInformation("DeleteJunk took " + JunkTime + "s");
+				using(Timeline.ScopeEvent("DeleteAllJunk()"))
+				{
+					List<string> JunkManifest = LoadJunkManifest();
+					DeleteAllJunk(JunkManifest);
+				}
 			}
 		}
 
diff --git a/Engine/Source/Programs/UnrealBuildTool/System/Timeline.cs b/Engine/Source/Programs/UnrealBuildTool/System/Timeline.cs
--- a/Engine/Source/Programs/UnrealBuildTool/System/Timeline.cs
+++ b/Engine/Source/Programs/UnrealBuildTool/System/Timeline.cs
@@ -117,18 +117,6 @@ namespace UnrealBuildTool
 			return Event;
 		}
 
-		/// <summary>
-		/// Enters a scope event with the calling method name as the event name
-		/// </summary>
-		/// <returns>Event to track the lifetime of the method. Dispose it before returning.</returns>
-		[MethodImplAttribute(MethodImplOptions.NoInlining)]
-		public static ITimelineEvent ScopeMethod()
-		{
-			StackFrame Frame = new StackFrame(1);
-			MethodBase Method = Frame.GetMethod();
-			return ScopeEvent(String.Format("{0}.{1}", Method.DeclaringType.Name, Method.Name));
-		}
-
 		/// <summary>
 		/// Prints this information to the log
 		/// </summary>
@@ -169,9 +157,16 @@ namespace UnrealBuildTool
 				Print(Event.StartTime, Event.FinishTime, Event.Name, OuterEvents, Verbosity);
 
 				// Push it onto the stack
-				if(Event.FinishTime.HasValue && EventIdx + 1 < Events.Count && Events[EventIdx + 1].StartTime < Event.FinishTime.Value)
+				if(Event.FinishTime.HasValue)
 				{
-					OuterEvents.Add(Event);
+					if(EventIdx + 1 < Events.Count && Events[EventIdx + 1].StartTime < Event.FinishTime.Value)
+					{
+						OuterEvents.Add(Event);
+					}
+					else
+					{
+						LastTime = Event.FinishTime.Value;
+					}
 				}
 			}
 
diff --git a/Engine/Source/Programs/UnrealBuildTool/System/UBTMakefile.cs b/Engine/Source/Programs/UnrealBuildTool/System/UBTMakefile.cs
--- a/Engine/Source/Programs/UnrealBuildTool/System/UBTMakefile.cs
+++ b/Engine/Source/Programs/UnrealBuildTool/System/UBTMakefile.cs
@@ -141,8 +141,6 @@ namespace UnrealBuildTool
 				throw new BuildException("Can't save a makefile that has invalid contents.  See UBTMakefile.IsValidMakefile()");
 			}
 
-			DateTime TimerStartTime = DateTime.UtcNow;
-
 			FileItem UBTMakefileItem = FileItem.GetItemByFileReference(GetUBTMakefilePath(TargetDescs));
 
 			// @todo ubtmake: Optimization: The UBTMakefile saved for game projects is upwards of 9 MB.  We should try to shrink its content if possible
@@ -162,12 +160,6 @@ namespace UnrealBuildTool
 			{
 				Log.TraceError("Failed to write makefile: {0}", Ex.Message);
 			}
-
-			if (UnrealBuildTool.bPrintPerformanceInfo)
-			{
-				TimeSpan TimerDuration = DateTime.UtcNow - TimerStartTime;
-				Log.TraceInformation("Saving makefile took " + TimerDuration.TotalSeconds + "s");
-			}
 		}
 
 
@@ -275,18 +267,13 @@ namespace UnrealBuildTool
 
 			try
 			{
-				DateTime LoadUBTMakefileStartTime = DateTime.UtcNow;
-
-				using (FileStream Stream = new FileStream(UBTMakefileInfo.FullName, FileMode.Open, FileAccess.Read))
+				using(Timeline.ScopeEvent("Loading makefile"))
 				{
-					BinaryFormatter Formatter = new BinaryFormatter();
-					LoadedUBTMakefile = Formatter.Deserialize(Stream) as UBTMakefile;
-				}
-
-				if (UnrealBuildTool.bPrintPerformanceInfo)
-				{
-					double LoadUBTMakefileTime = (DateTime.UtcNow - LoadUBTMakefileStartTime).TotalSeconds;
-					Log.TraceInformation("LoadUBTMakefile took " + LoadUBTMakefileTime + "s");
+					using (FileStream Stream = new FileStream(UBTMakefileInfo.FullName, FileMode.Open, FileAccess.Read))
+					{
+						BinaryFormatter Formatter = new BinaryFormatter();
+						LoadedUBTMakefile = Formatter.Deserialize(Stream) as UBTMakefile;
+					}
 				}
 			}
 			catch (Exception Ex)
diff --git a/Engine/Source/Programs/UnrealBuildTool/UnrealBuildTool.cs b/Engine/Source/Programs/UnrealBuildTool/UnrealBuildTool.cs
--- a/Engine/Source/Programs/UnrealBuildTool/UnrealBuildTool.cs
+++ b/Engine/Source/Programs/UnrealBuildTool/UnrealBuildTool.cs
@@ -17,7 +17,7 @@ namespace UnrealBuildTool
 		/// <summary>
 		/// How much time was spent scanning for include dependencies for outdated C++ files
 		/// </summary>
-		public static double TotalDeepIncludeScanTime = 0.0;
+		public static TimeSpan TotalDeepIncludeScanTime = TimeSpan.Zero;
 
 		/// <summary>
 		/// The environment at boot time.
@@ -109,15 +109,6 @@ namespace UnrealBuildTool
 		/// </summary>
 		static string RemoteIniPath = null;
 
-		/// <summary>
-		/// Whether to print performance information to the log
-		/// </summary>
-		static public bool bPrintPerformanceInfo
-		{
-			get;
-			set;
-		}
-
 		/// <summary>
 		/// Returns true if UnrealBuildTool is running using installed Engine components
 		/// </summary>
@@ -462,8 +453,7 @@ namespace UnrealBuildTool
 		{
 			bool bSuccess = true;
 
-			DateTime RunUBTInitStartTime = DateTime.UtcNow;
-
+			Stopwatch BuildTimer = Stopwatch.StartNew();
 
 			// Reset global configurations
 			ActionGraph ActionGraph = new ActionGraph();
@@ -476,7 +466,7 @@ namespace UnrealBuildTool
 			List<UEBuildTarget> Targets = null;
 			Dictionary<UEBuildTarget, CPPHeaders> TargetToHeaders = new Dictionary<UEBuildTarget, CPPHeaders>();
 
-			double TotalExecutorTime = 0.0;
+			Stopwatch ExecutorTimer = new Stopwatch();
 
 			try
 			{
@@ -488,25 +478,14 @@ namespace UnrealBuildTool
 					BuildConfiguration.bUseIncludeDependencyResolveCache = true;
 				}
 
-				if (UnrealBuildTool.bPrintPerformanceInfo)
-				{
-					double RunUBTInitTime = (DateTime.UtcNow - RunUBTInitStartTime).TotalSeconds;
-					Log.TraceInformation("RunUBT initialization took " + RunUBTInitTime + "s");
-				}
+				Timeline.AddEvent("RunUBT() initialization complete");
 
 				bool bSkipRulesCompile = Arguments.Any(x => x.Equals("-skiprulescompile", StringComparison.InvariantCultureIgnoreCase));
 
 				List<TargetDescriptor> TargetDescs = new List<TargetDescriptor>();
+				using(Timeline.ScopeEvent("TargetDescriptor.ParseCommandLine()"))
 				{
-					DateTime TargetDescstStartTime = DateTime.UtcNow;
-
 					TargetDescs.AddRange(TargetDescriptor.ParseCommandLine(new CommandLineArguments(Arguments), BuildConfiguration.bUsePrecompiled, bSkipRulesCompile));
-
-					if (UnrealBuildTool.bPrintPerformanceInfo)
-					{
-						double TargetDescsTime = (DateTime.UtcNow - TargetDescstStartTime).TotalSeconds;
-						Log.TraceInformation("Target descriptors took " + TargetDescsTime + "s");
-					}
 				}
 
 				// Handle remote builds
@@ -679,7 +658,10 @@ namespace UnrealBuildTool
 
 							// Try to load the UBTMakefile.  It will only be loaded if it has valid content and is not determined to be out of date.    
 							string ReasonNotLoaded;
-							UBTMakefile = UBTMakefile.LoadUBTMakefile(UBTMakefilePath, TargetDescs[0].ProjectFile, WorkingSet, out ReasonNotLoaded);
+							using(Timeline.ScopeEvent("UBTMakefile.LoadUBTMakefile"))
+							{
+								UBTMakefile = UBTMakefile.LoadUBTMakefile(UBTMakefilePath, TargetDescs[0].ProjectFile, WorkingSet, out ReasonNotLoaded);
+							}
 
 							if (UBTMakefile != null)
 							{
@@ -738,23 +720,18 @@ namespace UnrealBuildTool
 					}
 					else
 					{
-						DateTime TargetInitStartTime = DateTime.UtcNow;
-
 						Targets = new List<UEBuildTarget>();
 						foreach (TargetDescriptor TargetDesc in TargetDescs)
 						{
-							UEBuildTarget Target = UEBuildTarget.CreateTarget(TargetDesc, Arguments, bSkipRulesCompile, BuildConfiguration.SingleFileToCompile != null, BuildConfiguration.bUsePrecompiled);
-							if ((Target == null) && (BuildConfiguration.bCleanProject))
+							using(Timeline.ScopeEvent("UEBuildTarget.CreateTarget()"))
 							{
-								continue;
+								UEBuildTarget Target = UEBuildTarget.CreateTarget(TargetDesc, Arguments, bSkipRulesCompile, BuildConfiguration.SingleFileToCompile != null, BuildConfiguration.bUsePrecompiled);
+								if ((Target == null) && (BuildConfiguration.bCleanProject))
+								{
+									continue;
+								}
+								Targets.Add(Target);
 							}
-							Targets.Add(Target);
-						}
-
-						if (UnrealBuildTool.bPrintPerformanceInfo)
-						{
-							double TargetInitTime = (DateTime.UtcNow - TargetInitStartTime).TotalSeconds;
-							Log.TraceInformation("Target init took " + TargetInitTime + "s");
 						}
 					}
 
@@ -782,27 +759,33 @@ namespace UnrealBuildTool
 						if (!(!bIsGatheringBuild && bIsAssemblingBuild))
 						{
 							// Load the direct include dependency cache.
-							Headers.IncludeDependencyCache = DependencyCache.Create(DependencyCache.GetDependencyCachePathForTarget(Target.ProjectFile, Target.Platform, Target.GetTargetName()));
+							using(Timeline.ScopeEvent("DependencyCache.Create()"))
+							{
+								Headers.IncludeDependencyCache = DependencyCache.Create(DependencyCache.GetDependencyCachePathForTarget(Target.ProjectFile, Target.Platform, Target.GetTargetName()));
+							}
 						}
 
 						// We don't need this dependency cache in 'gather only' mode
 						if (BuildConfiguration.bUseUBTMakefiles &&
 								!(bIsGatheringBuild && !bIsAssemblingBuild))
 						{
-							// Load the cache that contains the list of flattened resolved includes for resolved source files
-							// @todo ubtmake: Ideally load this asynchronously at startup and only block when it is first needed and not finished loading
-							FileReference CacheFile = FlatCPPIncludeDependencyCache.GetDependencyCachePathForTarget(Target.ProjectFile, Target.TargetName, Target.Platform, Target.Architecture);
-							if (!FlatCPPIncludeDependencyCache.TryRead(CacheFile, out Headers.FlatCPPIncludeDependencyCache))
+							using(Timeline.ScopeEvent("FlatCPPIncludeDependencyCache.TryRead()"))
 							{
-								if (!bNeedsFullCPPIncludeRescan)
+								// Load the cache that contains the list of flattened resolved includes for resolved source files
+								// @todo ubtmake: Ideally load this asynchronously at startup and only block when it is first needed and not finished loading
+								FileReference CacheFile = FlatCPPIncludeDependencyCache.GetDependencyCachePathForTarget(Target.ProjectFile, Target.TargetName, Target.Platform, Target.Architecture);
+								if (!FlatCPPIncludeDependencyCache.TryRead(CacheFile, out Headers.FlatCPPIncludeDependencyCache))
 								{
-									if (!BuildConfiguration.bXGEExport && !BuildConfiguration.bGenerateManifest && !BuildConfiguration.bCleanProject)
+									if (!bNeedsFullCPPIncludeRescan)
 									{
-										bNeedsFullCPPIncludeRescan = true;
-										Log.TraceInformation("Performing full C++ include scan (no include cache file)");
+										if (!BuildConfiguration.bXGEExport && !BuildConfiguration.bGenerateManifest && !BuildConfiguration.bCleanProject)
+										{
+											bNeedsFullCPPIncludeRescan = true;
+											Log.TraceInformation("Performing full C++ include scan (no include cache file)");
+										}
 									}
+									Headers.FlatCPPIncludeDependencyCache = new FlatCPPIncludeDependencyCache(CacheFile);
 								}
-								Headers.FlatCPPIncludeDependencyCache = new FlatCPPIncludeDependencyCache(CacheFile);
 							}
 						}
 
@@ -888,7 +871,10 @@ namespace UnrealBuildTool
 								// to assemble the build.  Even if we are configured to assemble the build in this same invocation, we want to save out the
 								// Makefile so that it can be used on subsequent 'assemble only' runs, for the fastest possible iteration times
 								// @todo ubtmake: Optimization: We could make 'gather + assemble' mode slightly faster by saving this while busy compiling (on our worker thread)
-								UBTMakefile.SaveUBTMakefile(TargetDescs, UBTMakefile);
+								using(Timeline.ScopeEvent("UBTMakefile.SaveUBTMakefile()"))
+								{
+									UBTMakefile.SaveUBTMakefile(TargetDescs, UBTMakefile);
+								}
 							}
 						}
 
@@ -1224,17 +1210,23 @@ namespace UnrealBuildTool
 								}
 
 								// Execute the actions.
-								DateTime ExecutorStartTime = DateTime.UtcNow;
+								ExecutorTimer.Start();
 								if (BuildConfiguration.bXGEExport)
 								{
-									XGE.ExportActions(ActionsToExecute.ToList());
+									using(Timeline.ScopeEvent("XGE.ExportActions()"))
+									{
+										XGE.ExportActions(ActionsToExecute.ToList());
+									}
 									bSuccess = true;
 								}
 								else
 								{
-									bSuccess = ActionGraph.ExecuteActions(BuildConfiguration, ActionsToExecute.ToList(), out ExecutorName);
+									using(Timeline.ScopeEvent("ActionGraph.ExecuteActions()"))
+									{
+										bSuccess = ActionGraph.ExecuteActions(BuildConfiguration, ActionsToExecute.ToList(), out ExecutorName);
+									}
 								}
-								TotalExecutorTime += (DateTime.UtcNow - ExecutorStartTime).TotalSeconds;
+								ExecutorTimer.Stop();
 
 								// if the build succeeded, write the receipts and do any needed syncing
 								if (!bSuccess)
@@ -1273,7 +1265,10 @@ namespace UnrealBuildTool
 			// Wait until our CPPIncludes dependency scanner thread has finished
 			if (CppIncludeThread != null)
 			{
-				CppIncludeThread.Join();
+				using(Timeline.ScopeEvent("CppIncludeThread.Join()"))
+				{
+					CppIncludeThread.Join();
+				}
 			}
 
 			// Save the include dependency cache.
@@ -1286,20 +1281,25 @@ namespace UnrealBuildTool
 
 				if (Headers.IncludeDependencyCache != null)
 				{
-					Headers.IncludeDependencyCache.Save();
+					using(Timeline.ScopeEvent("DependencyCache.Save()"))
+					{
+						Headers.IncludeDependencyCache.Save();
+					}
 				}
 
 				if (Headers.FlatCPPIncludeDependencyCache != null)
 				{
-					Headers.FlatCPPIncludeDependencyCache.Save();
+					using (Timeline.ScopeEvent("FlatCPPIncludeDependencyCache.Save()"))
+					{
+						Headers.FlatCPPIncludeDependencyCache.Save();
+					}
 				}
 			}
 
 			// Figure out how long we took to execute.
 			if (ExecutorName != "Unknown")
 			{
-				double BuildDuration = (DateTime.UtcNow - RunUBTInitStartTime).TotalSeconds;
-				Log.TraceInformation("Total build time: {0:0.00} seconds ({1} executor: {2:0.00} seconds)", BuildDuration, ExecutorName, TotalExecutorTime);
+				Log.TraceInformation("Total build time: {0:0.00} seconds ({1} executor: {2:0.00} seconds)", BuildTimer.Elapsed.TotalSeconds, ExecutorName, ExecutorTimer.Elapsed.TotalSeconds);
 			}
 
 			return BuildResult;
