UBT: Change UBT makefile to use explicit serialization rather than BinaryFormatter et al. Improves load/save performance by 50%, and allows merging of serialized file items.

#rnx

diff --git a/Engine/Source/Programs/UnrealBuildTool/Configuration/UEBuildTarget.cs b/Engine/Source/Programs/UnrealBuildTool/Configuration/UEBuildTarget.cs
--- a/Engine/Source/Programs/UnrealBuildTool/Configuration/UEBuildTarget.cs
+++ b/Engine/Source/Programs/UnrealBuildTool/Configuration/UEBuildTarget.cs
@@ -253,25 +253,24 @@ namespace UnrealBuildTool
 		}
 	}
 
-	[Serializable]
-	class FlatModuleCsDataType : ISerializable
+	class FlatModuleCsDataType
 	{
-		public FlatModuleCsDataType(SerializationInfo Info, StreamingContext Context)
+		public FlatModuleCsDataType(BinaryReader Reader)
 		{
-			ModuleName = Info.GetString("mn");
-			BuildCsFilename = Info.GetString("bf");
-			ModuleSourceFolder = (DirectoryReference)Info.GetValue("mf", typeof(DirectoryReference));
-			ExternalDependencies = (List<string>)Info.GetValue("ed", typeof(List<string>));
-			UHTHeaderNames = (List<string>)Info.GetValue("hn", typeof(List<string>));
+			ModuleName = Reader.ReadString();
+			BuildCsFilename = Reader.ReadString();
+			ModuleSourceFolder = Reader.ReadDirectoryReference();
+			ExternalDependencies = Reader.ReadStringList();
+			UHTHeaderNames = Reader.ReadStringList();
 		}
 
-		public void GetObjectData(SerializationInfo Info, StreamingContext Context)
+		public void Write(BinaryWriter Writer)
 		{
-			Info.AddValue("mn", ModuleName);
-			Info.AddValue("bf", BuildCsFilename);
-			Info.AddValue("mf", ModuleSourceFolder);
-			Info.AddValue("ed", ExternalDependencies);
-			Info.AddValue("hn", UHTHeaderNames);
+			Writer.Write(ModuleName);
+			Writer.Write(BuildCsFilename);
+			Writer.Write(ModuleSourceFolder);
+			Writer.Write(ExternalDependencies);
+			Writer.Write(UHTHeaderNames);
 		}
 
 		public FlatModuleCsDataType(string InModuleName, string InBuildCsFilename, IEnumerable<string> InExternalDependencies)
@@ -291,8 +290,7 @@ namespace UnrealBuildTool
 	/// <summary>
 	/// A target that can be built
 	/// </summary>
-	[Serializable]
-	class UEBuildTarget : ISerializable
+	class UEBuildTarget
 	{
 		public string GetAppName()
 		{
@@ -743,58 +741,58 @@ namespace UnrealBuildTool
 			return bCompileMonolithic;	// @todo ubtmake: We need to make sure this function and similar things aren't called in assembler mode
 		}
 
-		public UEBuildTarget(SerializationInfo Info, StreamingContext Context)
-		{
-			TargetType = (TargetType)Info.GetInt32("tt");
-			ProjectFile = (FileReference)Info.GetValue("pf", typeof(FileReference));
-			AppName = Info.GetString("an");
-			TargetName = Info.GetString("tn");
-			bUseSharedBuildEnvironment = Info.GetBoolean("sb");
-			Platform = (UnrealTargetPlatform)Info.GetInt32("pl");
-			Configuration = (UnrealTargetConfiguration)Info.GetInt32("co");
-			Architecture = Info.GetString("ar");
-			PlatformIntermediateFolder = Info.GetString("if");
-			ProjectDirectory = (DirectoryReference)Info.GetValue("pd", typeof(DirectoryReference));
-			ProjectIntermediateDirectory = (DirectoryReference)Info.GetValue("pi", typeof(DirectoryReference));
-			EngineIntermediateDirectory = (DirectoryReference)Info.GetValue("ed", typeof(DirectoryReference));
-			OutputPaths = (List<FileReference>)Info.GetValue("op", typeof(List<FileReference>));
-			VersionFile = (FileReference)Info.GetValue("vf", typeof(FileReference));
-			bPrecompile = Info.GetBoolean("pc");
-			bCompileMonolithic = Info.GetBoolean("cm");
-			FlatModuleCsData = (List<FlatModuleCsDataType>)Info.GetValue("fm", typeof(List<FlatModuleCsDataType>));
-			ReceiptFileName = (FileReference)Info.GetValue("rf", typeof(FileReference));
-			TargetRulesFile = (FileReference)Info.GetValue("tc", typeof(FileReference));
-			PreBuildStepScripts = (FileReference[])Info.GetValue("pr", typeof(FileReference[]));
-			PostBuildStepScripts = (FileReference[])Info.GetValue("po", typeof(FileReference[]));
-			bDeployAfterCompile = Info.GetBoolean("dt");
-			bHasProjectScriptPlugin = Info.GetBoolean("sp");
-		}
-
-		public void GetObjectData(SerializationInfo Info, StreamingContext Context)
-		{
-			Info.AddValue("tt", (int)TargetType);
-			Info.AddValue("pf", ProjectFile);
-			Info.AddValue("an", AppName);
-			Info.AddValue("tn", TargetName);
-			Info.AddValue("sb", bUseSharedBuildEnvironment);
-			Info.AddValue("pl", (int)Platform);
-			Info.AddValue("co", (int)Configuration);
-			Info.AddValue("ar", Architecture);
-			Info.AddValue("if", PlatformIntermediateFolder);
-			Info.AddValue("pd", ProjectDirectory);
-			Info.AddValue("pi", ProjectIntermediateDirectory);
-			Info.AddValue("ed", EngineIntermediateDirectory);
-			Info.AddValue("op", OutputPaths);
-			Info.AddValue("vf", VersionFile);
-			Info.AddValue("pc", bPrecompile);
-			Info.AddValue("cm", bCompileMonolithic);
-			Info.AddValue("fm", FlatModuleCsData);
-			Info.AddValue("rf", ReceiptFileName);
-			Info.AddValue("tc", TargetRulesFile);
-			Info.AddValue("pr", PreBuildStepScripts);
-			Info.AddValue("po", PostBuildStepScripts);
-			Info.AddValue("dt", bDeployAfterCompile);
-			Info.AddValue("sp", bHasProjectScriptPlugin);
+		public UEBuildTarget(BinaryReader Reader)
+		{
+			TargetType = (TargetType)Reader.ReadInt32();
+			ProjectFile = Reader.ReadFileReference();
+			AppName = Reader.ReadString();
+			TargetName = Reader.ReadString();
+			bUseSharedBuildEnvironment = Reader.ReadBoolean();
+			Platform = (UnrealTargetPlatform)Reader.ReadInt32();
+			Configuration = (UnrealTargetConfiguration)Reader.ReadInt32();
+			Architecture = Reader.ReadString();
+			PlatformIntermediateFolder = Reader.ReadString();
+			ProjectDirectory = Reader.ReadDirectoryReference();
+			ProjectIntermediateDirectory = Reader.ReadDirectoryReference();
+			EngineIntermediateDirectory = Reader.ReadDirectoryReference();
+			OutputPaths = Reader.ReadList(() => Reader.ReadFileReference());
+			VersionFile = Reader.ReadFileReference();
+			bPrecompile = Reader.ReadBoolean();
+			bCompileMonolithic = Reader.ReadBoolean();
+			FlatModuleCsData = Reader.ReadList(() => new FlatModuleCsDataType(Reader));
+			ReceiptFileName = Reader.ReadFileReference();
+			TargetRulesFile = Reader.ReadFileReference();
+			PreBuildStepScripts = Reader.ReadArray(() => Reader.ReadFileReference());
+			PostBuildStepScripts = Reader.ReadArray(() => Reader.ReadFileReference());
+			bDeployAfterCompile = Reader.ReadBoolean();
+			bHasProjectScriptPlugin = Reader.ReadBoolean();
+		}
+
+		public void Write(BinaryWriter Writer)
+		{
+			Writer.Write((int)TargetType);
+			Writer.Write(ProjectFile);
+			Writer.Write(AppName);
+			Writer.Write(TargetName);
+			Writer.Write(bUseSharedBuildEnvironment);
+			Writer.Write((int)Platform);
+			Writer.Write((int)Configuration);
+			Writer.Write(Architecture);
+			Writer.Write(PlatformIntermediateFolder);
+			Writer.Write(ProjectDirectory);
+			Writer.Write(ProjectIntermediateDirectory);
+			Writer.Write(EngineIntermediateDirectory);
+			Writer.Write(OutputPaths, OutputPath => Writer.Write(OutputPath));
+			Writer.Write(VersionFile);
+			Writer.Write(bPrecompile);
+			Writer.Write(bCompileMonolithic);
+			Writer.Write(FlatModuleCsData, x => x.Write(Writer));
+			Writer.Write(ReceiptFileName);
+			Writer.Write(TargetRulesFile);
+			Writer.Write(PreBuildStepScripts, x => Writer.Write(x));
+			Writer.Write(PostBuildStepScripts, x => Writer.Write(x));
+			Writer.Write(bDeployAfterCompile);
+			Writer.Write(bHasProjectScriptPlugin);
 		}
 
 		/// <summary>
diff --git a/Engine/Source/Programs/UnrealBuildTool/System/ActionGraph.cs b/Engine/Source/Programs/UnrealBuildTool/System/ActionGraph.cs
--- a/Engine/Source/Programs/UnrealBuildTool/System/ActionGraph.cs
+++ b/Engine/Source/Programs/UnrealBuildTool/System/ActionGraph.cs
@@ -33,8 +33,7 @@ namespace UnrealBuildTool
 	/// <summary>
 	/// A build action.
 	/// </summary>
-	[Serializable]
-	class Action : ISerializable
+	class Action
 	{
 		///
 		/// Preparation and Assembly (serialized)
@@ -158,47 +157,47 @@ namespace UnrealBuildTool
 			ActionType = InActionType;
 		}
 
-		public Action(SerializationInfo SerializationInfo, StreamingContext StreamingContext)
+		public Action(BinaryReader Reader, List<FileItem> UniqueFileItems)
 		{
-			ActionType = (ActionType)SerializationInfo.GetByte("at");
-			WorkingDirectory = SerializationInfo.GetString("wd");
-			bPrintDebugInfo = SerializationInfo.GetBoolean("di");
-			CommandPath = SerializationInfo.GetString("cp");
-			CommandArguments = SerializationInfo.GetString("ca");
-			CommandDescription = SerializationInfo.GetString("cd");
-			StatusDescription = SerializationInfo.GetString("sd");
-			bCanExecuteRemotely = SerializationInfo.GetBoolean("ce");
-			bCanExecuteRemotelyWithSNDBS = SerializationInfo.GetBoolean("cs");
-			bIsGCCCompiler = SerializationInfo.GetBoolean("ig");
-			bIsUsingPCH = SerializationInfo.GetBoolean("iu");
-			bShouldOutputStatusDescription = SerializationInfo.GetBoolean("os");
-			bProducesImportLibrary = SerializationInfo.GetBoolean("il");
-			PrerequisiteItems = (List<FileItem>)SerializationInfo.GetValue("pr", typeof(List<FileItem>));
-			ProducedItems = (List<FileItem>)SerializationInfo.GetValue("pd", typeof(List<FileItem>));
-			DeleteItems = (List<FileItem>)SerializationInfo.GetValue("df", typeof(List<FileItem>));
+			ActionType = (ActionType)Reader.ReadByte();
+			WorkingDirectory = Reader.ReadString();
+			bPrintDebugInfo = Reader.ReadBoolean();
+			CommandPath = Reader.ReadString();
+			CommandArguments = Reader.ReadString();
+			CommandDescription = Reader.ReadNullable(() => Reader.ReadString());
+			StatusDescription = Reader.ReadString();
+			bCanExecuteRemotely = Reader.ReadBoolean();
+			bCanExecuteRemotelyWithSNDBS = Reader.ReadBoolean();
+			bIsGCCCompiler = Reader.ReadBoolean();
+			bIsUsingPCH = Reader.ReadBoolean();
+			bShouldOutputStatusDescription = Reader.ReadBoolean();
+			bProducesImportLibrary = Reader.ReadBoolean();
+			PrerequisiteItems = Reader.ReadFileItemList(UniqueFileItems);
+			ProducedItems = Reader.ReadFileItemList(UniqueFileItems);
+			DeleteItems = Reader.ReadFileItemList(UniqueFileItems);
 		}
 
 		/// <summary>
 		/// ISerializable: Called when serialized to report additional properties that should be saved
 		/// </summary>
-		public void GetObjectData(SerializationInfo SerializationInfo, StreamingContext StreamingContext)
+		public void Write(BinaryWriter Writer, Dictionary<FileItem, int> UniqueFileItemToIndex)
 		{
-			SerializationInfo.AddValue("at", (byte)ActionType);
-			SerializationInfo.AddValue("wd", WorkingDirectory);
-			SerializationInfo.AddValue("di", bPrintDebugInfo);
-			SerializationInfo.AddValue("cp", CommandPath);
-			SerializationInfo.AddValue("ca", CommandArguments);
-			SerializationInfo.AddValue("cd", CommandDescription);
-			SerializationInfo.AddValue("sd", StatusDescription);
-			SerializationInfo.AddValue("ce", bCanExecuteRemotely);
-			SerializationInfo.AddValue("cs", bCanExecuteRemotelyWithSNDBS);
-			SerializationInfo.AddValue("ig", bIsGCCCompiler);
-			SerializationInfo.AddValue("iu", bIsUsingPCH);
-			SerializationInfo.AddValue("os", bShouldOutputStatusDescription);
-			SerializationInfo.AddValue("il", bProducesImportLibrary);
-			SerializationInfo.AddValue("pr", PrerequisiteItems);
-			SerializationInfo.AddValue("pd", ProducedItems);
-			SerializationInfo.AddValue("df", DeleteItems);
+			Writer.Write((byte)ActionType);
+			Writer.Write(WorkingDirectory);
+			Writer.Write(bPrintDebugInfo);
+			Writer.Write(CommandPath);
+			Writer.Write(CommandArguments);
+			Writer.WriteNullable(CommandDescription, () => Writer.Write(CommandDescription));
+			Writer.Write(StatusDescription);
+			Writer.Write(bCanExecuteRemotely);
+			Writer.Write(bCanExecuteRemotelyWithSNDBS);
+			Writer.Write(bIsGCCCompiler);
+			Writer.Write(bIsUsingPCH);
+			Writer.Write(bShouldOutputStatusDescription);
+			Writer.Write(bProducesImportLibrary);
+			Writer.Write(PrerequisiteItems, UniqueFileItemToIndex);
+			Writer.Write(ProducedItems, UniqueFileItemToIndex);
+			Writer.Write(DeleteItems, UniqueFileItemToIndex);
 		}
 
 		/// <summary>
diff --git a/Engine/Source/Programs/UnrealBuildTool/System/BuildPredicateStore.cs b/Engine/Source/Programs/UnrealBuildTool/System/BuildPredicateStore.cs
--- a/Engine/Source/Programs/UnrealBuildTool/System/BuildPredicateStore.cs
+++ b/Engine/Source/Programs/UnrealBuildTool/System/BuildPredicateStore.cs
@@ -2,6 +2,7 @@
 
 using System;
 using System.Collections.Generic;
+using System.IO;
 using System.Linq;
 using System.Text;
 using System.Threading.Tasks;
@@ -29,5 +30,23 @@ namespace UnrealBuildTool
 		/// Set of files which are currently not part of the working set, but could be.
 		/// </summary>
 		public HashSet<FileItem> CandidatesForWorkingSet = new HashSet<FileItem>();
+
+		public BuildPredicateStore()
+		{
+		}
+
+		public BuildPredicateStore(BinaryReader Reader, List<FileItem> UniqueFileItems)
+		{
+			SourceDirectories = new HashSet<DirectoryReference>(Reader.ReadArray(() => Reader.ReadDirectoryReference()));
+			WorkingSet = new HashSet<FileItem>(Reader.ReadFileItemList(UniqueFileItems));
+			CandidatesForWorkingSet = new HashSet<FileItem>(Reader.ReadFileItemList(UniqueFileItems));
+		}
+
+		public void Write(BinaryWriter Writer, Dictionary<FileItem, int> UniqueFileItemToIndex)
+		{
+			Writer.Write(SourceDirectories.ToArray(), x => Writer.Write(x));
+			Writer.Write(WorkingSet.ToList(), UniqueFileItemToIndex);
+			Writer.Write(CandidatesForWorkingSet.ToList(), UniqueFileItemToIndex);
+		}
 	}
 }
diff --git a/Engine/Source/Programs/UnrealBuildTool/System/ExternalExecution.cs b/Engine/Source/Programs/UnrealBuildTool/System/ExternalExecution.cs
--- a/Engine/Source/Programs/UnrealBuildTool/System/ExternalExecution.cs
+++ b/Engine/Source/Programs/UnrealBuildTool/System/ExternalExecution.cs
@@ -143,8 +143,7 @@ namespace UnrealBuildTool
 	/// <summary>
 	/// Information about a module that needs to be passed to UnrealHeaderTool for code generation
 	/// </summary>
-	[Serializable]
-	class UHTModuleInfo : ISerializable
+	class UHTModuleInfo
 	{
 		/// <summary>
 		/// Module name
@@ -209,32 +208,32 @@ namespace UnrealBuildTool
 			this.bIsReadOnly = bIsReadOnly;
 		}
 
-		public UHTModuleInfo(SerializationInfo Info, StreamingContext Context)
+		public UHTModuleInfo(BinaryReader Reader, List<FileItem> UniqueFileItems)
 		{
-			ModuleName = Info.GetString("mn");
-			ModuleRulesFile = (FileReference)Info.GetValue("mr", typeof(FileReference));
-			ModuleDirectory = (DirectoryReference)Info.GetValue("md", typeof(DirectoryReference));
-			ModuleType = Info.GetString("mt");
-			PublicUObjectClassesHeaders = (List<FileItem>)Info.GetValue("cl", typeof(List<FileItem>));
-			PublicUObjectHeaders = (List<FileItem>)Info.GetValue("pu", typeof(List<FileItem>));
-			PrivateUObjectHeaders = (List<FileItem>)Info.GetValue("pr", typeof(List<FileItem>));
-			GeneratedCPPFilenameBase = Info.GetString("ge");
-			GeneratedCodeVersion = (EGeneratedCodeVersion)Info.GetInt32("gv");
-			bIsReadOnly = Info.GetBoolean("ro");
+			ModuleName = Reader.ReadString();
+			ModuleRulesFile = Reader.ReadFileReference();
+			ModuleDirectory = Reader.ReadDirectoryReference();
+			ModuleType = Reader.ReadString();
+			PublicUObjectClassesHeaders = Reader.ReadFileItemList(UniqueFileItems);
+			PublicUObjectHeaders = Reader.ReadFileItemList(UniqueFileItems);
+			PrivateUObjectHeaders = Reader.ReadFileItemList(UniqueFileItems);
+			GeneratedCPPFilenameBase = Reader.ReadString();
+			GeneratedCodeVersion = (EGeneratedCodeVersion)Reader.ReadInt32();
+			bIsReadOnly = Reader.ReadBoolean();
 		}
 
-		public void GetObjectData(SerializationInfo Info, StreamingContext Context)
+		public void Write(BinaryWriter Writer, Dictionary<FileItem, int> UniqueFileItemToIndex)
 		{
-			Info.AddValue("mn", ModuleName);
-			Info.AddValue("mr", ModuleRulesFile);
-			Info.AddValue("md", ModuleDirectory);
-			Info.AddValue("mt", ModuleType);
-			Info.AddValue("cl", PublicUObjectClassesHeaders);
-			Info.AddValue("pu", PublicUObjectHeaders);
-			Info.AddValue("pr", PrivateUObjectHeaders);
-			Info.AddValue("ge", GeneratedCPPFilenameBase);
-			Info.AddValue("gv", (int)GeneratedCodeVersion);
-			Info.AddValue("ro", bIsReadOnly);
+			Writer.Write(ModuleName);
+			Writer.Write(ModuleRulesFile);
+			Writer.Write(ModuleDirectory);
+			Writer.Write(ModuleType);
+			Writer.Write(PublicUObjectClassesHeaders, UniqueFileItemToIndex);
+			Writer.Write(PublicUObjectHeaders, UniqueFileItemToIndex);
+			Writer.Write(PrivateUObjectHeaders, UniqueFileItemToIndex);
+			Writer.Write(GeneratedCPPFilenameBase);
+			Writer.Write((int)GeneratedCodeVersion);
+			Writer.Write(bIsReadOnly);
 		}
 
 		public override string ToString()
diff --git a/Engine/Source/Programs/UnrealBuildTool/System/FileItem.cs b/Engine/Source/Programs/UnrealBuildTool/System/FileItem.cs
--- a/Engine/Source/Programs/UnrealBuildTool/System/FileItem.cs
+++ b/Engine/Source/Programs/UnrealBuildTool/System/FileItem.cs
@@ -425,4 +425,62 @@ namespace UnrealBuildTool
 		}
 	}
 
+	static class FileItemExtensionMethods
+	{
+		public static FileItem ReadFileItem(this BinaryReader Reader, List<FileItem> UniqueFileItems)
+		{
+			int Idx = Reader.ReadInt32();
+			if (Idx == -1)
+			{
+				return null;
+			}
+			else
+			{
+				return UniqueFileItems[Idx];
+			}
+		}
+
+		public static List<FileItem> ReadFileItemList(this BinaryReader Reader, List<FileItem> UniqueFileItems)
+		{
+			int Length = Reader.ReadInt32();
+			if(Length == -1)
+			{
+				return null;
+			}
+
+			List<FileItem> Items = new List<FileItem>(Length);
+			for(int Idx = 0; Idx < Length; Idx++)
+			{
+				Items.Add(ReadFileItem(Reader, UniqueFileItems));
+			}
+			return Items;
+		}
+
+		public static void Write(this BinaryWriter Writer, FileItem FileItem, Dictionary<FileItem, int> UniqueFileItemToIndex)
+		{
+			int Index = -1;
+			if (FileItem != null && !UniqueFileItemToIndex.TryGetValue(FileItem, out Index))
+			{
+				Index = UniqueFileItemToIndex.Count;
+				UniqueFileItemToIndex.Add(FileItem, Index);
+			}
+			Writer.Write(Index);
+		}
+
+		public static void Write(this BinaryWriter Writer, List<FileItem> FileItems, Dictionary<FileItem, int> UniqueFileItemToIndex)
+		{
+			if(FileItems == null)
+			{
+				Writer.Write(-1);
+			}
+			else
+			{
+				Writer.Write(FileItems.Count);
+				for(int Idx = 0; Idx < FileItems.Count; Idx++)
+				{
+					Write(Writer, FileItems[Idx], UniqueFileItemToIndex);
+				}
+			}
+		}
+	}
 }
diff --git a/Engine/Source/Programs/UnrealBuildTool/System/UBTMakefile.cs b/Engine/Source/Programs/UnrealBuildTool/System/UBTMakefile.cs
--- a/Engine/Source/Programs/UnrealBuildTool/System/UBTMakefile.cs
+++ b/Engine/Source/Programs/UnrealBuildTool/System/UBTMakefile.cs
@@ -15,19 +15,13 @@ namespace UnrealBuildTool
 	/// <summary>
 	/// A special Makefile that UBT is able to create in "-gather" mode, then load in "-assemble" mode to accelerate iterative compiling and linking
 	/// </summary>
-	[Serializable]
-	class UBTMakefile : ISerializable
+	class UBTMakefile
 	{
 		/// <summary>
 		/// The version number to write
 		/// </summary>
 		public const int CurrentVersion = 8;
 
-		/// <summary>
-		/// The version number that was read
-		/// </summary>
-		public int Version;
-
 		/// <summary>
 		/// Every action in the action graph
 		/// </summary>
@@ -76,41 +70,66 @@ namespace UnrealBuildTool
 
 		public UBTMakefile()
 		{
-			Version = CurrentVersion;
 		}
 
-		public UBTMakefile(SerializationInfo Info, StreamingContext Context)
+		public UBTMakefile(BinaryReader Reader)
 		{
-			Version = Info.GetInt32("ve");
-			if (Version != CurrentVersion)
+			int Version = Reader.ReadInt32();
+			if(Version != CurrentVersion)
 			{
 				throw new Exception(string.Format("Makefile version does not match - found {0}, expected: {1}", Version, CurrentVersion));
 			}
 
-			AllActions = (List<Action>)Info.GetValue("ac", typeof(List<Action>));
-			EnvironmentVariables = ((string[])Info.GetValue("e1", typeof(string[]))).Zip((string[])Info.GetValue("e2", typeof(string[])), (i1, i2) => new Tuple<string, string>(i1, i2)).ToList();
-			TargetNameToUObjectModules = (Dictionary<string, List<UHTModuleInfo>>)Info.GetValue("nu", typeof(Dictionary<string, List<UHTModuleInfo>>));
-			OutputItemsForAllTargets = (List<FileItem>)Info.GetValue("oi", typeof(List<FileItem>));
-			ModuleNameToOutputItems = (Dictionary<string, FileItem[]>)Info.GetValue("mn", typeof(Dictionary<string, FileItem[]>));
-			HotReloadModuleNamesForAllTargets = (HashSet<string>)Info.GetValue("gm", typeof(HashSet<string>));
-			Targets = (List<UEBuildTarget>)Info.GetValue("ta", typeof(List<UEBuildTarget>));
-			TargetBuildPredicates = (List<BuildPredicateStore>)Info.GetValue("tp", typeof(List<BuildPredicateStore>));
-			bUseAdaptiveUnityBuild = Info.GetBoolean("ua");
+			long Offset = Reader.ReadInt64();
+			long OriginalOffset = Reader.BaseStream.Position;
+
+			Reader.BaseStream.Seek(Offset, SeekOrigin.Begin);
+			List<FileReference> UniqueFileReferences = Reader.ReadList(() => Reader.ReadFileReference());
+			List<FileItem> UniqueFileItems = UniqueFileReferences.Select(x => FileItem.GetItemByFileReference(x)).ToList();
+			Reader.BaseStream.Seek(OriginalOffset, SeekOrigin.Begin);
+
+			AllActions = Reader.ReadList(() => new Action(Reader, UniqueFileItems));
+			EnvironmentVariables = Reader.ReadList(() => Tuple.Create(Reader.ReadString(), Reader.ReadString()));
+			TargetNameToUObjectModules = Reader.ReadDictionary(() => Reader.ReadString(), () => Reader.ReadList(() => new UHTModuleInfo(Reader, UniqueFileItems)));
+			OutputItemsForAllTargets = Reader.ReadFileItemList(UniqueFileItems);
+			ModuleNameToOutputItems = Reader.ReadDictionary(() => Reader.ReadString(), () => Reader.ReadArray(() => Reader.ReadFileItem(UniqueFileItems)));
+			HotReloadModuleNamesForAllTargets = new HashSet<string>(Reader.ReadStringArray());
+			Targets = Reader.ReadList(() => new UEBuildTarget(Reader));
+			TargetBuildPredicates = Reader.ReadList(() => new BuildPredicateStore(Reader, UniqueFileItems));
+			bUseAdaptiveUnityBuild = Reader.ReadBoolean();
 		}
 
-		public void GetObjectData(SerializationInfo Info, StreamingContext Context)
+
+		public void Write(BinaryWriter Writer)
 		{
-			Info.AddValue("ve", Version);
-			Info.AddValue("ac", AllActions);
-			Info.AddValue("e1", EnvironmentVariables.Select(x => x.Item1).ToArray());
-			Info.AddValue("e2", EnvironmentVariables.Select(x => x.Item2).ToArray());
-			Info.AddValue("nu", TargetNameToUObjectModules);
-			Info.AddValue("oi", OutputItemsForAllTargets);
-			Info.AddValue("mn", ModuleNameToOutputItems);
-			Info.AddValue("gm", HotReloadModuleNamesForAllTargets);
-			Info.AddValue("ta", Targets);
-			Info.AddValue("tp", TargetBuildPredicates);
-			Info.AddValue("ua", bUseAdaptiveUnityBuild);
+			Writer.Write(CurrentVersion);
+
+			long Offset = Writer.BaseStream.Position;
+			Writer.Write((Int64)0);
+			Dictionary<FileItem, int> UniqueFileItemToIndex = new Dictionary<FileItem, int>();
+
+			Writer.Write(AllActions, x => x.Write(Writer, UniqueFileItemToIndex));
+			Writer.Write(EnvironmentVariables, x => { Writer.Write(x.Item1); Writer.Write(x.Item2); });
+			Writer.Write(TargetNameToUObjectModules, x => Writer.Write(x), v => Writer.Write(v, e => e.Write(Writer, UniqueFileItemToIndex)));
+			Writer.Write(OutputItemsForAllTargets, UniqueFileItemToIndex);
+			Writer.Write(ModuleNameToOutputItems, k => Writer.Write(k), v => Writer.Write(v, e => Writer.Write(e, UniqueFileItemToIndex)));
+			Writer.Write(HotReloadModuleNamesForAllTargets.ToArray(), x => Writer.Write(x));
+			Writer.Write(Targets, x => x.Write(Writer));
+			Writer.Write(TargetBuildPredicates, x => x.Write(Writer, UniqueFileItemToIndex));
+			Writer.Write(bUseAdaptiveUnityBuild);
+			Writer.Flush();
+
+			long FileItemTableOffset = Writer.BaseStream.Position;
+			FileReference[] UniqueFileReferences = new FileReference[UniqueFileItemToIndex.Count];
+			foreach(KeyValuePair<FileItem, int> Pair in UniqueFileItemToIndex)
+			{
+				UniqueFileReferences[Pair.Value] = Pair.Key.Location;
+			}
+			Writer.Write(UniqueFileReferences, x => Writer.Write(x));
+
+			Writer.BaseStream.Seek(Offset, SeekOrigin.Begin);
+			Writer.Write(FileItemTableOffset);
+			Writer.Flush();
 		}
 
 
@@ -152,8 +171,10 @@ namespace UnrealBuildTool
 				Directory.CreateDirectory(Path.GetDirectoryName(UBTMakefileItem.AbsolutePath));
 				using (FileStream Stream = new FileStream(UBTMakefileItem.AbsolutePath, FileMode.Create, FileAccess.Write))
 				{
-					BinaryFormatter Formatter = new BinaryFormatter();
-					Formatter.Serialize(Stream, UBTMakefile);
+					using(BinaryWriter Writer = new BinaryWriter(Stream, Encoding.UTF8, true))
+					{
+						UBTMakefile.Write(Writer);
+					}
 				}
 			}
 			catch (Exception Ex)
@@ -271,8 +292,16 @@ namespace UnrealBuildTool
 				{
 					using (FileStream Stream = new FileStream(UBTMakefileInfo.FullName, FileMode.Open, FileAccess.Read))
 					{
-						BinaryFormatter Formatter = new BinaryFormatter();
-						LoadedUBTMakefile = Formatter.Deserialize(Stream) as UBTMakefile;
+						byte[] Data = new byte[Stream.Length];
+						if(Stream.Read(Data, 0, Data.Length) != Data.Length)
+						{
+							throw new Exception("Error");
+						}
+
+						Timeline.AddEvent("AFTER READ");
+
+						BinaryReader Reader = new BinaryReader(new MemoryStream(Data));
+						LoadedUBTMakefile = new UBTMakefile(Reader);
 					}
 				}
 			}
