UBT: Move makefile checks for outdated ini files into the regular load method.

#rnx

diff --git a/Engine/Source/Programs/UnrealBuildTool/System/UBTMakefile.cs b/Engine/Source/Programs/UnrealBuildTool/System/UBTMakefile.cs
--- a/Engine/Source/Programs/UnrealBuildTool/System/UBTMakefile.cs
+++ b/Engine/Source/Programs/UnrealBuildTool/System/UBTMakefile.cs
@@ -192,7 +192,7 @@ namespace UnrealBuildTool
 		/// <param name="ReasonNotLoaded">If the function returns null, this string will contain the reason why</param>
 		/// <param name="WorkingSet">Interface to query which source files are in the working set</param>
 		/// <returns>The loaded makefile, or null if it failed for some reason.  On failure, the 'ReasonNotLoaded' variable will contain information about why</returns>
-		public static UBTMakefile LoadUBTMakefile(FileReference MakefilePath, FileReference ProjectFile, ISourceFileWorkingSet WorkingSet, out string ReasonNotLoaded)
+		public static UBTMakefile Load(FileReference MakefilePath, FileReference ProjectFile, ISourceFileWorkingSet WorkingSet, out string ReasonNotLoaded)
 		{
 			// Check the directory timestamp on the project files directory.  If the user has generated project files more
 			// recently than the UBTMakefile, then we need to consider the file to be out of date
@@ -370,9 +370,26 @@ namespace UnrealBuildTool
 						}
 					}
 				}
+
+
+				// Check if ini files are newer. Ini files contain build settings too.
+				DirectoryReference ProjectDirectory = DirectoryReference.FromFile(ProjectFile);
+				foreach (ConfigHierarchyType IniType in (ConfigHierarchyType[])Enum.GetValues(typeof(ConfigHierarchyType)))
+				{
+					foreach (FileReference IniFilename in ConfigHierarchy.EnumerateConfigFileLocations(IniType, ProjectDirectory, Target.Platform))
+					{
+						FileInfo IniFileInfo = new FileInfo(IniFilename.FullName);
+						if (UBTMakefileInfo.LastWriteTime.CompareTo(IniFileInfo.LastWriteTime) < 0)
+						{
+							// Ini files are newer than UBTMakefile
+							ReasonNotLoaded = "ini files are newer than UBTMakefile";
+							return null;
+						}
+					}
+				}
 			}
 
-			foreach(BuildPredicateStore Predicates in LoadedUBTMakefile.TargetBuildPredicates)
+			foreach (BuildPredicateStore Predicates in LoadedUBTMakefile.TargetBuildPredicates)
 			{
 				foreach(DirectoryReference SourceDirectory in Predicates.SourceDirectories)
 				{
diff --git a/Engine/Source/Programs/UnrealBuildTool/UnrealBuildTool.cs b/Engine/Source/Programs/UnrealBuildTool/UnrealBuildTool.cs
--- a/Engine/Source/Programs/UnrealBuildTool/UnrealBuildTool.cs
+++ b/Engine/Source/Programs/UnrealBuildTool/UnrealBuildTool.cs
@@ -639,7 +639,7 @@ namespace UnrealBuildTool
 				// Create the working set provider
 				using (ISourceFileWorkingSet WorkingSet = SourceFileWorkingSet.Create(UnrealBuildTool.RootDirectory, ProjectDirs))
 				{
-					UBTMakefile UBTMakefile = null;
+					UBTMakefile Makefile = null;
 					{
 						FileReference UBTMakefilePath = UBTMakefile.GetUBTMakefilePath(TargetDescs);
 
@@ -660,43 +660,10 @@ namespace UnrealBuildTool
 							string ReasonNotLoaded;
 							using(Timeline.ScopeEvent("UBTMakefile.LoadUBTMakefile"))
 							{
-								UBTMakefile = UBTMakefile.LoadUBTMakefile(UBTMakefilePath, TargetDescs[0].ProjectFile, WorkingSet, out ReasonNotLoaded);
+								Makefile = UBTMakefile.Load(UBTMakefilePath, TargetDescs[0].ProjectFile, WorkingSet, out ReasonNotLoaded);
 							}
 
-							if (UBTMakefile != null)
-							{
-								// Check if ini files are newer. Ini files contain build settings too.
-								FileInfo UBTMakefileInfo = new FileInfo(UBTMakefilePath.FullName);
-								foreach (TargetDescriptor Desc in TargetDescs)
-								{
-									DirectoryReference ProjectDirectory = DirectoryReference.FromFile(Desc.ProjectFile);
-									foreach (ConfigHierarchyType IniType in (ConfigHierarchyType[])Enum.GetValues(typeof(ConfigHierarchyType)))
-									{
-										foreach (FileReference IniFilename in ConfigHierarchy.EnumerateConfigFileLocations(IniType, ProjectDirectory, Desc.Platform))
-										{
-											FileInfo IniFileInfo = new FileInfo(IniFilename.FullName);
-											if (UBTMakefileInfo.LastWriteTime.CompareTo(IniFileInfo.LastWriteTime) < 0)
-											{
-												// Ini files are newer than UBTMakefile
-												UBTMakefile = null;
-												ReasonNotLoaded = "ini files are newer than UBTMakefile";
-												break;
-											}
-										}
-
-										if (UBTMakefile == null)
-										{
-											break;
-										}
-									}
-									if (UBTMakefile == null)
-									{
-										break;
-									}
-								}
-							}
-
-							if (UBTMakefile == null)
+							if (Makefile == null)
 							{
 								// If the Makefile couldn't be loaded, then we're not going to be able to continue in "assembler only" mode.  We'll do both
 								// a 'gather' and 'assemble' in the same run.  This will take a while longer, but subsequent runs will be fast!
@@ -713,10 +680,10 @@ namespace UnrealBuildTool
 					}
 
 
-					if (UBTMakefile != null && !bIsGatheringBuild && bIsAssemblingBuild)
+					if (Makefile != null && !bIsGatheringBuild && bIsAssemblingBuild)
 					{
 						// If we've loaded a makefile, then we can fill target information from this file!
-						Targets = UBTMakefile.Targets;
+						Targets = Makefile.Targets;
 					}
 					else
 					{
@@ -851,19 +818,19 @@ namespace UnrealBuildTool
 						{
 							ActionGraph.FinalizeActionGraph();
 
-							UBTMakefile = new UBTMakefile();
-							UBTMakefile.AllActions = ActionGraph.AllActions;
-							UBTMakefile.OutputItemsForAllTargets = OutputItemsForAllTargets;
+							Makefile = new UBTMakefile();
+							Makefile.AllActions = ActionGraph.AllActions;
+							Makefile.OutputItemsForAllTargets = OutputItemsForAllTargets;
 							foreach (System.Collections.DictionaryEntry EnvironmentVariable in Environment.GetEnvironmentVariables())
 							{
-								UBTMakefile.EnvironmentVariables.Add(Tuple.Create((string)EnvironmentVariable.Key, (string)EnvironmentVariable.Value));
+								Makefile.EnvironmentVariables.Add(Tuple.Create((string)EnvironmentVariable.Key, (string)EnvironmentVariable.Value));
 							}
-							UBTMakefile.TargetNameToUObjectModules = TargetNameToUObjectModules;
-							UBTMakefile.ModuleNameToOutputItems = ModuleNameToOutputItems;
-							UBTMakefile.HotReloadModuleNamesForAllTargets = HotReloadModuleNamesForAllTargets;
-							UBTMakefile.Targets = Targets;
-							UBTMakefile.bUseAdaptiveUnityBuild = Targets.Any(x => x.Rules.bUseAdaptiveUnityBuild);
-							UBTMakefile.TargetBuildPredicates = TargetBuildPredicates;
+							Makefile.TargetNameToUObjectModules = TargetNameToUObjectModules;
+							Makefile.ModuleNameToOutputItems = ModuleNameToOutputItems;
+							Makefile.HotReloadModuleNamesForAllTargets = HotReloadModuleNamesForAllTargets;
+							Makefile.Targets = Targets;
+							Makefile.bUseAdaptiveUnityBuild = Targets.Any(x => x.Rules.bUseAdaptiveUnityBuild);
+							Makefile.TargetBuildPredicates = TargetBuildPredicates;
 
 							if (BuildConfiguration.bUseUBTMakefiles)
 							{
@@ -873,7 +840,7 @@ namespace UnrealBuildTool
 								// @todo ubtmake: Optimization: We could make 'gather + assemble' mode slightly faster by saving this while busy compiling (on our worker thread)
 								using(Timeline.ScopeEvent("UBTMakefile.SaveUBTMakefile()"))
 								{
-									UBTMakefile.SaveUBTMakefile(TargetDescs, UBTMakefile);
+									UBTMakefile.SaveUBTMakefile(TargetDescs, Makefile);
 								}
 							}
 						}
@@ -942,15 +909,15 @@ namespace UnrealBuildTool
 							// If we didn't build the graph in this session, then we'll need to load a cached one
 							if (!bIsGatheringBuild && BuildResult.Succeeded())
 							{
-								ActionGraph.AllActions = UBTMakefile.AllActions;
+								ActionGraph.AllActions = Makefile.AllActions;
 
-								OutputItemsForAllTargets = UBTMakefile.OutputItemsForAllTargets;
+								OutputItemsForAllTargets = Makefile.OutputItemsForAllTargets;
 
-								ModuleNameToOutputItems = UBTMakefile.ModuleNameToOutputItems;
+								ModuleNameToOutputItems = Makefile.ModuleNameToOutputItems;
 
-								HotReloadModuleNamesForAllTargets = UBTMakefile.HotReloadModuleNamesForAllTargets;
+								HotReloadModuleNamesForAllTargets = Makefile.HotReloadModuleNamesForAllTargets;
 
-								foreach (Tuple<string, string> EnvironmentVariable in UBTMakefile.EnvironmentVariables)
+								foreach (Tuple<string, string> EnvironmentVariable in Makefile.EnvironmentVariables)
 								{
 									// @todo ubtmake: There may be some variables we do NOT want to clobber.
 									Environment.SetEnvironmentVariable(EnvironmentVariable.Item1, EnvironmentVariable.Item2);
@@ -975,7 +942,7 @@ namespace UnrealBuildTool
 									foreach (UEBuildTarget Target in Targets)
 									{
 										List<UHTModuleInfo> TargetUObjectModules;
-										if (UBTMakefile.TargetNameToUObjectModules.TryGetValue(Target.GetTargetName(), out TargetUObjectModules))
+										if (Makefile.TargetNameToUObjectModules.TryGetValue(Target.GetTargetName(), out TargetUObjectModules))
 										{
 											if (TargetUObjectModules.Count > 0)
 											{
