UBT: Remove the PlatformRequiresCookedData() overload. This is implied by the target type being built and the platforms it supports (and is already set depending on the TargetType), and is not a property of the platform itself.

#rnx

diff --git a/Engine/Source/Programs/UnrealBuildTool/Configuration/UEBuildPlatform.cs b/Engine/Source/Programs/UnrealBuildTool/Configuration/UEBuildPlatform.cs
--- a/Engine/Source/Programs/UnrealBuildTool/Configuration/UEBuildPlatform.cs
+++ b/Engine/Source/Programs/UnrealBuildTool/Configuration/UEBuildPlatform.cs
@@ -585,17 +585,6 @@ namespace UnrealBuildTool
 			throw new BuildException("GetDebugInfoExtensions for {0} not handled in {1}", InBinaryType.ToString(), this.ToString());
 		}
 
-		/// <summary>
-		/// Whether this build should support ONLY cooked data or not
-		/// </summary>
-		/// <param name="InPlatform"> The UnrealTargetPlatform being built</param>
-		/// <param name="InConfiguration">The UnrealTargetConfiguration being built</param>
-		/// <returns>bool   true if the editor should be built, false if not</returns>
-		public virtual bool BuildRequiresCookedData(UnrealTargetPlatform InPlatform, UnrealTargetConfiguration InConfiguration)
-		{
-			return false;
-		}
-
 		/// <summary>
 		/// Whether this platform should build a monolithic binary
 		/// </summary>
diff --git a/Engine/Source/Programs/UnrealBuildTool/Configuration/UEBuildTarget.cs b/Engine/Source/Programs/UnrealBuildTool/Configuration/UEBuildTarget.cs
--- a/Engine/Source/Programs/UnrealBuildTool/Configuration/UEBuildTarget.cs
+++ b/Engine/Source/Programs/UnrealBuildTool/Configuration/UEBuildTarget.cs
@@ -392,13 +392,6 @@ namespace UnrealBuildTool
 				RulesObject.bDisableLinking = true;
 			}
 
-			// Disable the DDC and a few other things related to preparing assets
-			UEBuildPlatform BuildPlatform = UEBuildPlatform.GetBuildPlatform(RulesObject.Platform);
-			if (BuildPlatform.BuildRequiresCookedData(Desc.Platform, Desc.Configuration) == true)
-			{
-				RulesObject.bBuildRequiresCookedData = true;
-			}
-
 			if (!RulesObject.bAllowNonUFSIniWhenCooked)
 			{
 				RulesObject.GlobalDefinitions.Add("DISABLE_NONUFS_INI_WHEN_COOKED=1");
diff --git a/Engine/Source/Programs/UnrealBuildTool/Platform/Android/UEBuildAndroid.cs b/Engine/Source/Programs/UnrealBuildTool/Platform/Android/UEBuildAndroid.cs
--- a/Engine/Source/Programs/UnrealBuildTool/Platform/Android/UEBuildAndroid.cs
+++ b/Engine/Source/Programs/UnrealBuildTool/Platform/Android/UEBuildAndroid.cs
@@ -214,11 +214,6 @@ namespace UnrealBuildTool
 			return true;
 		}
 
-		public override bool BuildRequiresCookedData(UnrealTargetPlatform InPlatform, UnrealTargetConfiguration InConfiguration)
-		{
-			return true;
-		}
-
 		/// <summary>
 		/// Modify the rules for a newly created module, where the target is a different host platform.
 		/// This is not required - but allows for hiding details of a particular platform.
diff --git a/Engine/Source/Programs/UnrealBuildTool/Platform/HTML5/UEBuildHTML5.cs b/Engine/Source/Programs/UnrealBuildTool/Platform/HTML5/UEBuildHTML5.cs
--- a/Engine/Source/Programs/UnrealBuildTool/Platform/HTML5/UEBuildHTML5.cs
+++ b/Engine/Source/Programs/UnrealBuildTool/Platform/HTML5/UEBuildHTML5.cs
@@ -115,11 +115,6 @@ namespace UnrealBuildTool
 			return true;
 		}
 
-		public override bool BuildRequiresCookedData(UnrealTargetPlatform InPlatform, UnrealTargetConfiguration InConfiguration)
-		{
-			return true;
-		}
-
 		/// <summary>
 		/// Modify the rules for a newly created module, where the target is a different host platform.
 		/// This is not required - but allows for hiding details of a particular platform.
diff --git a/Engine/Source/Programs/UnrealBuildTool/Platform/IOS/UEBuildIOS.cs b/Engine/Source/Programs/UnrealBuildTool/Platform/IOS/UEBuildIOS.cs
--- a/Engine/Source/Programs/UnrealBuildTool/Platform/IOS/UEBuildIOS.cs
+++ b/Engine/Source/Programs/UnrealBuildTool/Platform/IOS/UEBuildIOS.cs
@@ -828,11 +828,6 @@ namespace UnrealBuildTool
 			return HasCustomIcons(ProjectDirectoryName);
 		}
 
-		public override bool BuildRequiresCookedData(UnrealTargetPlatform InPlatform, UnrealTargetConfiguration InConfiguration)
-		{
-			return true; // for iOS can only run cooked. this is mostly for testing console code paths.
-		}
-
 		public override bool ShouldCompileMonolithicBinary(UnrealTargetPlatform InPlatform)
 		{
 			// This platform currently always compiles monolithic
diff --git a/Engine/Source/Programs/UnrealBuildTool/Platform/Windows/UEBuildWindows.cs b/Engine/Source/Programs/UnrealBuildTool/Platform/Windows/UEBuildWindows.cs
--- a/Engine/Source/Programs/UnrealBuildTool/Platform/Windows/UEBuildWindows.cs
+++ b/Engine/Source/Programs/UnrealBuildTool/Platform/Windows/UEBuildWindows.cs
@@ -1421,11 +1421,6 @@ namespace UnrealBuildTool
 			return new string [] {};
 		}
 
-		public override bool BuildRequiresCookedData(UnrealTargetPlatform InPlatform, UnrealTargetConfiguration InConfiguration)
-		{
-			return false;
-		}
-
 		public override bool HasDefaultBuildConfig(UnrealTargetPlatform Platform, DirectoryReference ProjectPath)
 		{
 			if (Platform == UnrealTargetPlatform.Win32)
