UBT: Move more functionality related to constructing targets into the rules assembly.

#rnx

diff --git a/Engine/Source/Programs/UnrealBuildTool/Configuration/UEBuildTarget.cs b/Engine/Source/Programs/UnrealBuildTool/Configuration/UEBuildTarget.cs
--- a/Engine/Source/Programs/UnrealBuildTool/Configuration/UEBuildTarget.cs
+++ b/Engine/Source/Programs/UnrealBuildTool/Configuration/UEBuildTarget.cs
@@ -392,27 +392,6 @@ namespace UnrealBuildTool
 				RulesObject.bDisableLinking = true;
 			}
 
-			if (!RulesObject.bAllowNonUFSIniWhenCooked)
-			{
-				RulesObject.GlobalDefinitions.Add("DISABLE_NONUFS_INI_WHEN_COOKED=1");
-			}
-
-			if (RulesObject.bDisableUnverifiedCertificates)
-			{
-				RulesObject.GlobalDefinitions.Add("DISABLE_UNVERIFIED_CERTIFICATE_LOADING=1");
-			}
-
-			// Allow the platform to finalize the settings
-			UEBuildPlatform Platform = UEBuildPlatform.GetBuildPlatform(RulesObject.Platform);
-			Platform.ValidateTarget(RulesObject);
-
-			// Skip deploy step in UBT if UAT is going to do deploy step
-			if (Arguments.Contains("-skipdeploy") 
-				&& RulesObject.Platform == UnrealTargetPlatform.Android) // TODO: if safe on other platforms
-			{
-				RulesObject.bDeployAfterCompile = false;
-			}
-
 			// If we're compiling a plugin, and this target is monolithic, just create the object files
 			if(Desc.ForeignPlugin != null && RulesObject.LinkType == TargetLinkType.Monolithic)
 			{
diff --git a/Engine/Source/Programs/UnrealBuildTool/System/RulesAssembly.cs b/Engine/Source/Programs/UnrealBuildTool/System/RulesAssembly.cs
--- a/Engine/Source/Programs/UnrealBuildTool/System/RulesAssembly.cs
+++ b/Engine/Source/Programs/UnrealBuildTool/System/RulesAssembly.cs
@@ -510,6 +510,20 @@ namespace UnrealBuildTool
 				Rules.GlobalDefinitions.Add("DISABLE_GENERATED_INI_WHEN_COOKED=1");
 			}
 
+			if (!Rules.bAllowNonUFSIniWhenCooked)
+			{
+				Rules.GlobalDefinitions.Add("DISABLE_NONUFS_INI_WHEN_COOKED=1");
+			}
+
+			if (Rules.bDisableUnverifiedCertificates)
+			{
+				Rules.GlobalDefinitions.Add("DISABLE_UNVERIFIED_CERTIFICATE_LOADING=1");
+			}
+
+			// Allow the platform to finalize the settings
+			UEBuildPlatform Platform = UEBuildPlatform.GetBuildPlatform(Rules.Platform);
+			Platform.ValidateTarget(Rules);
+
 			return Rules;
 		}
 
