UBT: Remove unnecessary arguments array to UEBuildTarget.CreateTarget.

#rnx

diff --git a/Engine/Source/Programs/UnrealBuildTool/Configuration/UEBuildTarget.cs b/Engine/Source/Programs/UnrealBuildTool/Configuration/UEBuildTarget.cs
--- a/Engine/Source/Programs/UnrealBuildTool/Configuration/UEBuildTarget.cs
+++ b/Engine/Source/Programs/UnrealBuildTool/Configuration/UEBuildTarget.cs
@@ -333,16 +333,15 @@ namespace UnrealBuildTool
 		/// Creates a target object for the specified target name.
 		/// </summary>
 		/// <param name="Desc">Information about the target</param>
-		/// <param name="Arguments">Command line arguments</param>
 		/// <param name="bSkipRulesCompile">Whether to skip compiling any rules assemblies</param>
 		/// <param name="bCompilingSingleFile">Whether we're compiling a single file</param>
 		/// <param name="bUsePrecompiled">Whether to use a precompiled engine/enterprise build</param>
 		/// <returns>The build target object for the specified build rules source file</returns>
-		public static UEBuildTarget CreateTarget(TargetDescriptor Desc, string[] Arguments, bool bSkipRulesCompile, bool bCompilingSingleFile, bool bUsePrecompiled)
+		public static UEBuildTarget CreateTarget(TargetDescriptor Desc, bool bSkipRulesCompile, bool bCompilingSingleFile, bool bUsePrecompiled)
 		{
 			RulesAssembly RulesAssembly = RulesCompiler.CreateTargetRulesAssembly(Desc.ProjectFile, Desc.Name, bSkipRulesCompile, bUsePrecompiled, Desc.ForeignPlugin);
 
-			TargetRules RulesObject = RulesAssembly.CreateTargetRules(Desc.Name, Desc.Platform, Desc.Configuration, Desc.Architecture, Desc.ProjectFile, Arguments);
+			TargetRules RulesObject = RulesAssembly.CreateTargetRules(Desc.Name, Desc.Platform, Desc.Configuration, Desc.Architecture, Desc.ProjectFile, Desc.AdditionalArguments);
 			if ((ProjectFileGenerator.bGenerateProjectFiles == false) && !GetSupportedPlatforms(RulesObject).Contains(Desc.Platform))
 			{
 				throw new BuildException("{0} does not support the {1} platform.", Desc.Name, Desc.Platform.ToString());
diff --git a/Engine/Source/Programs/UnrealBuildTool/ProjectFiles/ProjectFileGenerator.cs b/Engine/Source/Programs/UnrealBuildTool/ProjectFiles/ProjectFileGenerator.cs
--- a/Engine/Source/Programs/UnrealBuildTool/ProjectFiles/ProjectFileGenerator.cs
+++ b/Engine/Source/Programs/UnrealBuildTool/ProjectFiles/ProjectFileGenerator.cs
@@ -1708,7 +1708,7 @@ namespace UnrealBuildTool
 							TargetDescriptor TargetDesc = new TargetDescriptor(CurTarget.UnrealProjectFilePath, CurTarget.Name, BuildHostPlatform.Current.Platform, UnrealTargetConfiguration.Development, DefaultArchitecture, null, new CommandLineArguments(NewArguments.ToArray()));
 
 							// Create the target
-							UEBuildTarget Target = UEBuildTarget.CreateTarget(TargetDesc, NewArguments.ToArray(), false, false, bUsePrecompiled);
+							UEBuildTarget Target = UEBuildTarget.CreateTarget(TargetDesc, false, false, bUsePrecompiled);
 
 							// Generate a compile environment for each module in the binary
 							CppCompileEnvironment GlobalCompileEnvironment = Target.CreateCompileEnvironmentForProjectFiles();
diff --git a/Engine/Source/Programs/UnrealBuildTool/System/RulesAssembly.cs b/Engine/Source/Programs/UnrealBuildTool/System/RulesAssembly.cs
--- a/Engine/Source/Programs/UnrealBuildTool/System/RulesAssembly.cs
+++ b/Engine/Source/Programs/UnrealBuildTool/System/RulesAssembly.cs
@@ -410,7 +410,7 @@ namespace UnrealBuildTool
 		/// <param name="TargetInfo">Target configuration information to pass to the constructor</param>
 		/// <param name="Arguments">Command line arguments for this target</param>
 		/// <returns>Instance of the corresponding TargetRules</returns>
-		protected TargetRules CreateTargetRulesInstance(string TypeName, TargetInfo TargetInfo, string[] Arguments)
+		protected TargetRules CreateTargetRulesInstance(string TypeName, TargetInfo TargetInfo, CommandLineArguments Arguments)
 		{
 			// The build module must define a type named '<TargetName>Target' that derives from our 'TargetRules' type.  
 			Type RulesType = CompiledAssembly.GetType(TypeName);
@@ -451,7 +451,7 @@ namespace UnrealBuildTool
 			{
 				foreach(object ConfigurableObject in Rules.GetConfigurableObjects())
 				{
-					CommandLine.ParseArguments(Arguments, ConfigurableObject);
+					Arguments.ApplyTo(ConfigurableObject);
 				}
 			}
 
@@ -537,7 +537,7 @@ namespace UnrealBuildTool
 		/// <param name="ProjectFile">Path to the project file for this target</param>
 		/// <param name="Arguments">Command line arguments for this target</param>
 		/// <returns>The build target rules for the specified target</returns>
-		public TargetRules CreateTargetRules(string TargetName, UnrealTargetPlatform Platform, UnrealTargetConfiguration Configuration, string Architecture, FileReference ProjectFile, string[] Arguments)
+		public TargetRules CreateTargetRules(string TargetName, UnrealTargetPlatform Platform, UnrealTargetConfiguration Configuration, string Architecture, FileReference ProjectFile, CommandLineArguments Arguments)
 		{
 			bool bFoundTargetName = TargetNameToTargetFile.ContainsKey(TargetName);
 			if (bFoundTargetName == false)
diff --git a/Engine/Source/Programs/UnrealBuildTool/ToolChain/RemoteMac.cs b/Engine/Source/Programs/UnrealBuildTool/ToolChain/RemoteMac.cs
--- a/Engine/Source/Programs/UnrealBuildTool/ToolChain/RemoteMac.cs
+++ b/Engine/Source/Programs/UnrealBuildTool/ToolChain/RemoteMac.cs
@@ -362,7 +362,7 @@ namespace UnrealBuildTool
 			RulesAssembly RulesAssembly = RulesCompiler.CreateTargetRulesAssembly(TargetDesc.ProjectFile, TargetDesc.Name, false, false, TargetDesc.ForeignPlugin);
 
 			// Create the target rules
-			TargetRules Rules = RulesAssembly.CreateTargetRules(TargetDesc.Name, TargetDesc.Platform, TargetDesc.Configuration, TargetDesc.Architecture, TargetDesc.ProjectFile, new string[0]);
+			TargetRules Rules = RulesAssembly.CreateTargetRules(TargetDesc.Name, TargetDesc.Platform, TargetDesc.Configuration, TargetDesc.Architecture, TargetDesc.ProjectFile, null);
 
 			// Check if we need to enable a nativized plugin, and compile the assembly for that if we do
 			FileReference NativizedPluginFile = Rules.GetNativizedPlugin();
diff --git a/Engine/Source/Programs/UnrealBuildTool/UnrealBuildTool.cs b/Engine/Source/Programs/UnrealBuildTool/UnrealBuildTool.cs
--- a/Engine/Source/Programs/UnrealBuildTool/UnrealBuildTool.cs
+++ b/Engine/Source/Programs/UnrealBuildTool/UnrealBuildTool.cs
@@ -692,7 +692,7 @@ namespace UnrealBuildTool
 						{
 							using(Timeline.ScopeEvent("UEBuildTarget.CreateTarget()"))
 							{
-								UEBuildTarget Target = UEBuildTarget.CreateTarget(TargetDesc, Arguments, bSkipRulesCompile, BuildConfiguration.SingleFileToCompile != null, BuildConfiguration.bUsePrecompiled);
+								UEBuildTarget Target = UEBuildTarget.CreateTarget(TargetDesc, bSkipRulesCompile, BuildConfiguration.SingleFileToCompile != null, BuildConfiguration.bUsePrecompiled);
 								if ((Target == null) && (BuildConfiguration.bCleanProject))
 								{
 									continue;
