UBT: Add flags to the ToolMode attribute which indicate which systems should be initialized before running it.

#rnx

diff --git a/Engine/Source/Programs/UnrealBuildTool/Modes/BuildMode.cs b/Engine/Source/Programs/UnrealBuildTool/Modes/BuildMode.cs
--- a/Engine/Source/Programs/UnrealBuildTool/Modes/BuildMode.cs
+++ b/Engine/Source/Programs/UnrealBuildTool/Modes/BuildMode.cs
@@ -13,7 +13,7 @@ namespace UnrealBuildTool
 	/// <summary>
 	/// Builds a target
 	/// </summary>
-	[ToolMode("Build")]
+	[ToolMode("Build", ToolModeOptions.XmlConfig | ToolModeOptions.BuildPlatforms | ToolModeOptions.SingleInstance)]
 	class BuildMode : ToolMode
 	{
 		/// <summary>
@@ -45,80 +45,65 @@ namespace UnrealBuildTool
 				throw new BuildException("Environment could not be read");
 			}
 
-			// Get a mutex for building in this branch
-			using(SingleInstanceMutex.Acquire(SingleInstanceMutexType.PerBranch, Arguments))
+			// Read the XML configuration files
+			XmlConfig.ApplyTo(this);
+
+			// Create the log file, and flush the startup listener to it
+			FileReference LogFile = null;
+			if(!Arguments.HasOption("-NoLog") && !Log.HasFileWriter())
 			{
-				// Change the working directory to be the Engine/Source folder. We are likely running from Engine/Binaries/DotNET
-				// This is critical to be done early so any code that relies on the current directory being Engine/Source will work.
-				DirectoryReference.SetCurrentDirectory(UnrealBuildTool.EngineSourceDirectory);
-
-				// Read the XML configuration files
-				FileReference XmlConfigCache = Arguments.GetFileReferenceOrDefault("-XmlConfigCache=", null);
-				XmlConfig.ReadConfigFiles(XmlConfigCache);
-				XmlConfig.ApplyTo(this);
-
-				// Create the log file, and flush the startup listener to it
-				FileReference LogFile = null;
-				if(!Arguments.HasOption("-NoLog") && !Log.HasFileWriter())
+				LogFile = new FileReference(BaseLogFileName);
+				foreach(string LogSuffix in Arguments.GetValues("-LogSuffix="))
 				{
-					LogFile = new FileReference(BaseLogFileName);
-					foreach(string LogSuffix in Arguments.GetValues("-LogSuffix="))
-					{
-						LogFile = LogFile.ChangeExtension(null) + "_" + LogSuffix + LogFile.GetExtension();
-					}
-
-					TextWriterTraceListener LogTraceListener = Log.AddFileWriter("DefaultLogTraceListener", LogFile);
-					StartupListener.CopyTo(LogTraceListener);
+					LogFile = LogFile.ChangeExtension(null) + "_" + LogSuffix + LogFile.GetExtension();
 				}
-				Trace.Listeners.Remove(StartupListener);
-
-				// Create the build configuration object, and read the settings
-				BuildConfiguration BuildConfiguration = new BuildConfiguration();
-				XmlConfig.ApplyTo(BuildConfiguration);
-				Arguments.ApplyTo(BuildConfiguration);
 
-				// Then let the command lines override any configs necessary.
-				if(BuildConfiguration.bXGEExport)
-				{
-					BuildConfiguration.bAllowXGE = true;
-				}
-				if(BuildConfiguration.SingleFileToCompile != null)
-				{
-					BuildConfiguration.bUseUBTMakefiles = false;
-				}
+				TextWriterTraceListener LogTraceListener = Log.AddFileWriter("DefaultLogTraceListener", LogFile);
+				StartupListener.CopyTo(LogTraceListener);
+			}
+			Trace.Listeners.Remove(StartupListener);
 
-				// Parse the remote INI setting
-				string RemoteIniPath;
-				Arguments.TryGetValue("-RemoteIni=", out RemoteIniPath);
-				UnrealBuildTool.SetRemoteIniPath(RemoteIniPath);
+			// Create the build configuration object, and read the settings
+			BuildConfiguration BuildConfiguration = new BuildConfiguration();
+			XmlConfig.ApplyTo(BuildConfiguration);
+			Arguments.ApplyTo(BuildConfiguration);
 
-				DateTime BasicInitStartTime = DateTime.UtcNow;
+			// Then let the command lines override any configs necessary.
+			if(BuildConfiguration.bXGEExport)
+			{
+				BuildConfiguration.bAllowXGE = true;
+			}
+			if(BuildConfiguration.SingleFileToCompile != null)
+			{
+				BuildConfiguration.bUseUBTMakefiles = false;
+			}
 
-				// Find and register all tool chains, build platforms, etc. that are present
-				UnrealBuildTool.RegisterAllUBTClasses(false);
+			// Parse the remote INI setting
+			string RemoteIniPath;
+			Arguments.TryGetValue("-RemoteIni=", out RemoteIniPath);
+			UnrealBuildTool.SetRemoteIniPath(RemoteIniPath);
 
-				Timeline.AddEvent("Basic UBT initialization");
+			Timeline.AddEvent("Basic UBT initialization");
 
-				// now that we know the available platforms, we can delete other platforms' junk. if we're only building specific modules from the editor, don't touch anything else (it may be in use).
-				if (!BuildConfiguration.bIgnoreJunk)
-				{
-					JunkDeleter.DeleteJunk();
-				}
+			// now that we know the available platforms, we can delete other platforms' junk. if we're only building specific modules from the editor, don't touch anything else (it may be in use).
+			if (!BuildConfiguration.bIgnoreJunk)
+			{
+				JunkDeleter.DeleteJunk();
+			}
 
-				// Build our project
-				using(Timeline.ScopeEvent("Calling RunUBT"))
-				{
-					ECompilationResult Result = UnrealBuildTool.RunUBT(BuildConfiguration, Arguments.GetRawArray(), LogFile);
+			// Build our project
+			using(Timeline.ScopeEvent("Calling RunUBT"))
+			{
+				ECompilationResult Result = UnrealBuildTool.RunUBT(BuildConfiguration, Arguments.GetRawArray(), LogFile);
 
-					// Print some performance info
-					Log.TraceLog("DirectIncludes cache miss time: {0}s ({1} misses)", CPPHeaders.DirectIncludeCacheMissesTotalTime, CPPHeaders.TotalDirectIncludeCacheMisses);
-					Log.TraceLog("FindIncludePaths calls: {0} ({1} searches)", CPPHeaders.TotalFindIncludedFileCalls, CPPHeaders.IncludePathSearchAttempts);
-					Log.TraceLog("Deep C++ include scan time: {0}s", UnrealBuildTool.TotalDeepIncludeScanTime);
-					Log.TraceLog("Include Resolves: {0} ({1} misses, {2:0.00}%)", CPPHeaders.TotalDirectIncludeResolves, CPPHeaders.TotalDirectIncludeResolveCacheMisses, (float)CPPHeaders.TotalDirectIncludeResolveCacheMisses / (float)CPPHeaders.TotalDirectIncludeResolves * 100);
-					Log.TraceLog("Total FileItems: {0} ({1} missing)", FileItem.TotalFileItemCount, FileItem.MissingFileItemCount);
+				// Print some performance info
+				Log.TraceLog("DirectIncludes cache miss time: {0}s ({1} misses)", CPPHeaders.DirectIncludeCacheMissesTotalTime, CPPHeaders.TotalDirectIncludeCacheMisses);
+				Log.TraceLog("FindIncludePaths calls: {0} ({1} searches)", CPPHeaders.TotalFindIncludedFileCalls, CPPHeaders.IncludePathSearchAttempts);
+				Log.TraceLog("Deep C++ include scan time: {0}s", UnrealBuildTool.TotalDeepIncludeScanTime);
+				Log.TraceLog("Include Resolves: {0} ({1} misses, {2:0.00}%)", CPPHeaders.TotalDirectIncludeResolves, CPPHeaders.TotalDirectIncludeResolveCacheMisses, (float)CPPHeaders.TotalDirectIncludeResolveCacheMisses / (float)CPPHeaders.TotalDirectIncludeResolves * 100);
+				Log.TraceLog("Total FileItems: {0} ({1} missing)", FileItem.TotalFileItemCount, FileItem.MissingFileItemCount);
 
-					return (int)Result;
-				}
+				return (int)Result;
 			}
 		}
 	}
diff --git a/Engine/Source/Programs/UnrealBuildTool/Modes/DeployMode.cs b/Engine/Source/Programs/UnrealBuildTool/Modes/DeployMode.cs
--- a/Engine/Source/Programs/UnrealBuildTool/Modes/DeployMode.cs
+++ b/Engine/Source/Programs/UnrealBuildTool/Modes/DeployMode.cs
@@ -12,7 +12,7 @@ namespace UnrealBuildTool
 	/// <summary>
 	/// Invokes the deployment handler for a target.
 	/// </summary>
-	[ToolMode("Deploy")]
+	[ToolMode("Deploy", ToolModeOptions.BuildPlatforms)]
 	class DeployMode : ToolMode
 	{
 		/// <summary>
@@ -32,13 +32,6 @@ namespace UnrealBuildTool
 			Arguments.ApplyTo(this);
 			Arguments.CheckAllArgumentsUsed();
 
-			// Change the working directory to be the Engine/Source folder. We are likely running from Engine/Binaries/DotNET
-			// This is critical to be done early so any code that relies on the current directory being Engine/Source will work.
-			DirectoryReference.SetCurrentDirectory(UnrealBuildTool.EngineSourceDirectory);
-
-			// Find and register all tool chains, build platforms, etc. that are present
-			UnrealBuildTool.RegisterAllUBTClasses(false);
-
 			// Execute the deploy
 			TargetReceipt Receipt = TargetReceipt.Read(ReceiptFile);
 			Log.WriteLine(LogEventType.Console, "Deploying {0} {1} {2}...", Receipt.TargetName, Receipt.Platform, Receipt.Configuration);
diff --git a/Engine/Source/Programs/UnrealBuildTool/Modes/GenerateProjectFilesMode.cs b/Engine/Source/Programs/UnrealBuildTool/Modes/GenerateProjectFilesMode.cs
--- a/Engine/Source/Programs/UnrealBuildTool/Modes/GenerateProjectFilesMode.cs
+++ b/Engine/Source/Programs/UnrealBuildTool/Modes/GenerateProjectFilesMode.cs
@@ -14,7 +14,7 @@ namespace UnrealBuildTool
 	/// <summary>
 	/// Generates project files for one or more projects
 	/// </summary>
-	[ToolMode("GenerateProjectFiles")]
+	[ToolMode("GenerateProjectFiles", ToolModeOptions.XmlConfig | ToolModeOptions.BuildPlatforms | ToolModeOptions.SingleInstance)]
 	class GenerateProjectFilesMode : ToolMode
 	{
 		/// <summary>
@@ -45,135 +45,122 @@ namespace UnrealBuildTool
 		/// <returns>Exit code</returns>
 		public override int Execute(CommandLineArguments Arguments)
 		{
-			using(SingleInstanceMutex.Acquire(SingleInstanceMutexType.PerBranch, Arguments))
-			{
-				// Apply any command line arguments to this class
-				Arguments.ApplyTo(this);
-
-				// Parse rocket-specific arguments.
-				FileReference ProjectFile;
-				UnrealBuildTool.TryParseProjectFileArgument(Arguments, out ProjectFile);
+			// Apply any command line arguments to this class
+			Arguments.ApplyTo(this);
 
-				// Change the working directory to be the Engine/Source folder. We are likely running from Engine/Binaries/DotNET
-				// This is critical to be done early so any code that relies on the current directory being Engine/Source will work.
-				DirectoryReference.SetCurrentDirectory(UnrealBuildTool.EngineSourceDirectory);
+			// Parse rocket-specific arguments.
+			FileReference ProjectFile;
+			UnrealBuildTool.TryParseProjectFileArgument(Arguments, out ProjectFile);
 
-				// Read the XML configuration files
-				XmlConfig.ReadConfigFiles(null);
+			// If there aren't any formats set, read the default project file format from the config file
+			if (ProjectFileFormats.Count == 0)
+			{
+				// Read from the XML config
+				if (!String.IsNullOrEmpty(ProjectFileGeneratorSettings.Format))
+				{
+					ProjectFileFormats.UnionWith(ProjectFileGeneratorSettings.ParseFormatList(ProjectFileGeneratorSettings.Format));
+				}
 
-				// Find and register all tool chains, build platforms, etc. that are present
-				UnrealBuildTool.RegisterAllUBTClasses(false);
+				// Read from the editor config
+				ProjectFileFormat PreferredSourceCodeAccessor;
+				if (ProjectFileGenerator.GetPreferredSourceCodeAccessor(ProjectFile, out PreferredSourceCodeAccessor))
+				{
+					ProjectFileFormats.Add(PreferredSourceCodeAccessor);
+				}
 
-				// If there aren't any formats set, read the default project file format from the config file
+				// If there's still nothing set, get the default project file format for this platform
 				if (ProjectFileFormats.Count == 0)
 				{
-					// Read from the XML config
-					if (!String.IsNullOrEmpty(ProjectFileGeneratorSettings.Format))
-					{
-						ProjectFileFormats.UnionWith(ProjectFileGeneratorSettings.ParseFormatList(ProjectFileGeneratorSettings.Format));
-					}
-
-					// Read from the editor config
-					ProjectFileFormat PreferredSourceCodeAccessor;
-					if (ProjectFileGenerator.GetPreferredSourceCodeAccessor(ProjectFile, out PreferredSourceCodeAccessor))
-					{
-						ProjectFileFormats.Add(PreferredSourceCodeAccessor);
-					}
-
-					// If there's still nothing set, get the default project file format for this platform
-					if (ProjectFileFormats.Count == 0)
-					{
-						ProjectFileFormats.UnionWith(BuildHostPlatform.Current.GetDefaultProjectFileFormats());
-					}
+					ProjectFileFormats.UnionWith(BuildHostPlatform.Current.GetDefaultProjectFileFormats());
 				}
+			}
 
-				// Register all the platform project generators
-				PlatformProjectGeneratorCollection PlatformProjectGenerators = new PlatformProjectGeneratorCollection();
-				foreach (Type CheckType in Assembly.GetExecutingAssembly().GetTypes())
+			// Register all the platform project generators
+			PlatformProjectGeneratorCollection PlatformProjectGenerators = new PlatformProjectGeneratorCollection();
+			foreach (Type CheckType in Assembly.GetExecutingAssembly().GetTypes())
+			{
+				if (CheckType.IsClass && !CheckType.IsAbstract && CheckType.IsSubclassOf(typeof(PlatformProjectGenerator)))
 				{
-					if (CheckType.IsClass && !CheckType.IsAbstract && CheckType.IsSubclassOf(typeof(PlatformProjectGenerator)))
+					PlatformProjectGenerator Generator = (PlatformProjectGenerator)Activator.CreateInstance(CheckType, Arguments);
+					foreach(UnrealTargetPlatform Platform in Generator.GetPlatforms())
 					{
-						PlatformProjectGenerator Generator = (PlatformProjectGenerator)Activator.CreateInstance(CheckType, Arguments);
-						foreach(UnrealTargetPlatform Platform in Generator.GetPlatforms())
-						{
-							Log.TraceVerbose("Registering project generator {0} for {1}", CheckType, Platform);
-							PlatformProjectGenerators.RegisterPlatformProjectGenerator(Platform, Generator);
-						}
+						Log.TraceVerbose("Registering project generator {0} for {1}", CheckType, Platform);
+						PlatformProjectGenerators.RegisterPlatformProjectGenerator(Platform, Generator);
 					}
 				}
+			}
 
-				// Create each project generator and run it
-				List<ProjectFileGenerator> Generators = new List<ProjectFileGenerator>();
-				foreach (ProjectFileFormat ProjectFileFormat in ProjectFileFormats.Distinct())
+			// Create each project generator and run it
+			List<ProjectFileGenerator> Generators = new List<ProjectFileGenerator>();
+			foreach (ProjectFileFormat ProjectFileFormat in ProjectFileFormats.Distinct())
+			{
+				ProjectFileGenerator Generator;
+				switch (ProjectFileFormat)
 				{
-					ProjectFileGenerator Generator;
-					switch (ProjectFileFormat)
-					{
-						case ProjectFileFormat.Make:
-							Generator = new MakefileGenerator(ProjectFile);
-							break;
-						case ProjectFileFormat.CMake:
-							Generator = new CMakefileGenerator(ProjectFile);
-							break;
-						case ProjectFileFormat.QMake:
-							Generator = new QMakefileGenerator(ProjectFile);
-							break;
-						case ProjectFileFormat.KDevelop:
-							Generator = new KDevelopGenerator(ProjectFile);
-							break;
-						case ProjectFileFormat.CodeLite:
-							Generator = new CodeLiteGenerator(ProjectFile, Arguments);
-							break;
-						case ProjectFileFormat.VisualStudio:
-							Generator = new VCProjectFileGenerator(ProjectFile, VCProjectFileFormat.Default, Arguments);
-							break;
-						case ProjectFileFormat.VisualStudio2012:
-							Generator = new VCProjectFileGenerator(ProjectFile, VCProjectFileFormat.VisualStudio2012, Arguments);
-							break;
-						case ProjectFileFormat.VisualStudio2013:
-							Generator = new VCProjectFileGenerator(ProjectFile, VCProjectFileFormat.VisualStudio2013, Arguments);
-							break;
-						case ProjectFileFormat.VisualStudio2015:
-							Generator = new VCProjectFileGenerator(ProjectFile, VCProjectFileFormat.VisualStudio2015, Arguments);
-							break;
-						case ProjectFileFormat.VisualStudio2017:
-							Generator = new VCProjectFileGenerator(ProjectFile, VCProjectFileFormat.VisualStudio2017, Arguments);
-							break;
-						case ProjectFileFormat.VisualStudio2019:
-							Generator = new VCProjectFileGenerator(ProjectFile, VCProjectFileFormat.VisualStudio2019, Arguments);
-							break;
-						case ProjectFileFormat.XCode:
-							Generator = new XcodeProjectFileGenerator(ProjectFile);
-							break;
-						case ProjectFileFormat.Eddie:
-							Generator = new EddieProjectFileGenerator(ProjectFile);
-							break;
-						case ProjectFileFormat.VisualStudioCode:
-							Generator = new VSCodeProjectFileGenerator(ProjectFile);
-							break;
-						case ProjectFileFormat.CLion:
-							Generator = new CLionGenerator(ProjectFile);
-							break;
-						case ProjectFileFormat.VisualStudioMac:
-							Generator = new VCMacProjectFileGenerator(ProjectFile, Arguments);
-							break;
-						default:
-							throw new BuildException("Unhandled project file type '{0}", ProjectFileFormat);
-					}
-					Generators.Add(Generator);
+					case ProjectFileFormat.Make:
+						Generator = new MakefileGenerator(ProjectFile);
+						break;
+					case ProjectFileFormat.CMake:
+						Generator = new CMakefileGenerator(ProjectFile);
+						break;
+					case ProjectFileFormat.QMake:
+						Generator = new QMakefileGenerator(ProjectFile);
+						break;
+					case ProjectFileFormat.KDevelop:
+						Generator = new KDevelopGenerator(ProjectFile);
+						break;
+					case ProjectFileFormat.CodeLite:
+						Generator = new CodeLiteGenerator(ProjectFile, Arguments);
+						break;
+					case ProjectFileFormat.VisualStudio:
+						Generator = new VCProjectFileGenerator(ProjectFile, VCProjectFileFormat.Default, Arguments);
+						break;
+					case ProjectFileFormat.VisualStudio2012:
+						Generator = new VCProjectFileGenerator(ProjectFile, VCProjectFileFormat.VisualStudio2012, Arguments);
+						break;
+					case ProjectFileFormat.VisualStudio2013:
+						Generator = new VCProjectFileGenerator(ProjectFile, VCProjectFileFormat.VisualStudio2013, Arguments);
+						break;
+					case ProjectFileFormat.VisualStudio2015:
+						Generator = new VCProjectFileGenerator(ProjectFile, VCProjectFileFormat.VisualStudio2015, Arguments);
+						break;
+					case ProjectFileFormat.VisualStudio2017:
+						Generator = new VCProjectFileGenerator(ProjectFile, VCProjectFileFormat.VisualStudio2017, Arguments);
+						break;
+					case ProjectFileFormat.VisualStudio2019:
+						Generator = new VCProjectFileGenerator(ProjectFile, VCProjectFileFormat.VisualStudio2019, Arguments);
+						break;
+					case ProjectFileFormat.XCode:
+						Generator = new XcodeProjectFileGenerator(ProjectFile);
+						break;
+					case ProjectFileFormat.Eddie:
+						Generator = new EddieProjectFileGenerator(ProjectFile);
+						break;
+					case ProjectFileFormat.VisualStudioCode:
+						Generator = new VSCodeProjectFileGenerator(ProjectFile);
+						break;
+					case ProjectFileFormat.CLion:
+						Generator = new CLionGenerator(ProjectFile);
+						break;
+					case ProjectFileFormat.VisualStudioMac:
+						Generator = new VCMacProjectFileGenerator(ProjectFile, Arguments);
+						break;
+					default:
+						throw new BuildException("Unhandled project file type '{0}", ProjectFileFormat);
 				}
+				Generators.Add(Generator);
+			}
 
-				// Check there are no superfluous command line arguments
-				Arguments.CheckAllArgumentsUsed();
+			// Check there are no superfluous command line arguments
+			Arguments.CheckAllArgumentsUsed();
 
-				// Now generate project files
-				ProjectFileGenerator.bGenerateProjectFiles = true;
-				foreach(ProjectFileGenerator Generator in Generators)
+			// Now generate project files
+			ProjectFileGenerator.bGenerateProjectFiles = true;
+			foreach(ProjectFileGenerator Generator in Generators)
+			{
+				if (!Generator.GenerateProjectFiles(PlatformProjectGenerators, Arguments.GetRawArray()))
 				{
-					if (!Generator.GenerateProjectFiles(PlatformProjectGenerators, Arguments.GetRawArray()))
-					{
-						return (int)ECompilationResult.OtherCompilationError;
-					}
+					return (int)ECompilationResult.OtherCompilationError;
 				}
 			}
 			return (int)ECompilationResult.Succeeded;
diff --git a/Engine/Source/Programs/UnrealBuildTool/Modes/SetupPlatformsMode.cs b/Engine/Source/Programs/UnrealBuildTool/Modes/SetupPlatformsMode.cs
--- a/Engine/Source/Programs/UnrealBuildTool/Modes/SetupPlatformsMode.cs
+++ b/Engine/Source/Programs/UnrealBuildTool/Modes/SetupPlatformsMode.cs
@@ -12,7 +12,7 @@ namespace UnrealBuildTool
 	/// <summary>
 	/// Register all platforms (and in the process, configure all autosdks)
 	/// </summary>
-	[ToolMode("SetupPlatforms")]
+	[ToolMode("SetupPlatforms", ToolModeOptions.XmlConfig | ToolModeOptions.BuildPlatforms | ToolModeOptions.SingleInstance)]
 	class SetupPlatforms : ToolMode
 	{
 		/// <summary>
@@ -22,17 +22,8 @@ namespace UnrealBuildTool
 		/// <returns>Exit code</returns>
 		public override int Execute(CommandLineArguments Arguments)
 		{
-			using(SingleInstanceMutex.Acquire(SingleInstanceMutexType.Global, Arguments))
-			{
-				// Output a warning if there are any arguments that are still unused
-				Arguments.CheckAllArgumentsUsed();
-
-				// Read the XML configuration files
-				XmlConfig.ReadConfigFiles(null);
-
-				// Find and register all tool chains, build platforms, etc. that are present
-				UnrealBuildTool.RegisterAllUBTClasses(false);
-			}
+			// Output a warning if there are any arguments that are still unused
+			Arguments.CheckAllArgumentsUsed();
 			return 0;
 		}
 	}
diff --git a/Engine/Source/Programs/UnrealBuildTool/Modes/ValidatePlatformsMode.cs b/Engine/Source/Programs/UnrealBuildTool/Modes/ValidatePlatformsMode.cs
--- a/Engine/Source/Programs/UnrealBuildTool/Modes/ValidatePlatformsMode.cs
+++ b/Engine/Source/Programs/UnrealBuildTool/Modes/ValidatePlatformsMode.cs
@@ -12,7 +12,7 @@ namespace UnrealBuildTool
 	/// <summary>
 	/// Validates the various platforms to determine if they are ready for building
 	/// </summary>
-	[ToolMode("ValidatePlatforms")]
+	[ToolMode("ValidatePlatforms", ToolModeOptions.XmlConfig | ToolModeOptions.BuildPlatformsForValidation | ToolModeOptions.SingleInstance)]
 	class ValidatePlatformsMode : ToolMode
 	{
 		/// <summary>
@@ -34,38 +34,29 @@ namespace UnrealBuildTool
 		/// <returns>Exit code</returns>
 		public override int Execute(CommandLineArguments Arguments)
 		{
+			// Output a message if there are any arguments that are still unused
 			Arguments.ApplyTo(this);
-			using(SingleInstanceMutex.Acquire(SingleInstanceMutexType.Global, Arguments))
+			Arguments.CheckAllArgumentsUsed();
+
+			// If the -AllPlatforms argument is specified, add all the known platforms into the list
+			if(bAllPlatforms)
 			{
-				// Output a message if there are any arguments that are still unused
-				Arguments.CheckAllArgumentsUsed();
+				Platforms.UnionWith(Enum.GetValues(typeof(UnrealTargetPlatform)).OfType<UnrealTargetPlatform>().Where(x => x != UnrealTargetPlatform.Unknown));
+			}
 
-				// If the -AllPlatforms argument is specified, add all the known platforms into the list
-				if(bAllPlatforms)
+			// Output a line for each registered platform
+			foreach (UnrealTargetPlatform Platform in Platforms)
+			{
+				UEBuildPlatform BuildPlatform = UEBuildPlatform.GetBuildPlatform(Platform, true);
+				if (BuildPlatform != null && BuildPlatform.HasRequiredSDKsInstalled() == SDKStatus.Valid)
 				{
-					Platforms.UnionWith(Enum.GetValues(typeof(UnrealTargetPlatform)).OfType<UnrealTargetPlatform>().Where(x => x != UnrealTargetPlatform.Unknown));
+					Log.TraceInformation("##PlatformValidate: {0} VALID", Platform.ToString());
 				}
-
-				// Read all the config files
-				XmlConfig.ReadConfigFiles(null);
-				
-				// Register all the platforms
-				UnrealBuildTool.RegisterAllUBTClasses(true);
-
-				// Output a line for each registered platform
-				foreach (UnrealTargetPlatform Platform in Platforms)
+				else
 				{
-					UEBuildPlatform BuildPlatform = UEBuildPlatform.GetBuildPlatform(Platform, true);
-					if (BuildPlatform != null && BuildPlatform.HasRequiredSDKsInstalled() == SDKStatus.Valid)
-					{
-						Log.TraceInformation("##PlatformValidate: {0} VALID", Platform.ToString());
-					}
-					else
-					{
-						Log.TraceInformation("##PlatformValidate: {0} INVALID", Platform.ToString());
-					}
-				} 
-			}
+					Log.TraceInformation("##PlatformValidate: {0} INVALID", Platform.ToString());
+				}
+			} 
 			return 0;
 		}
 	}
diff --git a/Engine/Source/Programs/UnrealBuildTool/Modes/WriteDocumentationMode.cs b/Engine/Source/Programs/UnrealBuildTool/Modes/WriteDocumentationMode.cs
--- a/Engine/Source/Programs/UnrealBuildTool/Modes/WriteDocumentationMode.cs
+++ b/Engine/Source/Programs/UnrealBuildTool/Modes/WriteDocumentationMode.cs
@@ -12,7 +12,7 @@ namespace UnrealBuildTool
 	/// <summary>
 	/// Generates documentation from reflection data
 	/// </summary>
-	[ToolMode("WriteDocumentation")]
+	[ToolMode("WriteDocumentation", ToolModeOptions.None)]
 	class WriteDocumentationMode : ToolMode
 	{
 		/// <summary>
diff --git a/Engine/Source/Programs/UnrealBuildTool/Modes/WriteMetadataMode.cs b/Engine/Source/Programs/UnrealBuildTool/Modes/WriteMetadataMode.cs
--- a/Engine/Source/Programs/UnrealBuildTool/Modes/WriteMetadataMode.cs
+++ b/Engine/Source/Programs/UnrealBuildTool/Modes/WriteMetadataMode.cs
@@ -63,7 +63,7 @@ namespace UnrealBuildTool
 	/// <summary>
 	/// Writes all metadata files at the end of a build (receipts, version files, etc...). This is implemented as a separate mode to allow it to be done as part of the action graph.
 	/// </summary>
-	[ToolMode("WriteMetadata")]
+	[ToolMode("WriteMetadata", ToolModeOptions.None)]
 	class WriteMetadataMode : ToolMode
 	{
 		/// <summary>
diff --git a/Engine/Source/Programs/UnrealBuildTool/Platform/IOS/IOSPostBuildSyncMode.cs b/Engine/Source/Programs/UnrealBuildTool/Platform/IOS/IOSPostBuildSyncMode.cs
--- a/Engine/Source/Programs/UnrealBuildTool/Platform/IOS/IOSPostBuildSyncMode.cs
+++ b/Engine/Source/Programs/UnrealBuildTool/Platform/IOS/IOSPostBuildSyncMode.cs
@@ -50,7 +50,7 @@ namespace UnrealBuildTool
 		}
 	}
 
-	[ToolMode("IOSPostBuildSync")]
+	[ToolMode("IOSPostBuildSync", ToolModeOptions.XmlConfig | ToolModeOptions.BuildPlatforms)]
 	class IOSPostBuildSyncMode : ToolMode
 	{
 		[CommandLine("-Input=", Required = true)]
@@ -64,16 +64,6 @@ namespace UnrealBuildTool
 			Arguments.ApplyTo(this);
 			Arguments.CheckAllArgumentsUsed();
 
-			// Read the XML config
-			XmlConfig.ReadConfigFiles(XmlConfigCache);
-
-			// Change the working directory to be the Engine/Source folder. We are likely running from Engine/Binaries/DotNET
-			// This is critical to be done early so any code that relies on the current directory being Engine/Source will work.
-			DirectoryReference.SetCurrentDirectory(UnrealBuildTool.EngineSourceDirectory);
-
-			// Find and register all tool chains, build platforms, etc. that are present
-			UnrealBuildTool.RegisterAllUBTClasses(false);
-
 			// Run the PostBuildSync command
 			IOSPostBuildSyncTarget Target = BinaryFormatterUtils.Load<IOSPostBuildSyncTarget>(InputFile);
 			IOSToolChain.PostBuildSync(Target);
diff --git a/Engine/Source/Programs/UnrealBuildTool/Platform/Windows/PVSToolChain.cs b/Engine/Source/Programs/UnrealBuildTool/Platform/Windows/PVSToolChain.cs
--- a/Engine/Source/Programs/UnrealBuildTool/Platform/Windows/PVSToolChain.cs
+++ b/Engine/Source/Programs/UnrealBuildTool/Platform/Windows/PVSToolChain.cs
@@ -38,7 +38,7 @@ namespace UnrealBuildTool
 	/// <summary>
 	/// Special mode for gathering all the messages into a single output file
 	/// </summary>
-	[ToolMode("PVSGather")]
+	[ToolMode("PVSGather", ToolModeOptions.None)]
 	class PVSGatherMode : ToolMode
 	{
 		/// <summary>
diff --git a/Engine/Source/Programs/UnrealBuildTool/System/ToolMode.cs b/Engine/Source/Programs/UnrealBuildTool/System/ToolMode.cs
--- a/Engine/Source/Programs/UnrealBuildTool/System/ToolMode.cs
+++ b/Engine/Source/Programs/UnrealBuildTool/System/ToolMode.cs
@@ -9,6 +9,38 @@ using Tools.DotNETCommon;
 
 namespace UnrealBuildTool
 {
+	/// <summary>
+	/// Systems that need to be configured to execute a tool mode
+	/// </summary>
+	[Flags]
+	enum ToolModeOptions
+	{
+		/// <summary>
+		/// Do not initialize anything
+		/// </summary>
+		None = 0,
+
+		/// <summary>
+		/// Initializes the XmlConfig system
+		/// </summary>
+		XmlConfig = 1,
+
+		/// <summary>
+		/// Registers build platforms
+		/// </summary>
+		BuildPlatforms = 2,
+
+		/// <summary>
+		/// Registers build platforms for validation
+		/// </summary>
+		BuildPlatformsForValidation = 4,
+
+		/// <summary>
+		/// Only allow a single instance running in the branch at once
+		/// </summary>
+		SingleInstance = 8,
+	}
+
 	/// <summary>
 	/// Attribute used to specify options for a UBT mode.
 	/// </summary>
@@ -19,13 +51,20 @@ namespace UnrealBuildTool
 		/// </summary>
 		public string Name;
 
+		/// <summary>
+		/// Options for executing this mode
+		/// </summary>
+		public ToolModeOptions Options;
+
 		/// <summary>
 		/// Constructor
 		/// </summary>
 		/// <param name="Name">Name of the mode</param>
-		public ToolModeAttribute(string Name)
+		/// <param name="Options">Options for this mode</param>
+		public ToolModeAttribute(string Name, ToolModeOptions Options)
 		{
 			this.Name = Name;
+			this.Options = Options;
 		}
 	}
 
diff --git a/Engine/Source/Programs/UnrealBuildTool/UnrealBuildTool.cs b/Engine/Source/Programs/UnrealBuildTool/UnrealBuildTool.cs
--- a/Engine/Source/Programs/UnrealBuildTool/UnrealBuildTool.cs
+++ b/Engine/Source/Programs/UnrealBuildTool/UnrealBuildTool.cs
@@ -370,6 +370,7 @@ namespace UnrealBuildTool
 		/// <returns>Zero on success, non-zero on error</returns>
 		private static int Main(string[] ArgumentsArray)
 		{
+			SingleInstanceMutex Mutex = null;
 			try
 			{
 				// Start capturing performance info
@@ -413,6 +414,10 @@ namespace UnrealBuildTool
 				// Ensure we can resolve any external assemblies that are not in the same folder as our assembly.
 				AssemblyUtils.InstallAssemblyResolver(Path.GetDirectoryName(Assembly.GetEntryAssembly().GetOriginalLocation()));
 
+				// Change the working directory to be the Engine/Source folder. We are likely running from Engine/Binaries/DotNET
+				// This is critical to be done early so any code that relies on the current directory being Engine/Source will work.
+				DirectoryReference.SetCurrentDirectory(UnrealBuildTool.EngineSourceDirectory);
+
 				// Try to get the correct mode
 				Type ModeType;
 				if(!ModeNameToType.TryGetValue(Options.Mode, out ModeType))
@@ -421,6 +426,44 @@ namespace UnrealBuildTool
 					return 1;
 				}
 
+				// Get the options for which systems have to be initialized for this mode
+				ToolModeOptions ModeOptions = ModeType.GetCustomAttribute<ToolModeAttribute>().Options;
+
+				// Read the XML configuration files
+				if((ModeOptions & ToolModeOptions.XmlConfig) != 0)
+				{
+					using(Timeline.ScopeEvent("XmlConfig.ReadConfigFiles()"))
+					{
+						FileReference XmlConfigCache = Arguments.GetFileReferenceOrDefault("-XmlConfigCache=", null);
+						XmlConfig.ReadConfigFiles(XmlConfigCache);
+					}
+				}
+
+				// Register all the build platforms
+				if((ModeOptions & ToolModeOptions.BuildPlatforms) != 0)
+				{
+					using(Timeline.ScopeEvent("UnrealBuildTool.RegisterAllUBTClasses()"))
+					{
+						RegisterAllUBTClasses(false);
+					}
+				}
+				if((ModeOptions & ToolModeOptions.BuildPlatformsForValidation) != 0)
+				{
+					using(Timeline.ScopeEvent("UnrealBuildTool.RegisterAllUBTClasses()"))
+					{
+						RegisterAllUBTClasses(true);
+					}
+				}
+
+				// Acquire a lock for this branch
+				if((ModeOptions & ToolModeOptions.SingleInstance) != 0)
+				{
+					using(Timeline.ScopeEvent("SingleInstanceMutex.Acquire()"))
+					{
+						Mutex = SingleInstanceMutex.Acquire(SingleInstanceMutexType.PerBranch, Arguments);
+					}
+				}
+
 				// Create the appropriate handler
 				ToolMode Mode = (ToolMode)Activator.CreateInstance(ModeType);
 				return Mode.Execute(Arguments);
@@ -441,6 +484,12 @@ namespace UnrealBuildTool
 			}
 			finally
 			{
+				// Dispose of the mutex
+				if(Mutex != null)
+				{
+					Mutex.Dispose();
+				}
+
 				// Print out all the performance info
 				Timeline.Print(TimeSpan.FromMilliseconds(20.0), LogEventType.Log);
 
