UBT: Move code for cleaning a target into a separate mode. This simplifies the codepath when building a target, and avoids having to spawn UBT recursively to clean UHT.

#rnx

diff --git a/Engine/Source/Programs/UnrealBuildTool/Configuration/BuildConfiguration.cs b/Engine/Source/Programs/UnrealBuildTool/Configuration/BuildConfiguration.cs
--- a/Engine/Source/Programs/UnrealBuildTool/Configuration/BuildConfiguration.cs
+++ b/Engine/Source/Programs/UnrealBuildTool/Configuration/BuildConfiguration.cs
@@ -154,12 +154,6 @@ namespace UnrealBuildTool
 		[CommandLine("-GenerateManifest")]
 		public bool bGenerateManifest = false;
 
-		/// <summary>
-		/// Whether to 'clean' the given project
-		/// </summary>
-		[CommandLine("-Clean")]
-		public bool bCleanProject = false;
-
 		/// <summary>
 		/// If true, force header regeneration. Intended for the build machine
 		/// </summary>
diff --git a/Engine/Source/Programs/UnrealBuildTool/Configuration/UEBuildTarget.cs b/Engine/Source/Programs/UnrealBuildTool/Configuration/UEBuildTarget.cs
--- a/Engine/Source/Programs/UnrealBuildTool/Configuration/UEBuildTarget.cs
+++ b/Engine/Source/Programs/UnrealBuildTool/Configuration/UEBuildTarget.cs
@@ -907,7 +907,7 @@ namespace UnrealBuildTool
 		/// </summary>
 		/// <param name="Type">The target type</param>
 		/// <returns>The app name for this target type</returns>
-		static string GetAppNameForTargetType(TargetType Type)
+		public static string GetAppNameForTargetType(TargetType Type)
 		{
 			switch(Type)
 			{
@@ -924,139 +924,6 @@ namespace UnrealBuildTool
 			}
 		}
 
-		/// <summary>
-		/// Cleans build products and intermediates for the target. This deletes files which are named consistently with the target being built
-		/// (e.g. UE4Editor-Foo-Win64-Debug.dll) rather than an actual record of previous build products.
-		/// </summary>
-		/// <param name="bIncludeUnrealHeaderTool">Whether to clean UnrealHeaderTool as well</param>
-		/// <returns>Whether the clean succeeded</returns>
-		public ECompilationResult Clean(bool bIncludeUnrealHeaderTool)
-		{
-			// Find the base folders that can contain binaries
-			List<DirectoryReference> BaseDirs = new List<DirectoryReference>();
-			BaseDirs.Add(UnrealBuildTool.EngineDirectory);
-			BaseDirs.Add(UnrealBuildTool.EnterpriseDirectory);
-			foreach (FileReference Plugin in Plugins.EnumeratePlugins(ProjectFile))
-			{
-				BaseDirs.Add(Plugin.Directory);
-			}
-			if (ProjectFile != null)
-			{
-				BaseDirs.Add(ProjectFile.Directory);
-			}
-
-			// If we're running a precompiled build, remove anything under the engine folder
-			BaseDirs.RemoveAll(x => RulesAssembly.IsReadOnly(x));
-
-			// Get all the names which can prefix build products
-			List<string> NamePrefixes = new List<string>();
-			if (Rules.Type != TargetType.Program)
-			{
-				NamePrefixes.Add(GetAppNameForTargetType(Rules.Type));
-			}
-			NamePrefixes.Add(TargetName);
-
-			// Get the suffixes for this configuration
-			List<string> NameSuffixes = new List<string>();
-			if (Configuration == Rules.UndecoratedConfiguration)
-			{
-				NameSuffixes.Add("");
-			}
-			NameSuffixes.Add(String.Format("-{0}-{1}", Platform.ToString(), Configuration.ToString()));
-			if (!String.IsNullOrEmpty(Architecture))
-			{
-				NameSuffixes.AddRange(NameSuffixes.ToArray().Select(x => x + Architecture));
-			}
-
-			// Add all the makefiles and caches to be deleted
-			List<FileReference> FilesToDelete = new List<FileReference>();
-			FilesToDelete.Add(FlatCPPIncludeDependencyCache.GetDependencyCachePathForTarget(ProjectFile, TargetName, Platform, Architecture));
-			FilesToDelete.Add(DependencyCache.GetDependencyCachePathForTarget(ProjectFile, Platform, TargetName));
-			FilesToDelete.Add(UBTMakefile.GetUBTMakefilePath(ProjectFile, Platform, Configuration, TargetName, false));
-			FilesToDelete.Add(UBTMakefile.GetUBTMakefilePath(ProjectFile, Platform, Configuration, TargetName, true));
-			FilesToDelete.Add(ActionHistory.GeneratePathForTarget(ProjectFile, TargetName, Platform, Architecture, bUseSharedBuildEnvironment));
-
-			// Add all the intermediate folders to be deleted
-			List<DirectoryReference> DirectoriesToDelete = new List<DirectoryReference>();
-			foreach (DirectoryReference BaseDir in BaseDirs)
-			{
-				foreach (string NamePrefix in NamePrefixes)
-				{
-					DirectoryReference GeneratedCodeDir = DirectoryReference.Combine(BaseDir, "Intermediate", "Build", Platform.ToString(), NamePrefix, "Inc");
-					if (DirectoryReference.Exists(GeneratedCodeDir))
-					{
-						DirectoriesToDelete.Add(GeneratedCodeDir);
-					}
-
-					DirectoryReference IntermediateDir = DirectoryReference.Combine(BaseDir, "Intermediate", "Build", Platform.ToString(), NamePrefix, Configuration.ToString());
-					if (DirectoryReference.Exists(IntermediateDir))
-					{
-						DirectoriesToDelete.Add(IntermediateDir);
-					}
-				}
-			}
-
-			// Add all the build products from this target
-			string[] NamePrefixesArray = NamePrefixes.Distinct().ToArray();
-			string[] NameSuffixesArray = NameSuffixes.Distinct().ToArray();
-			foreach (DirectoryReference BaseDir in BaseDirs)
-			{
-				DirectoryReference BinariesDir = DirectoryReference.Combine(BaseDir, "Binaries", Platform.ToString());
-				if(DirectoryReference.Exists(BinariesDir))
-				{
-					UEBuildPlatform.GetBuildPlatform(Platform).FindBuildProductsToClean(BinariesDir, NamePrefixesArray, NameSuffixesArray, FilesToDelete, DirectoriesToDelete);
-				}
-			}
-
-			// Get all the additional intermediate folders created by this platform
-			List<FileReference> AdditionalFilesToDelete = new List<FileReference>();
-			List<DirectoryReference> AdditionalDirectoriesToDelete = new List<DirectoryReference>();
-			UEBuildPlatform.GetBuildPlatform(Platform).FindAdditionalBuildProductsToClean(Rules, AdditionalFilesToDelete, AdditionalDirectoriesToDelete);
-			FilesToDelete.AddRange(AdditionalFilesToDelete);
-			DirectoriesToDelete.AddRange(AdditionalDirectoriesToDelete);
-
-			// Delete all the directories, then all the files. By sorting the list of directories before we delete them, we avoid spamming the log if a parent directory is deleted first.
-			foreach (DirectoryReference DirectoryToDelete in DirectoriesToDelete.OrderBy(x => x.FullName))
-			{
-				if (DirectoryReference.Exists(DirectoryToDelete))
-				{
-					Log.TraceVerbose("    Deleting {0}{1}...", DirectoryToDelete, Path.DirectorySeparatorChar);
-					try
-					{
-						DirectoryReference.Delete(DirectoryToDelete, true);
-					}
-					catch (Exception Ex)
-					{
-						throw new BuildException(Ex, "Unable to delete {0} ({1})", DirectoryToDelete, Ex.Message.TrimEnd());
-					}
-				}
-			}
-
-			foreach (FileReference FileToDelete in FilesToDelete.OrderBy(x => x.FullName))
-			{
-				if (FileReference.Exists(FileToDelete))
-				{
-					Log.TraceVerbose("    Deleting " + FileToDelete);
-					try
-					{
-						FileReference.Delete(FileToDelete);
-					}
-					catch (Exception Ex)
-					{
-						throw new BuildException(Ex, "Unable to delete {0} ({1})", FileToDelete, Ex.Message.TrimEnd());
-					}
-				}
-			}
-
-			// Finally clean UnrealHeaderTool if this target uses CoreUObject modules and we're not cleaning UHT already and we want UHT to be cleaned.
-			if (bIncludeUnrealHeaderTool && !RulesAssembly.IsReadOnly(ExternalExecution.GetHeaderToolReceiptFile(ProjectFile, UnrealTargetConfiguration.Development, bHasProjectScriptPlugin)) && TargetName != "UnrealHeaderTool")
-			{
-				ExternalExecution.RunExternalDotNETExecutable(UnrealBuildTool.GetUBTPath(), String.Format("UnrealHeaderTool {0} {1} -NoMutex -Clean -IgnoreJunk -NoLog", BuildHostPlatform.Current.Platform, UnrealTargetConfiguration.Development));
-			}
-
-			return ECompilationResult.Succeeded;
-		}
-
 		/// <summary>
 		/// Writes a list of all the externally referenced files required to use the precompiled data for this target
 		/// </summary>
diff --git a/Engine/Source/Programs/UnrealBuildTool/Modes/CleanMode.cs b/Engine/Source/Programs/UnrealBuildTool/Modes/CleanMode.cs
new file mode 100644
--- /dev/null
+++ b/Engine/Source/Programs/UnrealBuildTool/Modes/CleanMode.cs
@@ -0,0 +1,227 @@
+// Copyright 1998-2018 Epic Games, Inc. All Rights Reserved.
+
+using System;
+using System.Collections.Generic;
+using System.Diagnostics;
+using System.IO;
+using System.Linq;
+using System.Text;
+using System.Threading.Tasks;
+using Tools.DotNETCommon;
+
+namespace UnrealBuildTool
+{
+	/// <summary>
+	/// Cleans build products and intermediates for the target. This deletes files which are named consistently with the target being built
+	/// (e.g. UE4Editor-Foo-Win64-Debug.dll) rather than an actual record of previous build products.
+	/// </summary>
+	[ToolMode("Clean", ToolModeOptions.XmlConfig | ToolModeOptions.BuildPlatforms | ToolModeOptions.SingleInstance)]
+	class CleanMode : ToolMode
+	{
+		/// <summary>
+		/// Whether to avoid cleaning targets
+		/// </summary>
+		[CommandLine("-SkipRulesCompile")]
+		bool bSkipRulesCompile = false;
+
+		/// <summary>
+		/// Main entry point
+		/// </summary>
+		/// <param name="Arguments">Command-line arguments</param>
+		/// <returns>One of the values of ECompilationResult</returns>
+		public override int Execute(CommandLineArguments Arguments)
+		{
+			Arguments.ApplyTo(this);
+
+			// Create the build configuration object, and read the settings
+			BuildConfiguration BuildConfiguration = new BuildConfiguration();
+			XmlConfig.ApplyTo(BuildConfiguration);
+			Arguments.ApplyTo(BuildConfiguration);
+
+			// Parse all the targets being built
+			List<TargetDescriptor> TargetDescriptors = TargetDescriptor.ParseCommandLine(Arguments, BuildConfiguration.bUsePrecompiled, bSkipRulesCompile);
+			if(TargetDescriptors.Count == 0)
+			{
+				throw new BuildException("No targets specified to clean");
+			}
+
+			// Also add implicit descriptors for cleaning UnrealBuildTool
+			if(!BuildConfiguration.bDoNotBuildUHT)
+			{
+				const string UnrealHeaderToolTarget = "UnrealHeaderTool";
+
+				// Get a list of project files to clean UHT for
+				List<FileReference> ProjectFiles = new List<FileReference>();
+				foreach(TargetDescriptor TargetDesc in TargetDescriptors)
+				{
+					if(TargetDesc.Name != UnrealHeaderToolTarget && !RemoteMac.HandlesTargetPlatform(TargetDesc.Platform))
+					{
+						if(ProjectFiles.Count == 0)
+						{
+							ProjectFiles.Add(null);
+						}
+						if(TargetDesc.ProjectFile != null && !ProjectFiles.Contains(TargetDesc.ProjectFile))
+						{
+							ProjectFiles.Add(TargetDesc.ProjectFile);
+						}
+					}
+				}
+
+				// Add descriptors for cleaning UHT with all these projects
+				if(ProjectFiles.Count > 0)
+				{
+					UnrealTargetConfiguration Configuration = BuildConfiguration.bForceDebugUnrealHeaderTool ? UnrealTargetConfiguration.Debug : UnrealTargetConfiguration.Development;
+					string Architecture = UEBuildPlatform.GetBuildPlatform(BuildHostPlatform.Current.Platform).GetDefaultArchitecture(null);
+					foreach(FileReference ProjectFile in ProjectFiles)
+					{
+						TargetDescriptors.Add(new TargetDescriptor(ProjectFile, UnrealHeaderToolTarget, BuildHostPlatform.Current.Platform, Configuration, Architecture, null, null));
+					}
+				}
+			}
+
+			// Loop through all the targets, and clean them all
+			HashSet<FileReference> FilesToDelete = new HashSet<FileReference>();
+			HashSet<DirectoryReference> DirectoriesToDelete = new HashSet<DirectoryReference>();
+			foreach(TargetDescriptor TargetDescriptor in TargetDescriptors)
+			{
+				// Create the rules assembly
+				RulesAssembly RulesAssembly = RulesCompiler.CreateTargetRulesAssembly(TargetDescriptor.ProjectFile, TargetDescriptor.Name, bSkipRulesCompile, BuildConfiguration.bUsePrecompiled, TargetDescriptor.ForeignPlugin);
+
+				// Create the rules object
+				ReadOnlyTargetRules Target = new ReadOnlyTargetRules(RulesAssembly.CreateTargetRules(TargetDescriptor.Name, TargetDescriptor.Platform, TargetDescriptor.Configuration, TargetDescriptor.Architecture, TargetDescriptor.ProjectFile, TargetDescriptor.AdditionalArguments));
+
+				// Find the base folders that can contain binaries
+				List<DirectoryReference> BaseDirs = new List<DirectoryReference>();
+				BaseDirs.Add(UnrealBuildTool.EngineDirectory);
+				BaseDirs.Add(UnrealBuildTool.EnterpriseDirectory);
+				foreach (FileReference Plugin in Plugins.EnumeratePlugins(Target.ProjectFile))
+				{
+					BaseDirs.Add(Plugin.Directory);
+				}
+				if (Target.ProjectFile != null)
+				{
+					BaseDirs.Add(Target.ProjectFile.Directory);
+				}
+
+				// If we're running a precompiled build, remove anything under the engine folder
+				BaseDirs.RemoveAll(x => RulesAssembly.IsReadOnly(x));
+
+				// Get all the names which can prefix build products
+				List<string> NamePrefixes = new List<string>();
+				if (Target.Type != TargetType.Program)
+				{
+					NamePrefixes.Add(UEBuildTarget.GetAppNameForTargetType(Target.Type));
+				}
+				NamePrefixes.Add(Target.Name);
+
+				// Get the suffixes for this configuration
+				List<string> NameSuffixes = new List<string>();
+				if (Target.Configuration == Target.UndecoratedConfiguration)
+				{
+					NameSuffixes.Add("");
+				}
+				NameSuffixes.Add(String.Format("-{0}-{1}", Target.Platform.ToString(), Target.Configuration.ToString()));
+				if (!String.IsNullOrEmpty(Target.Architecture))
+				{
+					NameSuffixes.AddRange(NameSuffixes.ToArray().Select(x => x + Target.Architecture));
+				}
+
+				// Add all the makefiles and caches to be deleted
+				FilesToDelete.Add(FlatCPPIncludeDependencyCache.GetDependencyCachePathForTarget(Target.ProjectFile, Target.Name, Target.Platform, Target.Architecture));
+				FilesToDelete.Add(DependencyCache.GetDependencyCachePathForTarget(Target.ProjectFile, Target.Platform, Target.Name));
+				FilesToDelete.Add(UBTMakefile.GetUBTMakefilePath(Target.ProjectFile, Target.Platform, Target.Configuration, Target.Name, false));
+				FilesToDelete.Add(UBTMakefile.GetUBTMakefilePath(Target.ProjectFile, Target.Platform, Target.Configuration, Target.Name, true));
+				FilesToDelete.Add(ActionHistory.GeneratePathForTarget(Target.ProjectFile, Target.Name, Target.Platform, Target.Architecture, (Target.BuildEnvironment == TargetBuildEnvironment.Shared)));
+
+				// Add all the intermediate folders to be deleted
+				foreach (DirectoryReference BaseDir in BaseDirs)
+				{
+					foreach (string NamePrefix in NamePrefixes)
+					{
+						DirectoryReference GeneratedCodeDir = DirectoryReference.Combine(BaseDir, "Intermediate", "Build", Target.Platform.ToString(), NamePrefix, "Inc");
+						if (DirectoryReference.Exists(GeneratedCodeDir))
+						{
+							DirectoriesToDelete.Add(GeneratedCodeDir);
+						}
+
+						DirectoryReference IntermediateDir = DirectoryReference.Combine(BaseDir, "Intermediate", "Build", Target.Platform.ToString(), NamePrefix, Target.Configuration.ToString());
+						if (DirectoryReference.Exists(IntermediateDir))
+						{
+							DirectoriesToDelete.Add(IntermediateDir);
+						}
+					}
+				}
+
+				// List of additional files and directories to clean, specified by the target platform
+				List<FileReference> AdditionalFilesToDelete = new List<FileReference>();
+				List<DirectoryReference> AdditionalDirectoriesToDelete = new List<DirectoryReference>();
+
+				// Add all the build products from this target
+				string[] NamePrefixesArray = NamePrefixes.Distinct().ToArray();
+				string[] NameSuffixesArray = NameSuffixes.Distinct().ToArray();
+				foreach (DirectoryReference BaseDir in BaseDirs)
+				{
+					DirectoryReference BinariesDir = DirectoryReference.Combine(BaseDir, "Binaries", Target.Platform.ToString());
+					if(DirectoryReference.Exists(BinariesDir))
+					{
+						UEBuildPlatform.GetBuildPlatform(Target.Platform).FindBuildProductsToClean(BinariesDir, NamePrefixesArray, NameSuffixesArray, AdditionalFilesToDelete, AdditionalDirectoriesToDelete);
+					}
+				}
+
+				// Get all the additional intermediate folders created by this platform
+				UEBuildPlatform.GetBuildPlatform(Target.Platform).FindAdditionalBuildProductsToClean(Target, AdditionalFilesToDelete, AdditionalDirectoriesToDelete);
+
+				// Add the platform's files and directories to the main list
+				FilesToDelete.UnionWith(AdditionalFilesToDelete);
+				DirectoriesToDelete.UnionWith(AdditionalDirectoriesToDelete);
+			}
+
+			// Delete all the directories, then all the files. By sorting the list of directories before we delete them, we avoid spamming the log if a parent directory is deleted first.
+			foreach (DirectoryReference DirectoryToDelete in DirectoriesToDelete.OrderBy(x => x.FullName))
+			{
+				if (DirectoryReference.Exists(DirectoryToDelete))
+				{
+					Log.TraceVerbose("    Deleting {0}{1}...", DirectoryToDelete, Path.DirectorySeparatorChar);
+					try
+					{
+						DirectoryReference.Delete(DirectoryToDelete, true);
+					}
+					catch (Exception Ex)
+					{
+						throw new BuildException(Ex, "Unable to delete {0} ({1})", DirectoryToDelete, Ex.Message.TrimEnd());
+					}
+				}
+			}
+
+			foreach (FileReference FileToDelete in FilesToDelete.OrderBy(x => x.FullName))
+			{
+				if (FileReference.Exists(FileToDelete))
+				{
+					Log.TraceVerbose("    Deleting " + FileToDelete);
+					try
+					{
+						FileReference.Delete(FileToDelete);
+					}
+					catch (Exception Ex)
+					{
+						throw new BuildException(Ex, "Unable to delete {0} ({1})", FileToDelete, Ex.Message.TrimEnd());
+					}
+				}
+			}
+
+			// Also clean all the remote targets
+			for(int Idx = 0; Idx < TargetDescriptors.Count; Idx++)
+			{
+				TargetDescriptor TargetDescriptor = TargetDescriptors[Idx];
+				if(RemoteMac.HandlesTargetPlatform(TargetDescriptor.Platform))
+				{
+					RemoteMac RemoteMac = new RemoteMac(TargetDescriptor.ProjectFile);
+					RemoteMac.Clean(TargetDescriptor);
+				}
+			}
+
+			return 0;
+		}
+	}
+}
+
diff --git a/Engine/Source/Programs/UnrealBuildTool/ToolChain/RemoteMac.cs b/Engine/Source/Programs/UnrealBuildTool/ToolChain/RemoteMac.cs
--- a/Engine/Source/Programs/UnrealBuildTool/ToolChain/RemoteMac.cs
+++ b/Engine/Source/Programs/UnrealBuildTool/ToolChain/RemoteMac.cs
@@ -332,5 +332,43 @@
 		}
 
 		/// <summary>
+		/// Returns true if the remote executor supports this target platform
+		/// </summary>
+		/// <param name="Platform">The platform to check</param>
+		/// <returns>True if the remote mac handles this target platform</returns>
+		public static bool HandlesTargetPlatform(UnrealTargetPlatform Platform)
+		{
+			return BuildHostPlatform.Current.Platform == UnrealTargetPlatform.Win64 && (Platform == UnrealTargetPlatform.Mac || Platform == UnrealTargetPlatform.IOS || Platform == UnrealTargetPlatform.TVOS);
+		}
+
+		/// <summary>
+		/// Clean a target remotely
+		/// </summary>
+		/// <param name="TargetDesc">Descriptor for the target to build</param>
+		/// <returns>True if the build succeeded, false otherwise</returns>
+		public bool Clean(TargetDescriptor TargetDesc)
+		{
+			// Translate all the arguments for the remote
+			List<string> RemoteArguments = GetRemoteArgumentsForTarget(TargetDesc, null);
+			RemoteArguments.Add("-Clean");
+		
+			// Upload the workspace
+			DirectoryReference TempDir = CreateTempDirectory(TargetDesc);
+			UploadWorkspace(TempDir);
+
+			// Execute the compile
+			Log.TraceInformation("[Remote] Executing clean...");
+
+			StringBuilder BuildCommandLine = new StringBuilder("Engine/Build/BatchFiles/Mac/Build.sh");
+			foreach(string RemoteArgument in RemoteArguments)
+			{
+				BuildCommandLine.AppendFormat(" {0}", EscapeShellArgument(RemoteArgument));
+			}
+
+			int Result = Execute(GetRemotePath(UnrealBuildTool.RootDirectory), BuildCommandLine.ToString());
+			return Result == 0;
+		}
+
+		/// <summary>
 		/// Build a target remotely
 		/// </summary>
@@ -340,9 +378,7 @@
 		public bool Build(TargetDescriptor TargetDesc, FileReference RemoteLogFile)
 		{
 			// Get the directory for working files
-			DirectoryReference BaseDir = DirectoryReference.FromFile(TargetDesc.ProjectFile) ?? UnrealBuildTool.EngineDirectory;
-			DirectoryReference TempDir = DirectoryReference.Combine(BaseDir, "Intermediate", "Remote", TargetDesc.Name, TargetDesc.Platform.ToString(), TargetDesc.Configuration.ToString());
-			DirectoryReference.CreateDirectory(TempDir);
+			DirectoryReference TempDir = CreateTempDirectory(TargetDesc);
 
 			bool bLogIsMapped = false;
 			foreach (RemoteMapping Mapping in Mappings)
@@ -378,55 +414,10 @@
 			FileReference RemoteManifestFile = FileReference.Combine(TempDir, "Manifest.xml");
 
 			// Prepare the arguments we will pass to the remote build
-			List<string> RemoteArguments = new List<string>();
-			RemoteArguments.Add(TargetDesc.Name);
-			RemoteArguments.Add(TargetDesc.Platform.ToString());
-			RemoteArguments.Add(TargetDesc.Configuration.ToString());
-			RemoteArguments.Add("-SkipRulesCompile"); // Use the rules assembly built locally
-			RemoteArguments.Add(String.Format("-XmlConfigCache={0}", GetRemotePath(XmlConfig.CacheFile))); // Use the XML config cache built locally, since the remote won't have it
+			List<string> RemoteArguments = GetRemoteArgumentsForTarget(TargetDesc, LocalManifestFiles);
 			RemoteArguments.Add(String.Format("-Log={0}", GetRemotePath(RemoteLogFile)));
 			RemoteArguments.Add(String.Format("-Manifest={0}", GetRemotePath(RemoteManifestFile)));
 
-			if(TargetDesc.ProjectFile != null)
-			{
-				RemoteArguments.Add(String.Format("-Project={0}", GetRemotePath(TargetDesc.ProjectFile)));
-			}
-
-			foreach(string LocalArgument in TargetDesc.AdditionalArguments)
-			{
-				int EqualsIdx = LocalArgument.IndexOf('=');
-				if(EqualsIdx == -1)
-				{
-					RemoteArguments.Add(LocalArgument);
-					continue;
-				}
-
-				string Key = LocalArgument.Substring(0, EqualsIdx);
-				string Value = LocalArgument.Substring(EqualsIdx + 1);
-
-				if(Key.Equals("-Log", StringComparison.InvariantCultureIgnoreCase))
-				{
-					// We are already writing to the local log file. The remote will produce a different log (RemoteLogFile)
-					continue;
-				}
-				if(Key.Equals("-Manifest", StringComparison.InvariantCultureIgnoreCase))
-				{
-					LocalManifestFiles.Add(new FileReference(Value));
-					continue;
-				}
-
-				string RemoteArgument = LocalArgument;
-				foreach(RemoteMapping Mapping in Mappings)
-				{
-					if(Value.StartsWith(Mapping.LocalDirectory.FullName, StringComparison.InvariantCultureIgnoreCase))
-					{
-						RemoteArgument = String.Format("{0}={1}", Key, GetRemotePath(Value));
-						break;
-					}
-				}
-				RemoteArguments.Add(RemoteArgument);
-			}
-
 			// Handle any per-platform setup that is required
 			if(TargetDesc.Platform == UnrealTargetPlatform.IOS || TargetDesc.Platform == UnrealTargetPlatform.TVOS)
 			{
@@ -503,9 +494,6 @@
 			// Upload the workspace files
 			UploadWorkspace(TempDir);
 
-			// Fixup permissions on any shell scripts
-			Execute(RemoteBaseDir, String.Format("chmod +x {0}/Build/BatchFiles/Mac/*.sh", EscapeShellArgument(GetRemotePath(UnrealBuildTool.EngineDirectory))));
-
 			// Execute the compile
 			Log.TraceInformation("[Remote] Executing build");
 
@@ -562,5 +550,75 @@
 		}
 
 		/// <summary>
+		/// Creates a temporary directory for the given target
+		/// </summary>
+		/// <param name="TargetDesc">The target descriptor</param>
+		/// <returns>Directory to use for temporary files</returns>
+		static DirectoryReference CreateTempDirectory(TargetDescriptor TargetDesc)
+		{
+			DirectoryReference BaseDir = DirectoryReference.FromFile(TargetDesc.ProjectFile) ?? UnrealBuildTool.EngineDirectory;
+			DirectoryReference TempDir = DirectoryReference.Combine(BaseDir, "Intermediate", "Remote", TargetDesc.Name, TargetDesc.Platform.ToString(), TargetDesc.Configuration.ToString());
+			DirectoryReference.CreateDirectory(TempDir);
+			return TempDir;
+		}
+
+		/// <summary>
+		/// Translate the arguments for a target descriptor for the remote machine
+		/// </summary>
+		/// <param name="TargetDesc">The target descriptor</param>
+		/// <param name="LocalManifestFiles">Manifest files to be output from this target</param>
+		/// <return>List of remote arguments</return>
+		List<string> GetRemoteArgumentsForTarget(TargetDescriptor TargetDesc, List<FileReference> LocalManifestFiles)
+		{
+			List<string> RemoteArguments = new List<string>();
+			RemoteArguments.Add(TargetDesc.Name);
+			RemoteArguments.Add(TargetDesc.Platform.ToString());
+			RemoteArguments.Add(TargetDesc.Configuration.ToString());
+			RemoteArguments.Add("-SkipRulesCompile"); // Use the rules assembly built locally
+			RemoteArguments.Add(String.Format("-XmlConfigCache={0}", GetRemotePath(XmlConfig.CacheFile))); // Use the XML config cache built locally, since the remote won't have it
+
+			if(TargetDesc.ProjectFile != null)
+			{
+				RemoteArguments.Add(String.Format("-Project={0}", GetRemotePath(TargetDesc.ProjectFile)));
+			}
+
+			foreach(string LocalArgument in TargetDesc.AdditionalArguments)
+			{
+				int EqualsIdx = LocalArgument.IndexOf('=');
+				if(EqualsIdx == -1)
+				{
+					RemoteArguments.Add(LocalArgument);
+					continue;
+				}
+
+				string Key = LocalArgument.Substring(0, EqualsIdx);
+				string Value = LocalArgument.Substring(EqualsIdx + 1);
+
+				if(Key.Equals("-Log", StringComparison.InvariantCultureIgnoreCase))
+				{
+					// We are already writing to the local log file. The remote will produce a different log (RemoteLogFile)
+					continue;
+				}
+				if(Key.Equals("-Manifest", StringComparison.InvariantCultureIgnoreCase) && LocalManifestFiles != null)
+				{
+					LocalManifestFiles.Add(new FileReference(Value));
+					continue;
+				}
+
+				string RemoteArgument = LocalArgument;
+				foreach(RemoteMapping Mapping in Mappings)
+				{
+					if(Value.StartsWith(Mapping.LocalDirectory.FullName, StringComparison.InvariantCultureIgnoreCase))
+					{
+						RemoteArgument = String.Format("{0}={1}", Key, GetRemotePath(Value));
+						break;
+					}
+				}
+				RemoteArguments.Add(RemoteArgument);
+			}
+			return RemoteArguments;
+		}
+
+		/// <summary>
 		/// Runs the actool utility on a directory to create an Assets.car file
 		/// </summary>
@@ -823,6 +881,9 @@
 				Log.TraceInformation("[Remote] Uploading project files...");
 				UploadDirectory(ProjectFile.Directory, GetRemotePath(ProjectFile.Directory), ProjectFilters);
 			}
+
+			// Fixup permissions on any shell scripts
+			Execute(RemoteBaseDir, String.Format("chmod +x {0}/Build/BatchFiles/Mac/*.sh", EscapeShellArgument(GetRemotePath(UnrealBuildTool.EngineDirectory))));
 		}
 
 		/// <summary>
diff --git a/Engine/Source/Programs/UnrealBuildTool/UnrealBuildTool.cs b/Engine/Source/Programs/UnrealBuildTool/UnrealBuildTool.cs
--- a/Engine/Source/Programs/UnrealBuildTool/UnrealBuildTool.cs
+++ b/Engine/Source/Programs/UnrealBuildTool/UnrealBuildTool.cs
@@ -339,6 +339,7 @@ namespace UnrealBuildTool
 			/// The mode to execute
 			/// </summary>
 			[CommandLine]
+			[CommandLine("-Clean", Value="Clean")]
 			[CommandLine("-ProjectFiles", Value="GenerateProjectFiles")]
 			[CommandLine("-ProjectFileFormat=", Value="GenerateProjectFiles")]
 			[CommandLine("-Makefile", Value="GenerateProjectFiles")]
@@ -538,32 +539,29 @@ namespace UnrealBuildTool
 				}
 
 				// Handle remote builds
-				if(BuildHostPlatform.Current.Platform == UnrealTargetPlatform.Win64)
+				for(int Idx = 0; Idx < TargetDescs.Count; Idx++)
 				{
-					for(int Idx = 0; Idx < TargetDescs.Count; Idx++)
+					TargetDescriptor TargetDesc = TargetDescs[Idx];
+					if(RemoteMac.HandlesTargetPlatform(TargetDesc.Platform))
 					{
-						TargetDescriptor TargetDesc = TargetDescs[Idx];
-						if(TargetDesc.Platform == UnrealTargetPlatform.Mac || TargetDesc.Platform == UnrealTargetPlatform.IOS || TargetDesc.Platform == UnrealTargetPlatform.TVOS)
+						FileReference RemoteLogFile = null;
+						if(LogFile != null)
 						{
-							FileReference RemoteLogFile = null;
-							if(LogFile != null)
-							{
-								RemoteLogFile = FileReference.Combine(LogFile.Directory, LogFile.GetFileNameWithoutExtension() + "_Remote.txt");
-							}
+							RemoteLogFile = FileReference.Combine(LogFile.Directory, LogFile.GetFileNameWithoutExtension() + "_Remote.txt");
+						}
 
-							RemoteMac RemoteMac = new RemoteMac(TargetDesc.ProjectFile);
-							if(!RemoteMac.Build(TargetDesc, LogFile))
-							{
-								return ECompilationResult.Unknown;
-							}
+						RemoteMac RemoteMac = new RemoteMac(TargetDesc.ProjectFile);
+						if(!RemoteMac.Build(TargetDesc, LogFile))
+						{
+							return ECompilationResult.Unknown;
+						}
 
-							TargetDescs.RemoveAt(Idx--);
+						TargetDescs.RemoveAt(Idx--);
+						if(TargetDescs.Count == 0)
+						{
+							return ECompilationResult.Succeeded;
 						}
 					}
-					if(TargetDescs.Count == 0)
-					{
-						return ECompilationResult.Succeeded;
-					}
 				}
 
 				if (Arguments.Any(x => x.Equals("-InvalidateMakefilesOnly", StringComparison.InvariantCultureIgnoreCase)))
@@ -693,7 +691,7 @@ namespace UnrealBuildTool
 						FileReference UBTMakefilePath = UBTMakefile.GetUBTMakefilePath(TargetDescs);
 
 						// Make sure the gather phase is executed if we're not actually building anything
-						if (BuildConfiguration.bGenerateManifest || BuildConfiguration.bCleanProject || BuildConfiguration.bXGEExport || GeneratingActionGraph)
+						if (BuildConfiguration.bGenerateManifest || BuildConfiguration.bXGEExport || GeneratingActionGraph)
 						{
 							bIsGatheringBuild = true;
 						}
@@ -742,10 +740,6 @@ namespace UnrealBuildTool
 							using(Timeline.ScopeEvent("UEBuildTarget.CreateTarget()"))
 							{
 								UEBuildTarget Target = UEBuildTarget.CreateTarget(TargetDesc, bSkipRulesCompile, BuildConfiguration.SingleFileToCompile != null, BuildConfiguration.bUsePrecompiled);
-								if ((Target == null) && (BuildConfiguration.bCleanProject))
-								{
-									continue;
-								}
 								Targets.Add(Target);
 							}
 						}
@@ -794,7 +788,7 @@ namespace UnrealBuildTool
 								{
 									if (!bNeedsFullCPPIncludeRescan)
 									{
-										if (!BuildConfiguration.bXGEExport && !BuildConfiguration.bGenerateManifest && !BuildConfiguration.bCleanProject)
+										if (!BuildConfiguration.bXGEExport && !BuildConfiguration.bGenerateManifest)
 										{
 											bNeedsFullCPPIncludeRescan = true;
 											Log.TraceInformation("Performing full C++ include scan (no include cache file)");
@@ -811,14 +805,9 @@ namespace UnrealBuildTool
 							List<UHTModuleInfo> TargetUObjectModules = new List<UHTModuleInfo>();
 							Dictionary<string, FileItem[]> TargetModuleNameToOutputItems = new Dictionary<string, FileItem[]>(StringComparer.OrdinalIgnoreCase);
 							BuildPredicateStore Predicates = new BuildPredicateStore();
-							if (BuildConfiguration.bCleanProject)
-							{
-								BuildResult = Target.Clean(!BuildConfiguration.bDoNotBuildUHT);
-							}
-							else
-							{
-								BuildResult = Target.Build(BuildConfiguration, TargetToHeaders[Target], TargetOutputItems, TargetModuleNameToOutputItems, TargetUObjectModules, WorkingSet, ActionGraph, Predicates, bIsAssemblingBuild);
-							}
+
+							BuildResult = Target.Build(BuildConfiguration, TargetToHeaders[Target], TargetOutputItems, TargetModuleNameToOutputItems, TargetUObjectModules, WorkingSet, ActionGraph, Predicates, bIsAssemblingBuild);
+							
 							if (BuildResult != ECompilationResult.Succeeded)
 							{
 								break;
@@ -860,7 +849,7 @@ namespace UnrealBuildTool
 					if (BuildResult == ECompilationResult.Succeeded && !BuildConfiguration.bSkipBuild &&
 						(
 							(BuildConfiguration.bXGEExport && BuildConfiguration.bGenerateManifest) ||
-							(!GeneratingActionGraph && !BuildConfiguration.bGenerateManifest && !BuildConfiguration.bCleanProject)
+							(!GeneratingActionGraph && !BuildConfiguration.bGenerateManifest)
 						))
 					{
 						if (bIsGatheringBuild)
@@ -945,11 +934,6 @@ namespace UnrealBuildTool
 							{
 								HotReloadMode = HotReloadMode.FromEditor;
 							}
-
-							if (HotReloadMode != HotReloadMode.Disabled && BuildConfiguration.bCleanProject)
-							{
-								throw new BuildException("Unable to clean target while hot-reloading. Close the editor and try again.");
-							}
 						}
 						TargetDescriptor HotReloadTargetDesc = (HotReloadMode != HotReloadMode.Disabled) ? TargetDescs[0] : null;
 
@@ -1255,7 +1239,6 @@ namespace UnrealBuildTool
 							if (BuildResult.Succeeded()
 								&& BuildConfiguration.SingleFileToCompile == null
 								&& !BuildConfiguration.bGenerateManifest
-								&& !BuildConfiguration.bCleanProject
 								&& !BuildConfiguration.bXGEExport)
 							{
 								foreach (UEBuildTarget Target in Targets)
diff --git a/Engine/Source/Programs/UnrealBuildTool/UnrealBuildTool.csproj b/Engine/Source/Programs/UnrealBuildTool/UnrealBuildTool.csproj
--- a/Engine/Source/Programs/UnrealBuildTool/UnrealBuildTool.csproj
+++ b/Engine/Source/Programs/UnrealBuildTool/UnrealBuildTool.csproj
@@ -141,6 +141,7 @@
   </Choose>
   <ItemGroup>
     <Compile Include="Modes\BuildMode.cs" />
+    <Compile Include="Modes\CleanMode.cs" />
     <Compile Include="Modes\DeployMode.cs" />
     <Compile Include="Modes\GenerateProjectFilesMode.cs" />
     <Compile Include="Modes\SetupPlatformsMode.cs" />
