UBT: Fix exception in debug visualizer for TargetDescriptor when arguments is null.

#rnx

diff --git a/Engine/Source/Programs/UnrealBuildTool/Configuration/TargetDescriptor.cs b/Engine/Source/Programs/UnrealBuildTool/Configuration/TargetDescriptor.cs
--- a/Engine/Source/Programs/UnrealBuildTool/Configuration/TargetDescriptor.cs
+++ b/Engine/Source/Programs/UnrealBuildTool/Configuration/TargetDescriptor.cs
@@ -327,7 +327,7 @@ namespace UnrealBuildTool
 			{
 				Result.AppendFormat(" -Project={0}", Utils.MakePathSafeToUseWithCommandLine(ProjectFile));
 			}
-			if(AdditionalArguments.Count > 0)
+			if(AdditionalArguments != null && AdditionalArguments.Count > 0)
 			{
 				Result.AppendFormat(" {0}", AdditionalArguments);
 			}
