Add missing copyright boilerplate.

#rnx

diff --git a/Engine/Source/Programs/UnrealBuildTool/System/Timeline.cs b/Engine/Source/Programs/UnrealBuildTool/System/Timeline.cs
--- a/Engine/Source/Programs/UnrealBuildTool/System/Timeline.cs
+++ b/Engine/Source/Programs/UnrealBuildTool/System/Timeline.cs
@@ -1,3 +1,5 @@
+// Copyright 1998-2018 Epic Games, Inc. All Rights Reserved.
+
 using System;
 using System.Collections.Generic;
 using System.Diagnostics;
