UBT: Move all the external dependencies for a build onto the BuildPredicateStore class, and simplify the role of the FlatModuleCsDataType class (now renamed to UHTModuleHeaderInfo).

#rnx

diff --git a/Engine/Source/Programs/UnrealBuildTool/Configuration/UEBuildTarget.cs b/Engine/Source/Programs/UnrealBuildTool/Configuration/UEBuildTarget.cs
--- a/Engine/Source/Programs/UnrealBuildTool/Configuration/UEBuildTarget.cs
+++ b/Engine/Source/Programs/UnrealBuildTool/Configuration/UEBuildTarget.cs
@@ -253,40 +253,6 @@ namespace UnrealBuildTool
 		}
 	}
 
-	class FlatModuleCsDataType
-	{
-		public FlatModuleCsDataType(BinaryReader Reader)
-		{
-			ModuleName = Reader.ReadString();
-			BuildCsFilename = Reader.ReadString();
-			ModuleSourceFolder = Reader.ReadDirectoryReference();
-			ExternalDependencies = Reader.ReadStringList();
-			UHTHeaderNames = Reader.ReadStringList();
-		}
-
-		public void Write(BinaryWriter Writer)
-		{
-			Writer.Write(ModuleName);
-			Writer.Write(BuildCsFilename);
-			Writer.Write(ModuleSourceFolder);
-			Writer.Write(ExternalDependencies);
-			Writer.Write(UHTHeaderNames);
-		}
-
-		public FlatModuleCsDataType(string InModuleName, string InBuildCsFilename, IEnumerable<string> InExternalDependencies)
-		{
-			ModuleName = InModuleName;
-			BuildCsFilename = InBuildCsFilename;
-			ExternalDependencies = new List<string>(InExternalDependencies);
-		}
-
-		public string ModuleName;
-		public string BuildCsFilename;
-		public DirectoryReference ModuleSourceFolder;
-		public List<string> ExternalDependencies;
-		public List<string> UHTHeaderNames = new List<string>();
-	}
-
 	/// <summary>
 	/// A target that can be built
 	/// </summary>
@@ -660,11 +626,6 @@ namespace UnrealBuildTool
 		[NonSerialized]
 		private Dictionary<string, UEBuildModule> Modules = new Dictionary<string, UEBuildModule>(StringComparer.InvariantCultureIgnoreCase);
 
-		/// <summary>
-		/// Used to map names of modules to their .Build.cs filename
-		/// </summary>
-		public List<FlatModuleCsDataType> FlatModuleCsData = new List<FlatModuleCsDataType>();
-
 		/// <summary>
 		/// Filename for the receipt for this target.
 		/// </summary>
@@ -677,7 +638,7 @@ namespace UnrealBuildTool
 		/// <summary>
 		/// The name of the .Target.cs file, if the target was created with one
 		/// </summary>
-		public readonly FileReference TargetRulesFile;
+		readonly FileReference TargetRulesFile;
 
 		/// <summary>
 		/// List of scripts to run before building
@@ -694,15 +655,6 @@ namespace UnrealBuildTool
 		/// </summary>
 		public bool bDeployAfterCompile;
 
-		/// <summary>
-		/// A list of the module filenames which were used to build this target.
-		/// </summary>
-		/// <returns></returns>
-		public IEnumerable<string> GetAllModuleBuildCsFilenames()
-		{
-			return FlatModuleCsData.Select(Data => Data.BuildCsFilename);
-		}
-
 		/// <summary>
 		/// Whether this target should be compiled in monolithic mode
 		/// </summary>
@@ -730,9 +682,7 @@ namespace UnrealBuildTool
 			VersionFile = Reader.ReadFileReference();
 			bPrecompile = Reader.ReadBoolean();
 			bCompileMonolithic = Reader.ReadBoolean();
-			FlatModuleCsData = Reader.ReadList(() => new FlatModuleCsDataType(Reader));
 			ReceiptFileName = Reader.ReadFileReference();
-			TargetRulesFile = Reader.ReadFileReference();
 			PreBuildStepScripts = Reader.ReadArray(() => Reader.ReadFileReference());
 			PostBuildStepScripts = Reader.ReadArray(() => Reader.ReadFileReference());
 			bDeployAfterCompile = Reader.ReadBoolean();
@@ -757,9 +707,7 @@ namespace UnrealBuildTool
 			Writer.Write(VersionFile);
 			Writer.Write(bPrecompile);
 			Writer.Write(bCompileMonolithic);
-			Writer.Write(FlatModuleCsData, x => x.Write(Writer));
 			Writer.Write(ReceiptFileName);
-			Writer.Write(TargetRulesFile);
 			Writer.Write(PreBuildStepScripts, x => Writer.Write(x));
 			Writer.Write(PostBuildStepScripts, x => Writer.Write(x));
 			Writer.Write(bDeployAfterCompile);
@@ -1477,16 +1425,9 @@ namespace UnrealBuildTool
 
 			// Generate headers
 			HashSet<UEBuildModuleCPP> ModulesToGenerateHeadersFor = GatherDependencyModules(OriginalBinaries.ToList());
-
-			Dictionary<string, FlatModuleCsDataType> NameToFlatModuleData = new Dictionary<string, FlatModuleCsDataType>(StringComparer.InvariantCultureIgnoreCase);
-			foreach(FlatModuleCsDataType FlatModuleData in FlatModuleCsData)
-			{
-				NameToFlatModuleData[FlatModuleData.ModuleName] = FlatModuleData;
-			}
-
 			using(Timeline.ScopeEvent("ExternalExecution.SetupUObjectModules()"))
 			{
-				ExternalExecution.SetupUObjectModules(ModulesToGenerateHeadersFor, Rules.Platform, ProjectDescriptor, UObjectModules, NameToFlatModuleData, Rules.GeneratedCodeVersion, bIsAssemblingBuild);
+				ExternalExecution.SetupUObjectModules(ModulesToGenerateHeadersFor, Rules.Platform, ProjectDescriptor, UObjectModules, Predicates.UObjectModuleHeaders, Rules.GeneratedCodeVersion, bIsAssemblingBuild);
 			}
 
 			// NOTE: Even in Gather mode, we need to run UHT to make sure the files exist for the static action graph to be setup correctly.  This is because UHT generates .cpp
@@ -1716,6 +1657,18 @@ namespace UnrealBuildTool
 				CheckForEULAViolation();
 			}
 
+			// Add all the input files to the predicate store
+			Predicates.AdditionalDependencies.Add(FileItem.GetItemByFileReference(TargetRulesFile));
+			foreach(UEBuildModule Module in Modules.Values)
+			{
+				Predicates.AdditionalDependencies.Add(FileItem.GetItemByFileReference(Module.RulesFile));
+				foreach(string ExternalDependency in Module.Rules.ExternalDependencies)
+				{
+					FileReference Location = FileReference.Combine(Module.RulesFile.Directory, ExternalDependency);
+					Predicates.AdditionalDependencies.Add(FileItem.GetItemByFileReference(Location));
+				}
+			}
+
 			// Clean any stale modules which exist in multiple output directories. This can lead to the wrong DLL being loaded on Windows.
 			CleanStaleModules();
 			return ECompilationResult.Succeeded;
@@ -3543,7 +3496,6 @@ namespace UnrealBuildTool
 				// Now, go ahead and create the module builder instance
 				Module = InstantiateModule(RulesObject, GeneratedCodeDirectory);
 				Modules.Add(Module.Name, Module);
-				FlatModuleCsData.Add(new FlatModuleCsDataType(Module.Name, (Module.RulesFile == null) ? null : Module.RulesFile.FullName, RulesObject.ExternalDependencies));
 			}
 			return Module;
 		}
diff --git a/Engine/Source/Programs/UnrealBuildTool/System/BuildPredicateStore.cs b/Engine/Source/Programs/UnrealBuildTool/System/BuildPredicateStore.cs
--- a/Engine/Source/Programs/UnrealBuildTool/System/BuildPredicateStore.cs
+++ b/Engine/Source/Programs/UnrealBuildTool/System/BuildPredicateStore.cs
@@ -31,6 +31,16 @@ namespace UnrealBuildTool
 		/// </summary>
 		public HashSet<FileItem> CandidatesForWorkingSet = new HashSet<FileItem>();
 
+		/// <summary>
+		/// Additional files which are required 
+		/// </summary>
+		public HashSet<FileItem> AdditionalDependencies = new HashSet<FileItem>(); 
+
+		/// <summary>
+		/// Used to map names of modules to their .Build.cs filename
+		/// </summary>
+		public List<UHTModuleHeaderInfo> UObjectModuleHeaders = new List<UHTModuleHeaderInfo>();
+
 		public BuildPredicateStore()
 		{
 		}
@@ -40,6 +50,8 @@ namespace UnrealBuildTool
 			SourceDirectories = new HashSet<DirectoryReference>(Reader.ReadArray(() => Reader.ReadDirectoryReference()));
 			WorkingSet = new HashSet<FileItem>(Reader.ReadFileItemList(UniqueFileItems));
 			CandidatesForWorkingSet = new HashSet<FileItem>(Reader.ReadFileItemList(UniqueFileItems));
+			AdditionalDependencies = new HashSet<FileItem>(Reader.ReadFileItemList(UniqueFileItems));
+			UObjectModuleHeaders = Reader.ReadList(() => new UHTModuleHeaderInfo(Reader, UniqueFileItems));
 		}
 
 		public void Write(BinaryWriter Writer, Dictionary<FileItem, int> UniqueFileItemToIndex)
@@ -47,6 +59,8 @@ namespace UnrealBuildTool
 			Writer.Write(SourceDirectories.ToArray(), x => Writer.Write(x));
 			Writer.Write(WorkingSet.ToList(), UniqueFileItemToIndex);
 			Writer.Write(CandidatesForWorkingSet.ToList(), UniqueFileItemToIndex);
+			Writer.Write(AdditionalDependencies.ToList(), UniqueFileItemToIndex);
+			Writer.Write(UObjectModuleHeaders, x => x.Write(Writer, UniqueFileItemToIndex));
 		}
 	}
 }
diff --git a/Engine/Source/Programs/UnrealBuildTool/System/ExternalExecution.cs b/Engine/Source/Programs/UnrealBuildTool/System/ExternalExecution.cs
--- a/Engine/Source/Programs/UnrealBuildTool/System/ExternalExecution.cs
+++ b/Engine/Source/Programs/UnrealBuildTool/System/ExternalExecution.cs
@@ -322,6 +322,29 @@ namespace UnrealBuildTool
 		}
 	}
 
+	class UHTModuleHeaderInfo
+	{
+		public DirectoryReference SourceFolder;
+		public List<FileItem> HeaderFiles;
+
+		public UHTModuleHeaderInfo(DirectoryReference SourceFolder, List<FileItem> HeaderFiles)
+		{
+			this.SourceFolder = SourceFolder;
+			this.HeaderFiles = HeaderFiles;
+		}
+
+		public UHTModuleHeaderInfo(BinaryReader Reader, List<FileItem> UniqueFileItems)
+		{
+			SourceFolder = Reader.ReadDirectoryReference();
+			HeaderFiles = Reader.ReadFileItemList(UniqueFileItems);
+		}
+
+		public void Write(BinaryWriter Writer, Dictionary<FileItem, int> UniqueFileItemToIndex)
+		{
+			Writer.Write(SourceFolder);
+			Writer.Write(HeaderFiles, UniqueFileItemToIndex);
+		}
+	}
 
 	/// <summary>
 	/// This handles all running of the UnrealHeaderTool
@@ -586,7 +609,7 @@ namespace UnrealBuildTool
 			}
 		}
 
-		public static void SetupUObjectModules(IEnumerable<UEBuildModuleCPP> ModulesToGenerateHeadersFor, UnrealTargetPlatform Platform, ProjectDescriptor ProjectDescriptor, List<UHTModuleInfo> UObjectModules, Dictionary<string, FlatModuleCsDataType> FlatModuleCsData, EGeneratedCodeVersion GeneratedCodeVersion, bool bIsAssemblingBuild)
+		public static void SetupUObjectModules(IEnumerable<UEBuildModuleCPP> ModulesToGenerateHeadersFor, UnrealTargetPlatform Platform, ProjectDescriptor ProjectDescriptor, List<UHTModuleInfo> UObjectModules, List<UHTModuleHeaderInfo> UObjectModuleHeaders, EGeneratedCodeVersion GeneratedCodeVersion, bool bIsAssemblingBuild)
 		{
 			// Find the type of each module
 			Dictionary<UEBuildModuleCPP, UHTModuleType> ModuleToType = new Dictionary<UEBuildModuleCPP, UHTModuleType>();
@@ -620,8 +643,13 @@ namespace UnrealBuildTool
 					}
 
 					UObjectModules.Add(Info);
-					FlatModuleCsData[Module.Name].ModuleSourceFolder = Module.ModuleDirectory;
-					FlatModuleCsData[Module.Name].UHTHeaderNames = Info.PublicUObjectHeaders.Concat(Info.PublicUObjectClassesHeaders).Concat(Info.PrivateUObjectHeaders).Select(x => x.AbsolutePath).ToList();
+
+					List<FileItem> ReflectedHeaderFiles = new List<FileItem>();
+					ReflectedHeaderFiles.AddRange(Info.PublicUObjectHeaders);
+					ReflectedHeaderFiles.AddRange(Info.PublicUObjectClassesHeaders);
+					ReflectedHeaderFiles.AddRange(Info.PrivateUObjectHeaders);
+					UObjectModuleHeaders.Add(new UHTModuleHeaderInfo(Module.ModuleDirectory, ReflectedHeaderFiles));
+
 					Log.TraceVerbose("Detected UObject module: " + Info.ModuleName);
 				}
 				else
diff --git a/Engine/Source/Programs/UnrealBuildTool/System/UBTMakefile.cs b/Engine/Source/Programs/UnrealBuildTool/System/UBTMakefile.cs
--- a/Engine/Source/Programs/UnrealBuildTool/System/UBTMakefile.cs
+++ b/Engine/Source/Programs/UnrealBuildTool/System/UBTMakefile.cs
@@ -323,55 +323,6 @@ namespace UnrealBuildTool
 			// Check if any of the target's Build.cs files are newer than the makefile
 			foreach (UEBuildTarget Target in LoadedUBTMakefile.Targets)
 			{
-				string TargetCsFilename = Target.TargetRulesFile.FullName;
-				if (TargetCsFilename != null)
-				{
-					FileInfo TargetCsFile = new FileInfo(TargetCsFilename);
-					bool bTargetCsFileExists = TargetCsFile.Exists;
-					if (!bTargetCsFileExists || TargetCsFile.LastWriteTime > UBTMakefileInfo.LastWriteTime)
-					{
-						Log.TraceLog("{0} has been {1} since makefile was built, ignoring it ({2})", TargetCsFilename, bTargetCsFileExists ? "changed" : "deleted", UBTMakefileInfo.FullName);
-						ReasonNotLoaded = string.Format("changes to target files");
-						return null;
-					}
-				}
-
-				IEnumerable<string> BuildCsFilenames = Target.GetAllModuleBuildCsFilenames();
-				foreach (string BuildCsFilename in BuildCsFilenames)
-				{
-					if (BuildCsFilename != null)
-					{
-						FileInfo BuildCsFile = new FileInfo(BuildCsFilename);
-						bool bBuildCsFileExists = BuildCsFile.Exists;
-						if (!bBuildCsFileExists || BuildCsFile.LastWriteTime > UBTMakefileInfo.LastWriteTime)
-						{
-							Log.TraceLog("{0} has been {1} since makefile was built, ignoring it ({2})", BuildCsFilename, bBuildCsFileExists ? "changed" : "deleted", UBTMakefileInfo.FullName);
-							ReasonNotLoaded = string.Format("changes to module files");
-							return null;
-						}
-					}
-				}
-
-				foreach (FlatModuleCsDataType FlatCsModuleData in Target.FlatModuleCsData)
-				{
-					if (FlatCsModuleData.BuildCsFilename != null && FlatCsModuleData.ExternalDependencies.Count > 0)
-					{
-						string BaseDir = Path.GetDirectoryName(FlatCsModuleData.BuildCsFilename);
-						foreach (string ExternalDependency in FlatCsModuleData.ExternalDependencies)
-						{
-							FileInfo DependencyFile = new FileInfo(Path.Combine(BaseDir, ExternalDependency));
-							bool bDependencyFileExists = DependencyFile.Exists;
-							if (!bDependencyFileExists || DependencyFile.LastWriteTime > UBTMakefileInfo.LastWriteTime)
-							{
-								Log.TraceLog("{0} has been {1} since makefile was built, ignoring it ({2})", DependencyFile.FullName, bDependencyFileExists ? "changed" : "deleted", UBTMakefileInfo.FullName);
-								ReasonNotLoaded = string.Format("changes to external dependency");
-								return null;
-							}
-						}
-					}
-				}
-
-
 				// Check if ini files are newer. Ini files contain build settings too.
 				DirectoryReference ProjectDirectory = DirectoryReference.FromFile(ProjectFile);
 				foreach (ConfigHierarchyType IniType in (ConfigHierarchyType[])Enum.GetValues(typeof(ConfigHierarchyType)))
@@ -401,30 +352,44 @@ namespace UnrealBuildTool
 						return null;
 					}
 				}
-			}
 
-			// We do a check to see if any modules' headers have changed which have
-			// acquired or lost UHT types.  If so, which should be rare,
-			// we'll just invalidate the entire makefile and force it to be rebuilt.
-			foreach (UEBuildTarget Target in LoadedUBTMakefile.Targets)
-			{
+				foreach(FileItem AdditionalDependency in Predicates.AdditionalDependencies)
+				{
+					if (!AdditionalDependency.bExists)
+					{
+						Log.TraceLog("{0} has been deleted since makefile was built.", AdditionalDependency.Location);
+						ReasonNotLoaded = string.Format("{0} deleted", AdditionalDependency.Location.GetFileName());
+						return null;
+					}
+					if(AdditionalDependency.LastWriteTime > UBTMakefileInfo.LastWriteTime)
+					{
+						Log.TraceLog("{0} has been modified since makefile was built.", AdditionalDependency.Location);
+						ReasonNotLoaded = string.Format("{0} modified", AdditionalDependency.Location.GetFileName());
+						return null;
+					}
+				}
+
+				// We do a check to see if any modules' headers have changed which have
+				// acquired or lost UHT types.  If so, which should be rare,
+				// we'll just invalidate the entire makefile and force it to be rebuilt.
+
 				// Get all H files in processed modules newer than the makefile itself
-				HashSet<string> HFilesNewerThanMakefile =
-					new HashSet<string>(
-						Target.FlatModuleCsData
-						.SelectMany(x => x.ModuleSourceFolder != null ? Directory.EnumerateFiles(x.ModuleSourceFolder.FullName, "*.h", SearchOption.AllDirectories) : Enumerable.Empty<string>())
-						.Where(y => Directory.GetLastWriteTimeUtc(y) > UBTMakefileInfo.LastWriteTimeUtc)
+				HashSet<FileReference> HFilesNewerThanMakefile =
+					new HashSet<FileReference>(
+						Predicates.UObjectModuleHeaders
+						.SelectMany(x => x.SourceFolder != null ? DirectoryReference.EnumerateFiles(x.SourceFolder, "*.h", SearchOption.AllDirectories) : Enumerable.Empty<FileReference>())
+						.Where(y => FileReference.GetLastWriteTimeUtc(y) > UBTMakefileInfo.LastWriteTimeUtc)
 						.OrderBy(z => z).Distinct()
 					);
 
 				// Get all H files in all modules processed in the last makefile build
-				HashSet<string> AllUHTHeaders = new HashSet<string>(Target.FlatModuleCsData.SelectMany(x => x.UHTHeaderNames));
+				HashSet<FileReference> AllUHTHeaders = new HashSet<FileReference>(Predicates.UObjectModuleHeaders.SelectMany(x => x.HeaderFiles).Select(x => x.Location));
 
 				// Check whether any headers have been deleted. If they have, we need to regenerate the makefile since the module might now be empty. If we don't,
 				// and the file has been moved to a different module, we may include stale generated headers.
-				foreach (string FileName in AllUHTHeaders)
+				foreach (FileReference FileName in AllUHTHeaders)
 				{
-					if (!File.Exists(FileName))
+					if (!FileReference.Exists(FileName))
 					{
 						Log.TraceLog("File processed by UHT was deleted ({0}); invalidating makefile", FileName);
 						ReasonNotLoaded = string.Format("UHT file was deleted");
@@ -435,9 +400,9 @@ namespace UnrealBuildTool
 				// Makefile is invalid if:
 				// * There are any newer files which contain no UHT data, but were previously in the makefile
 				// * There are any newer files contain data which needs processing by UHT, but weren't not previously in the makefile
-				foreach (string Filename in HFilesNewerThanMakefile)
+				foreach (FileReference Filename in HFilesNewerThanMakefile)
 				{
-					bool bContainsUHTData = CPPHeaders.DoesFileContainUObjects(Filename);
+					bool bContainsUHTData = CPPHeaders.DoesFileContainUObjects(Filename.FullName);
 					bool bWasProcessed = AllUHTHeaders.Contains(Filename);
 					if (bContainsUHTData != bWasProcessed)
 					{
