UBT: Rename BuildPredicateStore to BuildPrerequisites.

#rnx

diff --git a/Engine/Source/Programs/UnrealBuildTool/Configuration/UEBuildBinary.cs b/Engine/Source/Programs/UnrealBuildTool/Configuration/UEBuildBinary.cs
--- a/Engine/Source/Programs/UnrealBuildTool/Configuration/UEBuildBinary.cs
+++ b/Engine/Source/Programs/UnrealBuildTool/Configuration/UEBuildBinary.cs
@@ -158,11 +158,11 @@ namespace UnrealBuildTool
 		/// <param name="LinkEnvironment">The environment to link the binary in</param>
 		/// <param name="SharedPCHs">List of templates for shared PCHs</param>
 		/// <param name="WorkingSet">The working set of source files</param>
-		/// <param name="Predicates">Predicates for the generated actions to be valid</param>
+		/// <param name="Prerequisites">Predicates for the generated actions to be valid</param>
 		/// <param name="ExeDir">Directory containing the output executable</param>
 		/// <param name="ActionGraph">Graph to add build actions to</param>
 		/// <returns>Set of built products</returns>
-		public List<FileItem> Build(ReadOnlyTargetRules Target, UEToolChain ToolChain, CppCompileEnvironment CompileEnvironment, LinkEnvironment LinkEnvironment, List<PrecompiledHeaderTemplate> SharedPCHs, ISourceFileWorkingSet WorkingSet, BuildPredicateStore Predicates, DirectoryReference ExeDir, ActionGraph ActionGraph)
+		public List<FileItem> Build(ReadOnlyTargetRules Target, UEToolChain ToolChain, CppCompileEnvironment CompileEnvironment, LinkEnvironment LinkEnvironment, List<PrecompiledHeaderTemplate> SharedPCHs, ISourceFileWorkingSet WorkingSet, BuildPrerequisites Prerequisites, DirectoryReference ExeDir, ActionGraph ActionGraph)
 		{
 			// Return nothing if we're using precompiled binaries. If we're not linking, we might want just one module to be compiled (eg. a foreign plugin), so allow any actions to run.
 			if (bUsePrecompiled && !Target.bDisableLinking)
@@ -171,7 +171,7 @@ namespace UnrealBuildTool
 			}
 
 			// Setup linking environment.
-			LinkEnvironment BinaryLinkEnvironment = SetupBinaryLinkEnvironment(Target, ToolChain, LinkEnvironment, CompileEnvironment, SharedPCHs, WorkingSet, Predicates, ExeDir, ActionGraph);
+			LinkEnvironment BinaryLinkEnvironment = SetupBinaryLinkEnvironment(Target, ToolChain, LinkEnvironment, CompileEnvironment, SharedPCHs, WorkingSet, Prerequisites, ExeDir, ActionGraph);
 
 			// If we're generating projects, we only need include paths and definitions, there is no need to run the linking logic.
 			if (ProjectFileGenerator.bGenerateProjectFiles)
@@ -650,7 +650,7 @@ namespace UnrealBuildTool
 			return BinaryCompileEnvironment;
 		}
 
-		private LinkEnvironment SetupBinaryLinkEnvironment(ReadOnlyTargetRules Target, UEToolChain ToolChain, LinkEnvironment LinkEnvironment, CppCompileEnvironment CompileEnvironment, List<PrecompiledHeaderTemplate> SharedPCHs, ISourceFileWorkingSet WorkingSet, BuildPredicateStore Predicates, DirectoryReference ExeDir, ActionGraph ActionGraph)
+		private LinkEnvironment SetupBinaryLinkEnvironment(ReadOnlyTargetRules Target, UEToolChain ToolChain, LinkEnvironment LinkEnvironment, CppCompileEnvironment CompileEnvironment, List<PrecompiledHeaderTemplate> SharedPCHs, ISourceFileWorkingSet WorkingSet, BuildPrerequisites Prerequisites, DirectoryReference ExeDir, ActionGraph ActionGraph)
 		{
 			LinkEnvironment BinaryLinkEnvironment = new LinkEnvironment(LinkEnvironment);
 			HashSet<UEBuildModule> LinkEnvironmentVisitedModules = new HashSet<UEBuildModule>();
@@ -665,7 +665,7 @@ namespace UnrealBuildTool
 				{
 					// Compile each module.
 					Log.TraceVerbose("Compile module: " + Module.Name);
-					LinkInputFiles = Module.Compile(Target, ToolChain, BinaryCompileEnvironment, SharedPCHs, WorkingSet, Predicates, ActionGraph);
+					LinkInputFiles = Module.Compile(Target, ToolChain, BinaryCompileEnvironment, SharedPCHs, WorkingSet, Prerequisites, ActionGraph);
 
 					// NOTE: Because of 'Shared PCHs', in monolithic builds the same PCH file may appear as a link input
 					// multiple times for a single binary.  We'll check for that here, and only add it once.  This avoids
diff --git a/Engine/Source/Programs/UnrealBuildTool/Configuration/UEBuildModule.cs b/Engine/Source/Programs/UnrealBuildTool/Configuration/UEBuildModule.cs
--- a/Engine/Source/Programs/UnrealBuildTool/Configuration/UEBuildModule.cs
+++ b/Engine/Source/Programs/UnrealBuildTool/Configuration/UEBuildModule.cs
@@ -733,7 +733,7 @@ namespace UnrealBuildTool
 		/// <summary>
 		/// Compiles the module, and returns a list of files output by the compiler.
 		/// </summary>
-		public abstract List<FileItem> Compile(ReadOnlyTargetRules Target, UEToolChain ToolChain, CppCompileEnvironment CompileEnvironment, List<PrecompiledHeaderTemplate> SharedPCHModules, ISourceFileWorkingSet WorkingSet, BuildPredicateStore Predicates, ActionGraph ActionGraph);
+		public abstract List<FileItem> Compile(ReadOnlyTargetRules Target, UEToolChain ToolChain, CppCompileEnvironment CompileEnvironment, List<PrecompiledHeaderTemplate> SharedPCHModules, ISourceFileWorkingSet WorkingSet, BuildPrerequisites Predicates, ActionGraph ActionGraph);
 
 		// Object interface.
 		public override string ToString()
diff --git a/Engine/Source/Programs/UnrealBuildTool/Configuration/UEBuildModuleCPP.cs b/Engine/Source/Programs/UnrealBuildTool/Configuration/UEBuildModuleCPP.cs
--- a/Engine/Source/Programs/UnrealBuildTool/Configuration/UEBuildModuleCPP.cs
+++ b/Engine/Source/Programs/UnrealBuildTool/Configuration/UEBuildModuleCPP.cs
@@ -294,7 +294,7 @@ namespace UnrealBuildTool
 		}
 
 		// UEBuildModule interface.
-		public override List<FileItem> Compile(ReadOnlyTargetRules Target, UEToolChain ToolChain, CppCompileEnvironment BinaryCompileEnvironment, List<PrecompiledHeaderTemplate> SharedPCHs, ISourceFileWorkingSet WorkingSet, BuildPredicateStore Predicates, ActionGraph ActionGraph)
+		public override List<FileItem> Compile(ReadOnlyTargetRules Target, UEToolChain ToolChain, CppCompileEnvironment BinaryCompileEnvironment, List<PrecompiledHeaderTemplate> SharedPCHs, ISourceFileWorkingSet WorkingSet, BuildPrerequisites Prerequisites, ActionGraph ActionGraph)
 		{
 			UEBuildPlatform BuildPlatform = UEBuildPlatform.GetBuildPlatformForCPPTargetPlatform(BinaryCompileEnvironment.Platform);
 
@@ -314,7 +314,7 @@ namespace UnrealBuildTool
 				return LinkInputFiles;
 			}
 
-			List<FileReference> SourceFilePaths = SourceFileSearch.FindModuleSourceFiles(ModuleRulesFile: RulesFile, SearchedDirectories: Predicates.SourceDirectories);
+			List<FileReference> SourceFilePaths = SourceFileSearch.FindModuleSourceFiles(ModuleRulesFile: RulesFile, SearchedDirectories: Prerequisites.SourceDirectories);
 			List<FileItem> SourceFiles = GetCPlusPlusFilesToBuild(SourceFilePaths, ModuleDirectory, Target.Platform);
 
 			SourceFilesClass SourceFilesFound = new SourceFilesClass();
@@ -472,7 +472,7 @@ namespace UnrealBuildTool
 			List<FileItem> CPPFilesToCompile = SourceFilesToBuild.CPPFiles;
 			if (bModuleUsesUnityBuild)
 			{
-				CPPFilesToCompile = Unity.GenerateUnityCPPs(Target, CPPFilesToCompile, CompileEnvironment, WorkingSet, Predicates, Rules.ShortName ?? Name, IntermediateDirectory);
+				CPPFilesToCompile = Unity.GenerateUnityCPPs(Target, CPPFilesToCompile, CompileEnvironment, WorkingSet, Prerequisites, Rules.ShortName ?? Name, IntermediateDirectory);
 				LinkInputFiles.AddRange(CompileUnityFilesWithToolChain(Target, ToolChain, CompileEnvironment, ModuleCompileEnvironment, CPPFilesToCompile, ActionGraph).ObjectFiles);
 			}
 			else
@@ -514,7 +514,7 @@ namespace UnrealBuildTool
 
 					if (bModuleUsesUnityBuild)
 					{
-						GeneratedFileItems = Unity.GenerateUnityCPPs(Target, GeneratedFileItems, GeneratedCPPCompileEnvironment, WorkingSet, Predicates, (Rules.ShortName ?? Name) + ".gen", IntermediateDirectory);
+						GeneratedFileItems = Unity.GenerateUnityCPPs(Target, GeneratedFileItems, GeneratedCPPCompileEnvironment, WorkingSet, Prerequisites, (Rules.ShortName ?? Name) + ".gen", IntermediateDirectory);
 						LinkInputFiles.AddRange(CompileUnityFilesWithToolChain(Target, ToolChain, GeneratedCPPCompileEnvironment, ModuleCompileEnvironment, GeneratedFileItems, ActionGraph).ObjectFiles);
 					}
 					else
diff --git a/Engine/Source/Programs/UnrealBuildTool/Configuration/UEBuildModuleExternal.cs b/Engine/Source/Programs/UnrealBuildTool/Configuration/UEBuildModuleExternal.cs
--- a/Engine/Source/Programs/UnrealBuildTool/Configuration/UEBuildModuleExternal.cs
+++ b/Engine/Source/Programs/UnrealBuildTool/Configuration/UEBuildModuleExternal.cs
@@ -20,7 +20,7 @@ namespace UnrealBuildTool
 		}
 
 		// UEBuildModule interface.
-		public override List<FileItem> Compile(ReadOnlyTargetRules Target, UEToolChain ToolChain, CppCompileEnvironment CompileEnvironment, List<PrecompiledHeaderTemplate> SharedPCHs, ISourceFileWorkingSet WorkingSet, BuildPredicateStore Predicates, ActionGraph ActionGraph)
+		public override List<FileItem> Compile(ReadOnlyTargetRules Target, UEToolChain ToolChain, CppCompileEnvironment CompileEnvironment, List<PrecompiledHeaderTemplate> SharedPCHs, ISourceFileWorkingSet WorkingSet, BuildPrerequisites Prerequisites, ActionGraph ActionGraph)
 		{
 			return new List<FileItem>();
 		}
diff --git a/Engine/Source/Programs/UnrealBuildTool/Configuration/UEBuildTarget.cs b/Engine/Source/Programs/UnrealBuildTool/Configuration/UEBuildTarget.cs
--- a/Engine/Source/Programs/UnrealBuildTool/Configuration/UEBuildTarget.cs
+++ b/Engine/Source/Programs/UnrealBuildTool/Configuration/UEBuildTarget.cs
@@ -1247,7 +1247,7 @@ namespace UnrealBuildTool
 		/// <summary>
 		/// Builds the target, appending list of output files and returns building result.
 		/// </summary>
-		public ECompilationResult Build(BuildConfiguration BuildConfiguration, CPPHeaders Headers, List<FileItem> OutputItems, Dictionary<string, FileItem[]> ModuleNameToOutputItems, List<UHTModuleInfo> UObjectModules, ISourceFileWorkingSet WorkingSet, ActionGraph ActionGraph, BuildPredicateStore Predicates, bool bIsAssemblingBuild)
+		public ECompilationResult Build(BuildConfiguration BuildConfiguration, CPPHeaders Headers, List<FileItem> OutputItems, Dictionary<string, FileItem[]> ModuleNameToOutputItems, List<UHTModuleInfo> UObjectModules, ISourceFileWorkingSet WorkingSet, ActionGraph ActionGraph, BuildPrerequisites Prerequisites, bool bIsAssemblingBuild)
 		{
 			CppPlatform CppPlatform = UEBuildPlatform.GetBuildPlatform(Platform).DefaultCppPlatform;
 			CppConfiguration CppConfiguration = GetCppConfiguration(Configuration);
@@ -1427,7 +1427,7 @@ namespace UnrealBuildTool
 			HashSet<UEBuildModuleCPP> ModulesToGenerateHeadersFor = GatherDependencyModules(OriginalBinaries.ToList());
 			using(Timeline.ScopeEvent("ExternalExecution.SetupUObjectModules()"))
 			{
-				ExternalExecution.SetupUObjectModules(ModulesToGenerateHeadersFor, Rules.Platform, ProjectDescriptor, UObjectModules, Predicates.UObjectModuleHeaders, Rules.GeneratedCodeVersion, bIsAssemblingBuild);
+				ExternalExecution.SetupUObjectModules(ModulesToGenerateHeadersFor, Rules.Platform, ProjectDescriptor, UObjectModules, Prerequisites.UObjectModuleHeaders, Rules.GeneratedCodeVersion, bIsAssemblingBuild);
 			}
 
 			// NOTE: Even in Gather mode, we need to run UHT to make sure the files exist for the static action graph to be setup correctly.  This is because UHT generates .cpp
@@ -1482,7 +1482,7 @@ namespace UnrealBuildTool
 			}
 			foreach (UEBuildBinary Binary in Binaries)
 			{
-				List<FileItem> BinaryOutputItems = Binary.Build(Rules, TargetToolChain, GlobalCompileEnvironment, GlobalLinkEnvironment, SharedPCHs, WorkingSet, Predicates, ExeDir, ActionGraph);
+				List<FileItem> BinaryOutputItems = Binary.Build(Rules, TargetToolChain, GlobalCompileEnvironment, GlobalLinkEnvironment, SharedPCHs, WorkingSet, Prerequisites, ExeDir, ActionGraph);
 				if(!bCompileMonolithic)
 				{
 					ModuleNameToOutputItems[Binary.PrimaryModule.Name] = BinaryOutputItems.ToArray();
@@ -1658,14 +1658,14 @@ namespace UnrealBuildTool
 			}
 
 			// Add all the input files to the predicate store
-			Predicates.AdditionalDependencies.Add(FileItem.GetItemByFileReference(TargetRulesFile));
+			Prerequisites.AdditionalDependencies.Add(FileItem.GetItemByFileReference(TargetRulesFile));
 			foreach(UEBuildModule Module in Modules.Values)
 			{
-				Predicates.AdditionalDependencies.Add(FileItem.GetItemByFileReference(Module.RulesFile));
+				Prerequisites.AdditionalDependencies.Add(FileItem.GetItemByFileReference(Module.RulesFile));
 				foreach(string ExternalDependency in Module.Rules.ExternalDependencies)
 				{
 					FileReference Location = FileReference.Combine(Module.RulesFile.Directory, ExternalDependency);
-					Predicates.AdditionalDependencies.Add(FileItem.GetItemByFileReference(Location));
+					Prerequisites.AdditionalDependencies.Add(FileItem.GetItemByFileReference(Location));
 				}
 			}
 
diff --git a/Engine/Source/Programs/UnrealBuildTool/System/BuildPredicateStore.cs b/Engine/Source/Programs/UnrealBuildTool/System/BuildPredicateStore.cs
deleted file mode 100644
diff --git a/Engine/Source/Programs/UnrealBuildTool/System/BuildPrerequisites.cs b/Engine/Source/Programs/UnrealBuildTool/System/BuildPrerequisites.cs
new file mode 100644
--- /dev/null
+++ b/Engine/Source/Programs/UnrealBuildTool/System/BuildPrerequisites.cs
@@ -0,0 +1,65 @@
+// Copyright 1998-2018 Epic Games, Inc. All Rights Reserved.
+
+using System;
+using System.Collections.Generic;
+using System.IO;
+using System.Linq;
+using System.Text;
+using System.Threading.Tasks;
+using Tools.DotNETCommon;
+
+namespace UnrealBuildTool
+{
+	/// <summary>
+	/// Contains predicates required for the action graph generated for a build to be valid. Unlike individual actions within the graph, these conditions must hold for the graph itself to be valid.
+	/// </summary>
+	class BuildPrerequisites
+	{
+		/// <summary>
+		/// Set of all source directories. Any files being added or removed from these directories will invalidate the makefile.
+		/// </summary>
+		public HashSet<DirectoryReference> SourceDirectories = new HashSet<DirectoryReference>();
+
+		/// <summary>
+		/// The set of source files that UnrealBuildTool determined to be part of the programmer's "working set". Used for adaptive non-unity builds.
+		/// </summary>
+		public HashSet<FileItem> WorkingSet = new HashSet<FileItem>();
+
+		/// <summary>
+		/// Set of files which are currently not part of the working set, but could be.
+		/// </summary>
+		public HashSet<FileItem> CandidatesForWorkingSet = new HashSet<FileItem>();
+
+		/// <summary>
+		/// Additional files which are required 
+		/// </summary>
+		public HashSet<FileItem> AdditionalDependencies = new HashSet<FileItem>(); 
+
+		/// <summary>
+		/// Used to map names of modules to their .Build.cs filename
+		/// </summary>
+		public List<UHTModuleHeaderInfo> UObjectModuleHeaders = new List<UHTModuleHeaderInfo>();
+
+		public BuildPrerequisites()
+		{
+		}
+
+		public BuildPrerequisites(BinaryReader Reader, List<FileItem> UniqueFileItems)
+		{
+			SourceDirectories = new HashSet<DirectoryReference>(Reader.ReadArray(() => Reader.ReadDirectoryReference()));
+			WorkingSet = new HashSet<FileItem>(Reader.ReadFileItemList(UniqueFileItems));
+			CandidatesForWorkingSet = new HashSet<FileItem>(Reader.ReadFileItemList(UniqueFileItems));
+			AdditionalDependencies = new HashSet<FileItem>(Reader.ReadFileItemList(UniqueFileItems));
+			UObjectModuleHeaders = Reader.ReadList(() => new UHTModuleHeaderInfo(Reader, UniqueFileItems));
+		}
+
+		public void Write(BinaryWriter Writer, Dictionary<FileItem, int> UniqueFileItemToIndex)
+		{
+			Writer.Write(SourceDirectories.ToArray(), x => Writer.Write(x));
+			Writer.Write(WorkingSet.ToList(), UniqueFileItemToIndex);
+			Writer.Write(CandidatesForWorkingSet.ToList(), UniqueFileItemToIndex);
+			Writer.Write(AdditionalDependencies.ToList(), UniqueFileItemToIndex);
+			Writer.Write(UObjectModuleHeaders, x => x.Write(Writer, UniqueFileItemToIndex));
+		}
+	}
+}
diff --git a/Engine/Source/Programs/UnrealBuildTool/System/UBTMakefile.cs b/Engine/Source/Programs/UnrealBuildTool/System/UBTMakefile.cs
--- a/Engine/Source/Programs/UnrealBuildTool/System/UBTMakefile.cs
+++ b/Engine/Source/Programs/UnrealBuildTool/System/UBTMakefile.cs
@@ -66,7 +66,7 @@ namespace UnrealBuildTool
 		/// <summary>
 		/// List of predicates for the action graph to be valid.
 		/// </summary>
-		public List<BuildPredicateStore> TargetBuildPredicates;
+		public List<BuildPrerequisites> TargetPrerequisites;
 
 		public UBTMakefile()
 		{
@@ -95,7 +95,7 @@ namespace UnrealBuildTool
 			ModuleNameToOutputItems = Reader.ReadDictionary(() => Reader.ReadString(), () => Reader.ReadArray(() => Reader.ReadFileItem(UniqueFileItems)));
 			HotReloadModuleNamesForAllTargets = new HashSet<string>(Reader.ReadStringArray());
 			Targets = Reader.ReadList(() => new UEBuildTarget(Reader));
-			TargetBuildPredicates = Reader.ReadList(() => new BuildPredicateStore(Reader, UniqueFileItems));
+			TargetPrerequisites = Reader.ReadList(() => new BuildPrerequisites(Reader, UniqueFileItems));
 			bUseAdaptiveUnityBuild = Reader.ReadBoolean();
 		}
 
@@ -115,7 +115,7 @@ namespace UnrealBuildTool
 			Writer.Write(ModuleNameToOutputItems, k => Writer.Write(k), v => Writer.Write(v, e => Writer.Write(e, UniqueFileItemToIndex)));
 			Writer.Write(HotReloadModuleNamesForAllTargets.ToArray(), x => Writer.Write(x));
 			Writer.Write(Targets, x => x.Write(Writer));
-			Writer.Write(TargetBuildPredicates, x => x.Write(Writer, UniqueFileItemToIndex));
+			Writer.Write(TargetPrerequisites, x => x.Write(Writer, UniqueFileItemToIndex));
 			Writer.Write(bUseAdaptiveUnityBuild);
 			Writer.Flush();
 
@@ -144,7 +144,7 @@ namespace UnrealBuildTool
 				ModuleNameToOutputItems != null && 
 				HotReloadModuleNamesForAllTargets != null &&
 				Targets != null && Targets.Count > 0 &&
-				TargetBuildPredicates != null;
+				TargetPrerequisites != null;
 		}
 
 
@@ -340,9 +340,9 @@ namespace UnrealBuildTool
 				}
 			}
 
-			foreach (BuildPredicateStore Predicates in LoadedUBTMakefile.TargetBuildPredicates)
+			foreach (BuildPrerequisites Prerequisites in LoadedUBTMakefile.TargetPrerequisites)
 			{
-				foreach(DirectoryReference SourceDirectory in Predicates.SourceDirectories)
+				foreach(DirectoryReference SourceDirectory in Prerequisites.SourceDirectories)
 				{
 					DirectoryInfo SourceDirectoryInfo = new DirectoryInfo(SourceDirectory.FullName);
 					if(!SourceDirectoryInfo.Exists || SourceDirectoryInfo.LastWriteTimeUtc > UBTMakefileInfo.LastWriteTimeUtc)
@@ -353,7 +353,7 @@ namespace UnrealBuildTool
 					}
 				}
 
-				foreach(FileItem AdditionalDependency in Predicates.AdditionalDependencies)
+				foreach(FileItem AdditionalDependency in Prerequisites.AdditionalDependencies)
 				{
 					if (!AdditionalDependency.bExists)
 					{
@@ -376,14 +376,14 @@ namespace UnrealBuildTool
 				// Get all H files in processed modules newer than the makefile itself
 				HashSet<FileReference> HFilesNewerThanMakefile =
 					new HashSet<FileReference>(
-						Predicates.UObjectModuleHeaders
+						Prerequisites.UObjectModuleHeaders
 						.SelectMany(x => x.SourceFolder != null ? DirectoryReference.EnumerateFiles(x.SourceFolder, "*.h", SearchOption.AllDirectories) : Enumerable.Empty<FileReference>())
 						.Where(y => FileReference.GetLastWriteTimeUtc(y) > UBTMakefileInfo.LastWriteTimeUtc)
 						.OrderBy(z => z).Distinct()
 					);
 
 				// Get all H files in all modules processed in the last makefile build
-				HashSet<FileReference> AllUHTHeaders = new HashSet<FileReference>(Predicates.UObjectModuleHeaders.SelectMany(x => x.HeaderFiles).Select(x => x.Location));
+				HashSet<FileReference> AllUHTHeaders = new HashSet<FileReference>(Prerequisites.UObjectModuleHeaders.SelectMany(x => x.HeaderFiles).Select(x => x.Location));
 
 				// Check whether any headers have been deleted. If they have, we need to regenerate the makefile since the module might now be empty. If we don't,
 				// and the file has been moved to a different module, we may include stale generated headers.
@@ -420,10 +420,10 @@ namespace UnrealBuildTool
 			// iteration times.)
 			if (LoadedUBTMakefile.bUseAdaptiveUnityBuild)
 			{
-				foreach(BuildPredicateStore Predicates in LoadedUBTMakefile.TargetBuildPredicates)
+				foreach(BuildPrerequisites Prerequisites in LoadedUBTMakefile.TargetPrerequisites)
 				{
 					// Check if any source files in the working set no longer belong in it
-					foreach (FileItem SourceFile in Predicates.WorkingSet)
+					foreach (FileItem SourceFile in Prerequisites.WorkingSet)
 					{
 						if (!WorkingSet.Contains(SourceFile.Location) && File.GetLastWriteTimeUtc(SourceFile.AbsolutePath) > UBTMakefileInfo.LastWriteTimeUtc)
 						{
@@ -434,7 +434,7 @@ namespace UnrealBuildTool
 					}
 
 					// Check if any source files that are eligible for being in the working set have been modified
-					foreach (FileItem SourceFile in Predicates.CandidatesForWorkingSet)
+					foreach (FileItem SourceFile in Prerequisites.CandidatesForWorkingSet)
 					{
 						if (WorkingSet.Contains(SourceFile.Location) && File.GetLastWriteTimeUtc(SourceFile.AbsolutePath) > UBTMakefileInfo.LastWriteTimeUtc)
 						{
diff --git a/Engine/Source/Programs/UnrealBuildTool/System/Unity.cs b/Engine/Source/Programs/UnrealBuildTool/System/Unity.cs
--- a/Engine/Source/Programs/UnrealBuildTool/System/Unity.cs
+++ b/Engine/Source/Programs/UnrealBuildTool/System/Unity.cs
@@ -152,7 +152,7 @@ namespace UnrealBuildTool
 		/// <param name="CPPFiles">The C++ files to #include.</param>
 		/// <param name="CompileEnvironment">The environment that is used to compile the C++ files.</param>
 		/// <param name="WorkingSet">Interface to query files which belong to the working set</param>
-		/// <param name="Predicates">Dependencies for this target</param>
+		/// <param name="Prerequisites">Dependencies for this target</param>
 		/// <param name="BaseName">Base name to use for the Unity files</param>
 		/// <param name="IntermediateDirectory">Intermediate directory for unity cpp files</param>
 		/// <returns>The "unity" C++ files.</returns>
@@ -161,7 +161,7 @@ namespace UnrealBuildTool
 			List<FileItem> CPPFiles,
 			CppCompileEnvironment CompileEnvironment,
 			ISourceFileWorkingSet WorkingSet,
-			BuildPredicateStore Predicates,
+			BuildPrerequisites Prerequisites,
 			string BaseName,
 			DirectoryReference IntermediateDirectory
 			)
@@ -218,7 +218,7 @@ namespace UnrealBuildTool
 							// Mark this file as part of the working set.  This will be saved into the UBT Makefile so that
 							// the assembler can automatically invalidate the Makefile when the working set changes (allowing this
 							// code to run again, to build up new unity blobs.)
-							Predicates.WorkingSet.Add(CPPFile);
+							Prerequisites.WorkingSet.Add(CPPFile);
 						}
 					}
 
@@ -240,7 +240,7 @@ namespace UnrealBuildTool
 					}
 
 					// When adaptive unity is enabled, go ahead and exclude any source files that we're actively working with
-					if (bUseAdaptiveUnityBuild && Predicates.WorkingSet.Contains(CPPFile))
+					if (bUseAdaptiveUnityBuild && Prerequisites.WorkingSet.Contains(CPPFile))
 					{
 						// Just compile this file normally, not as part of the unity blob
 						NewCPPFiles.Add(CPPFile);
@@ -266,7 +266,7 @@ namespace UnrealBuildTool
 						// If adaptive unity build is enabled for this module, add this source file to the set that will invalidate the makefile
 						if(bUseAdaptiveUnityBuild)
 						{
-							Predicates.CandidatesForWorkingSet.Add(CPPFile);
+							Prerequisites.CandidatesForWorkingSet.Add(CPPFile);
 						}
 
 						// Compile this file as part of the unity blob
diff --git a/Engine/Source/Programs/UnrealBuildTool/UnrealBuildTool.cs b/Engine/Source/Programs/UnrealBuildTool/UnrealBuildTool.cs
--- a/Engine/Source/Programs/UnrealBuildTool/UnrealBuildTool.cs
+++ b/Engine/Source/Programs/UnrealBuildTool/UnrealBuildTool.cs
@@ -750,7 +750,7 @@ namespace UnrealBuildTool
 					Dictionary<string, FileItem[]> ModuleNameToOutputItems = new Dictionary<string, FileItem[]>(StringComparer.OrdinalIgnoreCase);
 					Dictionary<string, List<UHTModuleInfo>> TargetNameToUObjectModules = new Dictionary<string, List<UHTModuleInfo>>(StringComparer.InvariantCultureIgnoreCase);
 					HashSet<string> HotReloadModuleNamesForAllTargets = new HashSet<string>(StringComparer.OrdinalIgnoreCase);
-					List<BuildPredicateStore> TargetBuildPredicates = new List<BuildPredicateStore>();
+					List<BuildPrerequisites> TargetPrerequisites = new List<BuildPrerequisites>();
 					foreach (UEBuildTarget Target in Targets)
 					{
 						// Create the header cache for this target
@@ -804,9 +804,9 @@ namespace UnrealBuildTool
 							List<FileItem> TargetOutputItems = new List<FileItem>();
 							List<UHTModuleInfo> TargetUObjectModules = new List<UHTModuleInfo>();
 							Dictionary<string, FileItem[]> TargetModuleNameToOutputItems = new Dictionary<string, FileItem[]>(StringComparer.OrdinalIgnoreCase);
-							BuildPredicateStore Predicates = new BuildPredicateStore();
+							BuildPrerequisites Prerequisites = new BuildPrerequisites();
 
-							BuildResult = Target.Build(BuildConfiguration, TargetToHeaders[Target], TargetOutputItems, TargetModuleNameToOutputItems, TargetUObjectModules, WorkingSet, ActionGraph, Predicates, bIsAssemblingBuild);
+							BuildResult = Target.Build(BuildConfiguration, TargetToHeaders[Target], TargetOutputItems, TargetModuleNameToOutputItems, TargetUObjectModules, WorkingSet, ActionGraph, Prerequisites, bIsAssemblingBuild);
 							
 							if (BuildResult != ECompilationResult.Succeeded)
 							{
@@ -822,7 +822,7 @@ namespace UnrealBuildTool
 
 							OutputItemsForAllTargets.AddRange(TargetOutputItems);
 
-							TargetBuildPredicates.Add(Predicates);
+							TargetPrerequisites.Add(Prerequisites);
 
 							// Update mapping of the target name to the list of UObject modules in this target
 							TargetNameToUObjectModules[Target.GetTargetName()] = TargetUObjectModules;
@@ -868,7 +868,7 @@ namespace UnrealBuildTool
 							Makefile.HotReloadModuleNamesForAllTargets = HotReloadModuleNamesForAllTargets;
 							Makefile.Targets = Targets;
 							Makefile.bUseAdaptiveUnityBuild = Targets.Any(x => x.Rules.bUseAdaptiveUnityBuild);
-							Makefile.TargetBuildPredicates = TargetBuildPredicates;
+							Makefile.TargetPrerequisites = TargetPrerequisites;
 
 							if (BuildConfiguration.bUseUBTMakefiles)
 							{
diff --git a/Engine/Source/Programs/UnrealBuildTool/UnrealBuildTool.csproj b/Engine/Source/Programs/UnrealBuildTool/UnrealBuildTool.csproj
--- a/Engine/Source/Programs/UnrealBuildTool/UnrealBuildTool.csproj
+++ b/Engine/Source/Programs/UnrealBuildTool/UnrealBuildTool.csproj
@@ -215,7 +215,7 @@
     <Compile Include="System\RulesAssembly.cs" />
     <Compile Include="System\ScopedTimer.cs" />
     <Compile Include="System\SingleInstanceMutex.cs" />
-    <Compile Include="System\BuildPredicateStore.cs" />
+    <Compile Include="System\BuildPrerequisites.cs" />
     <Compile Include="System\SourceFileWorkingSet.cs" />
     <Compile Include="System\StartupTraceListener.cs" />
     <Compile Include="System\SupportedConfigurationsAttribute.cs" />
