UBT: Remove some unnecessary fields from UEBuildTarget.

#rnx

diff --git a/Engine/Source/Programs/UnrealBuildTool/Configuration/UEBuildTarget.cs b/Engine/Source/Programs/UnrealBuildTool/Configuration/UEBuildTarget.cs
--- a/Engine/Source/Programs/UnrealBuildTool/Configuration/UEBuildTarget.cs
+++ b/Engine/Source/Programs/UnrealBuildTool/Configuration/UEBuildTarget.cs
@@ -551,36 +551,6 @@ namespace UnrealBuildTool
 		/// </summary>
 		public DirectoryReference EngineIntermediateDirectory;
 
-		/// <summary>
-		/// Output paths of final executable.
-		/// </summary>
-		public List<FileReference> OutputPaths;
-
-		/// <summary>
-		/// Returns the OutputPath is there is only one entry in OutputPaths
-		/// </summary>
-		public FileReference OutputPath
-		{
-			get
-			{
-				if (OutputPaths.Count != 1)
-				{
-					throw new BuildException("Attempted to use UEBuildTarget.OutputPath property, but there are multiple (or no) OutputPaths. You need to handle multiple in the code that called this (size = {0})", OutputPaths.Count);
-				}
-				return OutputPaths[0];
-			}
-		}
-
-		/// <summary>
-		/// Path to the file that contains the version for this target. Writing this file allows a target to read its version information at runtime.
-		/// </summary>
-		public FileReference VersionFile;
-
-		/// <summary>
-		/// Whether to build target modules that can be reused for future builds
-		/// </summary>
-		public bool bPrecompile;
-
 		/// <summary>
 		/// Identifies whether the project contains a script plugin. This will cause UHT to be rebuilt, even in installed builds.
 		/// </summary>
@@ -678,9 +648,6 @@ namespace UnrealBuildTool
 			ProjectDirectory = Reader.ReadDirectoryReference();
 			ProjectIntermediateDirectory = Reader.ReadDirectoryReference();
 			EngineIntermediateDirectory = Reader.ReadDirectoryReference();
-			OutputPaths = Reader.ReadList(() => Reader.ReadFileReference());
-			VersionFile = Reader.ReadFileReference();
-			bPrecompile = Reader.ReadBoolean();
 			bCompileMonolithic = Reader.ReadBoolean();
 			ReceiptFileName = Reader.ReadFileReference();
 			PreBuildStepScripts = Reader.ReadArray(() => Reader.ReadFileReference());
@@ -703,9 +670,6 @@ namespace UnrealBuildTool
 			Writer.Write(ProjectDirectory);
 			Writer.Write(ProjectIntermediateDirectory);
 			Writer.Write(EngineIntermediateDirectory);
-			Writer.Write(OutputPaths, OutputPath => Writer.Write(OutputPath));
-			Writer.Write(VersionFile);
-			Writer.Write(bPrecompile);
 			Writer.Write(bCompileMonolithic);
 			Writer.Write(ReceiptFileName);
 			Writer.Write(PreBuildStepScripts, x => Writer.Write(x));
@@ -731,7 +695,6 @@ namespace UnrealBuildTool
 			Rules = InRules;
 			RulesAssembly = InRulesAssembly;
 			TargetType = Rules.Type;
-			bPrecompile = InRules.bPrecompile;
 			ForeignPlugin = InDesc.ForeignPlugin;
 			bDeployAfterCompile = InRules.bDeployAfterCompile && !InRules.bDisableLinking;
 
@@ -809,31 +772,6 @@ namespace UnrealBuildTool
 				ProjectDescriptor = ProjectDescriptor.FromFile(ProjectFile);
 			}
 
-			// Construct the output paths for this target's executable
-			DirectoryReference OutputDirectory;
-			if (bCompileMonolithic || !bUseSharedBuildEnvironment)
-			{
-				OutputDirectory = ProjectDirectory;
-			}
-			else
-			{
-				OutputDirectory = UnrealBuildTool.EngineDirectory;
-			}
-
-			bool bCompileAsDLL = Rules.bShouldCompileAsDLL && bCompileMonolithic;
-			OutputPaths = MakeBinaryPaths(OutputDirectory, bCompileMonolithic ? TargetName : AppName, Platform, Configuration, bCompileAsDLL ? UEBuildBinaryType.DynamicLinkLibrary : UEBuildBinaryType.Executable, Rules.Architecture, Rules.UndecoratedConfiguration, bCompileMonolithic && ProjectFile != null, Rules.ExeBinariesSubFolder, ProjectFile, Rules);
-
-			// Get the path to the version file unless this is a formal build (where it will be compiled in)
-			if(Rules.LinkType != TargetLinkType.Monolithic)
-			{
-				UnrealTargetConfiguration VersionConfig = Configuration;
-				if(VersionConfig == UnrealTargetConfiguration.DebugGame && !bCompileMonolithic && TargetType != TargetType.Program && bUseSharedBuildEnvironment)
-				{
-					VersionConfig = UnrealTargetConfiguration.Development;
-				}
-				VersionFile = BuildVersion.GetFileNameForTarget(OutputDirectory, bCompileMonolithic? TargetName : AppName, Platform, VersionConfig, Architecture);
-			}
-
 			// Create all the binaries and modules
 			PreBuildSetup();
 		}
@@ -1144,9 +1082,13 @@ namespace UnrealBuildTool
 			TargetReceipt Receipt = new TargetReceipt(ProjectFile, TargetName, TargetType, Platform, Configuration, Version);
 
 			// Set the launch executable if there is one
-			if(!Rules.bShouldCompileAsDLL)
+			foreach(KeyValuePair<FileReference, BuildProductType> Pair in BuildProducts)
 			{
-				Receipt.Launch = OutputPaths[0];
+				if(Pair.Value == BuildProductType.Executable)
+				{
+					Receipt.Launch = Pair.Key;
+					break;
+				}
 			}
 
 			// Find all the build products and modules from this binary
@@ -1475,7 +1417,7 @@ namespace UnrealBuildTool
 			}
 
 			// Build the target's binaries.
-			DirectoryReference ExeDir = OutputPaths[0].Directory;
+			DirectoryReference ExeDir = Binaries[0].OutputDir;
 			if (CppPlatform == CppPlatform.Mac && ExeDir.FullName.EndsWith(".app/Contents/MacOS"))
 			{
 				ExeDir = ExeDir.ParentDirectory.ParentDirectory.ParentDirectory;
@@ -1556,6 +1498,18 @@ namespace UnrealBuildTool
 					}
 				}
 
+				// Get the path to the version file unless this is a formal build (where it will be compiled in)
+				FileReference VersionFile = null;
+				if(Rules.LinkType != TargetLinkType.Monolithic && Binaries[0].Type == UEBuildBinaryType.Executable)
+				{
+					UnrealTargetConfiguration VersionConfig = Configuration;
+					if(VersionConfig == UnrealTargetConfiguration.DebugGame && !bCompileMonolithic && TargetType != TargetType.Program && bUseSharedBuildEnvironment)
+					{
+						VersionConfig = UnrealTargetConfiguration.Development;
+					}
+					VersionFile = BuildVersion.GetFileNameForTarget(ExeDir, bCompileMonolithic? TargetName : AppName, Platform, VersionConfig, Architecture);
+				}
+
 				// Also add the version file as a build product
 				if(VersionFile != null)
 				{
@@ -2914,6 +2868,20 @@ namespace UnrealBuildTool
 				IntermediateDirectory = ProjectIntermediateDirectory;
 			}
 
+			// Construct the output paths for this target's executable
+			DirectoryReference OutputDirectory;
+			if (bCompileMonolithic || !bUseSharedBuildEnvironment)
+			{
+				OutputDirectory = ProjectDirectory;
+			}
+			else
+			{
+				OutputDirectory = UnrealBuildTool.EngineDirectory;
+			}
+
+			bool bCompileAsDLL = Rules.bShouldCompileAsDLL && bCompileMonolithic;
+			List<FileReference> OutputPaths = MakeBinaryPaths(OutputDirectory, bCompileMonolithic ? TargetName : AppName, Platform, Configuration, bCompileAsDLL ? UEBuildBinaryType.DynamicLinkLibrary : UEBuildBinaryType.Executable, Rules.Architecture, Rules.UndecoratedConfiguration, bCompileMonolithic && ProjectFile != null, Rules.ExeBinariesSubFolder, ProjectFile, Rules);
+
 			// Create the binary
 			UEBuildBinary Binary = new UEBuildBinary(
 				Type: Rules.bShouldCompileAsDLL? UEBuildBinaryType.DynamicLinkLibrary : UEBuildBinaryType.Executable,
@@ -2983,7 +2951,7 @@ namespace UnrealBuildTool
 			GlobalLinkEnvironment.bHasExports = Rules.bHasExports;
 			GlobalLinkEnvironment.bAllowASLR = (GlobalCompileEnvironment.Configuration == CppConfiguration.Shipping && Rules.bAllowASLRInShipping);
 			GlobalLinkEnvironment.bUsePDBFiles = Rules.bUsePDBFiles;
-			GlobalLinkEnvironment.BundleDirectory = BuildPlatform.GetBundleDirectory(Rules, OutputPaths);
+			GlobalLinkEnvironment.BundleDirectory = BuildPlatform.GetBundleDirectory(Rules, Binaries[0].OutputFilePaths);
 			GlobalLinkEnvironment.BundleVersion = Rules.BundleVersion;
 			GlobalLinkEnvironment.bAllowLTCG = Rules.bAllowLTCG;
             GlobalLinkEnvironment.bPGOOptimize = Rules.bPGOOptimize;
