UBT: Removing a few unused/unnecessary properties and accessors, and moving the GetSupportedPlatforms() method from UEBuildTarget to TargetRules.

#rnx

diff --git a/Engine/Source/Programs/UnrealBuildTool/Configuration/TargetRules.cs b/Engine/Source/Programs/UnrealBuildTool/Configuration/TargetRules.cs
--- a/Engine/Source/Programs/UnrealBuildTool/Configuration/TargetRules.cs
+++ b/Engine/Source/Programs/UnrealBuildTool/Configuration/TargetRules.cs
@@ -1401,6 +1401,37 @@ namespace UnrealBuildTool
 			return null;
 		}
 
+		/// <summary>
+		/// Gets a list of platforms that this target supports
+		/// </summary>
+		/// <returns>Array of platforms that the target supports</returns>
+		internal UnrealTargetPlatform[] GetSupportedPlatforms()
+		{
+			// Otherwise take the SupportedPlatformsAttribute from the first type in the inheritance chain that supports it
+			for (Type CurrentType = GetType(); CurrentType != null; CurrentType = CurrentType.BaseType)
+			{
+				object[] Attributes = GetType().GetCustomAttributes(typeof(SupportedPlatformsAttribute), false);
+				if (Attributes.Length > 0)
+				{
+					return Attributes.OfType<SupportedPlatformsAttribute>().SelectMany(x => x.Platforms).Distinct().ToArray();
+				}
+			}
+
+			// Otherwise, get the default for the target type
+			if (Type == TargetType.Program)
+			{
+				return Utils.GetPlatformsInClass(UnrealPlatformClass.Desktop);
+			}
+			else if (Type == TargetType.Editor)
+			{
+				return Utils.GetPlatformsInClass(UnrealPlatformClass.Editor);
+			}
+			else
+			{
+				return Utils.GetPlatformsInClass(UnrealPlatformClass.All);
+			}
+		}
+
 		/// <summary>
 		/// Finds all the subobjects which can be configured by command line options and config files
 		/// </summary>
diff --git a/Engine/Source/Programs/UnrealBuildTool/Configuration/UEBuildTarget.cs b/Engine/Source/Programs/UnrealBuildTool/Configuration/UEBuildTarget.cs
--- a/Engine/Source/Programs/UnrealBuildTool/Configuration/UEBuildTarget.cs
+++ b/Engine/Source/Programs/UnrealBuildTool/Configuration/UEBuildTarget.cs
@@ -258,43 +258,6 @@ namespace UnrealBuildTool
 	/// </summary>
 	class UEBuildTarget
 	{
-		public string GetAppName()
-		{
-			return AppName;
-		}
-
-		public string GetTargetName()
-		{
-			return TargetName;
-		}
-
-		public static UnrealTargetPlatform[] GetSupportedPlatforms(TargetRules Rules)
-		{
-			// Otherwise take the SupportedPlatformsAttribute from the first type in the inheritance chain that supports it
-			for (Type CurrentType = Rules.GetType(); CurrentType != null; CurrentType = CurrentType.BaseType)
-			{
-				object[] Attributes = Rules.GetType().GetCustomAttributes(typeof(SupportedPlatformsAttribute), false);
-				if (Attributes.Length > 0)
-				{
-					return Attributes.OfType<SupportedPlatformsAttribute>().SelectMany(x => x.Platforms).Distinct().ToArray();
-				}
-			}
-
-			// Otherwise, get the default for the target type
-			if (Rules.Type == TargetType.Program)
-			{
-				return Utils.GetPlatformsInClass(UnrealPlatformClass.Desktop);
-			}
-			else if (Rules.Type == TargetType.Editor)
-			{
-				return Utils.GetPlatformsInClass(UnrealPlatformClass.Editor);
-			}
-			else
-			{
-				return Utils.GetPlatformsInClass(UnrealPlatformClass.All);
-			}
-		}
-
 		/// <summary>
 		/// Creates a target object for the specified target name.
 		/// </summary>
@@ -308,7 +271,7 @@ namespace UnrealBuildTool
 			RulesAssembly RulesAssembly = RulesCompiler.CreateTargetRulesAssembly(Desc.ProjectFile, Desc.Name, bSkipRulesCompile, bUsePrecompiled, Desc.ForeignPlugin);
 
 			TargetRules RulesObject = RulesAssembly.CreateTargetRules(Desc.Name, Desc.Platform, Desc.Configuration, Desc.Architecture, Desc.ProjectFile, Desc.AdditionalArguments);
-			if ((ProjectFileGenerator.bGenerateProjectFiles == false) && !GetSupportedPlatforms(RulesObject).Contains(Desc.Platform))
+			if ((ProjectFileGenerator.bGenerateProjectFiles == false) && !RulesObject.GetSupportedPlatforms().Contains(Desc.Platform))
 			{
 				throw new BuildException("{0} does not support the {1} platform.", Desc.Name, Desc.Platform.ToString());
 			}
@@ -705,12 +668,6 @@ namespace UnrealBuildTool
 
 			bCompileMonolithic = (Rules.LinkType == TargetLinkType.Monolithic);
 
-			// Some platforms may *require* monolithic compilation...
-			if (!bCompileMonolithic && UEBuildPlatform.PlatformRequiresMonolithicBuilds(InDesc.Platform, InDesc.Configuration))
-			{
-				throw new BuildException(String.Format("{0} does not support modular builds", InDesc.Platform));
-			}
-
 			// Set the build environment
 			bUseSharedBuildEnvironment = (Rules.BuildEnvironment == TargetBuildEnvironment.Shared);
 			if (bUseSharedBuildEnvironment)
@@ -746,7 +703,7 @@ namespace UnrealBuildTool
 			}
 			else
 			{
-				ProjectIntermediateDirectory = DirectoryReference.Combine(ProjectDirectory, PlatformIntermediateFolder, GetTargetName(), Configuration.ToString());
+				ProjectIntermediateDirectory = DirectoryReference.Combine(ProjectDirectory, PlatformIntermediateFolder, TargetName, Configuration.ToString());
 			}
 
 			// Build the engine intermediate directory. If we're building agnostic engine binaries, we can use the engine intermediates folder. Otherwise we need to use the project intermediates directory.
@@ -1376,19 +1333,17 @@ namespace UnrealBuildTool
 			// files that are injected as top level prerequisites.  If UHT only emitted included header files, we wouldn't need to run it during the Gather phase at all.
 			if (UObjectModules.Count > 0)
 			{
-				FileReference ModuleInfoFileName = FileReference.Combine(ProjectIntermediateDirectory, GetTargetName() + ".uhtmanifest");
+				FileReference ModuleInfoFileName = FileReference.Combine(ProjectIntermediateDirectory, TargetName + ".uhtmanifest");
 
 				// Execute the header tool
 				ECompilationResult UHTResult = ExternalExecution.ExecuteHeaderToolIfNecessary(BuildConfiguration, ProjectFile, TargetName, TargetType, bHasProjectScriptPlugin, UObjectModules, ModuleInfoFileName, true, bIsAssemblingBuild);
 				if(!UHTResult.Succeeded())
 				{
-					Log.TraceInformation(String.Format("Error: UnrealHeaderTool failed for target '{0}' (platform: {1}, module info: {2}, exit code: {3} ({4})).", GetTargetName(), Platform.ToString(), ModuleInfoFileName, UHTResult.ToString(), (int)UHTResult));
+					Log.TraceInformation(String.Format("Error: UnrealHeaderTool failed for target '{0}' (platform: {1}, module info: {2}, exit code: {3} ({4})).", TargetName, Platform.ToString(), ModuleInfoFileName, UHTResult.ToString(), (int)UHTResult));
 					return UHTResult;
 				}
 			}
 
-			GlobalLinkEnvironment.bShouldCompileMonolithic = ShouldCompileMonolithic();
-
 			// Find all the shared PCHs.
 			List<PrecompiledHeaderTemplate> SharedPCHs = new List<PrecompiledHeaderTemplate>();
 			if (Rules.bUseSharedPCHs)
@@ -2916,7 +2871,7 @@ namespace UnrealBuildTool
 			// @Hack: This to prevent UHT from listing CoreUObject.init.gen.cpp as its dependency.
 			// We flag the compile environment when we build UHT so that we don't need to check
 			// this for each file when generating their dependencies.
-			GlobalCompileEnvironment.bHackHeaderGenerator = (GetAppName() == "UnrealHeaderTool");
+			GlobalCompileEnvironment.bHackHeaderGenerator = (AppName == "UnrealHeaderTool");
 
 			GlobalCompileEnvironment.bUseDebugCRT = GlobalCompileEnvironment.Configuration == CppConfiguration.Debug && Rules.bDebugBuildsActuallyUseDebugCRT;
 			GlobalCompileEnvironment.bEnableOSX109Support = Rules.bEnableOSX109Support;
@@ -2983,10 +2938,8 @@ namespace UnrealBuildTool
 
 			[REDACTED_EMAIL]: Do any platform specific tool chain initialization here if required
 
-			string OutputAppName = GetAppName();
-
 			UnrealTargetConfiguration EngineTargetConfiguration = Configuration == UnrealTargetConfiguration.DebugGame ? UnrealTargetConfiguration.Development : Configuration;
-			DirectoryReference LinkIntermediateDirectory = DirectoryReference.Combine(UnrealBuildTool.EngineDirectory, PlatformIntermediateFolder, OutputAppName, EngineTargetConfiguration.ToString());
+			DirectoryReference LinkIntermediateDirectory = DirectoryReference.Combine(UnrealBuildTool.EngineDirectory, PlatformIntermediateFolder, AppName, EngineTargetConfiguration.ToString());
 
 			// Installed Engine intermediates go to the project's intermediate folder. Installed Engine never writes to the engine intermediate folder. (Those files are immutable)
 			// Also, when compiling in monolithic, all intermediates go to the project's folder.  This is because a project can change definitions that affects all engine translation
@@ -2997,11 +2950,11 @@ namespace UnrealBuildTool
 				{
 					if (ProjectFile != null)
 					{
-						LinkIntermediateDirectory = DirectoryReference.Combine(ProjectFile.Directory, PlatformIntermediateFolder, OutputAppName, Configuration.ToString());
+						LinkIntermediateDirectory = DirectoryReference.Combine(ProjectFile.Directory, PlatformIntermediateFolder, AppName, Configuration.ToString());
 					}
 					else if (ForeignPlugin != null)
 					{
-						LinkIntermediateDirectory = DirectoryReference.Combine(ForeignPlugin.Directory, PlatformIntermediateFolder, OutputAppName, Configuration.ToString());
+						LinkIntermediateDirectory = DirectoryReference.Combine(ForeignPlugin.Directory, PlatformIntermediateFolder, AppName, Configuration.ToString());
 					}
 				}
 			}
diff --git a/Engine/Source/Programs/UnrealBuildTool/ProjectFiles/ProjectFileGenerator.cs b/Engine/Source/Programs/UnrealBuildTool/ProjectFiles/ProjectFileGenerator.cs
--- a/Engine/Source/Programs/UnrealBuildTool/ProjectFiles/ProjectFileGenerator.cs
+++ b/Engine/Source/Programs/UnrealBuildTool/ProjectFiles/ProjectFileGenerator.cs
@@ -2252,7 +2252,7 @@ namespace UnrealBuildTool
 							TargetFilePath = TargetFilePath,
 							ProjectFilePath = ProjectFilePath,
 							UnrealProjectFilePath = CheckProjectFile,
-							SupportedPlatforms = UEBuildTarget.GetSupportedPlatforms(TargetRulesObject).Where(x => UEBuildPlatform.GetBuildPlatform(x, true) != null).ToArray(),
+							SupportedPlatforms = TargetRulesObject.GetSupportedPlatforms().Where(x => UEBuildPlatform.GetBuildPlatform(x, true) != null).ToArray(),
 							CreateRulesDelegate = (Platform, Configuration) => RulesAssembly.CreateTargetRules(TargetName, Platform, Configuration, "", CheckProjectFile, null)
                         };
 
diff --git a/Engine/Source/Programs/UnrealBuildTool/ProjectFiles/VisualStudioCode/VSCodeProjectFileGenerator.cs b/Engine/Source/Programs/UnrealBuildTool/ProjectFiles/VisualStudioCode/VSCodeProjectFileGenerator.cs
--- a/Engine/Source/Programs/UnrealBuildTool/ProjectFiles/VisualStudioCode/VSCodeProjectFileGenerator.cs
+++ b/Engine/Source/Programs/UnrealBuildTool/ProjectFiles/VisualStudioCode/VSCodeProjectFileGenerator.cs
@@ -366,7 +366,7 @@ namespace UnrealBuildTool
 					foreach (ProjectTarget Target in Project.ProjectTargets)
 					{
 						Array Configs = Enum.GetValues(typeof(UnrealTargetConfiguration));
-						List<UnrealTargetPlatform> Platforms = new List<UnrealTargetPlatform>(UEBuildTarget.GetSupportedPlatforms(Target.TargetRules));
+						List<UnrealTargetPlatform> Platforms = new List<UnrealTargetPlatform>(Target.TargetRules.GetSupportedPlatforms());
 
 						ProjectData.Target NewTarget = new ProjectData.Target(NewProject, Target.TargetRules.Name, Target.TargetRules.Type);
 
diff --git a/Engine/Source/Programs/UnrealBuildTool/System/LinkEnvironment.cs b/Engine/Source/Programs/UnrealBuildTool/System/LinkEnvironment.cs
--- a/Engine/Source/Programs/UnrealBuildTool/System/LinkEnvironment.cs
+++ b/Engine/Source/Programs/UnrealBuildTool/System/LinkEnvironment.cs
@@ -261,11 +261,6 @@ namespace UnrealBuildTool
 		/// </summary>
 		public string BundleVersion;
 
-		/// <summary>
-		/// Whether we're linking in monolithic mode. Determines if linking should produce import library file. Relevant only for VC++, clang stores imports in shared library.
-		/// </summary>
-		public bool bShouldCompileMonolithic = false;
-
 		/// <summary>
 		/// A list of the object files to be linked.
 		/// </summary>
diff --git a/Engine/Source/Programs/UnrealBuildTool/System/RulesAssembly.cs b/Engine/Source/Programs/UnrealBuildTool/System/RulesAssembly.cs
--- a/Engine/Source/Programs/UnrealBuildTool/System/RulesAssembly.cs
+++ b/Engine/Source/Programs/UnrealBuildTool/System/RulesAssembly.cs
@@ -524,6 +524,12 @@ namespace UnrealBuildTool
 			UEBuildPlatform Platform = UEBuildPlatform.GetBuildPlatform(Rules.Platform);
 			Platform.ValidateTarget(Rules);
 
+			// Some platforms may *require* monolithic compilation...
+			if (Rules.LinkType != TargetLinkType.Monolithic && UEBuildPlatform.PlatformRequiresMonolithicBuilds(Rules.Platform, Rules.Configuration))
+			{
+				throw new BuildException(String.Format("{0} does not support modular builds", Rules.Platform));
+			}
+
 			return Rules;
 		}
 
diff --git a/Engine/Source/Programs/UnrealBuildTool/UnrealBuildTool.cs b/Engine/Source/Programs/UnrealBuildTool/UnrealBuildTool.cs
--- a/Engine/Source/Programs/UnrealBuildTool/UnrealBuildTool.cs
+++ b/Engine/Source/Programs/UnrealBuildTool/UnrealBuildTool.cs
@@ -771,7 +771,7 @@ namespace UnrealBuildTool
 							// Load the direct include dependency cache.
 							using(Timeline.ScopeEvent("DependencyCache.Create()"))
 							{
-								Headers.IncludeDependencyCache = DependencyCache.Create(DependencyCache.GetDependencyCachePathForTarget(Target.ProjectFile, Target.Platform, Target.GetTargetName()));
+								Headers.IncludeDependencyCache = DependencyCache.Create(DependencyCache.GetDependencyCachePathForTarget(Target.ProjectFile, Target.Platform, Target.TargetName));
 							}
 						}
 
@@ -825,7 +825,7 @@ namespace UnrealBuildTool
 							TargetPrerequisites.Add(Prerequisites);
 
 							// Update mapping of the target name to the list of UObject modules in this target
-							TargetNameToUObjectModules[Target.GetTargetName()] = TargetUObjectModules;
+							TargetNameToUObjectModules[Target.TargetName] = TargetUObjectModules;
 
 							// Merge the module name to output items map
 							foreach(KeyValuePair<string, FileItem[]> Pair in TargetModuleNameToOutputItems)
@@ -975,16 +975,16 @@ namespace UnrealBuildTool
 									foreach (UEBuildTarget Target in Targets)
 									{
 										List<UHTModuleInfo> TargetUObjectModules;
-										if (Makefile.TargetNameToUObjectModules.TryGetValue(Target.GetTargetName(), out TargetUObjectModules))
+										if (Makefile.TargetNameToUObjectModules.TryGetValue(Target.TargetName, out TargetUObjectModules))
 										{
 											if (TargetUObjectModules.Count > 0)
 											{
 												// Execute the header tool
-												FileReference ModuleInfoFileName = FileReference.Combine(Target.ProjectIntermediateDirectory, Target.GetTargetName() + ".uhtmanifest");
+												FileReference ModuleInfoFileName = FileReference.Combine(Target.ProjectIntermediateDirectory, Target.TargetName + ".uhtmanifest");
 												ECompilationResult UHTResult = ExternalExecution.ExecuteHeaderToolIfNecessary(BuildConfiguration, Target.ProjectFile, Target.TargetName, Target.TargetType, Target.bHasProjectScriptPlugin,  UObjectModules: TargetUObjectModules, ModuleInfoFileName: ModuleInfoFileName, bIsGatheringBuild: bIsGatheringBuild, bIsAssemblingBuild: bIsAssemblingBuild);
 												if(!UHTResult.Succeeded())
 												{
-													Log.TraceInformation("UnrealHeaderTool failed for target '" + Target.GetTargetName() + "' (platform: " + Target.Platform.ToString() + ", module info: " + ModuleInfoFileName + ").");
+													Log.TraceInformation("UnrealHeaderTool failed for target '" + Target.TargetName + "' (platform: " + Target.Platform.ToString() + ", module info: " + ModuleInfoFileName + ").");
 													BuildResult = UHTResult;
 													break;
 												}
