UBT: Move Action and ActionExecutor into their own file.

#rnx

diff --git a/Engine/Source/Programs/UnrealBuildTool/System/Action.cs b/Engine/Source/Programs/UnrealBuildTool/System/Action.cs
new file mode 100644
--- /dev/null
+++ b/Engine/Source/Programs/UnrealBuildTool/System/Action.cs
@@ -0,0 +1,254 @@
+using System;
+using System.Collections.Generic;
+using System.IO;
+using System.Linq;
+using System.Text;
+using System.Threading.Tasks;
+
+namespace UnrealBuildTool
+{
+	/// <summary>
+	/// Enumerates build action types.
+	/// </summary>
+	enum ActionType
+	{
+		BuildProject,
+
+		Compile,
+
+		CreateAppBundle,
+
+		GenerateDebugInfo,
+
+		Link,
+
+		WriteMetadata,
+
+		PostBuildStep,
+	}
+
+	/// <summary>
+	/// A build action.
+	/// </summary>
+	class Action
+	{
+		///
+		/// Preparation and Assembly (serialized)
+		/// 
+
+		/// <summary>
+		/// The type of this action (for debugging purposes).
+		/// </summary>
+		public readonly ActionType ActionType;
+
+		/// <summary>
+		/// Every file this action depends on.  These files need to exist and be up to date in order for this action to even be considered
+		/// </summary>
+		public List<FileItem> PrerequisiteItems = new List<FileItem>();
+
+		/// <summary>
+		/// The files that this action produces after completing
+		/// </summary>
+		public List<FileItem> ProducedItems = new List<FileItem>();
+
+		/// <summary>
+		/// Items that should be deleted before running this action
+		/// </summary>
+		public List<FileItem> DeleteItems = new List<FileItem>();
+
+		/// <summary>
+		/// Directory from which to execute the program to create produced items
+		/// </summary>
+		public string WorkingDirectory = null;
+
+		/// <summary>
+		/// True if we should log extra information when we run a program to create produced items
+		/// </summary>
+		public bool bPrintDebugInfo = false;
+
+		/// <summary>
+		/// The command to run to create produced items
+		/// </summary>
+		public string CommandPath = null;
+
+		/// <summary>
+		/// Command-line parameters to pass to the program
+		/// </summary>
+		public string CommandArguments = null;
+
+		/// <summary>
+		/// Optional friendly description of the type of command being performed, for example "Compile" or "Link".  Displayed by some executors.
+		/// </summary>
+		public string CommandDescription = null;
+
+		/// <summary>
+		/// Human-readable description of this action that may be displayed as status while invoking the action.  This is often the name of the file being compiled, or an executable file name being linked.  Displayed by some executors.
+		/// </summary>
+		public string StatusDescription = "...";
+
+		/// <summary>
+		/// True if this action is allowed to be run on a remote machine when a distributed build system is being used, such as XGE
+		/// </summary>
+		public bool bCanExecuteRemotely = false;
+
+		/// <summary>
+		/// True if this action is allowed to be run on a remote machine with SNDBS. Files with #import directives must be compiled locally. Also requires bCanExecuteRemotely = true.
+		/// </summary>
+		public bool bCanExecuteRemotelyWithSNDBS = true;
+
+		/// <summary>
+		/// True if this action is using the GCC compiler.  Some build systems may be able to optimize for this case.
+		/// </summary>
+		public bool bIsGCCCompiler = false;
+
+		/// <summary>
+		/// Whether the action is using a pre-compiled header to speed it up.
+		/// </summary>
+		public bool bIsUsingPCH = false;
+
+		/// <summary>
+		/// Whether we should log this action, whether executed locally or remotely.  This is useful for actions that take time
+		/// but invoke tools without any console output.
+		/// </summary>
+		public bool bShouldOutputStatusDescription = true;
+
+		/// <summary>
+		/// True if any libraries produced by this action should be considered 'import libraries'
+		/// </summary>
+		public bool bProducesImportLibrary = false;
+
+
+
+		///
+		/// Preparation only (not serialized)
+		///
+
+		/// <summary>
+		/// Total number of actions depending on this one.
+		/// </summary>
+		public int NumTotalDependentActions = 0;
+
+		/// <summary>
+		/// Relative cost of producing items for this action.
+		/// </summary>
+		public long RelativeCost = 0;
+
+
+		///
+		/// Assembly only (not serialized)
+		///
+
+		/// <summary>
+		/// Start time of action, optionally set by executor.
+		/// </summary>
+		public DateTimeOffset StartTime = DateTimeOffset.MinValue;
+
+		/// <summary>
+		/// End time of action, optionally set by executor.
+		/// </summary>
+		public DateTimeOffset EndTime = DateTimeOffset.MinValue;
+
+
+		public Action(ActionType InActionType)
+		{
+			ActionType = InActionType;
+		}
+
+		public Action(BinaryReader Reader, List<FileItem> UniqueFileItems)
+		{
+			ActionType = (ActionType)Reader.ReadByte();
+			WorkingDirectory = Reader.ReadString();
+			bPrintDebugInfo = Reader.ReadBoolean();
+			CommandPath = Reader.ReadString();
+			CommandArguments = Reader.ReadString();
+			CommandDescription = Reader.ReadNullable(() => Reader.ReadString());
+			StatusDescription = Reader.ReadString();
+			bCanExecuteRemotely = Reader.ReadBoolean();
+			bCanExecuteRemotelyWithSNDBS = Reader.ReadBoolean();
+			bIsGCCCompiler = Reader.ReadBoolean();
+			bIsUsingPCH = Reader.ReadBoolean();
+			bShouldOutputStatusDescription = Reader.ReadBoolean();
+			bProducesImportLibrary = Reader.ReadBoolean();
+			PrerequisiteItems = Reader.ReadFileItemList(UniqueFileItems);
+			ProducedItems = Reader.ReadFileItemList(UniqueFileItems);
+			DeleteItems = Reader.ReadFileItemList(UniqueFileItems);
+		}
+
+		/// <summary>
+		/// ISerializable: Called when serialized to report additional properties that should be saved
+		/// </summary>
+		public void Write(BinaryWriter Writer, Dictionary<FileItem, int> UniqueFileItemToIndex)
+		{
+			Writer.Write((byte)ActionType);
+			Writer.Write(WorkingDirectory);
+			Writer.Write(bPrintDebugInfo);
+			Writer.Write(CommandPath);
+			Writer.Write(CommandArguments);
+			Writer.WriteNullable(CommandDescription, () => Writer.Write(CommandDescription));
+			Writer.Write(StatusDescription);
+			Writer.Write(bCanExecuteRemotely);
+			Writer.Write(bCanExecuteRemotelyWithSNDBS);
+			Writer.Write(bIsGCCCompiler);
+			Writer.Write(bIsUsingPCH);
+			Writer.Write(bShouldOutputStatusDescription);
+			Writer.Write(bProducesImportLibrary);
+			Writer.Write(PrerequisiteItems, UniqueFileItemToIndex);
+			Writer.Write(ProducedItems, UniqueFileItemToIndex);
+			Writer.Write(DeleteItems, UniqueFileItemToIndex);
+		}
+
+		/// <summary>
+		/// Compares two actions based on total number of dependent items, descending.
+		/// </summary>
+		/// <param name="A">Action to compare</param>
+		/// <param name="B">Action to compare</param>
+		public static int Compare(Action A, Action B)
+		{
+			// Primary sort criteria is total number of dependent files, up to max depth.
+			if (B.NumTotalDependentActions != A.NumTotalDependentActions)
+			{
+				return Math.Sign(B.NumTotalDependentActions - A.NumTotalDependentActions);
+			}
+			// Secondary sort criteria is relative cost.
+			if (B.RelativeCost != A.RelativeCost)
+			{
+				return Math.Sign(B.RelativeCost - A.RelativeCost);
+			}
+			// Tertiary sort criteria is number of pre-requisites.
+			else
+			{
+				return Math.Sign(B.PrerequisiteItems.Count - A.PrerequisiteItems.Count);
+			}
+		}
+
+		public override string ToString()
+		{
+			string ReturnString = "";
+			if (CommandPath != null)
+			{
+				ReturnString += CommandPath + " - ";
+			}
+			if (CommandArguments != null)
+			{
+				ReturnString += CommandArguments;
+			}
+			return ReturnString;
+		}
+
+		/// <summary>
+		/// Returns the amount of time that this action is or has been executing in.
+		/// </summary>
+		public TimeSpan Duration
+		{
+			get
+			{
+				if (EndTime == DateTimeOffset.MinValue)
+				{
+					return DateTimeOffset.Now - StartTime;
+				}
+
+				return EndTime - StartTime;
+			}
+		}
+	}
+}
diff --git a/Engine/Source/Programs/UnrealBuildTool/System/ActionExecutor.cs b/Engine/Source/Programs/UnrealBuildTool/System/ActionExecutor.cs
new file mode 100644
--- /dev/null
+++ b/Engine/Source/Programs/UnrealBuildTool/System/ActionExecutor.cs
@@ -0,0 +1,19 @@
+using System;
+using System.Collections.Generic;
+using System.Linq;
+using System.Text;
+using System.Threading.Tasks;
+
+namespace UnrealBuildTool
+{
+	abstract class ActionExecutor
+	{
+		public abstract string Name
+		{
+			get;
+		}
+
+		public abstract bool ExecuteActions(List<Action> ActionsToExecute, bool bLogDetailedActionStats);
+	}
+
+}
diff --git a/Engine/Source/Programs/UnrealBuildTool/System/ActionGraph.cs b/Engine/Source/Programs/UnrealBuildTool/System/ActionGraph.cs
--- a/Engine/Source/Programs/UnrealBuildTool/System/ActionGraph.cs
+++ b/Engine/Source/Programs/UnrealBuildTool/System/ActionGraph.cs
@@ -10,261 +10,6 @@ using Tools.DotNETCommon;
 
 namespace UnrealBuildTool
 {
-	/// <summary>
-	/// Enumerates build action types.
-	/// </summary>
-	enum ActionType
-	{
-		BuildProject,
-
-		Compile,
-
-		CreateAppBundle,
-
-		GenerateDebugInfo,
-
-		Link,
-
-		WriteMetadata,
-
-		PostBuildStep,
-	}
-
-	/// <summary>
-	/// A build action.
-	/// </summary>
-	class Action
-	{
-		///
-		/// Preparation and Assembly (serialized)
-		/// 
-
-		/// <summary>
-		/// The type of this action (for debugging purposes).
-		/// </summary>
-		public readonly ActionType ActionType;
-
-		/// <summary>
-		/// Every file this action depends on.  These files need to exist and be up to date in order for this action to even be considered
-		/// </summary>
-		public List<FileItem> PrerequisiteItems = new List<FileItem>();
-
-		/// <summary>
-		/// The files that this action produces after completing
-		/// </summary>
-		public List<FileItem> ProducedItems = new List<FileItem>();
-
-		/// <summary>
-		/// Items that should be deleted before running this action
-		/// </summary>
-		public List<FileItem> DeleteItems = new List<FileItem>();
-
-		/// <summary>
-		/// Directory from which to execute the program to create produced items
-		/// </summary>
-		public string WorkingDirectory = null;
-
-		/// <summary>
-		/// True if we should log extra information when we run a program to create produced items
-		/// </summary>
-		public bool bPrintDebugInfo = false;
-
-		/// <summary>
-		/// The command to run to create produced items
-		/// </summary>
-		public string CommandPath = null;
-
-		/// <summary>
-		/// Command-line parameters to pass to the program
-		/// </summary>
-		public string CommandArguments = null;
-
-		/// <summary>
-		/// Optional friendly description of the type of command being performed, for example "Compile" or "Link".  Displayed by some executors.
-		/// </summary>
-		public string CommandDescription = null;
-
-		/// <summary>
-		/// Human-readable description of this action that may be displayed as status while invoking the action.  This is often the name of the file being compiled, or an executable file name being linked.  Displayed by some executors.
-		/// </summary>
-		public string StatusDescription = "...";
-
-		/// <summary>
-		/// True if this action is allowed to be run on a remote machine when a distributed build system is being used, such as XGE
-		/// </summary>
-		public bool bCanExecuteRemotely = false;
-
-		/// <summary>
-		/// True if this action is allowed to be run on a remote machine with SNDBS. Files with #import directives must be compiled locally. Also requires bCanExecuteRemotely = true.
-		/// </summary>
-		public bool bCanExecuteRemotelyWithSNDBS = true;
-
-		/// <summary>
-		/// True if this action is using the GCC compiler.  Some build systems may be able to optimize for this case.
-		/// </summary>
-		public bool bIsGCCCompiler = false;
-
-		/// <summary>
-		/// Whether the action is using a pre-compiled header to speed it up.
-		/// </summary>
-		public bool bIsUsingPCH = false;
-
-		/// <summary>
-		/// Whether we should log this action, whether executed locally or remotely.  This is useful for actions that take time
-		/// but invoke tools without any console output.
-		/// </summary>
-		public bool bShouldOutputStatusDescription = true;
-
-		/// <summary>
-		/// True if any libraries produced by this action should be considered 'import libraries'
-		/// </summary>
-		public bool bProducesImportLibrary = false;
-
-
-
-		///
-		/// Preparation only (not serialized)
-		///
-
-		/// <summary>
-		/// Total number of actions depending on this one.
-		/// </summary>
-		public int NumTotalDependentActions = 0;
-
-		/// <summary>
-		/// Relative cost of producing items for this action.
-		/// </summary>
-		public long RelativeCost = 0;
-
-
-		///
-		/// Assembly only (not serialized)
-		///
-
-		/// <summary>
-		/// Start time of action, optionally set by executor.
-		/// </summary>
-		public DateTimeOffset StartTime = DateTimeOffset.MinValue;
-
-		/// <summary>
-		/// End time of action, optionally set by executor.
-		/// </summary>
-		public DateTimeOffset EndTime = DateTimeOffset.MinValue;
-
-
-		public Action(ActionType InActionType)
-		{
-			ActionType = InActionType;
-		}
-
-		public Action(BinaryReader Reader, List<FileItem> UniqueFileItems)
-		{
-			ActionType = (ActionType)Reader.ReadByte();
-			WorkingDirectory = Reader.ReadString();
-			bPrintDebugInfo = Reader.ReadBoolean();
-			CommandPath = Reader.ReadString();
-			CommandArguments = Reader.ReadString();
-			CommandDescription = Reader.ReadNullable(() => Reader.ReadString());
-			StatusDescription = Reader.ReadString();
-			bCanExecuteRemotely = Reader.ReadBoolean();
-			bCanExecuteRemotelyWithSNDBS = Reader.ReadBoolean();
-			bIsGCCCompiler = Reader.ReadBoolean();
-			bIsUsingPCH = Reader.ReadBoolean();
-			bShouldOutputStatusDescription = Reader.ReadBoolean();
-			bProducesImportLibrary = Reader.ReadBoolean();
-			PrerequisiteItems = Reader.ReadFileItemList(UniqueFileItems);
-			ProducedItems = Reader.ReadFileItemList(UniqueFileItems);
-			DeleteItems = Reader.ReadFileItemList(UniqueFileItems);
-		}
-
-		/// <summary>
-		/// ISerializable: Called when serialized to report additional properties that should be saved
-		/// </summary>
-		public void Write(BinaryWriter Writer, Dictionary<FileItem, int> UniqueFileItemToIndex)
-		{
-			Writer.Write((byte)ActionType);
-			Writer.Write(WorkingDirectory);
-			Writer.Write(bPrintDebugInfo);
-			Writer.Write(CommandPath);
-			Writer.Write(CommandArguments);
-			Writer.WriteNullable(CommandDescription, () => Writer.Write(CommandDescription));
-			Writer.Write(StatusDescription);
-			Writer.Write(bCanExecuteRemotely);
-			Writer.Write(bCanExecuteRemotelyWithSNDBS);
-			Writer.Write(bIsGCCCompiler);
-			Writer.Write(bIsUsingPCH);
-			Writer.Write(bShouldOutputStatusDescription);
-			Writer.Write(bProducesImportLibrary);
-			Writer.Write(PrerequisiteItems, UniqueFileItemToIndex);
-			Writer.Write(ProducedItems, UniqueFileItemToIndex);
-			Writer.Write(DeleteItems, UniqueFileItemToIndex);
-		}
-
-		/// <summary>
-		/// Compares two actions based on total number of dependent items, descending.
-		/// </summary>
-		/// <param name="A">Action to compare</param>
-		/// <param name="B">Action to compare</param>
-		public static int Compare(Action A, Action B)
-		{
-			// Primary sort criteria is total number of dependent files, up to max depth.
-			if (B.NumTotalDependentActions != A.NumTotalDependentActions)
-			{
-				return Math.Sign(B.NumTotalDependentActions - A.NumTotalDependentActions);
-			}
-			// Secondary sort criteria is relative cost.
-			if (B.RelativeCost != A.RelativeCost)
-			{
-				return Math.Sign(B.RelativeCost - A.RelativeCost);
-			}
-			// Tertiary sort criteria is number of pre-requisites.
-			else
-			{
-				return Math.Sign(B.PrerequisiteItems.Count - A.PrerequisiteItems.Count);
-			}
-		}
-
-		public override string ToString()
-		{
-			string ReturnString = "";
-			if (CommandPath != null)
-			{
-				ReturnString += CommandPath + " - ";
-			}
-			if (CommandArguments != null)
-			{
-				ReturnString += CommandArguments;
-			}
-			return ReturnString;
-		}
-
-		/// <summary>
-		/// Returns the amount of time that this action is or has been executing in.
-		/// </summary>
-		public TimeSpan Duration
-		{
-			get
-			{
-				if (EndTime == DateTimeOffset.MinValue)
-				{
-					return DateTimeOffset.Now - StartTime;
-				}
-
-				return EndTime - StartTime;
-			}
-		}
-	};
-
-	abstract class ActionExecutor
-	{
-		public abstract string Name
-		{
-			get;
-		}
-
-		public abstract bool ExecuteActions(List<Action> ActionsToExecute, bool bLogDetailedActionStats);
-	}
-
 	class ActionGraph
 	{
 		/// <summary>
diff --git a/Engine/Source/Programs/UnrealBuildTool/UnrealBuildTool.csproj b/Engine/Source/Programs/UnrealBuildTool/UnrealBuildTool.csproj
--- a/Engine/Source/Programs/UnrealBuildTool/UnrealBuildTool.csproj
+++ b/Engine/Source/Programs/UnrealBuildTool/UnrealBuildTool.csproj
@@ -182,6 +182,8 @@
     <Compile Include="ProjectFiles\VisualStudio\VCSolutionOptions.cs" />
     <Compile Include="ProjectFiles\Xcode\XcodeProject.cs" />
     <Compile Include="ProjectFiles\Xcode\XcodeProjectFileGenerator.cs" />
+    <Compile Include="System\Action.cs" />
+    <Compile Include="System\ActionExecutor.cs" />
     <Compile Include="System\Asn.cs" />
     <Compile Include="System\BuildVersion.cs" />
     <Compile Include="System\CppCompileEnvironment.cs" />
