Fix missing reference.

#rnx

diff --git a/Engine/Source/Programs/UnrealBuildTool/System/Action.cs b/Engine/Source/Programs/UnrealBuildTool/System/Action.cs
--- a/Engine/Source/Programs/UnrealBuildTool/System/Action.cs
+++ b/Engine/Source/Programs/UnrealBuildTool/System/Action.cs
@@ -4,6 +4,7 @@ using System.IO;
 using System.Linq;
 using System.Text;
 using System.Threading.Tasks;
+using Tools.DotNETCommon;
 
 namespace UnrealBuildTool
 {
