UBT: Move UnrealBuildTool.RunUBT() into the BuildMode class.

#rnx

diff --git a/Engine/Source/Programs/UnrealBuildTool/Modes/BuildMode.cs b/Engine/Source/Programs/UnrealBuildTool/Modes/BuildMode.cs
--- a/Engine/Source/Programs/UnrealBuildTool/Modes/BuildMode.cs
+++ b/Engine/Source/Programs/UnrealBuildTool/Modes/BuildMode.cs
@@ -3,8 +3,10 @@
 using System;
 using System.Collections.Generic;
 using System.Diagnostics;
+using System.IO;
 using System.Linq;
 using System.Text;
+using System.Threading;
 using System.Threading.Tasks;
 using Tools.DotNETCommon;
 
@@ -91,10 +93,805 @@ namespace UnrealBuildTool
 				JunkDeleter.DeleteJunk();
 			}
 
-			// Build our project
-			using(Timeline.ScopeEvent("Calling RunUBT"))
+			bool bSuccess = true;
+
+			Stopwatch BuildTimer = Stopwatch.StartNew();
+
+			// Reset global configurations
+			ActionGraph ActionGraph = new ActionGraph();
+
+			string ExecutorName = "Unknown";
+			ECompilationResult BuildResult = ECompilationResult.Succeeded;
+
+			CppIncludeBackgroundThread CppIncludeThread = null;
+
+			List<UEBuildTarget> Targets = null;
+			Dictionary<UEBuildTarget, CPPHeaders> TargetToHeaders = new Dictionary<UEBuildTarget, CPPHeaders>();
+
+			Stopwatch ExecutorTimer = new Stopwatch();
+
+			try
 			{
-				ECompilationResult Result = UnrealBuildTool.RunUBT(BuildConfiguration, Arguments.GetRawArray(), LogFile);
+				// action graph implies using the dependency resolve cache
+				bool GeneratingActionGraph = Arguments.HasOption("-Graph");
+				if (GeneratingActionGraph)
+				{
+					BuildConfiguration.bUseIncludeDependencyResolveCache = true;
+				}
+
+				Timeline.AddEvent("RunUBT() initialization complete");
+
+				bool bSkipRulesCompile = Arguments.Any(x => x.Equals("-skiprulescompile", StringComparison.InvariantCultureIgnoreCase));
+
+				List<TargetDescriptor> TargetDescs = new List<TargetDescriptor>();
+				using(Timeline.ScopeEvent("TargetDescriptor.ParseCommandLine()"))
+				{
+					TargetDescs.AddRange(TargetDescriptor.ParseCommandLine(Arguments, BuildConfiguration.bUsePrecompiled, bSkipRulesCompile));
+				}
+
+				// Handle remote builds
+				for(int Idx = 0; Idx < TargetDescs.Count; Idx++)
+				{
+					TargetDescriptor TargetDesc = TargetDescs[Idx];
+					if(RemoteMac.HandlesTargetPlatform(TargetDesc.Platform))
+					{
+						FileReference RemoteLogFile = null;
+						if(LogFile != null)
+						{
+							RemoteLogFile = FileReference.Combine(LogFile.Directory, LogFile.GetFileNameWithoutExtension() + "_Remote.txt");
+						}
+
+						RemoteMac RemoteMac = new RemoteMac(TargetDesc.ProjectFile);
+						if(!RemoteMac.Build(TargetDesc, LogFile))
+						{
+							return (int)ECompilationResult.Unknown;
+						}
+
+						TargetDescs.RemoveAt(Idx--);
+						if(TargetDescs.Count == 0)
+						{
+							return (int)ECompilationResult.Succeeded;
+						}
+					}
+				}
+
+				if (Arguments.Any(x => x.Equals("-InvalidateMakefilesOnly", StringComparison.InvariantCultureIgnoreCase)))
+				{
+					Log.TraceInformation("Invalidating makefiles only in this run.");
+					if (TargetDescs.Count != 1)
+					{
+						Log.TraceError("You have to provide one target name for makefile invalidation.");
+						return (int)ECompilationResult.OtherCompilationError;
+					}
+
+					InvalidateMakefiles(TargetDescs[0]);
+					return (int)ECompilationResult.Succeeded;
+				}
+
+				// Used when BuildConfiguration.bUseUBTMakefiles is enabled.  If true, it means that our cached includes may not longer be
+				// valid (or never were), and we need to recover by forcibly scanning included headers for all build prerequisites to make sure that our
+				// cached set of includes is actually correct, before determining which files are outdated.
+				bool bNeedsFullCPPIncludeRescan = false;
+
+				if (BuildConfiguration.bUseUBTMakefiles)
+				{
+					// Only the modular editor and game targets will share build products.  Unfortunately, we can't determine at
+					// at this point whether we're dealing with modular or monolithic game binaries, so we opt to always invalidate
+					// cached includes if the target we're switching to is either a game target (has project file) or "UE4Editor".
+					bool bMightHaveSharedBuildProducts =
+						TargetDescs[0].ProjectFile != null ||  // Is this a game? (has a .uproject file for the target)
+						TargetDescs[0].Name.Equals("UE4Editor", StringComparison.InvariantCultureIgnoreCase); // Is the engine?
+					if (bMightHaveSharedBuildProducts)
+					{
+						bool bIsBuildingSameTargetsAsLastTime = false;
+
+						string TargetCollectionName = UBTMakefile.MakeTargetCollectionName(TargetDescs);
+
+						string LastBuiltTargetsFileName = "LastBuiltTargets.txt";
+						string LastBuiltTargetsFilePath = FileReference.Combine(UnrealBuildTool.EngineDirectory, "Intermediate", "Build", LastBuiltTargetsFileName).FullName;
+						if (File.Exists(LastBuiltTargetsFilePath) && Utils.ReadAllText(LastBuiltTargetsFilePath) == TargetCollectionName)
+						{
+							// @todo ubtmake: Because we're using separate files for hot reload vs. full compiles, it's actually possible that includes will
+							// become out of date without us knowing if the developer ping-pongs between hot reloading target A and building target B normally.
+							// To fix this we can not use a different file name for last built targets, but the downside is slower performance when
+							// performing the first hot reload after compiling normally (forces full include dependency scan)
+							bIsBuildingSameTargetsAsLastTime = true;
+						}
+
+						// Save out the name of the targets we're building
+						if (!bIsBuildingSameTargetsAsLastTime)
+						{
+							Directory.CreateDirectory(Path.GetDirectoryName(LastBuiltTargetsFilePath));
+							File.WriteAllText(LastBuiltTargetsFilePath, TargetCollectionName, Encoding.UTF8);
+
+							// Can't use super fast include checking unless we're building the same set of targets as last time, because
+							// we might not know about all of the C++ include dependencies for already-up-to-date shared build products
+							// between the targets
+							bNeedsFullCPPIncludeRescan = true;
+							Log.TraceInformation("Performing full C++ include scan (building a new target)");
+						}
+					}
+				}
+
+				// True if we should gather module dependencies for building in this run.  If this is false, then we'll expect to be loading information from disk about
+				// the target's modules before we'll be able to build anything.  One or both of IsGatheringBuild or IsAssemblingBuild must be true.
+				bool bIsGatheringBuild = true;
+
+				// True if we should go ahead and actually compile and link in this run.  If this is false, no actions will be executed and we'll stop after gathering
+				// and saving information about dependencies.  One or both of IsGatheringBuild or IsAssemblingBuild must be true.
+				bool bIsAssemblingBuild = true;
+
+				// If we were asked to enable fast build iteration, we want the 'gather' phase to default to off (unless it is overridden below
+				// using a command-line option.)
+				if (BuildConfiguration.bUseUBTMakefiles)
+				{
+					bIsGatheringBuild = false;
+				}
+
+				foreach (string Argument in Arguments)
+				{
+					string LowercaseArg = Argument.ToLowerInvariant();
+					if (LowercaseArg == "-gather")
+					{
+						bIsGatheringBuild = true;
+					}
+					else if (LowercaseArg == "-nogather")
+					{
+						bIsGatheringBuild = false;
+					}
+					else if (LowercaseArg == "-gatheronly")
+					{
+						bIsGatheringBuild = true;
+						bIsAssemblingBuild = false;
+					}
+					else if (LowercaseArg == "-assemble")
+					{
+						bIsAssemblingBuild = true;
+					}
+					else if (LowercaseArg == "-noassemble")
+					{
+						bIsAssemblingBuild = false;
+					}
+					else if (LowercaseArg == "-assembleonly")
+					{
+						bIsGatheringBuild = false;
+						bIsAssemblingBuild = true;
+					}
+				}
+
+				if (!bIsGatheringBuild && !bIsAssemblingBuild)
+				{
+					throw new BuildException("UnrealBuildTool: At least one of either IsGatheringBuild or IsAssemblingBuild must be true.  Did you pass '-NoGather' with '-NoAssemble'?");
+				}
+
+				// Get a set of all the project directories.
+				HashSet<DirectoryReference> ProjectDirs = new HashSet<DirectoryReference>();
+				foreach(TargetDescriptor TargetDesc in TargetDescs)
+				{
+					if(TargetDesc.ProjectFile != null)
+					{
+						ProjectDirs.Add(TargetDesc.ProjectFile.Directory);
+					}
+				}
+
+				// Create the working set provider
+				using (ISourceFileWorkingSet WorkingSet = SourceFileWorkingSet.Create(UnrealBuildTool.RootDirectory, ProjectDirs))
+				{
+					UBTMakefile Makefile = null;
+					{
+						FileReference UBTMakefilePath = UBTMakefile.GetUBTMakefilePath(TargetDescs);
+
+						// Make sure the gather phase is executed if we're not actually building anything
+						if (BuildConfiguration.bGenerateManifest || BuildConfiguration.bXGEExport || GeneratingActionGraph)
+						{
+							bIsGatheringBuild = true;
+						}
+
+						// Were we asked to run in 'assembler only' mode?  If so, let's check to see if that's even possible by seeing if
+						// we have a valid UBTMakefile already saved to disk, ready for us to load.
+						if (bIsAssemblingBuild && !bIsGatheringBuild)
+						{
+							// @todo ubtmake: Mildly terrified of BuildConfiguration/UEBuildConfiguration globals that were set during the Prepare phase but are not set now.  We may need to save/load all of these, otherwise
+							//		we'll need to call SetupGlobalEnvironment on all of the targets (maybe other stuff, too.  See PreBuildStep())
+
+							// Try to load the UBTMakefile.  It will only be loaded if it has valid content and is not determined to be out of date.    
+							string ReasonNotLoaded;
+							using(Timeline.ScopeEvent("UBTMakefile.LoadUBTMakefile"))
+							{
+								Makefile = UBTMakefile.Load(UBTMakefilePath, TargetDescs[0].ProjectFile, WorkingSet, out ReasonNotLoaded);
+							}
+
+							if (Makefile == null)
+							{
+								// If the Makefile couldn't be loaded, then we're not going to be able to continue in "assembler only" mode.  We'll do both
+								// a 'gather' and 'assemble' in the same run.  This will take a while longer, but subsequent runs will be fast!
+								bIsGatheringBuild = true;
+
+								FileItem.ClearCaches();
+
+								Log.TraceInformation("Creating makefile for {0}{1} ({2})",
+									TargetDescs[0].Name,
+									TargetDescs.Count > 1 ? (" (and " + (TargetDescs.Count - 1).ToString() + " more)") : "",
+									ReasonNotLoaded);
+							}
+						}
+					}
+
+
+					if (Makefile != null && !bIsGatheringBuild && bIsAssemblingBuild)
+					{
+						// If we've loaded a makefile, then we can fill target information from this file!
+						Targets = Makefile.Targets;
+					}
+					else
+					{
+						Targets = new List<UEBuildTarget>();
+						foreach (TargetDescriptor TargetDesc in TargetDescs)
+						{
+							using(Timeline.ScopeEvent("UEBuildTarget.CreateTarget()"))
+							{
+								UEBuildTarget Target = UEBuildTarget.CreateTarget(TargetDesc, bSkipRulesCompile, BuildConfiguration.SingleFileToCompile != null, BuildConfiguration.bUsePrecompiled);
+								Targets.Add(Target);
+							}
+						}
+					}
+
+					// Build action lists for all passed in targets.
+					List<FileItem> OutputItemsForAllTargets = new List<FileItem>();
+					Dictionary<string, FileItem[]> ModuleNameToOutputItems = new Dictionary<string, FileItem[]>(StringComparer.OrdinalIgnoreCase);
+					Dictionary<string, List<UHTModuleInfo>> TargetNameToUObjectModules = new Dictionary<string, List<UHTModuleInfo>>(StringComparer.InvariantCultureIgnoreCase);
+					HashSet<string> HotReloadModuleNamesForAllTargets = new HashSet<string>(StringComparer.OrdinalIgnoreCase);
+					List<BuildPrerequisites> TargetPrerequisites = new List<BuildPrerequisites>();
+					foreach (UEBuildTarget Target in Targets)
+					{
+						// Create the header cache for this target
+						FileReference DependencyCacheFile = DependencyCache.GetDependencyCachePathForTarget(Target.ProjectFile, Target.Platform, Target.TargetName);
+						bool bUseFlatCPPIncludeDependencyCache = BuildConfiguration.bUseUBTMakefiles && bIsAssemblingBuild;
+						CPPHeaders Headers = new CPPHeaders(Target.ProjectFile, DependencyCacheFile, bUseFlatCPPIncludeDependencyCache, BuildConfiguration.bUseUBTMakefiles, BuildConfiguration.bUseIncludeDependencyResolveCache, BuildConfiguration.bTestIncludeDependencyResolveCache);
+						TargetToHeaders[Target] = Headers;
+
+						// Make sure the appropriate executor is selected
+						UEBuildPlatform BuildPlatform = UEBuildPlatform.GetBuildPlatform(Target.Platform);
+						BuildConfiguration.bAllowXGE &= BuildPlatform.CanUseXGE();
+						BuildConfiguration.bAllowDistcc &= BuildPlatform.CanUseDistcc();
+						BuildConfiguration.bAllowSNDBS &= BuildPlatform.CanUseSNDBS();
+
+						// When in 'assembler only' mode, we'll load this cache later on a worker thread.  It takes a long time to load!
+						if (!(!bIsGatheringBuild && bIsAssemblingBuild))
+						{
+							// Load the direct include dependency cache.
+							using(Timeline.ScopeEvent("DependencyCache.Create()"))
+							{
+								Headers.IncludeDependencyCache = DependencyCache.Create(DependencyCache.GetDependencyCachePathForTarget(Target.ProjectFile, Target.Platform, Target.TargetName));
+							}
+						}
+
+						// We don't need this dependency cache in 'gather only' mode
+						if (BuildConfiguration.bUseUBTMakefiles &&
+								!(bIsGatheringBuild && !bIsAssemblingBuild))
+						{
+							using(Timeline.ScopeEvent("FlatCPPIncludeDependencyCache.TryRead()"))
+							{
+								// Load the cache that contains the list of flattened resolved includes for resolved source files
+								// @todo ubtmake: Ideally load this asynchronously at startup and only block when it is first needed and not finished loading
+								FileReference CacheFile = FlatCPPIncludeDependencyCache.GetDependencyCachePathForTarget(Target.ProjectFile, Target.TargetName, Target.Platform, Target.Architecture);
+								if (!FlatCPPIncludeDependencyCache.TryRead(CacheFile, out Headers.FlatCPPIncludeDependencyCache))
+								{
+									if (!bNeedsFullCPPIncludeRescan)
+									{
+										if (!BuildConfiguration.bXGEExport && !BuildConfiguration.bGenerateManifest)
+										{
+											bNeedsFullCPPIncludeRescan = true;
+											Log.TraceInformation("Performing full C++ include scan (no include cache file)");
+										}
+									}
+									Headers.FlatCPPIncludeDependencyCache = new FlatCPPIncludeDependencyCache(CacheFile);
+								}
+							}
+						}
+
+						if (bIsGatheringBuild)
+						{
+							List<FileItem> TargetOutputItems = new List<FileItem>();
+							List<UHTModuleInfo> TargetUObjectModules = new List<UHTModuleInfo>();
+							Dictionary<string, FileItem[]> TargetModuleNameToOutputItems = new Dictionary<string, FileItem[]>(StringComparer.OrdinalIgnoreCase);
+							BuildPrerequisites Prerequisites = new BuildPrerequisites();
+
+							BuildResult = Target.Build(BuildConfiguration, TargetToHeaders[Target], TargetOutputItems, TargetModuleNameToOutputItems, TargetUObjectModules, WorkingSet, ActionGraph, Prerequisites, bIsAssemblingBuild);
+							
+							if (BuildResult != ECompilationResult.Succeeded)
+							{
+								break;
+							}
+
+							// Export a JSON dump of this target
+							if (BuildConfiguration.JsonExportFile != null)
+							{
+								Directory.CreateDirectory(Path.GetDirectoryName(Path.GetFullPath(BuildConfiguration.JsonExportFile)));
+								Target.ExportJson(BuildConfiguration.JsonExportFile);
+							}
+
+							OutputItemsForAllTargets.AddRange(TargetOutputItems);
+
+							TargetPrerequisites.Add(Prerequisites);
+
+							// Update mapping of the target name to the list of UObject modules in this target
+							TargetNameToUObjectModules[Target.TargetName] = TargetUObjectModules;
+
+							// Merge the module name to output items map
+							foreach(KeyValuePair<string, FileItem[]> Pair in TargetModuleNameToOutputItems)
+							{
+								FileItem[] Items;
+								if(ModuleNameToOutputItems.TryGetValue(Pair.Key, out Items))
+								{
+									ModuleNameToOutputItems[Pair.Key] = Items.Concat(Pair.Value).ToArray();
+								}
+								else
+								{
+									ModuleNameToOutputItems[Pair.Key] = Pair.Value;
+								}
+							}
+
+							// Get the game module names
+							HotReloadModuleNamesForAllTargets.UnionWith(Target.GetHotReloadModuleNames());
+						}
+					}
+
+					if (BuildResult == ECompilationResult.Succeeded && !BuildConfiguration.bSkipBuild &&
+						(
+							(BuildConfiguration.bXGEExport && BuildConfiguration.bGenerateManifest) ||
+							(!GeneratingActionGraph && !BuildConfiguration.bGenerateManifest)
+						))
+					{
+						if (bIsGatheringBuild)
+						{
+							ActionGraph.FinalizeActionGraph();
+
+							Makefile = new UBTMakefile();
+							Makefile.AllActions = ActionGraph.AllActions;
+							Makefile.OutputItemsForAllTargets = OutputItemsForAllTargets;
+							foreach (System.Collections.DictionaryEntry EnvironmentVariable in Environment.GetEnvironmentVariables())
+							{
+								Makefile.EnvironmentVariables.Add(Tuple.Create((string)EnvironmentVariable.Key, (string)EnvironmentVariable.Value));
+							}
+							Makefile.TargetNameToUObjectModules = TargetNameToUObjectModules;
+							Makefile.ModuleNameToOutputItems = ModuleNameToOutputItems;
+							Makefile.HotReloadModuleNamesForAllTargets = HotReloadModuleNamesForAllTargets;
+							Makefile.Targets = Targets;
+							Makefile.bUseAdaptiveUnityBuild = Targets.Any(x => x.Rules.bUseAdaptiveUnityBuild);
+							Makefile.TargetPrerequisites = TargetPrerequisites;
+
+							if (BuildConfiguration.bUseUBTMakefiles)
+							{
+								// We've been told to prepare to build, so let's go ahead and save out our action graph so that we can use in a later invocation 
+								// to assemble the build.  Even if we are configured to assemble the build in this same invocation, we want to save out the
+								// Makefile so that it can be used on subsequent 'assemble only' runs, for the fastest possible iteration times
+								// @todo ubtmake: Optimization: We could make 'gather + assemble' mode slightly faster by saving this while busy compiling (on our worker thread)
+								using(Timeline.ScopeEvent("UBTMakefile.SaveUBTMakefile()"))
+								{
+									UBTMakefile.SaveUBTMakefile(TargetDescs, Makefile);
+								}
+							}
+						}
+
+						HashSet<string> OnlyModuleNames = new HashSet<string>(StringComparer.OrdinalIgnoreCase);
+						foreach(string Argument in Arguments)
+						{
+							const string ModuleArgument = "-Module=";
+							if(Argument.StartsWith(ModuleArgument, StringComparison.OrdinalIgnoreCase))
+							{
+								OnlyModuleNames.Add(Argument.Substring(ModuleArgument.Length));
+							}
+						}
+
+						Dictionary<string, int> HotReloadModuleNameToSuffix = new Dictionary<string, int>(StringComparer.OrdinalIgnoreCase);
+						foreach(string Argument in Arguments)
+						{
+							const string ModuleWithSuffixArgument = "-ModuleWithSuffix=";
+							if(Argument.StartsWith(ModuleWithSuffixArgument, StringComparison.OrdinalIgnoreCase))
+							{
+								string Value = Argument.Substring(ModuleWithSuffixArgument.Length);
+
+								int SuffixIdx = Value.LastIndexOf(',');
+								if(SuffixIdx == -1)
+								{
+									throw new BuildException("Missing suffix argument from -ModuleWithSuffix=Name,Suffix");
+								}
+
+								string ModuleName = Value.Substring(0, SuffixIdx);
+
+								int Suffix;
+								if(!Int32.TryParse(Value.Substring(SuffixIdx + 1), out Suffix))
+								{
+									throw new BuildException("Suffix for modules must be an integer");
+								}
+
+								HotReloadModuleNameToSuffix[ModuleName] = Suffix;
+							}
+						}
+
+						HotReloadMode HotReloadMode = HotReloadMode.Disabled;
+						if(Arguments.Any(x => x.Equals("-ForceHotReload", StringComparison.InvariantCultureIgnoreCase)))
+						{
+							HotReloadMode = HotReloadMode.FromIDE;
+						}
+						else if (TargetDescs.Count == 1 && !Arguments.Any(x => x.Equals("-NoHotReload", StringComparison.InvariantCultureIgnoreCase)))
+						{
+							if (BuildConfiguration.bAllowHotReloadFromIDE && HotReload.ShouldDoHotReloadFromIDE(BuildConfiguration, TargetDescs[0]))
+							{
+								HotReloadMode = HotReloadMode.FromIDE;
+							}
+							else if (HotReloadModuleNameToSuffix.Count > 0 && TargetDescs[0].ForeignPlugin == null)
+							{
+								HotReloadMode = HotReloadMode.FromEditor;
+							}
+						}
+						TargetDescriptor HotReloadTargetDesc = (HotReloadMode != HotReloadMode.Disabled) ? TargetDescs[0] : null;
+
+						if (bIsAssemblingBuild)
+						{
+							// If we didn't build the graph in this session, then we'll need to load a cached one
+							if (!bIsGatheringBuild && BuildResult.Succeeded())
+							{
+								ActionGraph.AllActions = Makefile.AllActions;
+
+								OutputItemsForAllTargets = Makefile.OutputItemsForAllTargets;
+
+								ModuleNameToOutputItems = Makefile.ModuleNameToOutputItems;
+
+								HotReloadModuleNamesForAllTargets = Makefile.HotReloadModuleNamesForAllTargets;
+
+								foreach (Tuple<string, string> EnvironmentVariable in Makefile.EnvironmentVariables)
+								{
+									// @todo ubtmake: There may be some variables we do NOT want to clobber.
+									Environment.SetEnvironmentVariable(EnvironmentVariable.Item1, EnvironmentVariable.Item2);
+								}
+
+								// Execute all the pre-build steps
+								if (!BuildConfiguration.bXGEExport)
+								{
+									foreach (UEBuildTarget Target in Targets)
+									{
+										if (!Target.ExecuteCustomPreBuildSteps())
+										{
+											BuildResult = ECompilationResult.OtherCompilationError;
+											break;
+										}
+									}
+								}
+
+								// If any of the targets need UHT to be run, we'll go ahead and do that now
+								if (BuildResult.Succeeded())
+								{
+									foreach (UEBuildTarget Target in Targets)
+									{
+										List<UHTModuleInfo> TargetUObjectModules;
+										if (Makefile.TargetNameToUObjectModules.TryGetValue(Target.TargetName, out TargetUObjectModules))
+										{
+											if (TargetUObjectModules.Count > 0)
+											{
+												// Execute the header tool
+												FileReference ModuleInfoFileName = FileReference.Combine(Target.ProjectIntermediateDirectory, Target.TargetName + ".uhtmanifest");
+												ECompilationResult UHTResult = ExternalExecution.ExecuteHeaderToolIfNecessary(BuildConfiguration, Target.ProjectFile, Target.TargetName, Target.TargetType, Target.bHasProjectScriptPlugin,  UObjectModules: TargetUObjectModules, ModuleInfoFileName: ModuleInfoFileName, bIsGatheringBuild: bIsGatheringBuild, bIsAssemblingBuild: bIsAssemblingBuild);
+												if(!UHTResult.Succeeded())
+												{
+													Log.TraceInformation("UnrealHeaderTool failed for target '" + Target.TargetName + "' (platform: " + Target.Platform.ToString() + ", module info: " + ModuleInfoFileName + ").");
+													BuildResult = UHTResult;
+													break;
+												}
+											}
+										}
+									}
+								}
+							}
+
+							if (BuildResult.Succeeded())
+							{
+								// Get the output items we're trying to build right now
+								List<FileItem> OutputItems = OutputItemsForAllTargets;
+
+								// Filter them if we have any specific modules specified
+								if (OnlyModuleNames.Count > 0)
+								{
+									OutputItems = new List<FileItem>();
+									foreach(string OnlyModuleName in OnlyModuleNames)
+									{
+										FileItem[] OutputItemsForModule;
+										if(!ModuleNameToOutputItems.TryGetValue(OnlyModuleName, out OutputItemsForModule))
+										{
+											throw new BuildException("Unable to find output items for module '{0}'", OnlyModuleName);
+										}
+										OutputItems.AddRange(OutputItemsForModule);
+									}
+								}
+
+								// For now simply treat all object files as the root target.
+								Action[] PrerequisiteActions;
+								{
+									HashSet<Action> PrerequisiteActionsSet = new HashSet<Action>();
+									foreach (FileItem OutputItem in OutputItems)
+									{
+										ActionGraph.GatherPrerequisiteActions(OutputItem, ref PrerequisiteActionsSet);
+									}
+									PrerequisiteActions = PrerequisiteActionsSet.ToArray();
+								}
+
+								// Apply the previous hot reload state
+								HotReloadState HotReloadState = null;
+								if(HotReloadMode == HotReloadMode.Disabled)
+								{
+									// Delete any previous state files; they are no longer valid
+									foreach(TargetDescriptor TargetDesc in TargetDescs)
+									{
+										FileReference HotReloadStateFile = HotReloadState.GetLocation(TargetDesc.ProjectFile, TargetDesc.Name, TargetDesc.Platform, TargetDesc.Configuration, TargetDesc.Architecture);
+										HotReload.DeleteTemporaryFiles(HotReloadStateFile);
+									}
+								}
+								else
+								{
+									// Read the previous state file and apply it to the action graph
+									FileReference HotReloadStateFile = HotReloadState.GetLocation(TargetDescs[0].ProjectFile, TargetDescs[0].Name, TargetDescs[0].Platform, TargetDescs[0].Configuration, TargetDescs[0].Architecture);
+									if(FileReference.Exists(HotReloadStateFile))
+									{
+										HotReloadState = HotReloadState.Load(HotReloadStateFile);
+									}
+									else
+									{
+										HotReloadState = new HotReloadState();
+									}
+
+									// Update the action graph to produce these new files
+									HotReload.PatchActionGraph(PrerequisiteActions, HotReloadState.OriginalFileToHotReloadFile);
+
+									// Update the module to output file mapping
+									foreach(string HotReloadModuleName in HotReloadModuleNamesForAllTargets)
+									{
+										FileItem[] ModuleOutputItems = ModuleNameToOutputItems[HotReloadModuleName];
+										for(int Idx = 0; Idx < ModuleOutputItems.Length; Idx++)
+										{
+											FileReference NewLocation;
+											if(HotReloadState.OriginalFileToHotReloadFile.TryGetValue(ModuleOutputItems[Idx].Location, out NewLocation))
+											{
+												ModuleOutputItems[Idx] = FileItem.GetItemByFileReference(NewLocation);
+											}
+										}
+									}
+
+									// If we want a specific suffix on any modules, apply that now. We'll track the outputs later, but the suffix has to be forced  (and is always out of date if it doesn't exist).
+									if(HotReloadModuleNameToSuffix.Count > 0)
+									{
+										Dictionary<FileReference, FileReference> OldLocationToNewLocation = new Dictionary<FileReference, FileReference>();
+										foreach(string HotReloadModuleName in HotReloadModuleNamesForAllTargets)
+										{
+											int ModuleSuffix;
+											if(HotReloadModuleNameToSuffix.TryGetValue(HotReloadModuleName, out ModuleSuffix))
+											{
+												FileItem[] ModuleOutputItems = ModuleNameToOutputItems[HotReloadModuleName];
+												foreach(FileItem ModuleOutputItem in ModuleOutputItems)
+												{
+													FileReference OldLocation = ModuleOutputItem.Location;
+													FileReference NewLocation = HotReload.ReplaceSuffix(OldLocation, ModuleSuffix);
+													OldLocationToNewLocation[OldLocation] = NewLocation;
+												}
+											}
+										}
+										HotReload.PatchActionGraph(PrerequisiteActions, OldLocationToNewLocation);
+									}
+								}
+
+								// Plan the actions to execute for the build.
+								Dictionary<UEBuildTarget, List<FileItem>> TargetToOutdatedPrerequisitesMap;
+								HashSet<Action> ActionsToExecute = ActionGraph.GetActionsToExecute(BuildConfiguration, PrerequisiteActions, Targets, TargetToHeaders, bIsAssemblingBuild, bNeedsFullCPPIncludeRescan, out TargetToOutdatedPrerequisitesMap);
+
+								// Patch action history for hot reload when running in assembler mode.  In assembler mode, the suffix on the output file will be
+								// the same for every invocation on that makefile, but we need a new suffix each time.
+								if (HotReloadMode != HotReloadMode.Disabled)
+								{
+									// For all the hot-reloadable modules that may need a unique suffix appended, build a mapping from output item to all the output items in that module. We can't 
+									// apply a suffix to one without applying a suffix to all of them.
+									Dictionary<FileItem, FileItem[]> HotReloadItemToDependentItems = new Dictionary<FileItem, FileItem[]>();
+									foreach(string HotReloadModuleName in HotReloadModuleNamesForAllTargets)
+									{
+										int ModuleSuffix;
+										if(!HotReloadModuleNameToSuffix.TryGetValue(HotReloadModuleName, out ModuleSuffix) || ModuleSuffix == -1)
+										{
+											FileItem[] ModuleOutputItems;
+											if(ModuleNameToOutputItems.TryGetValue(HotReloadModuleName, out ModuleOutputItems))
+											{
+												foreach(FileItem ModuleOutputItem in ModuleOutputItems)
+												{
+													HotReloadItemToDependentItems[ModuleOutputItem] = ModuleOutputItems;
+												}
+											}
+										} 
+									}
+
+									// Expand the list of actions to execute to include everything that references any files with a new suffix. Unlike a regular build, we can't ignore
+									// dependencies on import libraries under the assumption that a header would change if the API changes, because the dependency will be on a different DLL.
+									HashSet<FileItem> FilesRequiringSuffix = new HashSet<FileItem>(ActionsToExecute.SelectMany(x => x.ProducedItems).Where(x => HotReloadItemToDependentItems.ContainsKey(x)));
+									for(int LastNumFilesWithNewSuffix = 0; FilesRequiringSuffix.Count > LastNumFilesWithNewSuffix;)
+									{
+										LastNumFilesWithNewSuffix = FilesRequiringSuffix.Count;
+										foreach(Action PrerequisiteAction in PrerequisiteActions)
+										{
+											if(!ActionsToExecute.Contains(PrerequisiteAction))
+											{
+												foreach(FileItem ProducedItem in PrerequisiteAction.ProducedItems)
+												{
+													FileItem[] DependentItems;
+													if(HotReloadItemToDependentItems.TryGetValue(ProducedItem, out DependentItems))
+													{
+														ActionsToExecute.Add(PrerequisiteAction);
+														FilesRequiringSuffix.UnionWith(DependentItems);
+													}
+												}
+											}
+										}
+									}
+
+									// Build a list of file mappings
+									Dictionary<FileReference, FileReference> OldLocationToNewLocation = new Dictionary<FileReference, FileReference>();
+									foreach(FileItem FileRequiringSuffix in FilesRequiringSuffix)
+									{
+										FileReference OldLocation = FileRequiringSuffix.Location;
+										FileReference NewLocation = HotReload.ReplaceSuffix(OldLocation, HotReloadState.NextSuffix);
+										OldLocationToNewLocation[OldLocation] = NewLocation;
+									}
+
+									// Update the action graph with these new paths
+									HotReload.PatchActionGraph(PrerequisiteActions, OldLocationToNewLocation);
+
+									// Get a new list of actions to execute now that the graph has been modified
+									ActionsToExecute = ActionGraph.GetActionsToExecute(BuildConfiguration, PrerequisiteActions, Targets, TargetToHeaders, bIsAssemblingBuild, bNeedsFullCPPIncludeRescan, out TargetToOutdatedPrerequisitesMap);
+
+									// Build a mapping of all file items to their original
+									Dictionary<FileReference, FileReference> HotReloadFileToOriginalFile = new Dictionary<FileReference, FileReference>();
+									foreach(KeyValuePair<FileReference, FileReference> Pair in HotReloadState.OriginalFileToHotReloadFile)
+									{
+										HotReloadFileToOriginalFile[Pair.Value] = Pair.Key;
+									}
+									foreach(KeyValuePair<FileReference, FileReference> Pair in OldLocationToNewLocation)
+									{
+										FileReference OriginalLocation;
+										if(!HotReloadFileToOriginalFile.TryGetValue(Pair.Key, out OriginalLocation))
+										{
+											OriginalLocation = Pair.Key;
+										}
+										HotReloadFileToOriginalFile[Pair.Value] = OriginalLocation;
+									}
+
+									// Now filter out all the hot reload files and update the state
+									foreach(Action Action in ActionsToExecute)
+									{
+										foreach(FileItem ProducedItem in Action.ProducedItems)
+										{
+											FileReference OriginalLocation;
+											if(HotReloadFileToOriginalFile.TryGetValue(ProducedItem.Location, out OriginalLocation))
+											{
+												HotReloadState.OriginalFileToHotReloadFile[OriginalLocation] = ProducedItem.Location;
+												HotReloadState.TemporaryFiles.Add(ProducedItem.Location);
+											}
+										}
+									}
+
+									// Increment the suffix for the next iteration
+									if(ActionsToExecute.Count > 0)
+									{
+										HotReloadState.NextSuffix++;
+									}
+
+									// Save the new state
+									FileReference HotReloadStateFile = HotReloadState.GetLocation(TargetDescs[0].ProjectFile, TargetDescs[0].Name, TargetDescs[0].Platform, TargetDescs[0].Configuration, TargetDescs[0].Architecture);
+									HotReloadState.Save(HotReloadStateFile);
+								}
+
+								// Display some stats to the user.
+								Log.TraceVerbose(
+										"{0} actions, {1} outdated and requested actions",
+										ActionGraph.AllActions.Count,
+										ActionsToExecute.Count
+										);
+
+								// Cache indirect includes for all outdated C++ files.  We kick this off as a background thread so that it can
+								// perform the scan while we're compiling.  It usually only takes up to a few seconds, but we don't want to hurt
+								// our best case UBT iteration times for this task which can easily be performed asynchronously
+								if (BuildConfiguration.bUseUBTMakefiles && TargetToOutdatedPrerequisitesMap.Count > 0)
+								{
+									CppIncludeThread = new CppIncludeBackgroundThread(TargetToOutdatedPrerequisitesMap, TargetToHeaders);
+								}
+
+								// Execute the actions.
+								ExecutorTimer.Start();
+								if (BuildConfiguration.bXGEExport)
+								{
+									using(Timeline.ScopeEvent("XGE.ExportActions()"))
+									{
+										XGE.ExportActions(ActionsToExecute.ToList());
+									}
+									bSuccess = true;
+								}
+								else
+								{
+									using(Timeline.ScopeEvent("ActionGraph.ExecuteActions()"))
+									{
+										bSuccess = ActionGraph.ExecuteActions(BuildConfiguration, ActionsToExecute.ToList(), out ExecutorName);
+									}
+								}
+								ExecutorTimer.Stop();
+
+								// if the build succeeded, write the receipts and do any needed syncing
+								if (!bSuccess)
+								{
+									BuildResult = ECompilationResult.OtherCompilationError;
+								}
+							}
+
+							// Run the deployment steps
+							if (BuildResult.Succeeded()
+								&& BuildConfiguration.SingleFileToCompile == null
+								&& !BuildConfiguration.bGenerateManifest
+								&& !BuildConfiguration.bXGEExport)
+							{
+								foreach (UEBuildTarget Target in Targets)
+								{
+									if (Target.bDeployAfterCompile && HotReloadModuleNameToSuffix.Count == 0)
+									{
+										Log.WriteLine(LogEventType.Console, "Deploying {0} {1} {2}...", Target.TargetName, Target.Platform, Target.Configuration);
+										TargetReceipt Receipt = TargetReceipt.Read(Target.ReceiptFileName);
+										UEBuildPlatform.GetBuildPlatform(Target.Platform).Deploy(Receipt);
+									}
+								}
+							}
+						}
+					}
+				}
+			}
+			catch (Exception Ex)
+			{
+				Log.WriteException(Ex, (LogFile == null) ? null : LogFile.FullName);
+				BuildResult = ECompilationResult.OtherCompilationError;
+			}
+			finally
+			{
+				// Wait until our CPPIncludes dependency scanner thread has finished
+				if (CppIncludeThread != null)
+				{
+					using(Timeline.ScopeEvent("CppIncludeThread.Join()"))
+					{
+						CppIncludeThread.Join();
+					}
+				}
+
+				// Save the include dependency cache.
+				foreach (CPPHeaders Headers in TargetToHeaders.Values)
+				{
+					// NOTE: It's very important that we save the include cache, even if a build exception was thrown (compile error, etc), because we need to make sure that
+					//    any C++ include dependencies that we computed for out of date source files are saved.  Remember, the build may fail *after* some build products
+					//    are successfully built.  If we didn't save our dependency cache after build failures, source files for those build products that were successfully
+					//    built before the failure would not be considered out of date on the next run, so this is our only chance to cache C++ includes for those files!
+
+					if (Headers.IncludeDependencyCache != null)
+					{
+						using(Timeline.ScopeEvent("DependencyCache.Save()"))
+						{
+							Headers.IncludeDependencyCache.Save();
+						}
+					}
+
+					if (Headers.FlatCPPIncludeDependencyCache != null)
+					{
+						using (Timeline.ScopeEvent("FlatCPPIncludeDependencyCache.Save()"))
+						{
+							Headers.FlatCPPIncludeDependencyCache.Save();
+						}
+					}
+				}
+
+				// Figure out how long we took to execute.
+				if (ExecutorName != "Unknown")
+				{
+					Log.TraceInformation("Total build time: {0:0.00} seconds ({1} executor: {2:0.00} seconds)", BuildTimer.Elapsed.TotalSeconds, ExecutorName, ExecutorTimer.Elapsed.TotalSeconds);
+				}
 
 				// Print some performance info
 				Log.TraceLog("DirectIncludes cache miss time: {0}s ({1} misses)", CPPHeaders.DirectIncludeCacheMissesTotalTime, CPPHeaders.TotalDirectIncludeCacheMisses);
@@ -102,8 +899,76 @@ namespace UnrealBuildTool
 				Log.TraceLog("Deep C++ include scan time: {0}s", UnrealBuildTool.TotalDeepIncludeScanTime);
 				Log.TraceLog("Include Resolves: {0} ({1} misses, {2:0.00}%)", CPPHeaders.TotalDirectIncludeResolves, CPPHeaders.TotalDirectIncludeResolveCacheMisses, (float)CPPHeaders.TotalDirectIncludeResolveCacheMisses / (float)CPPHeaders.TotalDirectIncludeResolves * 100);
 				Log.TraceLog("Total FileItems: {0} ({1} missing)", FileItem.TotalFileItemCount, FileItem.MissingFileItemCount);
+			}
+			return (int)BuildResult;
+		}
+
+		/// <summary>
+		/// Invalidates makefiles for given target.
+		/// </summary>
+		/// <param name="Target">Target</param>
+		private static void InvalidateMakefiles(TargetDescriptor Target)
+		{
+			string[] MakefileNames = new string[] { "HotReloadMakefile.ubt", "Makefile.ubt" };
+			DirectoryReference BaseDir = UBTMakefile.GetUBTMakefileDirectoryPathForSingleTarget(Target);
+
+			foreach (string MakefileName in MakefileNames)
+			{
+				FileReference MakefileRef = FileReference.Combine(BaseDir, MakefileName);
+				if (FileReference.Exists(MakefileRef))
+				{
+					FileReference.Delete(MakefileRef);
+				}
+			}
+		}
 
-				return (int)Result;
+		/// <summary>
+		/// Helper class to update the C++ dependency cache on a background thread. Captures exceptions and re-throws on the main thread when joined.
+		/// </summary>
+		class CppIncludeBackgroundThread
+		{
+			Thread BackgroundThread;
+			Exception CaughtException;
+
+			public CppIncludeBackgroundThread(Dictionary<UEBuildTarget, List<FileItem>> TargetToOutdatedPrerequisitesMap, Dictionary<UEBuildTarget, CPPHeaders> TargetToHeaders)
+			{
+				BackgroundThread = new Thread(() => Run(TargetToOutdatedPrerequisitesMap, TargetToHeaders));
+				BackgroundThread.Start();
+			}
+
+			public void Join()
+			{
+				BackgroundThread.Join();
+				if (CaughtException != null)
+				{
+					throw CaughtException;
+				}
+			}
+
+			private void Run(Dictionary<UEBuildTarget, List<FileItem>> TargetToOutdatedPrerequisitesMap, Dictionary<UEBuildTarget, CPPHeaders> TargetToHeaders)
+			{
+				try
+				{
+					// @todo ubtmake: This thread will access data structures that are also used on the main UBT thread, but during this time UBT
+					// is only invoking the build executor, so should not be touching this stuff.  However, we need to at some guards to make sure.
+					foreach (KeyValuePair<UEBuildTarget, List<FileItem>> TargetAndOutdatedPrerequisites in TargetToOutdatedPrerequisitesMap)
+					{
+						UEBuildTarget Target = TargetAndOutdatedPrerequisites.Key;
+						List<FileItem> OutdatedPrerequisites = TargetAndOutdatedPrerequisites.Value;
+						CPPHeaders Headers = TargetToHeaders[Target];
+
+						foreach (FileItem PrerequisiteItem in OutdatedPrerequisites)
+						{
+							// Invoke our deep include scanner to figure out whether any of the files included by this source file have
+							// changed since the build product was built
+							Headers.FindAndCacheAllIncludedFiles(PrerequisiteItem, PrerequisiteItem.CachedIncludePaths, bOnlyCachedDependencies: false);
+						}
+					}
+				}
+				catch (Exception Ex)
+				{
+					CaughtException = Ex;
+				}
 			}
 		}
 	}
diff --git a/Engine/Source/Programs/UnrealBuildTool/Modes/GenerateProjectFilesMode.cs b/Engine/Source/Programs/UnrealBuildTool/Modes/GenerateProjectFilesMode.cs
--- a/Engine/Source/Programs/UnrealBuildTool/Modes/GenerateProjectFilesMode.cs
+++ b/Engine/Source/Programs/UnrealBuildTool/Modes/GenerateProjectFilesMode.cs
@@ -50,7 +50,7 @@ namespace UnrealBuildTool
 
 			// Parse rocket-specific arguments.
 			FileReference ProjectFile;
-			UnrealBuildTool.TryParseProjectFileArgument(Arguments, out ProjectFile);
+			TryParseProjectFileArgument(Arguments, out ProjectFile);
 
 			// If there aren't any formats set, read the default project file format from the config file
 			if (ProjectFileFormats.Count == 0)
@@ -165,5 +165,41 @@ namespace UnrealBuildTool
 			}
 			return (int)ECompilationResult.Succeeded;
 		}
+
+		/// <summary>
+		/// Try to parse the project file from the command line
+		/// </summary>
+		/// <param name="Arguments">The command line arguments</param>
+		/// <param name="ProjectFile">The project file that was parsed</param>
+		/// <returns>True if the project file was parsed, false otherwise</returns>
+		private static bool TryParseProjectFileArgument(CommandLineArguments Arguments, out FileReference ProjectFile)
+		{
+			FileReference ExplicitProjectFile;
+			if(Arguments.TryGetValue("-Project=", out ExplicitProjectFile))
+			{
+				ProjectFile = ExplicitProjectFile;
+				return true;
+			}
+
+			for(int Idx = 0; Idx < Arguments.Count; Idx++)
+			{
+				if(Arguments[Idx][0] != '-' && Arguments[Idx].EndsWith(".uproject", StringComparison.OrdinalIgnoreCase))
+				{
+					Arguments.MarkAsUsed(Idx);
+					ProjectFile = new FileReference(Arguments[Idx]);
+					return true;
+				}
+			}
+
+			FileReference InstalledProjectFile = UnrealBuildTool.GetInstalledProjectFile();
+			if(InstalledProjectFile != null)
+			{
+				ProjectFile = InstalledProjectFile;
+				return true;
+			}
+
+			ProjectFile = null;
+			return false;
+		}
 	}
 }
diff --git a/Engine/Source/Programs/UnrealBuildTool/UnrealBuildTool.cs b/Engine/Source/Programs/UnrealBuildTool/UnrealBuildTool.cs
--- a/Engine/Source/Programs/UnrealBuildTool/UnrealBuildTool.cs
+++ b/Engine/Source/Programs/UnrealBuildTool/UnrealBuildTool.cs
@@ -264,41 +264,6 @@ namespace UnrealBuildTool
 			}
 		}
 
-		/// <summary>
-		/// Try to parse the project file from the command line
-		/// </summary>
-		/// <param name="Arguments">The command line arguments</param>
-		/// <param name="ProjectFile">The project file that was parsed</param>
-		/// <returns>True if the project file was parsed, false otherwise</returns>
-		public static bool TryParseProjectFileArgument(CommandLineArguments Arguments, out FileReference ProjectFile)
-		{
-			FileReference ExplicitProjectFile;
-			if(Arguments.TryGetValue("-Project=", out ExplicitProjectFile))
-			{
-				ProjectFile = ExplicitProjectFile;
-				return true;
-			}
-
-			for(int Idx = 0; Idx < Arguments.Count; Idx++)
-			{
-				if(Arguments[Idx][0] != '-' && Arguments[Idx].EndsWith(".uproject", StringComparison.OrdinalIgnoreCase))
-				{
-					Arguments.MarkAsUsed(Idx);
-					ProjectFile = new FileReference(Arguments[Idx]);
-					return true;
-				}
-			}
-
-			if(IsProjectInstalled())
-			{
-				ProjectFile = InstalledProjectFile;
-				return true;
-			}
-
-			ProjectFile = null;
-			return false;
-		}
-
 		/// <summary>
 		/// Global options for UBT (any modes)
 		/// </summary>
@@ -498,879 +463,5 @@ namespace UnrealBuildTool
 				Trace.Close();
 			}
 		}
-
-		internal static ECompilationResult RunUBT(BuildConfiguration BuildConfiguration, string[] Arguments, FileReference LogFile)
-		{
-			bool bSuccess = true;
-
-			Stopwatch BuildTimer = Stopwatch.StartNew();
-
-			// Reset global configurations
-			ActionGraph ActionGraph = new ActionGraph();
-
-			string ExecutorName = "Unknown";
-			ECompilationResult BuildResult = ECompilationResult.Succeeded;
-
-			CppIncludeBackgroundThread CppIncludeThread = null;
-
-			List<UEBuildTarget> Targets = null;
-			Dictionary<UEBuildTarget, CPPHeaders> TargetToHeaders = new Dictionary<UEBuildTarget, CPPHeaders>();
-
-			Stopwatch ExecutorTimer = new Stopwatch();
-
-			try
-			{
-				int ArgumentIndex;
-				// action graph implies using the dependency resolve cache
-				bool GeneratingActionGraph = Utils.ParseCommandLineFlag(Arguments, "-graph", out ArgumentIndex);
-				if (GeneratingActionGraph)
-				{
-					BuildConfiguration.bUseIncludeDependencyResolveCache = true;
-				}
-
-				Timeline.AddEvent("RunUBT() initialization complete");
-
-				bool bSkipRulesCompile = Arguments.Any(x => x.Equals("-skiprulescompile", StringComparison.InvariantCultureIgnoreCase));
-
-				List<TargetDescriptor> TargetDescs = new List<TargetDescriptor>();
-				using(Timeline.ScopeEvent("TargetDescriptor.ParseCommandLine()"))
-				{
-					TargetDescs.AddRange(TargetDescriptor.ParseCommandLine(new CommandLineArguments(Arguments), BuildConfiguration.bUsePrecompiled, bSkipRulesCompile));
-				}
-
-				// Handle remote builds
-				for(int Idx = 0; Idx < TargetDescs.Count; Idx++)
-				{
-					TargetDescriptor TargetDesc = TargetDescs[Idx];
-					if(RemoteMac.HandlesTargetPlatform(TargetDesc.Platform))
-					{
-						FileReference RemoteLogFile = null;
-						if(LogFile != null)
-						{
-							RemoteLogFile = FileReference.Combine(LogFile.Directory, LogFile.GetFileNameWithoutExtension() + "_Remote.txt");
-						}
-
-						RemoteMac RemoteMac = new RemoteMac(TargetDesc.ProjectFile);
-						if(!RemoteMac.Build(TargetDesc, LogFile))
-						{
-							return ECompilationResult.Unknown;
-						}
-
-						TargetDescs.RemoveAt(Idx--);
-						if(TargetDescs.Count == 0)
-						{
-							return ECompilationResult.Succeeded;
-						}
-					}
-				}
-
-				if (Arguments.Any(x => x.Equals("-InvalidateMakefilesOnly", StringComparison.InvariantCultureIgnoreCase)))
-				{
-					Log.TraceInformation("Invalidating makefiles only in this run.");
-					if (TargetDescs.Count != 1)
-					{
-						Log.TraceError("You have to provide one target name for makefile invalidation.");
-						return ECompilationResult.OtherCompilationError;
-					}
-
-					InvalidateMakefiles(TargetDescs[0]);
-					return ECompilationResult.Succeeded;
-				}
-
-				// Used when BuildConfiguration.bUseUBTMakefiles is enabled.  If true, it means that our cached includes may not longer be
-				// valid (or never were), and we need to recover by forcibly scanning included headers for all build prerequisites to make sure that our
-				// cached set of includes is actually correct, before determining which files are outdated.
-				bool bNeedsFullCPPIncludeRescan = false;
-
-				if (BuildConfiguration.bUseUBTMakefiles)
-				{
-					// Only the modular editor and game targets will share build products.  Unfortunately, we can't determine at
-					// at this point whether we're dealing with modular or monolithic game binaries, so we opt to always invalidate
-					// cached includes if the target we're switching to is either a game target (has project file) or "UE4Editor".
-					bool bMightHaveSharedBuildProducts =
-						TargetDescs[0].ProjectFile != null ||  // Is this a game? (has a .uproject file for the target)
-						TargetDescs[0].Name.Equals("UE4Editor", StringComparison.InvariantCultureIgnoreCase); // Is the engine?
-					if (bMightHaveSharedBuildProducts)
-					{
-						bool bIsBuildingSameTargetsAsLastTime = false;
-
-						string TargetCollectionName = UBTMakefile.MakeTargetCollectionName(TargetDescs);
-
-						string LastBuiltTargetsFileName = "LastBuiltTargets.txt";
-						string LastBuiltTargetsFilePath = FileReference.Combine(UnrealBuildTool.EngineDirectory, "Intermediate", "Build", LastBuiltTargetsFileName).FullName;
-						if (File.Exists(LastBuiltTargetsFilePath) && Utils.ReadAllText(LastBuiltTargetsFilePath) == TargetCollectionName)
-						{
-							// @todo ubtmake: Because we're using separate files for hot reload vs. full compiles, it's actually possible that includes will
-							// become out of date without us knowing if the developer ping-pongs between hot reloading target A and building target B normally.
-							// To fix this we can not use a different file name for last built targets, but the downside is slower performance when
-							// performing the first hot reload after compiling normally (forces full include dependency scan)
-							bIsBuildingSameTargetsAsLastTime = true;
-						}
-
-						// Save out the name of the targets we're building
-						if (!bIsBuildingSameTargetsAsLastTime)
-						{
-							Directory.CreateDirectory(Path.GetDirectoryName(LastBuiltTargetsFilePath));
-							File.WriteAllText(LastBuiltTargetsFilePath, TargetCollectionName, Encoding.UTF8);
-
-							// Can't use super fast include checking unless we're building the same set of targets as last time, because
-							// we might not know about all of the C++ include dependencies for already-up-to-date shared build products
-							// between the targets
-							bNeedsFullCPPIncludeRescan = true;
-							Log.TraceInformation("Performing full C++ include scan (building a new target)");
-						}
-					}
-				}
-
-				// True if we should gather module dependencies for building in this run.  If this is false, then we'll expect to be loading information from disk about
-				// the target's modules before we'll be able to build anything.  One or both of IsGatheringBuild or IsAssemblingBuild must be true.
-				bool bIsGatheringBuild = true;
-
-				// True if we should go ahead and actually compile and link in this run.  If this is false, no actions will be executed and we'll stop after gathering
-				// and saving information about dependencies.  One or both of IsGatheringBuild or IsAssemblingBuild must be true.
-				bool bIsAssemblingBuild = true;
-
-				// If we were asked to enable fast build iteration, we want the 'gather' phase to default to off (unless it is overridden below
-				// using a command-line option.)
-				if (BuildConfiguration.bUseUBTMakefiles)
-				{
-					bIsGatheringBuild = false;
-				}
-
-				foreach (string Argument in Arguments)
-				{
-					string LowercaseArg = Argument.ToLowerInvariant();
-					if (LowercaseArg == "-gather")
-					{
-						bIsGatheringBuild = true;
-					}
-					else if (LowercaseArg == "-nogather")
-					{
-						bIsGatheringBuild = false;
-					}
-					else if (LowercaseArg == "-gatheronly")
-					{
-						bIsGatheringBuild = true;
-						bIsAssemblingBuild = false;
-					}
-					else if (LowercaseArg == "-assemble")
-					{
-						bIsAssemblingBuild = true;
-					}
-					else if (LowercaseArg == "-noassemble")
-					{
-						bIsAssemblingBuild = false;
-					}
-					else if (LowercaseArg == "-assembleonly")
-					{
-						bIsGatheringBuild = false;
-						bIsAssemblingBuild = true;
-					}
-				}
-
-				if (!bIsGatheringBuild && !bIsAssemblingBuild)
-				{
-					throw new BuildException("UnrealBuildTool: At least one of either IsGatheringBuild or IsAssemblingBuild must be true.  Did you pass '-NoGather' with '-NoAssemble'?");
-				}
-
-				// Get a set of all the project directories.
-				HashSet<DirectoryReference> ProjectDirs = new HashSet<DirectoryReference>();
-				foreach(TargetDescriptor TargetDesc in TargetDescs)
-				{
-					if(TargetDesc.ProjectFile != null)
-					{
-						ProjectDirs.Add(TargetDesc.ProjectFile.Directory);
-					}
-				}
-
-				// Create the working set provider
-				using (ISourceFileWorkingSet WorkingSet = SourceFileWorkingSet.Create(UnrealBuildTool.RootDirectory, ProjectDirs))
-				{
-					UBTMakefile Makefile = null;
-					{
-						FileReference UBTMakefilePath = UBTMakefile.GetUBTMakefilePath(TargetDescs);
-
-						// Make sure the gather phase is executed if we're not actually building anything
-						if (BuildConfiguration.bGenerateManifest || BuildConfiguration.bXGEExport || GeneratingActionGraph)
-						{
-							bIsGatheringBuild = true;
-						}
-
-						// Were we asked to run in 'assembler only' mode?  If so, let's check to see if that's even possible by seeing if
-						// we have a valid UBTMakefile already saved to disk, ready for us to load.
-						if (bIsAssemblingBuild && !bIsGatheringBuild)
-						{
-							// @todo ubtmake: Mildly terrified of BuildConfiguration/UEBuildConfiguration globals that were set during the Prepare phase but are not set now.  We may need to save/load all of these, otherwise
-							//		we'll need to call SetupGlobalEnvironment on all of the targets (maybe other stuff, too.  See PreBuildStep())
-
-							// Try to load the UBTMakefile.  It will only be loaded if it has valid content and is not determined to be out of date.    
-							string ReasonNotLoaded;
-							using(Timeline.ScopeEvent("UBTMakefile.LoadUBTMakefile"))
-							{
-								Makefile = UBTMakefile.Load(UBTMakefilePath, TargetDescs[0].ProjectFile, WorkingSet, out ReasonNotLoaded);
-							}
-
-							if (Makefile == null)
-							{
-								// If the Makefile couldn't be loaded, then we're not going to be able to continue in "assembler only" mode.  We'll do both
-								// a 'gather' and 'assemble' in the same run.  This will take a while longer, but subsequent runs will be fast!
-								bIsGatheringBuild = true;
-
-								FileItem.ClearCaches();
-
-								Log.TraceInformation("Creating makefile for {0}{1} ({2})",
-									TargetDescs[0].Name,
-									TargetDescs.Count > 1 ? (" (and " + (TargetDescs.Count - 1).ToString() + " more)") : "",
-									ReasonNotLoaded);
-							}
-						}
-					}
-
-
-					if (Makefile != null && !bIsGatheringBuild && bIsAssemblingBuild)
-					{
-						// If we've loaded a makefile, then we can fill target information from this file!
-						Targets = Makefile.Targets;
-					}
-					else
-					{
-						Targets = new List<UEBuildTarget>();
-						foreach (TargetDescriptor TargetDesc in TargetDescs)
-						{
-							using(Timeline.ScopeEvent("UEBuildTarget.CreateTarget()"))
-							{
-								UEBuildTarget Target = UEBuildTarget.CreateTarget(TargetDesc, bSkipRulesCompile, BuildConfiguration.SingleFileToCompile != null, BuildConfiguration.bUsePrecompiled);
-								Targets.Add(Target);
-							}
-						}
-					}
-
-					// Build action lists for all passed in targets.
-					List<FileItem> OutputItemsForAllTargets = new List<FileItem>();
-					Dictionary<string, FileItem[]> ModuleNameToOutputItems = new Dictionary<string, FileItem[]>(StringComparer.OrdinalIgnoreCase);
-					Dictionary<string, List<UHTModuleInfo>> TargetNameToUObjectModules = new Dictionary<string, List<UHTModuleInfo>>(StringComparer.InvariantCultureIgnoreCase);
-					HashSet<string> HotReloadModuleNamesForAllTargets = new HashSet<string>(StringComparer.OrdinalIgnoreCase);
-					List<BuildPrerequisites> TargetPrerequisites = new List<BuildPrerequisites>();
-					foreach (UEBuildTarget Target in Targets)
-					{
-						// Create the header cache for this target
-						FileReference DependencyCacheFile = DependencyCache.GetDependencyCachePathForTarget(Target.ProjectFile, Target.Platform, Target.TargetName);
-						bool bUseFlatCPPIncludeDependencyCache = BuildConfiguration.bUseUBTMakefiles && bIsAssemblingBuild;
-						CPPHeaders Headers = new CPPHeaders(Target.ProjectFile, DependencyCacheFile, bUseFlatCPPIncludeDependencyCache, BuildConfiguration.bUseUBTMakefiles, BuildConfiguration.bUseIncludeDependencyResolveCache, BuildConfiguration.bTestIncludeDependencyResolveCache);
-						TargetToHeaders[Target] = Headers;
-
-						// Make sure the appropriate executor is selected
-						UEBuildPlatform BuildPlatform = UEBuildPlatform.GetBuildPlatform(Target.Platform);
-						BuildConfiguration.bAllowXGE &= BuildPlatform.CanUseXGE();
-						BuildConfiguration.bAllowDistcc &= BuildPlatform.CanUseDistcc();
-						BuildConfiguration.bAllowSNDBS &= BuildPlatform.CanUseSNDBS();
-
-						// When in 'assembler only' mode, we'll load this cache later on a worker thread.  It takes a long time to load!
-						if (!(!bIsGatheringBuild && bIsAssemblingBuild))
-						{
-							// Load the direct include dependency cache.
-							using(Timeline.ScopeEvent("DependencyCache.Create()"))
-							{
-								Headers.IncludeDependencyCache = DependencyCache.Create(DependencyCache.GetDependencyCachePathForTarget(Target.ProjectFile, Target.Platform, Target.TargetName));
-							}
-						}
-
-						// We don't need this dependency cache in 'gather only' mode
-						if (BuildConfiguration.bUseUBTMakefiles &&
-								!(bIsGatheringBuild && !bIsAssemblingBuild))
-						{
-							using(Timeline.ScopeEvent("FlatCPPIncludeDependencyCache.TryRead()"))
-							{
-								// Load the cache that contains the list of flattened resolved includes for resolved source files
-								// @todo ubtmake: Ideally load this asynchronously at startup and only block when it is first needed and not finished loading
-								FileReference CacheFile = FlatCPPIncludeDependencyCache.GetDependencyCachePathForTarget(Target.ProjectFile, Target.TargetName, Target.Platform, Target.Architecture);
-								if (!FlatCPPIncludeDependencyCache.TryRead(CacheFile, out Headers.FlatCPPIncludeDependencyCache))
-								{
-									if (!bNeedsFullCPPIncludeRescan)
-									{
-										if (!BuildConfiguration.bXGEExport && !BuildConfiguration.bGenerateManifest)
-										{
-											bNeedsFullCPPIncludeRescan = true;
-											Log.TraceInformation("Performing full C++ include scan (no include cache file)");
-										}
-									}
-									Headers.FlatCPPIncludeDependencyCache = new FlatCPPIncludeDependencyCache(CacheFile);
-								}
-							}
-						}
-
-						if (bIsGatheringBuild)
-						{
-							List<FileItem> TargetOutputItems = new List<FileItem>();
-							List<UHTModuleInfo> TargetUObjectModules = new List<UHTModuleInfo>();
-							Dictionary<string, FileItem[]> TargetModuleNameToOutputItems = new Dictionary<string, FileItem[]>(StringComparer.OrdinalIgnoreCase);
-							BuildPrerequisites Prerequisites = new BuildPrerequisites();
-
-							BuildResult = Target.Build(BuildConfiguration, TargetToHeaders[Target], TargetOutputItems, TargetModuleNameToOutputItems, TargetUObjectModules, WorkingSet, ActionGraph, Prerequisites, bIsAssemblingBuild);
-							
-							if (BuildResult != ECompilationResult.Succeeded)
-							{
-								break;
-							}
-
-							// Export a JSON dump of this target
-							if (BuildConfiguration.JsonExportFile != null)
-							{
-								Directory.CreateDirectory(Path.GetDirectoryName(Path.GetFullPath(BuildConfiguration.JsonExportFile)));
-								Target.ExportJson(BuildConfiguration.JsonExportFile);
-							}
-
-							OutputItemsForAllTargets.AddRange(TargetOutputItems);
-
-							TargetPrerequisites.Add(Prerequisites);
-
-							// Update mapping of the target name to the list of UObject modules in this target
-							TargetNameToUObjectModules[Target.TargetName] = TargetUObjectModules;
-
-							// Merge the module name to output items map
-							foreach(KeyValuePair<string, FileItem[]> Pair in TargetModuleNameToOutputItems)
-							{
-								FileItem[] Items;
-								if(ModuleNameToOutputItems.TryGetValue(Pair.Key, out Items))
-								{
-									ModuleNameToOutputItems[Pair.Key] = Items.Concat(Pair.Value).ToArray();
-								}
-								else
-								{
-									ModuleNameToOutputItems[Pair.Key] = Pair.Value;
-								}
-							}
-
-							// Get the game module names
-							HotReloadModuleNamesForAllTargets.UnionWith(Target.GetHotReloadModuleNames());
-						}
-					}
-
-					if (BuildResult == ECompilationResult.Succeeded && !BuildConfiguration.bSkipBuild &&
-						(
-							(BuildConfiguration.bXGEExport && BuildConfiguration.bGenerateManifest) ||
-							(!GeneratingActionGraph && !BuildConfiguration.bGenerateManifest)
-						))
-					{
-						if (bIsGatheringBuild)
-						{
-							ActionGraph.FinalizeActionGraph();
-
-							Makefile = new UBTMakefile();
-							Makefile.AllActions = ActionGraph.AllActions;
-							Makefile.OutputItemsForAllTargets = OutputItemsForAllTargets;
-							foreach (System.Collections.DictionaryEntry EnvironmentVariable in Environment.GetEnvironmentVariables())
-							{
-								Makefile.EnvironmentVariables.Add(Tuple.Create((string)EnvironmentVariable.Key, (string)EnvironmentVariable.Value));
-							}
-							Makefile.TargetNameToUObjectModules = TargetNameToUObjectModules;
-							Makefile.ModuleNameToOutputItems = ModuleNameToOutputItems;
-							Makefile.HotReloadModuleNamesForAllTargets = HotReloadModuleNamesForAllTargets;
-							Makefile.Targets = Targets;
-							Makefile.bUseAdaptiveUnityBuild = Targets.Any(x => x.Rules.bUseAdaptiveUnityBuild);
-							Makefile.TargetPrerequisites = TargetPrerequisites;
-
-							if (BuildConfiguration.bUseUBTMakefiles)
-							{
-								// We've been told to prepare to build, so let's go ahead and save out our action graph so that we can use in a later invocation 
-								// to assemble the build.  Even if we are configured to assemble the build in this same invocation, we want to save out the
-								// Makefile so that it can be used on subsequent 'assemble only' runs, for the fastest possible iteration times
-								// @todo ubtmake: Optimization: We could make 'gather + assemble' mode slightly faster by saving this while busy compiling (on our worker thread)
-								using(Timeline.ScopeEvent("UBTMakefile.SaveUBTMakefile()"))
-								{
-									UBTMakefile.SaveUBTMakefile(TargetDescs, Makefile);
-								}
-							}
-						}
-
-						HashSet<string> OnlyModuleNames = new HashSet<string>(StringComparer.OrdinalIgnoreCase);
-						foreach(string Argument in Arguments)
-						{
-							const string ModuleArgument = "-Module=";
-							if(Argument.StartsWith(ModuleArgument, StringComparison.OrdinalIgnoreCase))
-							{
-								OnlyModuleNames.Add(Argument.Substring(ModuleArgument.Length));
-							}
-						}
-
-						Dictionary<string, int> HotReloadModuleNameToSuffix = new Dictionary<string, int>(StringComparer.OrdinalIgnoreCase);
-						foreach(string Argument in Arguments)
-						{
-							const string ModuleWithSuffixArgument = "-ModuleWithSuffix=";
-							if(Argument.StartsWith(ModuleWithSuffixArgument, StringComparison.OrdinalIgnoreCase))
-							{
-								string Value = Argument.Substring(ModuleWithSuffixArgument.Length);
-
-								int SuffixIdx = Value.LastIndexOf(',');
-								if(SuffixIdx == -1)
-								{
-									throw new BuildException("Missing suffix argument from -ModuleWithSuffix=Name,Suffix");
-								}
-
-								string ModuleName = Value.Substring(0, SuffixIdx);
-
-								int Suffix;
-								if(!Int32.TryParse(Value.Substring(SuffixIdx + 1), out Suffix))
-								{
-									throw new BuildException("Suffix for modules must be an integer");
-								}
-
-								HotReloadModuleNameToSuffix[ModuleName] = Suffix;
-							}
-						}
-
-						HotReloadMode HotReloadMode = HotReloadMode.Disabled;
-						if(Arguments.Any(x => x.Equals("-ForceHotReload", StringComparison.InvariantCultureIgnoreCase)))
-						{
-							HotReloadMode = HotReloadMode.FromIDE;
-						}
-						else if (TargetDescs.Count == 1 && !Arguments.Any(x => x.Equals("-NoHotReload", StringComparison.InvariantCultureIgnoreCase)))
-						{
-							if (BuildConfiguration.bAllowHotReloadFromIDE && HotReload.ShouldDoHotReloadFromIDE(BuildConfiguration, TargetDescs[0]))
-							{
-								HotReloadMode = HotReloadMode.FromIDE;
-							}
-							else if (HotReloadModuleNameToSuffix.Count > 0 && TargetDescs[0].ForeignPlugin == null)
-							{
-								HotReloadMode = HotReloadMode.FromEditor;
-							}
-						}
-						TargetDescriptor HotReloadTargetDesc = (HotReloadMode != HotReloadMode.Disabled) ? TargetDescs[0] : null;
-
-						if (bIsAssemblingBuild)
-						{
-							// If we didn't build the graph in this session, then we'll need to load a cached one
-							if (!bIsGatheringBuild && BuildResult.Succeeded())
-							{
-								ActionGraph.AllActions = Makefile.AllActions;
-
-								OutputItemsForAllTargets = Makefile.OutputItemsForAllTargets;
-
-								ModuleNameToOutputItems = Makefile.ModuleNameToOutputItems;
-
-								HotReloadModuleNamesForAllTargets = Makefile.HotReloadModuleNamesForAllTargets;
-
-								foreach (Tuple<string, string> EnvironmentVariable in Makefile.EnvironmentVariables)
-								{
-									// @todo ubtmake: There may be some variables we do NOT want to clobber.
-									Environment.SetEnvironmentVariable(EnvironmentVariable.Item1, EnvironmentVariable.Item2);
-								}
-
-								// Execute all the pre-build steps
-								if (!BuildConfiguration.bXGEExport)
-								{
-									foreach (UEBuildTarget Target in Targets)
-									{
-										if (!Target.ExecuteCustomPreBuildSteps())
-										{
-											BuildResult = ECompilationResult.OtherCompilationError;
-											break;
-										}
-									}
-								}
-
-								// If any of the targets need UHT to be run, we'll go ahead and do that now
-								if (BuildResult.Succeeded())
-								{
-									foreach (UEBuildTarget Target in Targets)
-									{
-										List<UHTModuleInfo> TargetUObjectModules;
-										if (Makefile.TargetNameToUObjectModules.TryGetValue(Target.TargetName, out TargetUObjectModules))
-										{
-											if (TargetUObjectModules.Count > 0)
-											{
-												// Execute the header tool
-												FileReference ModuleInfoFileName = FileReference.Combine(Target.ProjectIntermediateDirectory, Target.TargetName + ".uhtmanifest");
-												ECompilationResult UHTResult = ExternalExecution.ExecuteHeaderToolIfNecessary(BuildConfiguration, Target.ProjectFile, Target.TargetName, Target.TargetType, Target.bHasProjectScriptPlugin,  UObjectModules: TargetUObjectModules, ModuleInfoFileName: ModuleInfoFileName, bIsGatheringBuild: bIsGatheringBuild, bIsAssemblingBuild: bIsAssemblingBuild);
-												if(!UHTResult.Succeeded())
-												{
-													Log.TraceInformation("UnrealHeaderTool failed for target '" + Target.TargetName + "' (platform: " + Target.Platform.ToString() + ", module info: " + ModuleInfoFileName + ").");
-													BuildResult = UHTResult;
-													break;
-												}
-											}
-										}
-									}
-								}
-							}
-
-							if (BuildResult.Succeeded())
-							{
-								// Get the output items we're trying to build right now
-								List<FileItem> OutputItems = OutputItemsForAllTargets;
-
-								// Filter them if we have any specific modules specified
-								if (OnlyModuleNames.Count > 0)
-								{
-									OutputItems = new List<FileItem>();
-									foreach(string OnlyModuleName in OnlyModuleNames)
-									{
-										FileItem[] OutputItemsForModule;
-										if(!ModuleNameToOutputItems.TryGetValue(OnlyModuleName, out OutputItemsForModule))
-										{
-											throw new BuildException("Unable to find output items for module '{0}'", OnlyModuleName);
-										}
-										OutputItems.AddRange(OutputItemsForModule);
-									}
-								}
-
-								// For now simply treat all object files as the root target.
-								Action[] PrerequisiteActions;
-								{
-									HashSet<Action> PrerequisiteActionsSet = new HashSet<Action>();
-									foreach (FileItem OutputItem in OutputItems)
-									{
-										ActionGraph.GatherPrerequisiteActions(OutputItem, ref PrerequisiteActionsSet);
-									}
-									PrerequisiteActions = PrerequisiteActionsSet.ToArray();
-								}
-
-								// Apply the previous hot reload state
-								HotReloadState HotReloadState = null;
-								if(HotReloadMode == HotReloadMode.Disabled)
-								{
-									// Delete any previous state files; they are no longer valid
-									foreach(TargetDescriptor TargetDesc in TargetDescs)
-									{
-										FileReference HotReloadStateFile = HotReloadState.GetLocation(TargetDesc.ProjectFile, TargetDesc.Name, TargetDesc.Platform, TargetDesc.Configuration, TargetDesc.Architecture);
-										HotReload.DeleteTemporaryFiles(HotReloadStateFile);
-									}
-								}
-								else
-								{
-									// Read the previous state file and apply it to the action graph
-									FileReference HotReloadStateFile = HotReloadState.GetLocation(TargetDescs[0].ProjectFile, TargetDescs[0].Name, TargetDescs[0].Platform, TargetDescs[0].Configuration, TargetDescs[0].Architecture);
-									if(FileReference.Exists(HotReloadStateFile))
-									{
-										HotReloadState = HotReloadState.Load(HotReloadStateFile);
-									}
-									else
-									{
-										HotReloadState = new HotReloadState();
-									}
-
-									// Update the action graph to produce these new files
-									HotReload.PatchActionGraph(PrerequisiteActions, HotReloadState.OriginalFileToHotReloadFile);
-
-									// Update the module to output file mapping
-									foreach(string HotReloadModuleName in HotReloadModuleNamesForAllTargets)
-									{
-										FileItem[] ModuleOutputItems = ModuleNameToOutputItems[HotReloadModuleName];
-										for(int Idx = 0; Idx < ModuleOutputItems.Length; Idx++)
-										{
-											FileReference NewLocation;
-											if(HotReloadState.OriginalFileToHotReloadFile.TryGetValue(ModuleOutputItems[Idx].Location, out NewLocation))
-											{
-												ModuleOutputItems[Idx] = FileItem.GetItemByFileReference(NewLocation);
-											}
-										}
-									}
-
-									// If we want a specific suffix on any modules, apply that now. We'll track the outputs later, but the suffix has to be forced  (and is always out of date if it doesn't exist).
-									if(HotReloadModuleNameToSuffix.Count > 0)
-									{
-										Dictionary<FileReference, FileReference> OldLocationToNewLocation = new Dictionary<FileReference, FileReference>();
-										foreach(string HotReloadModuleName in HotReloadModuleNamesForAllTargets)
-										{
-											int ModuleSuffix;
-											if(HotReloadModuleNameToSuffix.TryGetValue(HotReloadModuleName, out ModuleSuffix))
-											{
-												FileItem[] ModuleOutputItems = ModuleNameToOutputItems[HotReloadModuleName];
-												foreach(FileItem ModuleOutputItem in ModuleOutputItems)
-												{
-													FileReference OldLocation = ModuleOutputItem.Location;
-													FileReference NewLocation = HotReload.ReplaceSuffix(OldLocation, ModuleSuffix);
-													OldLocationToNewLocation[OldLocation] = NewLocation;
-												}
-											}
-										}
-										HotReload.PatchActionGraph(PrerequisiteActions, OldLocationToNewLocation);
-									}
-								}
-
-								// Plan the actions to execute for the build.
-								Dictionary<UEBuildTarget, List<FileItem>> TargetToOutdatedPrerequisitesMap;
-								HashSet<Action> ActionsToExecute = ActionGraph.GetActionsToExecute(BuildConfiguration, PrerequisiteActions, Targets, TargetToHeaders, bIsAssemblingBuild, bNeedsFullCPPIncludeRescan, out TargetToOutdatedPrerequisitesMap);
-
-								// Patch action history for hot reload when running in assembler mode.  In assembler mode, the suffix on the output file will be
-								// the same for every invocation on that makefile, but we need a new suffix each time.
-								if (HotReloadMode != HotReloadMode.Disabled)
-								{
-									// For all the hot-reloadable modules that may need a unique suffix appended, build a mapping from output item to all the output items in that module. We can't 
-									// apply a suffix to one without applying a suffix to all of them.
-									Dictionary<FileItem, FileItem[]> HotReloadItemToDependentItems = new Dictionary<FileItem, FileItem[]>();
-									foreach(string HotReloadModuleName in HotReloadModuleNamesForAllTargets)
-									{
-										int ModuleSuffix;
-										if(!HotReloadModuleNameToSuffix.TryGetValue(HotReloadModuleName, out ModuleSuffix) || ModuleSuffix == -1)
-										{
-											FileItem[] ModuleOutputItems;
-											if(ModuleNameToOutputItems.TryGetValue(HotReloadModuleName, out ModuleOutputItems))
-											{
-												foreach(FileItem ModuleOutputItem in ModuleOutputItems)
-												{
-													HotReloadItemToDependentItems[ModuleOutputItem] = ModuleOutputItems;
-												}
-											}
-										} 
-									}
-
-									// Expand the list of actions to execute to include everything that references any files with a new suffix. Unlike a regular build, we can't ignore
-									// dependencies on import libraries under the assumption that a header would change if the API changes, because the dependency will be on a different DLL.
-									HashSet<FileItem> FilesRequiringSuffix = new HashSet<FileItem>(ActionsToExecute.SelectMany(x => x.ProducedItems).Where(x => HotReloadItemToDependentItems.ContainsKey(x)));
-									for(int LastNumFilesWithNewSuffix = 0; FilesRequiringSuffix.Count > LastNumFilesWithNewSuffix;)
-									{
-										LastNumFilesWithNewSuffix = FilesRequiringSuffix.Count;
-										foreach(Action PrerequisiteAction in PrerequisiteActions)
-										{
-											if(!ActionsToExecute.Contains(PrerequisiteAction))
-											{
-												foreach(FileItem ProducedItem in PrerequisiteAction.ProducedItems)
-												{
-													FileItem[] DependentItems;
-													if(HotReloadItemToDependentItems.TryGetValue(ProducedItem, out DependentItems))
-													{
-														ActionsToExecute.Add(PrerequisiteAction);
-														FilesRequiringSuffix.UnionWith(DependentItems);
-													}
-												}
-											}
-										}
-									}
-
-									// Build a list of file mappings
-									Dictionary<FileReference, FileReference> OldLocationToNewLocation = new Dictionary<FileReference, FileReference>();
-									foreach(FileItem FileRequiringSuffix in FilesRequiringSuffix)
-									{
-										FileReference OldLocation = FileRequiringSuffix.Location;
-										FileReference NewLocation = HotReload.ReplaceSuffix(OldLocation, HotReloadState.NextSuffix);
-										OldLocationToNewLocation[OldLocation] = NewLocation;
-									}
-
-									// Update the action graph with these new paths
-									HotReload.PatchActionGraph(PrerequisiteActions, OldLocationToNewLocation);
-
-									// Get a new list of actions to execute now that the graph has been modified
-									ActionsToExecute = ActionGraph.GetActionsToExecute(BuildConfiguration, PrerequisiteActions, Targets, TargetToHeaders, bIsAssemblingBuild, bNeedsFullCPPIncludeRescan, out TargetToOutdatedPrerequisitesMap);
-
-									// Build a mapping of all file items to their original
-									Dictionary<FileReference, FileReference> HotReloadFileToOriginalFile = new Dictionary<FileReference, FileReference>();
-									foreach(KeyValuePair<FileReference, FileReference> Pair in HotReloadState.OriginalFileToHotReloadFile)
-									{
-										HotReloadFileToOriginalFile[Pair.Value] = Pair.Key;
-									}
-									foreach(KeyValuePair<FileReference, FileReference> Pair in OldLocationToNewLocation)
-									{
-										FileReference OriginalLocation;
-										if(!HotReloadFileToOriginalFile.TryGetValue(Pair.Key, out OriginalLocation))
-										{
-											OriginalLocation = Pair.Key;
-										}
-										HotReloadFileToOriginalFile[Pair.Value] = OriginalLocation;
-									}
-
-									// Now filter out all the hot reload files and update the state
-									foreach(Action Action in ActionsToExecute)
-									{
-										foreach(FileItem ProducedItem in Action.ProducedItems)
-										{
-											FileReference OriginalLocation;
-											if(HotReloadFileToOriginalFile.TryGetValue(ProducedItem.Location, out OriginalLocation))
-											{
-												HotReloadState.OriginalFileToHotReloadFile[OriginalLocation] = ProducedItem.Location;
-												HotReloadState.TemporaryFiles.Add(ProducedItem.Location);
-											}
-										}
-									}
-
-									// Increment the suffix for the next iteration
-									if(ActionsToExecute.Count > 0)
-									{
-										HotReloadState.NextSuffix++;
-									}
-
-									// Save the new state
-									FileReference HotReloadStateFile = HotReloadState.GetLocation(TargetDescs[0].ProjectFile, TargetDescs[0].Name, TargetDescs[0].Platform, TargetDescs[0].Configuration, TargetDescs[0].Architecture);
-									HotReloadState.Save(HotReloadStateFile);
-								}
-
-								// Display some stats to the user.
-								Log.TraceVerbose(
-										"{0} actions, {1} outdated and requested actions",
-										ActionGraph.AllActions.Count,
-										ActionsToExecute.Count
-										);
-
-								// Cache indirect includes for all outdated C++ files.  We kick this off as a background thread so that it can
-								// perform the scan while we're compiling.  It usually only takes up to a few seconds, but we don't want to hurt
-								// our best case UBT iteration times for this task which can easily be performed asynchronously
-								if (BuildConfiguration.bUseUBTMakefiles && TargetToOutdatedPrerequisitesMap.Count > 0)
-								{
-									CppIncludeThread = new CppIncludeBackgroundThread(TargetToOutdatedPrerequisitesMap, TargetToHeaders);
-								}
-
-								// Execute the actions.
-								ExecutorTimer.Start();
-								if (BuildConfiguration.bXGEExport)
-								{
-									using(Timeline.ScopeEvent("XGE.ExportActions()"))
-									{
-										XGE.ExportActions(ActionsToExecute.ToList());
-									}
-									bSuccess = true;
-								}
-								else
-								{
-									using(Timeline.ScopeEvent("ActionGraph.ExecuteActions()"))
-									{
-										bSuccess = ActionGraph.ExecuteActions(BuildConfiguration, ActionsToExecute.ToList(), out ExecutorName);
-									}
-								}
-								ExecutorTimer.Stop();
-
-								// if the build succeeded, write the receipts and do any needed syncing
-								if (!bSuccess)
-								{
-									BuildResult = ECompilationResult.OtherCompilationError;
-								}
-							}
-
-							// Run the deployment steps
-							if (BuildResult.Succeeded()
-								&& BuildConfiguration.SingleFileToCompile == null
-								&& !BuildConfiguration.bGenerateManifest
-								&& !BuildConfiguration.bXGEExport)
-							{
-								foreach (UEBuildTarget Target in Targets)
-								{
-									if (Target.bDeployAfterCompile && HotReloadModuleNameToSuffix.Count == 0)
-									{
-										Log.WriteLine(LogEventType.Console, "Deploying {0} {1} {2}...", Target.TargetName, Target.Platform, Target.Configuration);
-										TargetReceipt Receipt = TargetReceipt.Read(Target.ReceiptFileName);
-										UEBuildPlatform.GetBuildPlatform(Target.Platform).Deploy(Receipt);
-									}
-								}
-							}
-						}
-					}
-				}
-			}
-			catch (Exception Ex)
-			{
-				Log.WriteException(Ex, (LogFile == null) ? null : LogFile.FullName);
-				BuildResult = ECompilationResult.OtherCompilationError;
-			}
-
-			// Wait until our CPPIncludes dependency scanner thread has finished
-			if (CppIncludeThread != null)
-			{
-				using(Timeline.ScopeEvent("CppIncludeThread.Join()"))
-				{
-					CppIncludeThread.Join();
-				}
-			}
-
-			// Save the include dependency cache.
-			foreach (CPPHeaders Headers in TargetToHeaders.Values)
-			{
-				// NOTE: It's very important that we save the include cache, even if a build exception was thrown (compile error, etc), because we need to make sure that
-				//    any C++ include dependencies that we computed for out of date source files are saved.  Remember, the build may fail *after* some build products
-				//    are successfully built.  If we didn't save our dependency cache after build failures, source files for those build products that were successfully
-				//    built before the failure would not be considered out of date on the next run, so this is our only chance to cache C++ includes for those files!
-
-				if (Headers.IncludeDependencyCache != null)
-				{
-					using(Timeline.ScopeEvent("DependencyCache.Save()"))
-					{
-						Headers.IncludeDependencyCache.Save();
-					}
-				}
-
-				if (Headers.FlatCPPIncludeDependencyCache != null)
-				{
-					using (Timeline.ScopeEvent("FlatCPPIncludeDependencyCache.Save()"))
-					{
-						Headers.FlatCPPIncludeDependencyCache.Save();
-					}
-				}
-			}
-
-			// Figure out how long we took to execute.
-			if (ExecutorName != "Unknown")
-			{
-				Log.TraceInformation("Total build time: {0:0.00} seconds ({1} executor: {2:0.00} seconds)", BuildTimer.Elapsed.TotalSeconds, ExecutorName, ExecutorTimer.Elapsed.TotalSeconds);
-			}
-
-			return BuildResult;
-		}
-
-		/// <summary>
-		/// Invalidates makefiles for given target.
-		/// </summary>
-		/// <param name="Target">Target</param>
-		private static void InvalidateMakefiles(TargetDescriptor Target)
-		{
-			string[] MakefileNames = new string[] { "HotReloadMakefile.ubt", "Makefile.ubt" };
-			DirectoryReference BaseDir = UBTMakefile.GetUBTMakefileDirectoryPathForSingleTarget(Target);
-
-			foreach (string MakefileName in MakefileNames)
-			{
-				FileReference MakefileRef = FileReference.Combine(BaseDir, MakefileName);
-				if (FileReference.Exists(MakefileRef))
-				{
-					FileReference.Delete(MakefileRef);
-				}
-			}
-		}
-
-		/// <summary>
-		/// Helper class to update the C++ dependency cache on a background thread. Captures exceptions and re-throws on the main thread when joined.
-		/// </summary>
-		class CppIncludeBackgroundThread
-		{
-			Thread BackgroundThread;
-			Exception CaughtException;
-
-			public CppIncludeBackgroundThread(Dictionary<UEBuildTarget, List<FileItem>> TargetToOutdatedPrerequisitesMap, Dictionary<UEBuildTarget, CPPHeaders> TargetToHeaders)
-			{
-				BackgroundThread = new Thread(() => Run(TargetToOutdatedPrerequisitesMap, TargetToHeaders));
-				BackgroundThread.Start();
-			}
-
-			public void Join()
-			{
-				BackgroundThread.Join();
-				if (CaughtException != null)
-				{
-					throw CaughtException;
-				}
-			}
-
-			private void Run(Dictionary<UEBuildTarget, List<FileItem>> TargetToOutdatedPrerequisitesMap, Dictionary<UEBuildTarget, CPPHeaders> TargetToHeaders)
-			{
-				try
-				{
-					// @todo ubtmake: This thread will access data structures that are also used on the main UBT thread, but during this time UBT
-					// is only invoking the build executor, so should not be touching this stuff.  However, we need to at some guards to make sure.
-					foreach (KeyValuePair<UEBuildTarget, List<FileItem>> TargetAndOutdatedPrerequisites in TargetToOutdatedPrerequisitesMap)
-					{
-						UEBuildTarget Target = TargetAndOutdatedPrerequisites.Key;
-						List<FileItem> OutdatedPrerequisites = TargetAndOutdatedPrerequisites.Value;
-						CPPHeaders Headers = TargetToHeaders[Target];
-
-						foreach (FileItem PrerequisiteItem in OutdatedPrerequisites)
-						{
-							// Invoke our deep include scanner to figure out whether any of the files included by this source file have
-							// changed since the build product was built
-							Headers.FindAndCacheAllIncludedFiles(PrerequisiteItem, PrerequisiteItem.CachedIncludePaths, bOnlyCachedDependencies: false);
-						}
-					}
-				}
-				catch (Exception Ex)
-				{
-					CaughtException = Ex;
-				}
-			}
-		}
 	}
 }
