UBT: Move UnrealBuildTool.RegisterAllUBTClasses() to UEBuildPlatform.RegisterPlatforms().

#rnx

diff --git a/Engine/Source/Programs/UnrealBuildTool/Configuration/UEBuildPlatform.cs b/Engine/Source/Programs/UnrealBuildTool/Configuration/UEBuildPlatform.cs
--- a/Engine/Source/Programs/UnrealBuildTool/Configuration/UEBuildPlatform.cs
+++ b/Engine/Source/Programs/UnrealBuildTool/Configuration/UEBuildPlatform.cs
@@ -7,6 +7,7 @@ using System.Diagnostics;
 using System.IO;
 using System.Linq;
 using Tools.DotNETCommon;
+using System.Reflection;
 
 namespace UnrealBuildTool
 {
@@ -53,6 +54,37 @@ namespace UnrealBuildTool
 			DefaultCppPlatform = InDefaultCPPPlatform;
 		}
 
+		/// <summary>
+		/// Finds all the UEBuildPlatformFactory types in this assembly and uses them to register all the available platforms
+		/// </summary>
+		/// <param name="bIncludeNonInstalledPlatforms">Whether to register platforms that are not installed</param>
+		public static void RegisterPlatforms(bool bIncludeNonInstalledPlatforms)
+		{
+			// Find and register all tool chains and build platforms that are present
+			Assembly UBTAssembly = Assembly.GetExecutingAssembly();
+
+			Type[] AllTypes = UBTAssembly.GetTypes();
+
+			// register all build platforms first, since they implement SDK-switching logic that can set environment variables
+			foreach (Type CheckType in AllTypes)
+			{
+				if (CheckType.IsClass && !CheckType.IsAbstract)
+				{
+					if (CheckType.IsSubclassOf(typeof(UEBuildPlatformFactory)))
+					{
+						Log.TraceVerbose("    Registering build platform: {0}", CheckType.ToString());
+						UEBuildPlatformFactory TempInst = (UEBuildPlatformFactory)(UBTAssembly.CreateInstance(CheckType.FullName, true));
+
+						// We need all platforms to be registered when we run -validateplatform command to check SDK status of each
+						if (bIncludeNonInstalledPlatforms || InstalledPlatformInfo.IsValidPlatform(TempInst.TargetPlatform))
+						{
+							TempInst.RegisterBuildPlatforms();
+						}
+					}
+				}
+			}
+		}
+
 		/// <summary>
 		/// Gets an array of all platform folder names
 		/// </summary>
diff --git a/Engine/Source/Programs/UnrealBuildTool/System/PlatformExports.cs b/Engine/Source/Programs/UnrealBuildTool/System/PlatformExports.cs
--- a/Engine/Source/Programs/UnrealBuildTool/System/PlatformExports.cs
+++ b/Engine/Source/Programs/UnrealBuildTool/System/PlatformExports.cs
@@ -156,7 +156,7 @@ namespace UnrealBuildTool
 			XmlConfig.ReadConfigFiles(null);
 
 			// Register all the platform classes
-			UnrealBuildTool.RegisterAllUBTClasses(false);
+			UEBuildPlatform.RegisterPlatforms(false);
 			return true;
 		}
 	}
diff --git a/Engine/Source/Programs/UnrealBuildTool/UnrealBuildTool.cs b/Engine/Source/Programs/UnrealBuildTool/UnrealBuildTool.cs
--- a/Engine/Source/Programs/UnrealBuildTool/UnrealBuildTool.cs
+++ b/Engine/Source/Programs/UnrealBuildTool/UnrealBuildTool.cs
@@ -228,40 +228,10 @@ namespace UnrealBuildTool
 		static public bool IsUnderAnEngineDirectory(DirectoryReference InDirectory)
 		{
 			// Enterprise modules are considered as engine modules
-			return InDirectory.IsUnderDirectory(UnrealBuildTool.EngineDirectory) || InDirectory.IsUnderDirectory(UnrealBuildTool.EnterpriseSourceDirectory) ||
-				InDirectory.IsUnderDirectory(UnrealBuildTool.EnterprisePluginsDirectory) || InDirectory.IsUnderDirectory(UnrealBuildTool.EnterpriseIntermediateDirectory);
-		}
-
-		public static void RegisterAllUBTClasses(bool bValidatingPlatforms)
-		{
-			// Find and register all tool chains and build platforms that are present
-			Assembly UBTAssembly = Assembly.GetExecutingAssembly();
-			if (UBTAssembly != null)
-			{
-				Log.TraceVerbose("Searching for ToolChains, BuildPlatforms, BuildDeploys and ProjectGenerators...");
-
-				List<System.Type> ProjectGeneratorList = new List<System.Type>();
-				Type[] AllTypes = UBTAssembly.GetTypes();
-
-				// register all build platforms first, since they implement SDK-switching logic that can set environment variables
-				foreach (Type CheckType in AllTypes)
-				{
-					if (CheckType.IsClass && !CheckType.IsAbstract)
-					{
-						if (CheckType.IsSubclassOf(typeof(UEBuildPlatformFactory)))
-						{
-							Log.TraceVerbose("    Registering build platform: {0}", CheckType.ToString());
-							UEBuildPlatformFactory TempInst = (UEBuildPlatformFactory)(UBTAssembly.CreateInstance(CheckType.FullName, true));
-
-							// We need all platforms to be registered when we run -validateplatform command to check SDK status of each
-							if (bValidatingPlatforms || InstalledPlatformInfo.IsValidPlatform(TempInst.TargetPlatform))
-							{
-								TempInst.RegisterBuildPlatforms();
-							}
-						}
-					}
-				}
-			}
+			return InDirectory.IsUnderDirectory(UnrealBuildTool.EngineDirectory) 
+				|| InDirectory.IsUnderDirectory(UnrealBuildTool.EnterpriseSourceDirectory) 
+				|| InDirectory.IsUnderDirectory(UnrealBuildTool.EnterprisePluginsDirectory) 
+				|| InDirectory.IsUnderDirectory(UnrealBuildTool.EnterpriseIntermediateDirectory);
 		}
 
 		/// <summary>
@@ -408,16 +378,16 @@ namespace UnrealBuildTool
 				// Register all the build platforms
 				if((ModeOptions & ToolModeOptions.BuildPlatforms) != 0)
 				{
-					using(Timeline.ScopeEvent("UnrealBuildTool.RegisterAllUBTClasses()"))
+					using(Timeline.ScopeEvent("UEBuildPlatform.RegisterPlatforms()"))
 					{
-						RegisterAllUBTClasses(false);
+						UEBuildPlatform.RegisterPlatforms(false);
 					}
 				}
 				if((ModeOptions & ToolModeOptions.BuildPlatformsForValidation) != 0)
 				{
-					using(Timeline.ScopeEvent("UnrealBuildTool.RegisterAllUBTClasses()"))
+					using(Timeline.ScopeEvent("UEBuildPlatform.RegisterPlatforms()"))
 					{
-						RegisterAllUBTClasses(true);
+						UEBuildPlatform.RegisterPlatforms(true);
 					}
 				}
 
