Add missing copyright notices.

#rnx

diff --git a/Engine/Source/Programs/UnrealBuildTool/System/Action.cs b/Engine/Source/Programs/UnrealBuildTool/System/Action.cs
--- a/Engine/Source/Programs/UnrealBuildTool/System/Action.cs
+++ b/Engine/Source/Programs/UnrealBuildTool/System/Action.cs
@@ -1,3 +1,5 @@
+// Copyright 1998-2018 Epic Games, Inc. All Rights Reserved.
+
 using System;
 using System.Collections.Generic;
 using System.IO;
diff --git a/Engine/Source/Programs/UnrealBuildTool/System/ActionExecutor.cs b/Engine/Source/Programs/UnrealBuildTool/System/ActionExecutor.cs
--- a/Engine/Source/Programs/UnrealBuildTool/System/ActionExecutor.cs
+++ b/Engine/Source/Programs/UnrealBuildTool/System/ActionExecutor.cs
@@ -1,3 +1,5 @@
+// Copyright 1998-2018 Epic Games, Inc. All Rights Reserved.
+
 using System;
 using System.Collections.Generic;
 using System.Linq;
