UBT: Remove the unused FileItem.bNeedsHotReloadNumbersDLLCleanUp field.

#rnx

diff --git a/Engine/Source/Programs/UnrealBuildTool/Platform/Mac/MacToolChain.cs b/Engine/Source/Programs/UnrealBuildTool/Platform/Mac/MacToolChain.cs
--- a/Engine/Source/Programs/UnrealBuildTool/Platform/Mac/MacToolChain.cs
+++ b/Engine/Source/Programs/UnrealBuildTool/Platform/Mac/MacToolChain.cs
@@ -646,7 +646,6 @@ namespace UnrealBuildTool
 
 			// Add the output file as a production of the link action.
 			FileItem OutputFile = FileItem.GetItemByFileReference(LinkEnvironment.OutputFilePath);
-			OutputFile.bNeedsHotReloadNumbersDLLCleanUp = LinkEnvironment.bIsBuildingDLL;
 
 			// To solve the problem with cross dependencies, for now we create a broken dylib that does not link with other engine dylibs.
 			// This is fixed in later step, FixDylibDependencies. For this and to know what libraries to copy whilst creating an app bundle,
diff --git a/Engine/Source/Programs/UnrealBuildTool/Platform/Windows/VCToolChain.cs b/Engine/Source/Programs/UnrealBuildTool/Platform/Windows/VCToolChain.cs
--- a/Engine/Source/Programs/UnrealBuildTool/Platform/Windows/VCToolChain.cs
+++ b/Engine/Source/Programs/UnrealBuildTool/Platform/Windows/VCToolChain.cs
@@ -1287,7 +1287,6 @@ namespace UnrealBuildTool
 			else
 			{
 				OutputFile = FileItem.GetItemByFileReference(LinkEnvironment.OutputFilePath);
-				OutputFile.bNeedsHotReloadNumbersDLLCleanUp = LinkEnvironment.bIsBuildingDLL;
 			}
 
 			List<FileItem> ProducedItems = new List<FileItem>();
diff --git a/Engine/Source/Programs/UnrealBuildTool/System/FileItem.cs b/Engine/Source/Programs/UnrealBuildTool/System/FileItem.cs
--- a/Engine/Source/Programs/UnrealBuildTool/System/FileItem.cs
+++ b/Engine/Source/Programs/UnrealBuildTool/System/FileItem.cs
@@ -32,11 +32,6 @@ namespace UnrealBuildTool
 		/// </summary>
 		public FileReference Location;
 
-		/// <summary>
-		/// True if any DLLs produced by this
-		/// </summary>
-		public bool bNeedsHotReloadNumbersDLLCleanUp = false;
-
 		/// <summary>
 		/// Accessor for the absolute path to the file
 		/// </summary>
@@ -323,7 +318,6 @@ namespace UnrealBuildTool
 		{
 			ProducingAction = (Action)SerializationInfo.GetValue("pa", typeof(Action));
 			Location = (FileReference)SerializationInfo.GetValue("fi", typeof(FileReference));
-			bNeedsHotReloadNumbersDLLCleanUp = SerializationInfo.GetBoolean("hr");
 			CachedIncludePaths = (CppIncludePaths)SerializationInfo.GetValue("ci", typeof(CppIncludePaths));
 
 			// Go ahead and init normally now
@@ -349,7 +343,6 @@ namespace UnrealBuildTool
 		{
 			SerializationInfo.AddValue("pa", ProducingAction);
 			SerializationInfo.AddValue("fi", Location);
-			SerializationInfo.AddValue("hr", bNeedsHotReloadNumbersDLLCleanUp);
 			SerializationInfo.AddValue("ci", CachedIncludePaths);
 		}
 
diff --git a/Engine/Source/Programs/UnrealBuildTool/System/HotReload.cs b/Engine/Source/Programs/UnrealBuildTool/System/HotReload.cs
--- a/Engine/Source/Programs/UnrealBuildTool/System/HotReload.cs
+++ b/Engine/Source/Programs/UnrealBuildTool/System/HotReload.cs
@@ -425,7 +425,6 @@ namespace UnrealBuildTool
 							Action.PrerequisiteItems[ItemIndex] = NewPrerequisiteItem;
 
 							// Copy the other important settings from the original file item
-							NewPrerequisiteItem.bNeedsHotReloadNumbersDLLCleanUp = OriginalPrerequisiteItem.bNeedsHotReloadNumbersDLLCleanUp;
 							NewPrerequisiteItem.ProducingAction = OriginalPrerequisiteItem.ProducingAction;
 
 							// Keep track of it so we can fix up dependencies in a second pass afterwards
@@ -466,7 +465,6 @@ namespace UnrealBuildTool
 						Action.ProducedItems[ItemIndex] = NewProducedItem;
 
 						// Copy the other important settings from the original file item
-						NewProducedItem.bNeedsHotReloadNumbersDLLCleanUp = OriginalProducedItem.bNeedsHotReloadNumbersDLLCleanUp;
 						NewProducedItem.ProducingAction = OriginalProducedItem.ProducingAction;
 
 						// Keep track of it so we can fix up dependencies in a second pass afterwards
