UBT: Remove the FileItem.GetExistingItemByFileReference() method. None of the current callers of this function correctly check for null, which results in harder to find errors down the line.

diff --git a/Engine/Source/Programs/UnrealBuildTool/Configuration/UEBuildModuleCPP.cs b/Engine/Source/Programs/UnrealBuildTool/Configuration/UEBuildModuleCPP.cs
--- a/Engine/Source/Programs/UnrealBuildTool/Configuration/UEBuildModuleCPP.cs
+++ b/Engine/Source/Programs/UnrealBuildTool/Configuration/UEBuildModuleCPP.cs
@@ -308,7 +308,11 @@ namespace UnrealBuildTool
 				PrecompiledManifest Manifest = PrecompiledManifest.Read(PrecompiledManifestLocation);
 				foreach(FileReference OutputFile in Manifest.OutputFiles)
 				{
-					FileItem ObjectFile = FileItem.GetExistingItemByFileReference(OutputFile);
+					FileItem ObjectFile = FileItem.GetItemByFileReference(OutputFile);
+					if(!ObjectFile.bExists)
+					{
+						throw new BuildException("Missing object file {0} listed in {1}", OutputFile, PrecompiledManifestLocation);
+					}
 					LinkInputFiles.Add(ObjectFile);
 				}
 				return LinkInputFiles;
@@ -497,8 +501,14 @@ namespace UnrealBuildTool
 					// Always force include the PCH, even if PCHs are disabled, for generated code. Legacy code can rely on PCHs being included to compile correctly, and this used to be done by UHT manually including it.
 					if(GeneratedCPPCompileEnvironment.PrecompiledHeaderFile == null && Rules.PrivatePCHHeaderFile != null && Rules.PCHUsage != ModuleRules.PCHUsageMode.UseExplicitOrSharedPCHs)
 					{
+						FileItem PrivatePchFileItem = FileItem.GetItemByFileReference(FileReference.Combine(ModuleDirectory, Rules.PrivatePCHHeaderFile));
+						if(!PrivatePchFileItem.bExists)
+						{
+							throw new BuildException("Unable to find private PCH file '{0}', referenced by '{1}'", PrivatePchFileItem.Location, RulesFile);
+						}
+
 						GeneratedCPPCompileEnvironment = new CppCompileEnvironment(GeneratedCPPCompileEnvironment);
-						GeneratedCPPCompileEnvironment.ForceIncludeFiles.Add(FileItem.GetExistingItemByFileReference(FileReference.Combine(ModuleDirectory, Rules.PrivatePCHHeaderFile)));
+						GeneratedCPPCompileEnvironment.ForceIncludeFiles.Add(PrivatePchFileItem);
 					}
 
 					// Compile all the generated files
diff --git a/Engine/Source/Programs/UnrealBuildTool/Configuration/UEBuildTarget.cs b/Engine/Source/Programs/UnrealBuildTool/Configuration/UEBuildTarget.cs
--- a/Engine/Source/Programs/UnrealBuildTool/Configuration/UEBuildTarget.cs
+++ b/Engine/Source/Programs/UnrealBuildTool/Configuration/UEBuildTarget.cs
@@ -1495,7 +1495,7 @@ namespace UnrealBuildTool
 				WriteMetadataAction.StatusDescription = ReceiptFileName.GetFileName();
 				WriteMetadataAction.bCanExecuteRemotely = false;
 				WriteMetadataAction.PrerequisiteItems.AddRange(OutputItems);
-				WriteMetadataAction.PrerequisiteItems.Add(FileItem.GetExistingItemByPath(Assembly.GetExecutingAssembly().Location));
+				WriteMetadataAction.PrerequisiteItems.Add(FileItem.GetItemByPath(Assembly.GetExecutingAssembly().Location));
 				WriteMetadataAction.ProducedItems.Add(FileItem.GetItemByFileReference(ReceiptFileName));
 
 				OutputItems.AddRange(WriteMetadataAction.ProducedItems);
diff --git a/Engine/Source/Programs/UnrealBuildTool/System/ExternalExecution.cs b/Engine/Source/Programs/UnrealBuildTool/System/ExternalExecution.cs
--- a/Engine/Source/Programs/UnrealBuildTool/System/ExternalExecution.cs
+++ b/Engine/Source/Programs/UnrealBuildTool/System/ExternalExecution.cs
@@ -368,7 +368,7 @@ namespace UnrealBuildTool
 			{
 				// Check to see if we know anything about this file.  If we have up-to-date cached information about whether it has
 				// UObjects or not, we can skip doing a test here.
-				FileItem UObjectHeaderFileItem = FileItem.GetExistingItemByFileReference(Header);
+				FileItem UObjectHeaderFileItem = FileItem.GetItemByFileReference(Header);
 
 				if (CPPHeaders.DoesFileContainUObjects(UObjectHeaderFileItem.AbsolutePath))
 				{
diff --git a/Engine/Source/Programs/UnrealBuildTool/System/FileItem.cs b/Engine/Source/Programs/UnrealBuildTool/System/FileItem.cs
--- a/Engine/Source/Programs/UnrealBuildTool/System/FileItem.cs
+++ b/Engine/Source/Programs/UnrealBuildTool/System/FileItem.cs
@@ -167,30 +167,6 @@ namespace UnrealBuildTool
 			}
 		}
 
-		/// <summary>
-		/// If the given file path identifies a file that already exists, returns the FileItem that represents it.
-		/// </summary>
-		public static FileItem GetExistingItemByPath(string FileName)
-		{
-			return GetExistingItemByFileReference(new FileReference(FileName));
-		}
-
-		/// <summary>
-		/// If the given file path identifies a file that already exists, returns the FileItem that represents it.
-		/// </summary>
-		public static FileItem GetExistingItemByFileReference(FileReference FileRef)
-		{
-			FileItem Result = GetItemByFileReference(FileRef);
-			if (Result.bExists)
-			{
-				return Result;
-			}
-			else
-			{
-				return null;
-			}
-		}
-
 		/// <summary>
 		/// Determines the appropriate encoding for a string: either ASCII or UTF-8.
 		/// </summary>
