Merging //UE4/Dev-Main to Dev-Build (//UE4/Dev-Build)

#rnx

diff --git a/.gitignore b/.gitignore
--- a/.gitignore
+++ b/.gitignore
@@ -206,3 +206,6 @@ Engine/Binaries/DotNET/SwarmAgent.Options.xml
 
 # Ignore CLion directory
 .idea/
+
+# Annoying windows files
+[Dd]esktop.ini
diff --git a/Engine/Source/Programs/UnrealBuildTool/Configuration/UEBuildModuleCPP.cs b/Engine/Source/Programs/UnrealBuildTool/Configuration/UEBuildModuleCPP.cs
--- a/Engine/Source/Programs/UnrealBuildTool/Configuration/UEBuildModuleCPP.cs
+++ b/Engine/Source/Programs/UnrealBuildTool/Configuration/UEBuildModuleCPP.cs
@@ -329,6 +329,24 @@ namespace UnrealBuildTool
 
 			SourceDirectories = new HashSet<DirectoryReference>(SourceFiles.Select(x => x.Location.Directory));
 
+			// Store the module compile environment along with the source file.  This is so that we can use it later on when looking for header dependencies
+			foreach (FileItem CFile in SourceFilesFound.CFiles)
+			{
+				CFile.CachedIncludePaths = ModuleCompileEnvironment.IncludePaths;
+			}
+			foreach (FileItem CCFile in SourceFilesFound.CCFiles)
+			{
+				CCFile.CachedIncludePaths = ModuleCompileEnvironment.IncludePaths;
+			}
+			foreach (FileItem CPPFile in SourceFilesFound.CPPFiles)
+			{
+				CPPFile.CachedIncludePaths = ModuleCompileEnvironment.IncludePaths;
+			}
+			foreach (FileItem MMFile in SourceFilesFound.MMFiles)
+			{
+				MMFile.CachedIncludePaths = ModuleCompileEnvironment.IncludePaths;
+			}
+
 			// Process all of the header file dependencies for this module
 			CheckFirstIncludeMatchesEachCppFile(Target, ModuleCompileEnvironment, SourceFilesToBuild.CPPFiles);
 
@@ -1228,38 +1246,38 @@ namespace UnrealBuildTool
 				}
 
 				if (IsCompilableSourceFile)
-				{
+		{
 					if (SourceFilePath.IsUnderDirectory(SourceFilesBaseDirectory))
-					{
+			{
 						// Store the path as relative to the project file
 						string RelativeFilePath = SourceFilePath.MakeRelativeTo(SourceFilesBaseDirectory);
 
 						// All compiled files should always be in a sub-directory under the project file directory.  We enforce this here.
 						if (Path.IsPathRooted(RelativeFilePath) || RelativeFilePath.StartsWith(".."))
-						{
+				{
 							throw new BuildException("Error: Found source file {0} in project whose path was not relative to the base directory of the source files", RelativeFilePath);
-						}
+				}
 
 						// Check for source files that don't belong to the platform we're currently compiling.  We'll filter
 						// those source files out
 						bool IncludeThisFile = true;
 						foreach (string CurPlatformName in OtherPlatformNameStrings)
-						{
+				{
 							if (RelativeFilePath.IndexOf(Path.DirectorySeparatorChar + CurPlatformName + Path.DirectorySeparatorChar, StringComparison.InvariantCultureIgnoreCase) != -1
 								|| RelativeFilePath.StartsWith(CurPlatformName + Path.DirectorySeparatorChar))
 							{
 								IncludeThisFile = false;
 								break;
 							}
-						}
+				}
 
 						if (IncludeThisFile)
-						{
+				{
 							FilteredFileItems.Add(FileItem.GetItemByFileReference(SourceFilePath));
 						}
 					}
 				}
-			}
+				}
 
 			// @todo projectfiles: Consider enabling this error but changing it to a warning instead.  It can fire for
 			//    projects that are being digested for IntelliSense (because the module was set as a cross-
@@ -1267,7 +1285,7 @@ namespace UnrealBuildTool
 			//    in the project generator
 			bool AllowEmptyProjects = true;
 			if (!AllowEmptyProjects)
-			{
+				{
 				if (FilteredFileItems.Count == 0)
 				{
 					throw new BuildException("Could not find any valid source files for base directory {0}.  Project has {1} files in it", SourceFilesBaseDirectory, SourceFiles.Count);
diff --git a/Engine/Source/Runtime/CoreUObject/Private/UObject/SavePackage.cpp b/Engine/Source/Runtime/CoreUObject/Private/UObject/SavePackage.cpp
--- a/Engine/Source/Runtime/CoreUObject/Private/UObject/SavePackage.cpp
+++ b/Engine/Source/Runtime/CoreUObject/Private/UObject/SavePackage.cpp
@@ -3947,6 +3947,14 @@ FSavePackageResultStruct UPackage::Save(UPackage* InOuter, UObject* Base, EObjec
 				}
 #endif
 
+				if (!(Linker->Summary.PackageFlags & PKG_FilterEditorOnly))
+				{
+					// The Editor version is used as part of the check to see if a package is too old to use the gather cache, so we always have to add it if we have gathered loc for this asset
+					// We need to set the editor custom version before we copy the version container to the summary, otherwise we may end up with corrupt assets
+					// because we later do it on the Linker when actually gathering loc data
+					ExportTaggerArchive.UsingCustomVersion(FEditorObjectVersion::GUID);
+				}
+
 				// Use the custom versions we had previously gleaned from the export tag pass
 				Linker->Summary.SetCustomVersionContainer(ExportTaggerArchive.GetCustomVersions());
 
@@ -4537,6 +4545,7 @@ FSavePackageResultStruct UPackage::Save(UPackage* InOuter, UObject* Base, EObjec
 				if (!(Linker->Summary.PackageFlags & PKG_FilterEditorOnly))
 				{
 					// The Editor version is used as part of the check to see if a package is too old to use the gather cache, so we always have to add it if we have gathered loc for this asset
+					// Note that using custom version here only works because we already added it to the export tagger before the package summary was serialized
 					Linker->UsingCustomVersion(FEditorObjectVersion::GUID);
 
 					bool bCanCacheGatheredText = false;
