UBT: Fix resource strings on Windows for program targets. Custom version strings can now only affect resources with a unique build environment, and the default product name for programs is now taken from the target name (via the UE_APP_NAME setting). The default product name in the editor is left blank to allow the target name to be used as the default.

diff --git a/Engine/Build/Windows/Default.rc2 b/Engine/Build/Windows/Default.rc2
--- a/Engine/Build/Windows/Default.rc2
+++ b/Engine/Build/Windows/Default.rc2
@@ -10,6 +10,31 @@
 #include <windows.h>
 #undef APSTUDIO_READONLY_SYMBOLS
 
+// Various strings used for project resources
+#ifdef PROJECT_COMPANY_NAME
+	#define BUILD_PROJECT_COMPANY_NAME PREPROCESSOR_TO_STRING(PROJECT_COMPANY_NAME)
+#else
+	#define BUILD_PROJECT_COMPANY_NAME EPIC_COMPANY_NAME
+#endif
+
+#ifdef PROJECT_COPYRIGHT_STRING
+	#define BUILD_PROJECT_COPYRIGHT_STRING PREPROCESSOR_TO_STRING(PROJECT_COPYRIGHT_STRING)
+#else
+	#define BUILD_PROJECT_COPYRIGHT_STRING EPIC_COPYRIGHT_STRING
+#endif
+
+#ifdef PROJECT_PRODUCT_NAME
+	#define BUILD_PROJECT_PRODUCT_NAME PREPROCESSOR_TO_STRING(PROJECT_PRODUCT_NAME)
+#else
+	#define BUILD_PROJECT_PRODUCT_NAME PREPROCESSOR_TO_STRING(UE_APP_NAME)
+#endif
+
+#ifdef PROJECT_PRODUCT_IDENTIFIER
+	#define BUILD_PROJECT_PRODUCT_IDENTIFIER PREPROCESSOR_TO_STRING(PROJECT_PRODUCT_IDENTIFIER)
+#else
+	#define BUILD_PROJECT_PRODUCT_IDENTIFIER EPIC_PRODUCT_IDENTIFIER
+#endif
+
 /////////////////////////////////////////////////////////////////////////////
 // English (United States) resources
 
diff --git a/Engine/Config/BaseGame.ini b/Engine/Config/BaseGame.ini
--- a/Engine/Config/BaseGame.ini
+++ b/Engine/Config/BaseGame.ini
@@ -51,7 +51,7 @@ CopyrightNotice=Fill out your copyright notice in the Description page of Projec
 Description=
 LicensingTerms=
 PrivacyPolicy=
-ProjectName=UE4Game
+ProjectName=
 ProjectVersion=1.0.0.0
 Homepage=
 SupportContact=
diff --git a/Engine/Source/Programs/UnrealBuildTool/Configuration/UEBuildBinary.cs b/Engine/Source/Programs/UnrealBuildTool/Configuration/UEBuildBinary.cs
--- a/Engine/Source/Programs/UnrealBuildTool/Configuration/UEBuildBinary.cs
+++ b/Engine/Source/Programs/UnrealBuildTool/Configuration/UEBuildBinary.cs
@@ -657,7 +657,12 @@ namespace UnrealBuildTool
 			List<UEBuildBinary> BinaryDependencies = new List<UEBuildBinary>();
 
 			CppCompileEnvironment BinaryCompileEnvironment = CreateBinaryCompileEnvironment(CompileEnvironment);
-			
+
+			if(BinaryCompileEnvironment.bUseSharedBuildEnvironment && Target.ProjectFile != null && IntermediateDirectory.IsUnderDirectory(Target.ProjectFile.Directory))
+			{
+				BinaryCompileEnvironment.bUseSharedBuildEnvironment = false;
+			}
+
 			foreach (UEBuildModule Module in Modules)
 			{
 				List<FileItem> LinkInputFiles;
@@ -729,7 +734,7 @@ namespace UnrealBuildTool
 
 						// Create a compile environment for resource files
 						CppCompileEnvironment ResourceCompileEnvironment = new CppCompileEnvironment(BinaryCompileEnvironment);
-						WindowsPlatform.SetupResourceCompileEnvironment(ResourceCompileEnvironment, ResourceIntermediateDirectory, Target);
+						WindowsPlatform.SetupResourceCompileEnvironment(ResourceCompileEnvironment, Target);
 
 						// @todo: This should be in some Windows code somewhere...
 						// Set the original file name macro; used in Default.rc2 to set the binary metadata fields.
diff --git a/Engine/Source/Programs/UnrealBuildTool/Configuration/UEBuildModuleCPP.cs b/Engine/Source/Programs/UnrealBuildTool/Configuration/UEBuildModuleCPP.cs
--- a/Engine/Source/Programs/UnrealBuildTool/Configuration/UEBuildModuleCPP.cs
+++ b/Engine/Source/Programs/UnrealBuildTool/Configuration/UEBuildModuleCPP.cs
@@ -576,7 +576,7 @@ namespace UnrealBuildTool
 			if(SourceFilesToBuild.RCFiles.Count > 0)
 			{
 				CppCompileEnvironment ResourceCompileEnvironment = new CppCompileEnvironment(BinaryCompileEnvironment);
-				WindowsPlatform.SetupResourceCompileEnvironment(ResourceCompileEnvironment, IntermediateDirectory, Target);
+				WindowsPlatform.SetupResourceCompileEnvironment(ResourceCompileEnvironment, Target);
 				LinkInputFiles.AddRange(ToolChain.CompileRCFiles(ResourceCompileEnvironment, SourceFilesToBuild.RCFiles, IntermediateDirectory, ActionGraph).ObjectFiles);
 			}
 
diff --git a/Engine/Source/Programs/UnrealBuildTool/Configuration/UEBuildTarget.cs b/Engine/Source/Programs/UnrealBuildTool/Configuration/UEBuildTarget.cs
--- a/Engine/Source/Programs/UnrealBuildTool/Configuration/UEBuildTarget.cs
+++ b/Engine/Source/Programs/UnrealBuildTool/Configuration/UEBuildTarget.cs
@@ -1363,7 +1363,7 @@ namespace UnrealBuildTool
 						FileItem DefaultResourceFile = FileItem.GetItemByFileReference(FileReference.Combine(UnrealBuildTool.EngineDirectory, "Build", "Windows", "Default.rc2"));
 						DefaultResourceFile.CachedIncludePaths = DefaultResourceCompileEnvironment.IncludePaths;
 
-						WindowsPlatform.SetupResourceCompileEnvironment(DefaultResourceCompileEnvironment, EngineIntermediateDirectory, Rules);
+						WindowsPlatform.SetupResourceCompileEnvironment(DefaultResourceCompileEnvironment, Rules);
 
 						CPPOutput DefaultResourceOutput = TargetToolChain.CompileRCFiles(DefaultResourceCompileEnvironment, new List<FileItem> { DefaultResourceFile }, EngineIntermediateDirectory, ActionGraph);
 						GlobalLinkEnvironment.DefaultResourceFiles.AddRange(DefaultResourceOutput.ObjectFiles);
@@ -2877,6 +2877,7 @@ namespace UnrealBuildTool
 			GlobalCompileEnvironment.bEnableOSX109Support = Rules.bEnableOSX109Support;
 			GlobalCompileEnvironment.Definitions.Add(String.Format("IS_PROGRAM={0}", TargetType == TargetType.Program ? "1" : "0"));
 			GlobalCompileEnvironment.Definitions.AddRange(Rules.GlobalDefinitions);
+			GlobalCompileEnvironment.bUseSharedBuildEnvironment = (Rules.BuildEnvironment == TargetBuildEnvironment.Shared);
 			GlobalCompileEnvironment.bEnableExceptions = Rules.bForceEnableExceptions || Rules.bBuildEditor;
 			GlobalCompileEnvironment.bEnableObjCExceptions = Rules.bForceEnableObjCExceptions || Rules.bBuildEditor;
 			GlobalCompileEnvironment.bShadowVariableWarningsAsErrors = Rules.bShadowVariableErrors;
diff --git a/Engine/Source/Programs/UnrealBuildTool/Platform/Windows/UEBuildWindows.cs b/Engine/Source/Programs/UnrealBuildTool/Platform/Windows/UEBuildWindows.cs
--- a/Engine/Source/Programs/UnrealBuildTool/Platform/Windows/UEBuildWindows.cs
+++ b/Engine/Source/Programs/UnrealBuildTool/Platform/Windows/UEBuildWindows.cs
@@ -129,7 +129,7 @@ namespace UnrealBuildTool
 		public string CopyrightNotice;
 
 		/// <summary>
-		/// The project's name.
+		/// The product name.
 		/// </summary>
 		[ConfigFile(ConfigHierarchyType.Game, "/Script/EngineSettings.GeneralProjectSettings", "ProjectName")]
 		public string ProductName;
@@ -1477,13 +1477,12 @@ namespace UnrealBuildTool
 		/// Configures the resource compile environment for the given target
 		/// </summary>
 		/// <param name="ResourceCompileEnvironment">The compile environment</param>
-		/// <param name="IntermediateDirectory">The output directory for compiled files</param>
 		/// <param name="Target">The target being built</param>
-		public static void SetupResourceCompileEnvironment(CppCompileEnvironment ResourceCompileEnvironment, DirectoryReference IntermediateDirectory, ReadOnlyTargetRules Target)
+		public static void SetupResourceCompileEnvironment(CppCompileEnvironment ResourceCompileEnvironment, ReadOnlyTargetRules Target)
 		{
 			// Figure the icon to use. We can only use a custom icon when compiling to a project-specific intemediate directory (and not for the shared editor executable, for example).
 			FileReference IconFile;
-			if(Target.ProjectFile != null && IntermediateDirectory.IsUnderDirectory(Target.ProjectFile.Directory))
+			if(Target.ProjectFile != null && !ResourceCompileEnvironment.bUseSharedBuildEnvironment)
 			{
 				IconFile = WindowsPlatform.GetApplicationIcon(Target.ProjectFile);
 			}
@@ -1494,6 +1493,30 @@ namespace UnrealBuildTool
 
 			// Setup the compile environment, setting the icon to use via a macro. This is used in Default.rc2.
 			ResourceCompileEnvironment.Definitions.Add(String.Format("BUILD_ICON_FILE_NAME=\"\\\"{0}\\\"\"", IconFile.FullName.Replace("\\", "\\\\")));
+
+			// Apply the target settings for the resources
+			if(!ResourceCompileEnvironment.bUseSharedBuildEnvironment)
+			{
+				if (!String.IsNullOrEmpty(Target.WindowsPlatform.CompanyName))
+				{
+					ResourceCompileEnvironment.Definitions.Add(String.Format("PROJECT_COMPANY_NAME={0}", SanitizeMacroValue(Target.WindowsPlatform.CompanyName)));
+				}
+
+				if (!String.IsNullOrEmpty(Target.WindowsPlatform.CopyrightNotice))
+				{
+					ResourceCompileEnvironment.Definitions.Add(String.Format("PROJECT_COPYRIGHT_STRING={0}", SanitizeMacroValue(Target.WindowsPlatform.CopyrightNotice)));
+				}
+
+				if (!String.IsNullOrEmpty(Target.WindowsPlatform.ProductName))
+				{
+					ResourceCompileEnvironment.Definitions.Add(String.Format("PROJECT_PRODUCT_NAME={0}", SanitizeMacroValue(Target.WindowsPlatform.ProductName)));
+				}
+
+				if (Target.ProjectFile != null)
+				{
+					ResourceCompileEnvironment.Definitions.Add(String.Format("PROJECT_PRODUCT_IDENTIFIER={0}", SanitizeMacroValue(Target.ProjectFile.GetFileNameWithoutExtension())));
+				}
+			}
 		}
 
 		/// <summary>
@@ -1687,29 +1710,6 @@ namespace UnrealBuildTool
 				LinkEnvironment.AdditionalArguments += " /ignore:4078";
 			}
 
-			if (Target.Type != TargetType.Editor)
-			{
-				if (!string.IsNullOrEmpty(Target.WindowsPlatform.CompanyName))
-				{
-					CompileEnvironment.Definitions.Add(String.Format("PROJECT_COMPANY_NAME={0}", SanitizeMacroValue(Target.WindowsPlatform.CompanyName)));
-				}
-
-				if (!string.IsNullOrEmpty(Target.WindowsPlatform.CopyrightNotice))
-				{
-					CompileEnvironment.Definitions.Add(String.Format("PROJECT_COPYRIGHT_STRING={0}", SanitizeMacroValue(Target.WindowsPlatform.CopyrightNotice)));
-				}
-
-				if (!string.IsNullOrEmpty(Target.WindowsPlatform.ProductName))
-				{
-					CompileEnvironment.Definitions.Add(String.Format("PROJECT_PRODUCT_NAME={0}", SanitizeMacroValue(Target.WindowsPlatform.ProductName)));
-				}
-
-				if (Target.ProjectFile != null)
-				{
-					CompileEnvironment.Definitions.Add(String.Format("PROJECT_PRODUCT_IDENTIFIER={0}", SanitizeMacroValue(Target.ProjectFile.GetFileNameWithoutExtension())));
-				}
-			}
-
 			// Set up default stack size
 			LinkEnvironment.DefaultStackSize = Target.WindowsPlatform.DefaultStackSize;
 			LinkEnvironment.DefaultStackSizeCommit = Target.WindowsPlatform.DefaultStackSizeCommit;
diff --git a/Engine/Source/Programs/UnrealBuildTool/System/CppCompileEnvironment.cs b/Engine/Source/Programs/UnrealBuildTool/System/CppCompileEnvironment.cs
--- a/Engine/Source/Programs/UnrealBuildTool/System/CppCompileEnvironment.cs
+++ b/Engine/Source/Programs/UnrealBuildTool/System/CppCompileEnvironment.cs
@@ -100,6 +100,11 @@ namespace UnrealBuildTool
 		/// </summary>
 		public PrecompiledHeaderAction PrecompiledHeaderAction = PrecompiledHeaderAction.None;
 
+		/// <summary>
+		/// Whether artifacts from this compile are shared with other targets. If so, we should not apply any target-wide modifications to the compile environment.
+		/// </summary>
+		public bool bUseSharedBuildEnvironment;
+
 		/// <summary>
 		/// Use run time type information
 		/// </summary>
@@ -333,6 +338,7 @@ namespace UnrealBuildTool
 			Architecture = Other.Architecture;
 			PrecompiledHeaderIncludeFilename = Other.PrecompiledHeaderIncludeFilename;
 			PrecompiledHeaderAction = Other.PrecompiledHeaderAction;
+			bUseSharedBuildEnvironment = Other.bUseSharedBuildEnvironment;
 			bUseRTTI = Other.bUseRTTI;
 			bUseInlining = Other.bUseInlining;
 			bUseAVX = Other.bUseAVX;
diff --git a/Engine/Source/Runtime/Launch/Resources/Version.h b/Engine/Source/Runtime/Launch/Resources/Version.h
--- a/Engine/Source/Runtime/Launch/Resources/Version.h
+++ b/Engine/Source/Runtime/Launch/Resources/Version.h
@@ -69,28 +69,6 @@
 #define EPIC_PRODUCT_NAME "Unreal Engine"
 #define EPIC_PRODUCT_IDENTIFIER "UnrealEngine"
 
-// Various strings used for project resources
-#ifdef PROJECT_COMPANY_NAME
-	#define BUILD_PROJECT_COMPANY_NAME PREPROCESSOR_TO_STRING(PROJECT_COMPANY_NAME)
-#else
-	#define BUILD_PROJECT_COMPANY_NAME EPIC_COMPANY_NAME
-#endif // !PROJECT_COMPANY_NAME
-#ifdef PROJECT_COPYRIGHT_STRING
-	#define BUILD_PROJECT_COPYRIGHT_STRING PREPROCESSOR_TO_STRING(PROJECT_COPYRIGHT_STRING)
-#else
-	#define BUILD_PROJECT_COPYRIGHT_STRING EPIC_COPYRIGHT_STRING
-#endif // !PROJECT_COPYRIGHT_STRING
-#ifdef PROJECT_PRODUCT_NAME
-	#define BUILD_PROJECT_PRODUCT_NAME PREPROCESSOR_TO_STRING(PROJECT_PRODUCT_NAME)
-#else
-	#define BUILD_PROJECT_PRODUCT_NAME EPIC_PRODUCT_NAME
-#endif // !PROJECT_PRODUCT_NAME
-#ifdef PROJECT_PRODUCT_IDENTIFIER
-	#define BUILD_PROJECT_PRODUCT_IDENTIFIER PREPROCESSOR_TO_STRING(PROJECT_PRODUCT_IDENTIFIER)
-#else
-	#define BUILD_PROJECT_PRODUCT_IDENTIFIER EPIC_PRODUCT_IDENTIFIER
-#endif // !PROJECT_PRODUCT_IDENTIFIER
-
 #if defined(BUILT_FROM_CHANGELIST) && defined(BRANCH_NAME)
 #define ENGINE_VERSION_STRING \
 	VERSION_STRINGIFY(ENGINE_MAJOR_VERSION) \
