UBT: Remove hacks to support FileItems being directories. This was only used for app bundles on Mac, which we handle as folders of regular files now.

#rnx

diff --git a/Engine/Source/Programs/UnrealBuildTool/System/ActionGraph.cs b/Engine/Source/Programs/UnrealBuildTool/System/ActionGraph.cs
--- a/Engine/Source/Programs/UnrealBuildTool/System/ActionGraph.cs
+++ b/Engine/Source/Programs/UnrealBuildTool/System/ActionGraph.cs
@@ -483,7 +483,7 @@ namespace UnrealBuildTool
 					// If the produced file doesn't exist or has zero size, consider it outdated.  The zero size check is to detect cases
 					// where aborting an earlier compile produced invalid zero-sized obj files, but that may cause actions where that's
 					// legitimate output to always be considered outdated.
-					if (ProducedItem.bExists && (RootAction.ActionType != ActionType.Compile || ProducedItem.Length > 0 || ProducedItem.IsDirectory))
+					if (ProducedItem.bExists && (RootAction.ActionType != ActionType.Compile || ProducedItem.Length > 0))
 					{
 						// Use the oldest produced item's time as the last execution time.
 						if (ProducedItem.LastWriteTime < LastExecutionTime)
diff --git a/Engine/Source/Programs/UnrealBuildTool/System/FileItem.cs b/Engine/Source/Programs/UnrealBuildTool/System/FileItem.cs
--- a/Engine/Source/Programs/UnrealBuildTool/System/FileItem.cs
+++ b/Engine/Source/Programs/UnrealBuildTool/System/FileItem.cs
@@ -72,11 +72,6 @@ namespace UnrealBuildTool
 		/// </summary>
 		public FileInfo Info;
 
-		/// <summary>
-		/// This is true if this item is actually a directory. Consideration for Mac application bundles. Note that Info will be null if true!
-		/// </summary>
-		public bool IsDirectory;
-
 		/// <summary>
 		/// Relative cost of action associated with producing this file.
 		/// </summary>
@@ -328,25 +323,13 @@ namespace UnrealBuildTool
 		/// </summary>
 		public void ResetFileInfo()
 		{
-			if (Directory.Exists(AbsolutePath))
-			{
-				// path is actually a directory (such as a Mac app bundle)
-				_bExists = true;
-				LastWriteTime = Directory.GetLastWriteTimeUtc(AbsolutePath);
-				IsDirectory = true;
-				_Length = 0;
-				Info = null;
-			}
-			else
-			{
-				Info = new FileInfo(AbsolutePath);
+			Info = new FileInfo(AbsolutePath);
 
-				_bExists = Info.Exists;
-				if (_bExists)
-				{
-					_LastWriteTime = Info.LastWriteTimeUtc;
-					_Length = Info.Length;
-				}
+			_bExists = Info.Exists;
+			if (_bExists)
+			{
+				_LastWriteTime = Info.LastWriteTimeUtc;
+				_Length = Info.Length;
 			}
 		}
 
