UBT: Allow creating FileItem objects on different threads.

#rnx

diff --git a/Engine/Source/Programs/UnrealBuildTool/System/FileItem.cs b/Engine/Source/Programs/UnrealBuildTool/System/FileItem.cs
--- a/Engine/Source/Programs/UnrealBuildTool/System/FileItem.cs
+++ b/Engine/Source/Programs/UnrealBuildTool/System/FileItem.cs
@@ -8,6 +8,7 @@ using System.Diagnostics;
 using System.Threading;
 using System.Runtime.Serialization;
 using Tools.DotNETCommon;
+using System.Collections.Concurrent;
 
 namespace UnrealBuildTool
 {
@@ -15,8 +16,7 @@ namespace UnrealBuildTool
 	/// Represents a file on disk that is used as an input or output of a build action.
 	/// FileItems are created by calling FileItem.GetItemByFileReference, which creates a single FileItem for each unique file path.
 	/// </summary>
-	[Serializable]
-	class FileItem : ISerializable
+	class FileItem
 	{
 		///
 		/// Preparation and Assembly (serialized)
@@ -80,31 +80,25 @@ namespace UnrealBuildTool
 		/// <summary>
 		/// The last write time of the file.
 		/// </summary>
-		public DateTimeOffset _LastWriteTime;
 		public DateTimeOffset LastWriteTime
 		{
-			get { return _LastWriteTime; }
-			set { _LastWriteTime = value; }
+			get { return Info.LastWriteTimeUtc; }
 		}
 
 		/// <summary>
 		/// Whether the file exists.
 		/// </summary>
-		public bool _bExists = false;
 		public bool bExists
 		{
-			get { return _bExists; }
-			set { _bExists = value; }
+			get { return Info.Exists; }
 		}
 
 		/// <summary>
 		/// Size of the file if it exists, otherwise -1
 		/// </summary>
-		public long _Length = -1;
 		public long Length
 		{
-			get { return _Length; }
-			set { _Length = value; }
+			get { return Info.Length; }
 		}
 
 
@@ -112,16 +106,10 @@ namespace UnrealBuildTool
 		/// Statics
 		///
 
-		/// <summary>
-		/// Used for performance debugging
-		/// </summary>
-		public static long TotalFileItemCount = 0;
-		public static long MissingFileItemCount = 0;
-
 		/// <summary>
 		/// A case-insensitive dictionary that's used to map each unique file name to a single FileItem object.
 		/// </summary>
-		static Dictionary<FileReference, FileItem> UniqueSourceFileMap = new Dictionary<FileReference, FileItem>();
+		static ConcurrentDictionary<FileReference, FileItem> UniqueSourceFileMap = new ConcurrentDictionary<FileReference, FileItem>();
 
 		/// <summary>
 		/// Clears the FileItem caches.
@@ -151,15 +139,20 @@ namespace UnrealBuildTool
 		/// <returns>The FileItem that represents the given a full file path.</returns>
 		public static FileItem GetItemByFileReference(FileReference Reference)
 		{
-			FileItem Result = null;
-			if (UniqueSourceFileMap.TryGetValue(Reference, out Result))
+			FileItem Result;
+			if (!UniqueSourceFileMap.TryGetValue(Reference, out Result))
 			{
-				return Result;
-			}
-			else
-			{
-				return new FileItem(Reference);
+				FileItem NewFileItem = new FileItem(Reference, Reference.ToFileInfo());
+				if(UniqueSourceFileMap.TryAdd(Reference, NewFileItem))
+				{
+					Result = NewFileItem;
+				}
+				else
+				{
+					Result = UniqueSourceFileMap[Reference];
+				}
 			}
+			return Result;
 		}
 
 		/// <summary>
@@ -220,7 +213,7 @@ namespace UnrealBuildTool
 		/// </summary>
 		public void Delete()
 		{
-			Debug.Assert(_bExists);
+			Debug.Assert(bExists);
 
 			int MaxRetryCount = 3;
 			int DeleteTryCount = 0;
@@ -263,74 +256,22 @@ namespace UnrealBuildTool
 		}
 
 		/// <summary>
-		/// Initialization constructor.
+		/// Constructor
 		/// </summary>
-		protected FileItem(FileReference InFile)
+		/// <param name="Location">Location of the file</param>
+		/// <param name="Info">File info</param>
+		private FileItem(FileReference Location, FileInfo Info)
 		{
-			Location = InFile;
-
-			ResetFileInfo();
-
-			++TotalFileItemCount;
-			if (!_bExists)
-			{
-				++MissingFileItemCount;
-				// Log.TraceInformation( "Missing: " + FileAbsolutePath );
-			}
-
-			UniqueSourceFileMap[Location] = this;
+			this.Location = Location;
+			this.Info = Info;
 		}
 
-
 		/// <summary>
-		/// ISerializable: Constructor called when this object is deserialized
-		/// </summary>
-		protected FileItem(SerializationInfo SerializationInfo, StreamingContext StreamingContext)
-		{
-			ProducingAction = (Action)SerializationInfo.GetValue("pa", typeof(Action));
-			Location = (FileReference)SerializationInfo.GetValue("fi", typeof(FileReference));
-			CachedIncludePaths = (CppIncludePaths)SerializationInfo.GetValue("ci", typeof(CppIncludePaths));
-
-			// Go ahead and init normally now
-			{
-				ResetFileInfo();
-
-				++TotalFileItemCount;
-				if (!_bExists)
-				{
-					++MissingFileItemCount;
-					// Log.TraceInformation( "Missing: " + FileAbsolutePath );
-				}
-
-				UniqueSourceFileMap[Location] = this;
-			}
-		}
-
-
-		/// <summary>
-		/// ISerializable: Called when serialized to report additional properties that should be saved
-		/// </summary>
-		public void GetObjectData(SerializationInfo SerializationInfo, StreamingContext StreamingContext)
-		{
-			SerializationInfo.AddValue("pa", ProducingAction);
-			SerializationInfo.AddValue("fi", Location);
-			SerializationInfo.AddValue("ci", CachedIncludePaths);
-		}
-
-
-		/// <summary>
-		/// (Re-)set file information for this FileItem
+		/// Resets the cached file info
 		/// </summary>
 		public void ResetFileInfo()
 		{
-			Info = new FileInfo(AbsolutePath);
-
-			_bExists = Info.Exists;
-			if (_bExists)
-			{
-				_LastWriteTime = Info.LastWriteTimeUtc;
-				_Length = Info.Length;
-			}
+			Info = Location.ToFileInfo();
 		}
 
 		/// <summary>
@@ -344,33 +285,6 @@ namespace UnrealBuildTool
 			}
 		}
 
-		/// <summary>
-		/// Initialization constructor for optionally remote files.
-		/// </summary>
-		protected FileItem(FileReference InReference, bool InIsRemoteFile, UnrealTargetPlatform Platform)
-		{
-			Location = InReference;
-
-			FileInfo Info = new FileInfo(AbsolutePath);
-
-			_bExists = Info.Exists;
-			if (_bExists)
-			{
-				_LastWriteTime = Info.LastWriteTimeUtc;
-				_Length = Info.Length;
-			}
-
-			++TotalFileItemCount;
-			if (!_bExists)
-			{
-				++MissingFileItemCount;
-				// Log.TraceInformation( "Missing: " + FileAbsolutePath );
-			}
-
-			// @todo iosmerge: This was in UE3, why commented out now?
-			//UniqueSourceFileMap[AbsolutePathUpperInvariant] = this;
-		}
-
 		public override string ToString()
 		{
 			return AbsolutePath;
