UBT: Minor cleanup of FileItem names to match FileInfo, and add some comments.

#rnx

diff --git a/Engine/Source/Programs/UnrealBuildTool/Configuration/UEBuildModuleCPP.cs b/Engine/Source/Programs/UnrealBuildTool/Configuration/UEBuildModuleCPP.cs
--- a/Engine/Source/Programs/UnrealBuildTool/Configuration/UEBuildModuleCPP.cs
+++ b/Engine/Source/Programs/UnrealBuildTool/Configuration/UEBuildModuleCPP.cs
@@ -309,7 +309,7 @@ namespace UnrealBuildTool
 				foreach(FileReference OutputFile in Manifest.OutputFiles)
 				{
 					FileItem ObjectFile = FileItem.GetItemByFileReference(OutputFile);
-					if(!ObjectFile.bExists)
+					if(!ObjectFile.Exists)
 					{
 						throw new BuildException("Missing object file {0} listed in {1}", OutputFile, PrecompiledManifestLocation);
 					}
@@ -520,7 +520,7 @@ namespace UnrealBuildTool
 					if(GeneratedCPPCompileEnvironment.PrecompiledHeaderFile == null && Rules.PrivatePCHHeaderFile != null && Rules.PCHUsage != ModuleRules.PCHUsageMode.UseExplicitOrSharedPCHs)
 					{
 						FileItem PrivatePchFileItem = FileItem.GetItemByFileReference(FileReference.Combine(ModuleDirectory, Rules.PrivatePCHHeaderFile));
-						if(!PrivatePchFileItem.bExists)
+						if(!PrivatePchFileItem.Exists)
 						{
 							throw new BuildException("Unable to find private PCH file '{0}', referenced by '{1}'", PrivatePchFileItem.Location, RulesFile);
 						}
@@ -923,7 +923,7 @@ namespace UnrealBuildTool
 			WrapperFile.CachedIncludePaths = IncludedFile.CachedIncludePaths;
 
 			// Touch it if the included file is newer, to make sure our timestamp dependency checking is accurate.
-			if (IncludedFile.LastWriteTime > WrapperFile.LastWriteTime)
+			if (IncludedFile.LastWriteTimeUtc > WrapperFile.LastWriteTimeUtc)
 			{
 				File.SetLastWriteTimeUtc(WrapperFile.AbsolutePath, DateTime.UtcNow);
 				WrapperFile.ResetFileInfo();
diff --git a/Engine/Source/Programs/UnrealBuildTool/Modes/BuildMode.cs b/Engine/Source/Programs/UnrealBuildTool/Modes/BuildMode.cs
--- a/Engine/Source/Programs/UnrealBuildTool/Modes/BuildMode.cs
+++ b/Engine/Source/Programs/UnrealBuildTool/Modes/BuildMode.cs
@@ -898,7 +898,6 @@ namespace UnrealBuildTool
 				Log.TraceLog("FindIncludePaths calls: {0} ({1} searches)", CPPHeaders.TotalFindIncludedFileCalls, CPPHeaders.IncludePathSearchAttempts);
 				Log.TraceLog("Deep C++ include scan time: {0}s", UnrealBuildTool.TotalDeepIncludeScanTime);
 				Log.TraceLog("Include Resolves: {0} ({1} misses, {2:0.00}%)", CPPHeaders.TotalDirectIncludeResolves, CPPHeaders.TotalDirectIncludeResolveCacheMisses, (float)CPPHeaders.TotalDirectIncludeResolveCacheMisses / (float)CPPHeaders.TotalDirectIncludeResolves * 100);
-				Log.TraceLog("Total FileItems: {0} ({1} missing)", FileItem.TotalFileItemCount, FileItem.MissingFileItemCount);
 			}
 			return (int)BuildResult;
 		}
diff --git a/Engine/Source/Programs/UnrealBuildTool/System/ActionGraph.cs b/Engine/Source/Programs/UnrealBuildTool/System/ActionGraph.cs
--- a/Engine/Source/Programs/UnrealBuildTool/System/ActionGraph.cs
+++ b/Engine/Source/Programs/UnrealBuildTool/System/ActionGraph.cs
@@ -450,7 +450,7 @@ namespace UnrealBuildTool
 			{
 				// Determine the last time the action was run based on the write times of its produced files.
 				string LatestUpdatedProducedItemName = null;
-				DateTimeOffset LastExecutionTime = DateTimeOffset.MaxValue;
+				DateTimeOffset LastExecutionTimeUtc = DateTimeOffset.MaxValue;
 				foreach (FileItem ProducedItem in RootAction.ProducedItems)
 				{
 					// Optionally skip the action history check, as this only works for local builds
@@ -462,7 +462,7 @@ namespace UnrealBuildTool
 						if (!ActionHistory.GetProducingCommandLine(ProducedItem, out OldProducingCommandLine)
 						|| !String.Equals(OldProducingCommandLine, NewProducingCommandLine, StringComparison.InvariantCultureIgnoreCase))
 						{
-							if(ProducedItem.bExists)
+							if(ProducedItem.Exists)
 							{
 								Log.TraceLog(
 									"{0}: Produced item \"{1}\" was produced by outdated command-line.\n  Old command-line: {2}\n  New command-line: {3}",
@@ -483,12 +483,12 @@ namespace UnrealBuildTool
 					// If the produced file doesn't exist or has zero size, consider it outdated.  The zero size check is to detect cases
 					// where aborting an earlier compile produced invalid zero-sized obj files, but that may cause actions where that's
 					// legitimate output to always be considered outdated.
-					if (ProducedItem.bExists && (RootAction.ActionType != ActionType.Compile || ProducedItem.Length > 0))
+					if (ProducedItem.Exists && (RootAction.ActionType != ActionType.Compile || ProducedItem.Length > 0))
 					{
 						// Use the oldest produced item's time as the last execution time.
-						if (ProducedItem.LastWriteTime < LastExecutionTime)
+						if (ProducedItem.LastWriteTimeUtc < LastExecutionTimeUtc)
 						{
-							LastExecutionTime = ProducedItem.LastWriteTime;
+							LastExecutionTimeUtc = ProducedItem.LastWriteTimeUtc;
 							LatestUpdatedProducedItemName = ProducedItem.AbsolutePath;
 						}
 					}
@@ -547,10 +547,10 @@ namespace UnrealBuildTool
 							{
 								foreach (FileItem IncludedFile in IncludedFileList)
 								{
-									if (IncludedFile.bExists)
+									if (IncludedFile.Exists)
 									{
 										// allow a 1 second slop for network copies
-										TimeSpan TimeDifference = IncludedFile.LastWriteTime - LastExecutionTime;
+										TimeSpan TimeDifference = IncludedFile.LastWriteTimeUtc - LastExecutionTimeUtc;
 										bool bPrerequisiteItemIsNewerThanLastExecution = TimeDifference.TotalSeconds > 1;
 										if (bPrerequisiteItemIsNewerThanLastExecution)
 										{
@@ -558,8 +558,8 @@ namespace UnrealBuildTool
 												"{0}: Included file {1} is newer than the last execution of the action: {2} vs {3}",
 												RootAction.StatusDescription,
 												Path.GetFileName(IncludedFile.AbsolutePath),
-												IncludedFile.LastWriteTime.LocalDateTime,
-												LastExecutionTime.LocalDateTime
+												IncludedFile.LastWriteTimeUtc.LocalDateTime,
+												LastExecutionTimeUtc.LocalDateTime
 												);
 											bIsOutdated = true;
 
@@ -609,10 +609,10 @@ namespace UnrealBuildTool
 								}
 							}
 
-							if (PrerequisiteItem.bExists)
+							if (PrerequisiteItem.Exists)
 							{
 								// allow a 1 second slop for network copies
-								TimeSpan TimeDifference = PrerequisiteItem.LastWriteTime - LastExecutionTime;
+								TimeSpan TimeDifference = PrerequisiteItem.LastWriteTimeUtc - LastExecutionTimeUtc;
 								bool bPrerequisiteItemIsNewerThanLastExecution = TimeDifference.TotalSeconds > 1;
 								if (bPrerequisiteItemIsNewerThanLastExecution)
 								{
@@ -620,8 +620,8 @@ namespace UnrealBuildTool
 										"{0}: Prerequisite {1} is newer than the last execution of the action: {2} vs {3}",
 										RootAction.StatusDescription,
 										Path.GetFileName(PrerequisiteItem.AbsolutePath),
-										PrerequisiteItem.LastWriteTime.LocalDateTime,
-										LastExecutionTime.LocalDateTime
+										PrerequisiteItem.LastWriteTimeUtc.LocalDateTime,
+										LastExecutionTimeUtc.LocalDateTime
 										);
 									bIsOutdated = true;
 								}
@@ -711,7 +711,7 @@ namespace UnrealBuildTool
 					Action OutdatedAction = OutdatedActionInfo.Key;
 					foreach (FileItem DeleteItem in OutdatedActionInfo.Key.DeleteItems)
 					{
-						if (DeleteItem.bExists)
+						if (DeleteItem.Exists)
 						{
 							Log.TraceLog("Deleting outdated item: {0}", DeleteItem.AbsolutePath);
 							DeleteItem.Delete();
diff --git a/Engine/Source/Programs/UnrealBuildTool/System/DependencyCache.cs b/Engine/Source/Programs/UnrealBuildTool/System/DependencyCache.cs
--- a/Engine/Source/Programs/UnrealBuildTool/System/DependencyCache.cs
+++ b/Engine/Source/Programs/UnrealBuildTool/System/DependencyCache.cs
@@ -273,7 +273,7 @@ namespace UnrealBuildTool
 			}
 
 			// File is in cache, now check whether last write time is prior to cache creation time.
-			if (File.LastWriteTime.ToUniversalTime() >= CreationTimeUtc)
+			if (File.LastWriteTimeUtc.ToUniversalTime() >= CreationTimeUtc)
 			{
 				// Remove entry from cache as it's stale.
 				DependencyMap.Remove(File.Location);
diff --git a/Engine/Source/Programs/UnrealBuildTool/System/FileItem.cs b/Engine/Source/Programs/UnrealBuildTool/System/FileItem.cs
--- a/Engine/Source/Programs/UnrealBuildTool/System/FileItem.cs
+++ b/Engine/Source/Programs/UnrealBuildTool/System/FileItem.cs
@@ -13,15 +13,11 @@ using System.Collections.Concurrent;
 namespace UnrealBuildTool
 {
 	/// <summary>
-	/// Represents a file on disk that is used as an input or output of a build action.
-	/// FileItems are created by calling FileItem.GetItemByFileReference, which creates a single FileItem for each unique file path.
+	/// Represents a file on disk that is used as an input or output of a build action. FileItem instances are unique for a given path. Use FileItem.GetItemByFileReference 
+	/// to get the FileItem for a specific path.
 	/// </summary>
 	class FileItem
 	{
-		///
-		/// Preparation and Assembly (serialized)
-		/// 
-
 		/// <summary>
 		/// The action that produces the file.
 		/// </summary>
@@ -32,14 +28,6 @@ namespace UnrealBuildTool
 		/// </summary>
 		public FileReference Location;
 
-		/// <summary>
-		/// Accessor for the absolute path to the file
-		/// </summary>
-		public string AbsolutePath
-		{
-			get { return Location.FullName; }
-		}
-
 		/// <summary>
 		/// For C++ file items, this stores cached information about the include paths needed in order to include header files from these C++ files.  This is part of UBT's dependency caching optimizations.
 		/// </summary>
@@ -62,11 +50,6 @@ namespace UnrealBuildTool
 		}
 		private CppIncludePaths CachedIncludePathsValue;
 
-
-		///
-		/// Transients (not serialized)
-		///
-
 		/// <summary>
 		/// The information about the file.
 		/// </summary>
@@ -78,17 +61,33 @@ namespace UnrealBuildTool
 		public long RelativeCost = 0;
 
 		/// <summary>
-		/// The last write time of the file.
+		/// A case-insensitive dictionary that's used to map each unique file name to a single FileItem object.
 		/// </summary>
-		public DateTimeOffset LastWriteTime
+		static ConcurrentDictionary<FileReference, FileItem> UniqueSourceFileMap = new ConcurrentDictionary<FileReference, FileItem>();
+
+		/// <summary>
+		/// Constructor
+		/// </summary>
+		/// <param name="Location">Location of the file</param>
+		/// <param name="Info">File info</param>
+		private FileItem(FileReference Location, FileInfo Info)
 		{
-			get { return Info.LastWriteTimeUtc; }
+			this.Location = Location;
+			this.Info = Info;
+		}
+
+		/// <summary>
+		/// Accessor for the absolute path to the file
+		/// </summary>
+		public string AbsolutePath
+		{
+			get { return Location.FullName; }
 		}
 
 		/// <summary>
 		/// Whether the file exists.
 		/// </summary>
-		public bool bExists
+		public bool Exists
 		{
 			get { return Info.Exists; }
 		}
@@ -101,15 +100,13 @@ namespace UnrealBuildTool
 			get { return Info.Length; }
 		}
 
-
-		///
-		/// Statics
-		///
-
 		/// <summary>
-		/// A case-insensitive dictionary that's used to map each unique file name to a single FileItem object.
+		/// The last write time of the file.
 		/// </summary>
-		static ConcurrentDictionary<FileReference, FileItem> UniqueSourceFileMap = new ConcurrentDictionary<FileReference, FileItem>();
+		public DateTimeOffset LastWriteTimeUtc
+		{
+			get { return Info.LastWriteTimeUtc; }
+		}
 
 		/// <summary>
 		/// Clears the FileItem caches.
@@ -130,26 +127,34 @@ namespace UnrealBuildTool
 			}
 		}
 
+		/// <summary>
+		/// Gets a FileItem corresponding to the given path
+		/// </summary>
+		/// <param name="FilePath">Path for the FileItem</param>
 		/// <returns>The FileItem that represents the given file path.</returns>
 		public static FileItem GetItemByPath(string FilePath)
 		{
 			return GetItemByFileReference(new FileReference(FilePath));
 		}
 
+		/// <summary>
+		/// Gets a FileItem for a given path
+		/// </summary>
+		/// <param name="Location">Location of the file</param>
 		/// <returns>The FileItem that represents the given a full file path.</returns>
-		public static FileItem GetItemByFileReference(FileReference Reference)
+		public static FileItem GetItemByFileReference(FileReference Location)
 		{
 			FileItem Result;
-			if (!UniqueSourceFileMap.TryGetValue(Reference, out Result))
+			if (!UniqueSourceFileMap.TryGetValue(Location, out Result))
 			{
-				FileItem NewFileItem = new FileItem(Reference, Reference.ToFileInfo());
-				if(UniqueSourceFileMap.TryAdd(Reference, NewFileItem))
+				FileItem NewFileItem = new FileItem(Location, Location.ToFileInfo());
+				if(UniqueSourceFileMap.TryAdd(Location, NewFileItem))
 				{
 					Result = NewFileItem;
 				}
 				else
 				{
-					Result = UniqueSourceFileMap[Reference];
+					Result = UniqueSourceFileMap[Location];
 				}
 			}
 			return Result;
@@ -213,7 +218,7 @@ namespace UnrealBuildTool
 		/// </summary>
 		public void Delete()
 		{
-			Debug.Assert(bExists);
+			Debug.Assert(Exists);
 
 			int MaxRetryCount = 3;
 			int DeleteTryCount = 0;
@@ -255,17 +260,6 @@ namespace UnrealBuildTool
 			while (!bFileDeletedSuccessfully && (DeleteTryCount < MaxRetryCount));
 		}
 
-		/// <summary>
-		/// Constructor
-		/// </summary>
-		/// <param name="Location">Location of the file</param>
-		/// <param name="Info">File info</param>
-		private FileItem(FileReference Location, FileInfo Info)
-		{
-			this.Location = Location;
-			this.Info = Info;
-		}
-
 		/// <summary>
 		/// Resets the cached file info
 		/// </summary>
@@ -285,6 +279,10 @@ namespace UnrealBuildTool
 			}
 		}
 
+		/// <summary>
+		/// Return the path to this FileItem to debugging
+		/// </summary>
+		/// <returns>Absolute path to this file item</returns>
 		public override string ToString()
 		{
 			return AbsolutePath;
diff --git a/Engine/Source/Programs/UnrealBuildTool/System/UBTMakefile.cs b/Engine/Source/Programs/UnrealBuildTool/System/UBTMakefile.cs
--- a/Engine/Source/Programs/UnrealBuildTool/System/UBTMakefile.cs
+++ b/Engine/Source/Programs/UnrealBuildTool/System/UBTMakefile.cs
@@ -355,13 +355,13 @@ namespace UnrealBuildTool
 
 				foreach(FileItem AdditionalDependency in Prerequisites.AdditionalDependencies)
 				{
-					if (!AdditionalDependency.bExists)
+					if (!AdditionalDependency.Exists)
 					{
 						Log.TraceLog("{0} has been deleted since makefile was built.", AdditionalDependency.Location);
 						ReasonNotLoaded = string.Format("{0} deleted", AdditionalDependency.Location.GetFileName());
 						return null;
 					}
-					if(AdditionalDependency.LastWriteTime > UBTMakefileInfo.LastWriteTime)
+					if(AdditionalDependency.LastWriteTimeUtc > UBTMakefileInfo.LastWriteTime)
 					{
 						Log.TraceLog("{0} has been modified since makefile was built.", AdditionalDependency.Location);
 						ReasonNotLoaded = string.Format("{0} modified", AdditionalDependency.Location.GetFileName());
