PR #5356: UE-67548: UBT: Fix for IntelliSense generation for game projects using engine source (Contributed by [REDACTED_PERSON])

#rnx

diff --git a/Engine/Source/Programs/UnrealBuildTool/ProjectFiles/ProjectFileGenerator.cs b/Engine/Source/Programs/UnrealBuildTool/ProjectFiles/ProjectFileGenerator.cs
--- a/Engine/Source/Programs/UnrealBuildTool/ProjectFiles/ProjectFileGenerator.cs
+++ b/Engine/Source/Programs/UnrealBuildTool/ProjectFiles/ProjectFileGenerator.cs
@@ -1693,6 +1693,10 @@ namespace UnrealBuildTool
 						{
 							NewArguments.Add("-precompile");
 						}
+						if (CurTarget.TargetRules.Type == TargetType.Editor && CurTarget.TargetRules.BuildEnvironment == TargetBuildEnvironment.Default)
+						{
+							NewArguments.Add("-SharedBuildEnvironment"); // Force the shared build environment, even though we have -Monolithic
+						}
 						NewArguments.Add("-Monolithic"); // Don't want any _API macros to be DLLIMPORT; this will cause Intellisense errors in the file that they are defined.
 						NewArguments.AddRange(Arguments);
 
