UBT: Remove the unused RelativeCost field from FileItem.

#rnx

diff --git a/Engine/Source/Programs/UnrealBuildTool/System/ActionGraph.cs b/Engine/Source/Programs/UnrealBuildTool/System/ActionGraph.cs
--- a/Engine/Source/Programs/UnrealBuildTool/System/ActionGraph.cs
+++ b/Engine/Source/Programs/UnrealBuildTool/System/ActionGraph.cs
@@ -195,7 +195,6 @@ namespace UnrealBuildTool
 				foreach (FileItem ProducedItem in Action.ProducedItems)
 				{
 					ProducedItem.ProducingAction = Action;
-					Action.RelativeCost += ProducedItem.RelativeCost;
 				}
 			}
 		}
diff --git a/Engine/Source/Programs/UnrealBuildTool/System/FileItem.cs b/Engine/Source/Programs/UnrealBuildTool/System/FileItem.cs
--- a/Engine/Source/Programs/UnrealBuildTool/System/FileItem.cs
+++ b/Engine/Source/Programs/UnrealBuildTool/System/FileItem.cs
@@ -55,11 +55,6 @@ namespace UnrealBuildTool
 		/// </summary>
 		public FileInfo Info;
 
-		/// <summary>
-		/// Relative cost of action associated with producing this file.
-		/// </summary>
-		public long RelativeCost = 0;
-
 		/// <summary>
 		/// A case-insensitive dictionary that's used to map each unique file name to a single FileItem object.
 		/// </summary>
diff --git a/Engine/Source/Programs/UnrealBuildTool/System/Unity.cs b/Engine/Source/Programs/UnrealBuildTool/System/Unity.cs
--- a/Engine/Source/Programs/UnrealBuildTool/System/Unity.cs
+++ b/Engine/Source/Programs/UnrealBuildTool/System/Unity.cs
@@ -340,7 +340,6 @@ namespace UnrealBuildTool
 				// Write the unity file to the intermediate folder.
 				FileItem UnityCPPFile = FileItem.CreateIntermediateTextFile(UnityCPPFilePath, OutputUnityCPPWriter.ToString());
 				UnityCPPFile.CachedIncludePaths = CompileEnvironment.IncludePaths;
-				UnityCPPFile.RelativeCost = UnityFile.TotalLength;
 				NewCPPFiles.Add(UnityCPPFile);
 			}
 
