UBT: Fix incremental builds with makefiles always showing the target as up to date, due to deserialized FileItems not having an associated ProducingAction.

#rnx

diff --git a/Engine/Source/Programs/UnrealBuildTool/Modes/BuildMode.cs b/Engine/Source/Programs/UnrealBuildTool/Modes/BuildMode.cs
--- a/Engine/Source/Programs/UnrealBuildTool/Modes/BuildMode.cs
+++ b/Engine/Source/Programs/UnrealBuildTool/Modes/BuildMode.cs
@@ -534,6 +534,7 @@ namespace UnrealBuildTool
 							if (!bIsGatheringBuild && BuildResult.Succeeded())
 							{
 								ActionGraph.AllActions = Makefile.AllActions;
+								ActionGraph.LinkActionsAndItems();
 
 								OutputItemsForAllTargets = Makefile.OutputItemsForAllTargets;
 
diff --git a/Engine/Source/Programs/UnrealBuildTool/System/ActionGraph.cs b/Engine/Source/Programs/UnrealBuildTool/System/ActionGraph.cs
--- a/Engine/Source/Programs/UnrealBuildTool/System/ActionGraph.cs
+++ b/Engine/Source/Programs/UnrealBuildTool/System/ActionGraph.cs
@@ -188,7 +188,7 @@ namespace UnrealBuildTool
 		/// <summary>
 		/// Links actions with their prerequisite and produced items into an action graph.
 		/// </summary>
-		void LinkActionsAndItems()
+		public void LinkActionsAndItems()
 		{
 			foreach (Action Action in AllActions)
 			{
