UBT: Remove debug code from makefile serialization.

#rnx

diff --git a/Engine/Source/Programs/UnrealBuildTool/System/UBTMakefile.cs b/Engine/Source/Programs/UnrealBuildTool/System/UBTMakefile.cs
--- a/Engine/Source/Programs/UnrealBuildTool/System/UBTMakefile.cs
+++ b/Engine/Source/Programs/UnrealBuildTool/System/UBTMakefile.cs
@@ -292,15 +292,7 @@ namespace UnrealBuildTool
 				{
 					using (FileStream Stream = new FileStream(UBTMakefileInfo.FullName, FileMode.Open, FileAccess.Read))
 					{
-						byte[] Data = new byte[Stream.Length];
-						if(Stream.Read(Data, 0, Data.Length) != Data.Length)
-						{
-							throw new Exception("Error");
-						}
-
-						Timeline.AddEvent("AFTER READ");
-
-						BinaryReader Reader = new BinaryReader(new MemoryStream(Data));
+						BinaryReader Reader = new BinaryReader(Stream);
 						LoadedUBTMakefile = new UBTMakefile(Reader);
 					}
 				}
