UBT: Use compiler-generated dependency information for source file dependencies. This ensures that the dependency list is accurate, and makes incremental builds more reliable.

* A standalone Win32 application called cl-filter is used to parse output from running the Visual Studio compiler with the /showIncludes option, and generate a dependency list similar to Clang.
* Actions which produce an external dependency list declare it explicitly. If the dependency list is missing, or any of the contents are out of date, the action is rebuilt.
* The contents of dependency lists are cached by the CppDependencyCache class for future iterations.

diff --git a/Engine/Extras/cl-filter/cl-filter.cpp b/Engine/Extras/cl-filter/cl-filter.cpp
new file mode 100644
--- /dev/null
+++ b/Engine/Extras/cl-filter/cl-filter.cpp
@@ -0,0 +1,203 @@
+// Copyright 1998-2019 Epic Games, Inc. All Rights Reserved.
+
+/**
+ * This program is designed to execute the Visual C++ compiler (cl.exe) and filter off any output lines from the /showIncludes directive
+ * into a separate file for dependency checking. GCC/Clang have a specific option for this, whereas MSVC does not.
+ */
+
+#include <windows.h>
+#include <stdio.h>
+#include <string>
+#include <vector>
+
+int wmain(int ArgC, const wchar_t* ArgV[])
+{
+	// Get the full command line, and find the '--' marker
+	wchar_t* CmdLine = ::GetCommandLineW();
+
+	wchar_t *ChildCmdLine = wcsstr(CmdLine, L" -- ");
+	if (ChildCmdLine == nullptr)
+	{
+		wprintf(L"ERROR: Unable to find child command line (%s)", CmdLine);
+		return -1;
+	}
+
+	*ChildCmdLine = 0;
+	ChildCmdLine += 4;
+
+	int NumArgs = 0;
+	wchar_t** Args = CommandLineToArgvW(CmdLine, &NumArgs);
+
+	if (NumArgs != 2)
+	{
+		wprintf(L"ERROR: Syntax: cl-filter <dependencies-file> -- <child command line>\n");
+		return -1;
+	}
+
+	const wchar_t* OutputFileName = Args[1];
+
+	// Create the child process
+	PROCESS_INFORMATION ProcessInfo;
+	ZeroMemory(&ProcessInfo, sizeof(ProcessInfo));
+
+	SECURITY_ATTRIBUTES SecurityAttributes;
+	ZeroMemory(&SecurityAttributes, sizeof(SecurityAttributes));
+	SecurityAttributes.bInheritHandle = TRUE;
+
+	HANDLE StdOutReadHandle;
+	HANDLE StdOutWriteHandle;
+	if (CreatePipe(&StdOutReadHandle, &StdOutWriteHandle, &SecurityAttributes, 0) == 0)
+	{
+		wprintf(L"ERROR: Unable to create output pipe for child process\n");
+		return -1;
+	}
+
+	HANDLE StdErrWriteHandle;
+	if (DuplicateHandle(GetCurrentProcess(), StdOutWriteHandle, GetCurrentProcess(), &StdErrWriteHandle, 0, true, DUPLICATE_SAME_ACCESS) == 0)
+	{
+		wprintf(L"ERROR: Unable to create stderr pipe handle for child process\n");
+		return -1;
+	}
+
+	// Create the new process as suspended, so we can modify it before it starts executing (and potentially preempting us)
+	STARTUPINFO StartupInfo;
+	ZeroMemory(&StartupInfo, sizeof(StartupInfo));
+	StartupInfo.cb = sizeof(StartupInfo);
+	StartupInfo.hStdInput = NULL;
+	StartupInfo.hStdOutput = StdOutWriteHandle;
+	StartupInfo.hStdError = StdErrWriteHandle;
+	StartupInfo.dwFlags = STARTF_USESTDHANDLES;
+
+	DWORD ProcessCreationFlags = GetPriorityClass(GetCurrentProcess());
+	if (CreateProcessW(NULL, ChildCmdLine, NULL, NULL, TRUE, ProcessCreationFlags, NULL, NULL, &StartupInfo, &ProcessInfo) == 0)
+	{
+		wprintf(L"ERROR: Unable to create child process\n");
+		return -1;
+	}
+
+	// Close the write ends of the handle. We don't want any other process to be able to inherit these.
+	CloseHandle(StdOutWriteHandle);
+	CloseHandle(StdErrWriteHandle);
+
+	// Create a file to contain the dependency list
+	FILE* OutputFile;
+	if (_wfopen_s(&OutputFile, OutputFileName, L"wt") != 0)
+	{
+		wprintf(L"ERROR: Unable to open %s for output", OutputFileName);
+		return -1;
+	}
+
+	// Get the default console codepage
+	UINT CodePage = GetConsoleCP();
+	
+	// Pipe the output to stdout
+	wchar_t Buffer[4096];
+	size_t BufferSize = 0;
+
+	for (;;)
+	{
+		char InputData[4096];
+
+		// Read the next chunk of data from the output stream
+		DWORD BytesRead;
+		if (!ReadFile(StdOutReadHandle, InputData, sizeof(InputData), &BytesRead, NULL))
+		{
+			if (GetLastError() != ERROR_BROKEN_PIPE)
+			{
+				wprintf(L"ERROR: Unable to read data from child process (%08x)", GetLastError());
+			}
+			BytesRead = 0;
+		}
+
+		// Convert it to a string
+		if (BytesRead > 0)
+		{
+			int NumCharacters = MultiByteToWideChar(CodePage, 0, InputData, (int)BytesRead, Buffer + BufferSize, (int)((sizeof(Buffer) / sizeof(Buffer[0])) - BufferSize));
+			if (NumCharacters > 0)
+			{
+				BufferSize += (size_t)NumCharacters;
+			}
+		}
+
+		// Parse individual lines from the output
+		size_t LineStart = 0;
+		while(LineStart < BufferSize)
+		{
+			// Find the end of this line
+			size_t LineEnd = LineStart;
+			while (LineEnd < BufferSize && Buffer[LineEnd] != '\r' && Buffer[LineEnd] != '\n')
+			{
+				LineEnd++;
+			}
+
+			// If we didn't reach a line terminator, and we can still read more data, clear up some space and try again
+			if (LineEnd == BufferSize && LineStart != 0)
+			{
+				break;
+			}
+
+			// Skip past the EOL markers
+			if (LineEnd < BufferSize && Buffer[LineEnd] == '\r')
+			{
+				LineEnd++;
+			}
+			if (LineEnd < BufferSize && Buffer[LineEnd] == '\n')
+			{
+				LineEnd++;
+			}
+
+			// Filter out any lines that are 
+			const wchar_t Prefix[] = L"Note: including file: ";
+			if (wcsncmp(Buffer + LineStart, Prefix, sizeof(Prefix) / sizeof(Prefix[0]) - 1) == 0)
+			{
+				size_t FileNameStart = LineStart + sizeof(Prefix) / sizeof(Prefix[0]) - 1;
+				while (FileNameStart < LineEnd && iswspace(Buffer[FileNameStart]))
+				{
+					FileNameStart++;
+				}
+
+				size_t FileNameEnd = LineEnd;
+				while (FileNameEnd > FileNameStart && iswspace(Buffer[FileNameEnd - 1]))
+				{
+					FileNameEnd--;
+				}
+
+				Buffer[FileNameEnd] = 0;
+				fwprintf(OutputFile, L"%s\n", Buffer + FileNameStart);
+			}
+			else
+			{
+				char RawOutput[4096];
+				int RawOutputSize = WideCharToMultiByte(CodePage, 0, Buffer + LineStart, (int)(LineEnd - LineStart), RawOutput, sizeof(RawOutput), nullptr, nullptr);
+
+				DWORD BytesWritten;
+				WriteFile(GetStdHandle(STD_OUTPUT_HANDLE), RawOutput, RawOutputSize, &BytesWritten, NULL);
+			}
+
+			// Move to the next line
+			LineStart = LineEnd;
+		}
+
+		// Check if we've finished
+		if (BufferSize == 0 && BytesRead == 0)
+		{
+			break;
+		}
+
+		// Shuffle everything down
+		memmove(Buffer, Buffer + LineStart, (BufferSize - LineStart) * sizeof(wchar_t));
+		BufferSize -= LineStart;
+	}
+
+	WaitForSingleObject(ProcessInfo.hProcess, INFINITE);
+
+	DWORD ExitCode;
+	if (!GetExitCodeProcess(ProcessInfo.hProcess, &ExitCode))
+	{
+		ExitCode = (DWORD)-1;
+	}
+
+	fclose(OutputFile);
+
+	return ExitCode;
+}
diff --git a/Engine/Source/Programs/UnrealBuildTool/Configuration/BuildConfiguration.cs b/Engine/Source/Programs/UnrealBuildTool/Configuration/BuildConfiguration.cs
--- a/Engine/Source/Programs/UnrealBuildTool/Configuration/BuildConfiguration.cs
+++ b/Engine/Source/Programs/UnrealBuildTool/Configuration/BuildConfiguration.cs
@@ -20,13 +20,6 @@ namespace UnrealBuildTool
 		[XmlConfigFile]
 		public bool bIgnoreOutdatedImportLibraries = true;
 
-		/// <summary>
-		/// Whether to generate command line dependencies for compile actions when requested
-		/// </summary>
-		[XmlConfigFile]
-		[CommandLine("-SkipActionHistory", Value = "false")]
-		public bool bUseActionHistory = true;
-
 		/// <summary>
 		/// Use existing static libraries for all engine modules in this target.
 		/// </summary>
@@ -83,20 +76,6 @@ namespace UnrealBuildTool
 		[CommandLine("-SkipBuild")]
 		public bool bSkipBuild = false;
 
-		/// <summary>
-		/// Whether the dependency cache includes pre-resolved include locations so UBT doesn't have to re-resolve each include location just to check the timestamp.
-		/// This is technically not fully correct because the dependency cache is global and each module could have a different set of include paths that could cause headers
-		/// to resolve files differently. In practice this is not the case, and significantly speeds up UBT when nothing is to be done.
-		/// </summary>
-		[XmlConfigFile]
-		public bool bUseIncludeDependencyResolveCache = true;
-
-		/// <summary>
-		/// Used to test the dependency resolve cache. This will verify the resolve cache has no conflicts by resolving every time and checking against any previous resolve attempts.
-		/// </summary>
-		[XmlConfigFile]
-		public bool bTestIncludeDependencyResolveCache = false;
-
 		/// <summary>
 		/// Enables support for very fast iterative builds by caching target data.  Turning this on causes Unreal Build Tool to emit
 		/// 'UBT Makefiles' for targets when they're built the first time.  Subsequent builds will load these Makefiles and begin
diff --git a/Engine/Source/Programs/UnrealBuildTool/Configuration/UEBuildModuleCPP.cs b/Engine/Source/Programs/UnrealBuildTool/Configuration/UEBuildModuleCPP.cs
--- a/Engine/Source/Programs/UnrealBuildTool/Configuration/UEBuildModuleCPP.cs
+++ b/Engine/Source/Programs/UnrealBuildTool/Configuration/UEBuildModuleCPP.cs
@@ -329,37 +329,9 @@ namespace UnrealBuildTool
 
 			SourceDirectories = new HashSet<DirectoryReference>(SourceFiles.Select(x => x.Location.Directory));
 
-			// Store the module compile environment along with the source file.  This is so that we can use it later on when looking for header dependencies
-			foreach (FileItem CFile in SourceFilesFound.CFiles)
-			{
-				CFile.CachedIncludePaths = ModuleCompileEnvironment.IncludePaths;
-			}
-			foreach (FileItem CCFile in SourceFilesFound.CCFiles)
-			{
-				CCFile.CachedIncludePaths = ModuleCompileEnvironment.IncludePaths;
-			}
-			foreach (FileItem CPPFile in SourceFilesFound.CPPFiles)
-			{
-				CPPFile.CachedIncludePaths = ModuleCompileEnvironment.IncludePaths;
-			}
-			foreach (FileItem MMFile in SourceFilesFound.MMFiles)
-			{
-				MMFile.CachedIncludePaths = ModuleCompileEnvironment.IncludePaths;
-			}
-
 			// Process all of the header file dependencies for this module
 			CheckFirstIncludeMatchesEachCppFile(Target, ModuleCompileEnvironment, SourceFilesToBuild.CPPFiles);
 
-			// Make sure our RC files have cached includes.  
-			foreach (FileItem RCFile in SourceFilesToBuild.RCFiles)
-			{
-				// The default resource file (PCLaunch.rc) is created in a module-agnostic way, so we want to avoid overriding the include paths for it
-				if(RCFile.CachedIncludePaths == null)
-				{
-					RCFile.CachedIncludePaths = ModuleCompileEnvironment.IncludePaths;
-				}
-			}
-
 			// Should we force a precompiled header to be generated for this module?  Usually, we only bother with a
 			// precompiled header if there are at least several source files in the module (after combining them for unity
 			// builds.)  But for game modules, it can be convenient to always have a precompiled header to single-file
@@ -534,7 +506,6 @@ namespace UnrealBuildTool
 					foreach (string GeneratedFilename in GeneratedFiles)
 					{
 						FileItem GeneratedCppFileItem = FileItem.GetItemByPath(GeneratedFilename);
-						GeneratedCppFileItem.CachedIncludePaths = ModuleCompileEnvironment.IncludePaths;
 
 						// @todo ubtmake: Check for ALL other places where we might be injecting .cpp or .rc files for compiling without caching CachedCPPIncludeInfo first (anything platform specific?)
 						GeneratedFileItems.Add(GeneratedCppFileItem);
@@ -603,7 +574,6 @@ namespace UnrealBuildTool
 		{
 			CppCompileEnvironment CompileEnvironment = CreateSharedPCHCompileEnvironment(Target, BaseCompileEnvironment);
 			FileItem HeaderFile = FileItem.GetItemByFileReference(FileReference.Combine(ModuleDirectory, Rules.SharedPCHHeaderFile));
-			HeaderFile.CachedIncludePaths = CompileEnvironment.IncludePaths;
 
 			DirectoryReference PrecompiledHeaderDir;
 			if(Rules.bUsePrecompiled)
@@ -628,12 +598,6 @@ namespace UnrealBuildTool
 		/// <returns>The created PCH instance.</returns>
 		private PrecompiledHeaderInstance CreatePrivatePCH(UEToolChain ToolChain, FileItem HeaderFile, CppCompileEnvironment ModuleCompileEnvironment, ActionGraph ActionGraph)
 		{
-			// Cache the header file include paths. This file could have been a shared PCH too, so ignore if the include paths are already set.
-			if(HeaderFile.CachedIncludePaths == null)
-			{
-				HeaderFile.CachedIncludePaths = ModuleCompileEnvironment.IncludePaths;
-			}
-
 			// Create the wrapper file, which sets all the definitions needed to compile it
 			FileReference WrapperLocation = FileReference.Combine(IntermediateDirectory, String.Format("PCH.{0}.h", Name));
 			FileItem WrapperFile = CreatePCHWrapperFile(WrapperLocation, ModuleCompileEnvironment.Definitions, HeaderFile);
@@ -839,7 +803,6 @@ namespace UnrealBuildTool
 				// Write the PCH header
 				FileReference DedicatedPchLocation = FileReference.Combine(IntermediateDirectory, String.Format("PCH.Dedicated.{0}.h", File.Location.GetFileNameWithoutExtension()));
 				FileItem DedicatedPchFile = FileItem.CreateIntermediateTextFile(DedicatedPchLocation, WrapperContents.ToString());
-				DedicatedPchFile.CachedIncludePaths = File.CachedIncludePaths;
 
 				// Create a new C++ environment to compile the PCH
 				CppCompileEnvironment PchEnvironment = new CppCompileEnvironment(CompileEnvironment);
@@ -920,7 +883,6 @@ namespace UnrealBuildTool
 
 			// Create the item
 			FileItem WrapperFile = FileItem.CreateIntermediateTextFile(OutputFile, WrapperContents.ToString());
-			WrapperFile.CachedIncludePaths = IncludedFile.CachedIncludePaths;
 
 			// Touch it if the included file is newer, to make sure our timestamp dependency checking is accurate.
 			if (IncludedFile.LastWriteTimeUtc > WrapperFile.LastWriteTimeUtc)
@@ -977,12 +939,6 @@ namespace UnrealBuildTool
 						}
 					}
 
-					// Store the module compile environment along with the .cpp file.  This is so that we can use it later on when looking for header dependencies
-					foreach (FileItem CPPFile in CppFiles)
-					{
-						CPPFile.CachedIncludePaths = ModuleCompileEnvironment.IncludePaths;
-					}
-
 					// Find the directly included files for each source file, and make sure it includes the matching header if possible
 					InvalidIncludeDirectiveMessages = new List<string>();
 					if (Rules != null && Rules.bEnforceIWYU && Target.bEnforceIWYU)
diff --git a/Engine/Source/Programs/UnrealBuildTool/Configuration/UEBuildTarget.cs b/Engine/Source/Programs/UnrealBuildTool/Configuration/UEBuildTarget.cs
--- a/Engine/Source/Programs/UnrealBuildTool/Configuration/UEBuildTarget.cs
+++ b/Engine/Source/Programs/UnrealBuildTool/Configuration/UEBuildTarget.cs
@@ -1360,7 +1360,6 @@ namespace UnrealBuildTool
 						CppCompileEnvironment DefaultResourceCompileEnvironment = new CppCompileEnvironment(GlobalCompileEnvironment);
 
 						FileItem DefaultResourceFile = FileItem.GetItemByFileReference(FileReference.Combine(UnrealBuildTool.EngineDirectory, "Build", "Windows", "Default.rc2"));
-						DefaultResourceFile.CachedIncludePaths = DefaultResourceCompileEnvironment.IncludePaths;
 
 						WindowsPlatform.SetupResourceCompileEnvironment(DefaultResourceCompileEnvironment, Rules);
 
diff --git a/Engine/Source/Programs/UnrealBuildTool/Modes/BuildMode.cs b/Engine/Source/Programs/UnrealBuildTool/Modes/BuildMode.cs
--- a/Engine/Source/Programs/UnrealBuildTool/Modes/BuildMode.cs
+++ b/Engine/Source/Programs/UnrealBuildTool/Modes/BuildMode.cs
@@ -103,8 +103,6 @@ namespace UnrealBuildTool
 			string ExecutorName = "Unknown";
 			ECompilationResult BuildResult = ECompilationResult.Succeeded;
 
-			CppIncludeBackgroundThread CppIncludeThread = null;
-
 			List<UEBuildTarget> Targets = null;
 			Dictionary<UEBuildTarget, CPPHeaders> TargetToHeaders = new Dictionary<UEBuildTarget, CPPHeaders>();
 
@@ -114,10 +112,6 @@ namespace UnrealBuildTool
 			{
 				// action graph implies using the dependency resolve cache
 				bool GeneratingActionGraph = Arguments.HasOption("-Graph");
-				if (GeneratingActionGraph)
-				{
-					BuildConfiguration.bUseIncludeDependencyResolveCache = true;
-				}
 
 				Timeline.AddEvent("RunUBT() initialization complete");
 
@@ -168,51 +162,6 @@ namespace UnrealBuildTool
 					return (int)ECompilationResult.Succeeded;
 				}
 
-				// Used when BuildConfiguration.bUseUBTMakefiles is enabled.  If true, it means that our cached includes may not longer be
-				// valid (or never were), and we need to recover by forcibly scanning included headers for all build prerequisites to make sure that our
-				// cached set of includes is actually correct, before determining which files are outdated.
-				bool bNeedsFullCPPIncludeRescan = false;
-
-				if (BuildConfiguration.bUseUBTMakefiles)
-				{
-					// Only the modular editor and game targets will share build products.  Unfortunately, we can't determine at
-					// at this point whether we're dealing with modular or monolithic game binaries, so we opt to always invalidate
-					// cached includes if the target we're switching to is either a game target (has project file) or "UE4Editor".
-					bool bMightHaveSharedBuildProducts =
-						TargetDescs[0].ProjectFile != null ||  // Is this a game? (has a .uproject file for the target)
-						TargetDescs[0].Name.Equals("UE4Editor", StringComparison.InvariantCultureIgnoreCase); // Is the engine?
-					if (bMightHaveSharedBuildProducts)
-					{
-						bool bIsBuildingSameTargetsAsLastTime = false;
-
-						string TargetCollectionName = UBTMakefile.MakeTargetCollectionName(TargetDescs);
-
-						string LastBuiltTargetsFileName = "LastBuiltTargets.txt";
-						string LastBuiltTargetsFilePath = FileReference.Combine(UnrealBuildTool.EngineDirectory, "Intermediate", "Build", LastBuiltTargetsFileName).FullName;
-						if (File.Exists(LastBuiltTargetsFilePath) && Utils.ReadAllText(LastBuiltTargetsFilePath) == TargetCollectionName)
-						{
-							// @todo ubtmake: Because we're using separate files for hot reload vs. full compiles, it's actually possible that includes will
-							// become out of date without us knowing if the developer ping-pongs between hot reloading target A and building target B normally.
-							// To fix this we can not use a different file name for last built targets, but the downside is slower performance when
-							// performing the first hot reload after compiling normally (forces full include dependency scan)
-							bIsBuildingSameTargetsAsLastTime = true;
-						}
-
-						// Save out the name of the targets we're building
-						if (!bIsBuildingSameTargetsAsLastTime)
-						{
-							Directory.CreateDirectory(Path.GetDirectoryName(LastBuiltTargetsFilePath));
-							File.WriteAllText(LastBuiltTargetsFilePath, TargetCollectionName, Encoding.UTF8);
-
-							// Can't use super fast include checking unless we're building the same set of targets as last time, because
-							// we might not know about all of the C++ include dependencies for already-up-to-date shared build products
-							// between the targets
-							bNeedsFullCPPIncludeRescan = true;
-							Log.TraceInformation("Performing full C++ include scan (building a new target)");
-						}
-					}
-				}
-
 				// True if we should gather module dependencies for building in this run.  If this is false, then we'll expect to be loading information from disk about
 				// the target's modules before we'll be able to build anything.  One or both of IsGatheringBuild or IsAssemblingBuild must be true.
 				bool bIsGatheringBuild = true;
@@ -307,8 +256,6 @@ namespace UnrealBuildTool
 								// a 'gather' and 'assemble' in the same run.  This will take a while longer, but subsequent runs will be fast!
 								bIsGatheringBuild = true;
 
-								FileItem.ClearCaches();
-
 								Log.TraceInformation("Creating makefile for {0}{1} ({2})",
 									TargetDescs[0].Name,
 									TargetDescs.Count > 1 ? (" (and " + (TargetDescs.Count - 1).ToString() + " more)") : "",
@@ -346,8 +293,7 @@ namespace UnrealBuildTool
 					{
 						// Create the header cache for this target
 						FileReference DependencyCacheFile = DependencyCache.GetDependencyCachePathForTarget(Target.ProjectFile, Target.Platform, Target.TargetName);
-						bool bUseFlatCPPIncludeDependencyCache = BuildConfiguration.bUseUBTMakefiles && bIsAssemblingBuild;
-						CPPHeaders Headers = new CPPHeaders(Target.ProjectFile, DependencyCacheFile, bUseFlatCPPIncludeDependencyCache, BuildConfiguration.bUseUBTMakefiles, BuildConfiguration.bUseIncludeDependencyResolveCache, BuildConfiguration.bTestIncludeDependencyResolveCache);
+						CPPHeaders Headers = new CPPHeaders(Target.ProjectFile, DependencyCacheFile);
 						TargetToHeaders[Target] = Headers;
 
 						// Make sure the appropriate executor is selected
@@ -366,30 +312,6 @@ namespace UnrealBuildTool
 							}
 						}
 
-						// We don't need this dependency cache in 'gather only' mode
-						if (BuildConfiguration.bUseUBTMakefiles &&
-								!(bIsGatheringBuild && !bIsAssemblingBuild))
-						{
-							using(Timeline.ScopeEvent("FlatCPPIncludeDependencyCache.TryRead()"))
-							{
-								// Load the cache that contains the list of flattened resolved includes for resolved source files
-								// @todo ubtmake: Ideally load this asynchronously at startup and only block when it is first needed and not finished loading
-								FileReference CacheFile = FlatCPPIncludeDependencyCache.GetDependencyCachePathForTarget(Target.ProjectFile, Target.TargetName, Target.Platform, Target.Architecture);
-								if (!FlatCPPIncludeDependencyCache.TryRead(CacheFile, out Headers.FlatCPPIncludeDependencyCache))
-								{
-									if (!bNeedsFullCPPIncludeRescan)
-									{
-										if (!BuildConfiguration.bXGEExport && !BuildConfiguration.bGenerateManifest)
-										{
-											bNeedsFullCPPIncludeRescan = true;
-											Log.TraceInformation("Performing full C++ include scan (no include cache file)");
-										}
-									}
-									Headers.FlatCPPIncludeDependencyCache = new FlatCPPIncludeDependencyCache(CacheFile);
-								}
-							}
-						}
-
 						if (bIsGatheringBuild)
 						{
 							List<FileItem> TargetOutputItems = new List<FileItem>();
@@ -680,9 +602,16 @@ namespace UnrealBuildTool
 									}
 								}
 
+								// Create the dependencies cache
+								CppDependencyCache CppDependencies;
+								using(Timeline.ScopeEvent("Reading dependency cache"))
+								{
+									CppDependencies = CppDependencyCache.CreateHierarchy(Targets[0].ProjectFile, Targets[0].TargetName, Targets[0].Platform, Targets[0].Configuration, Targets[0].TargetType);
+								}
+
 								// Plan the actions to execute for the build.
 								Dictionary<UEBuildTarget, List<FileItem>> TargetToOutdatedPrerequisitesMap;
-								HashSet<Action> ActionsToExecute = ActionGraph.GetActionsToExecute(BuildConfiguration, PrerequisiteActions, Targets, TargetToHeaders, bIsAssemblingBuild, bNeedsFullCPPIncludeRescan, out TargetToOutdatedPrerequisitesMap);
+								HashSet<Action> ActionsToExecute = ActionGraph.GetActionsToExecute(BuildConfiguration, PrerequisiteActions, Targets, CppDependencies, out TargetToOutdatedPrerequisitesMap);
 
 								// Patch action history for hot reload when running in assembler mode.  In assembler mode, the suffix on the output file will be
 								// the same for every invocation on that makefile, but we need a new suffix each time.
@@ -743,7 +672,7 @@ namespace UnrealBuildTool
 									HotReload.PatchActionGraph(PrerequisiteActions, OldLocationToNewLocation);
 
 									// Get a new list of actions to execute now that the graph has been modified
-									ActionsToExecute = ActionGraph.GetActionsToExecute(BuildConfiguration, PrerequisiteActions, Targets, TargetToHeaders, bIsAssemblingBuild, bNeedsFullCPPIncludeRescan, out TargetToOutdatedPrerequisitesMap);
+									ActionsToExecute = ActionGraph.GetActionsToExecute(BuildConfiguration, PrerequisiteActions, Targets, CppDependencies, out TargetToOutdatedPrerequisitesMap);
 
 									// Build a mapping of all file items to their original
 									Dictionary<FileReference, FileReference> HotReloadFileToOriginalFile = new Dictionary<FileReference, FileReference>();
@@ -793,14 +722,6 @@ namespace UnrealBuildTool
 										ActionsToExecute.Count
 										);
 
-								// Cache indirect includes for all outdated C++ files.  We kick this off as a background thread so that it can
-								// perform the scan while we're compiling.  It usually only takes up to a few seconds, but we don't want to hurt
-								// our best case UBT iteration times for this task which can easily be performed asynchronously
-								if (BuildConfiguration.bUseUBTMakefiles && TargetToOutdatedPrerequisitesMap.Count > 0)
-								{
-									CppIncludeThread = new CppIncludeBackgroundThread(TargetToOutdatedPrerequisitesMap, TargetToHeaders);
-								}
-
 								// Execute the actions.
 								ExecutorTimer.Start();
 								if (BuildConfiguration.bXGEExport)
@@ -854,14 +775,7 @@ namespace UnrealBuildTool
 			}
 			finally
 			{
-				// Wait until our CPPIncludes dependency scanner thread has finished
-				if (CppIncludeThread != null)
-				{
-					using(Timeline.ScopeEvent("CppIncludeThread.Join()"))
-					{
-						CppIncludeThread.Join();
-					}
-				}
+				CppDependencyCache.SaveAll();
 
 				// Save the include dependency cache.
 				foreach (CPPHeaders Headers in TargetToHeaders.Values)
@@ -878,14 +792,6 @@ namespace UnrealBuildTool
 							Headers.IncludeDependencyCache.Save();
 						}
 					}
-
-					if (Headers.FlatCPPIncludeDependencyCache != null)
-					{
-						using (Timeline.ScopeEvent("FlatCPPIncludeDependencyCache.Save()"))
-						{
-							Headers.FlatCPPIncludeDependencyCache.Save();
-						}
-					}
 				}
 
 				// Figure out how long we took to execute.
@@ -896,9 +802,6 @@ namespace UnrealBuildTool
 
 				// Print some performance info
 				Log.TraceLog("DirectIncludes cache miss time: {0}s ({1} misses)", CPPHeaders.DirectIncludeCacheMissesTotalTime, CPPHeaders.TotalDirectIncludeCacheMisses);
-				Log.TraceLog("FindIncludePaths calls: {0} ({1} searches)", CPPHeaders.TotalFindIncludedFileCalls, CPPHeaders.IncludePathSearchAttempts);
-				Log.TraceLog("Deep C++ include scan time: {0}s", UnrealBuildTool.TotalDeepIncludeScanTime);
-				Log.TraceLog("Include Resolves: {0} ({1} misses, {2:0.00}%)", CPPHeaders.TotalDirectIncludeResolves, CPPHeaders.TotalDirectIncludeResolveCacheMisses, (float)CPPHeaders.TotalDirectIncludeResolveCacheMisses / (float)CPPHeaders.TotalDirectIncludeResolves * 100);
 			}
 			return (int)BuildResult;
 		}
@@ -921,56 +824,6 @@ namespace UnrealBuildTool
 				}
 			}
 		}
-
-		/// <summary>
-		/// Helper class to update the C++ dependency cache on a background thread. Captures exceptions and re-throws on the main thread when joined.
-		/// </summary>
-		class CppIncludeBackgroundThread
-		{
-			Thread BackgroundThread;
-			Exception CaughtException;
-
-			public CppIncludeBackgroundThread(Dictionary<UEBuildTarget, List<FileItem>> TargetToOutdatedPrerequisitesMap, Dictionary<UEBuildTarget, CPPHeaders> TargetToHeaders)
-			{
-				BackgroundThread = new Thread(() => Run(TargetToOutdatedPrerequisitesMap, TargetToHeaders));
-				BackgroundThread.Start();
-			}
-
-			public void Join()
-			{
-				BackgroundThread.Join();
-				if (CaughtException != null)
-				{
-					throw CaughtException;
-				}
-			}
-
-			private void Run(Dictionary<UEBuildTarget, List<FileItem>> TargetToOutdatedPrerequisitesMap, Dictionary<UEBuildTarget, CPPHeaders> TargetToHeaders)
-			{
-				try
-				{
-					// @todo ubtmake: This thread will access data structures that are also used on the main UBT thread, but during this time UBT
-					// is only invoking the build executor, so should not be touching this stuff.  However, we need to at some guards to make sure.
-					foreach (KeyValuePair<UEBuildTarget, List<FileItem>> TargetAndOutdatedPrerequisites in TargetToOutdatedPrerequisitesMap)
-					{
-						UEBuildTarget Target = TargetAndOutdatedPrerequisites.Key;
-						List<FileItem> OutdatedPrerequisites = TargetAndOutdatedPrerequisites.Value;
-						CPPHeaders Headers = TargetToHeaders[Target];
-
-						foreach (FileItem PrerequisiteItem in OutdatedPrerequisites)
-						{
-							// Invoke our deep include scanner to figure out whether any of the files included by this source file have
-							// changed since the build product was built
-							Headers.FindAndCacheAllIncludedFiles(PrerequisiteItem, PrerequisiteItem.CachedIncludePaths, bOnlyCachedDependencies: false);
-						}
-					}
-				}
-				catch (Exception Ex)
-				{
-					CaughtException = Ex;
-				}
-			}
-		}
 	}
 }
 
diff --git a/Engine/Source/Programs/UnrealBuildTool/Modes/CleanMode.cs b/Engine/Source/Programs/UnrealBuildTool/Modes/CleanMode.cs
--- a/Engine/Source/Programs/UnrealBuildTool/Modes/CleanMode.cs
+++ b/Engine/Source/Programs/UnrealBuildTool/Modes/CleanMode.cs
@@ -127,7 +127,6 @@ namespace UnrealBuildTool
 				}
 
 				// Add all the makefiles and caches to be deleted
-				FilesToDelete.Add(FlatCPPIncludeDependencyCache.GetDependencyCachePathForTarget(Target.ProjectFile, Target.Name, Target.Platform, Target.Architecture));
 				FilesToDelete.Add(DependencyCache.GetDependencyCachePathForTarget(Target.ProjectFile, Target.Platform, Target.Name));
 				FilesToDelete.Add(UBTMakefile.GetUBTMakefilePath(Target.ProjectFile, Target.Platform, Target.Configuration, Target.Name, false));
 				FilesToDelete.Add(UBTMakefile.GetUBTMakefilePath(Target.ProjectFile, Target.Platform, Target.Configuration, Target.Name, true));
diff --git a/Engine/Source/Programs/UnrealBuildTool/Platform/Android/AndroidToolChain.cs b/Engine/Source/Programs/UnrealBuildTool/Platform/Android/AndroidToolChain.cs
--- a/Engine/Source/Programs/UnrealBuildTool/Platform/Android/AndroidToolChain.cs
+++ b/Engine/Source/Programs/UnrealBuildTool/Platform/Android/AndroidToolChain.cs
@@ -1362,7 +1362,7 @@ namespace UnrealBuildTool
 						}
 
 						// Add the C++ source file and its included files to the prerequisite item list.
-						AddPrerequisiteSourceFile(CompileEnvironment, SourceFile, CompileAction.PrerequisiteItems);
+						CompileAction.PrerequisiteItems.Add(SourceFile);
 
 						if (CompileEnvironment.PrecompiledHeaderAction == PrecompiledHeaderAction.Create && !bIsPlainCFile)
 						{
@@ -1415,6 +1415,15 @@ namespace UnrealBuildTool
 						// Add the source file path to the command-line.
 						FileArguments += string.Format(" \"{0}\"", SourceFile.AbsolutePath);
 
+						// Generate the included header dependency list
+						if(CompileEnvironment.bGenerateDependenciesFile)
+						{
+							FileItem DependencyListFile = FileItem.GetItemByFileReference(FileReference.Combine(OutputDir, InlineArchName(Path.GetFileName(SourceFile.AbsolutePath) + ".d", Arch, GPUArchitecture, true)));
+							FileArguments += string.Format(" -MD -MF\"{0}\"", DependencyListFile.AbsolutePath.Replace('\\', '/'));
+							CompileAction.DependencyListFile = DependencyListFile;
+							CompileAction.ProducedItems.Add(DependencyListFile);
+						}
+
 						// Build a full argument list
 						string AllArguments = Arguments + FileArguments + CompileEnvironment.AdditionalArguments;
 
diff --git a/Engine/Source/Programs/UnrealBuildTool/Platform/HTML5/HTML5ToolChain.cs b/Engine/Source/Programs/UnrealBuildTool/Platform/HTML5/HTML5ToolChain.cs
--- a/Engine/Source/Programs/UnrealBuildTool/Platform/HTML5/HTML5ToolChain.cs
+++ b/Engine/Source/Programs/UnrealBuildTool/Platform/HTML5/HTML5ToolChain.cs
@@ -383,7 +383,7 @@ namespace UnrealBuildTool
 				bool bIsPlainCFile = Path.GetExtension(SourceFile.AbsolutePath).ToUpperInvariant() == ".C";
 
 				// Add the C++ source file and its included files to the prerequisite item list.
-				AddPrerequisiteSourceFile(CompileEnvironment, SourceFile, CompileAction.PrerequisiteItems);
+				CompileAction.PrerequisiteItems.Add(SourceFile);
 
 				// Add the source file path to the command-line.
 				string FileArguments = string.Format(" \"{0}\"", SourceFile.AbsolutePath);
@@ -402,6 +402,15 @@ namespace UnrealBuildTool
 					FileArguments += GetCLArguments_CPP(CompileEnvironment);
 				}
 
+				// Generate the included header dependency list
+				if(CompileEnvironment.bGenerateDependenciesFile)
+				{
+					FileItem DependencyListFile = FileItem.GetItemByFileReference(FileReference.Combine(OutputDir, Path.GetFileName(SourceFile.AbsolutePath) + ".d"));
+					FileArguments += string.Format(" -MD -MF\"{0}\"", DependencyListFile.AbsolutePath.Replace('\\', '/'));
+					CompileAction.DependencyListFile = DependencyListFile;
+					CompileAction.ProducedItems.Add(DependencyListFile);
+				}
+
 				CompileAction.WorkingDirectory = UnrealBuildTool.EngineSourceDirectory.FullName;
 				CompileAction.CommandPath = HTML5SDKInfo.Python();
 
diff --git a/Engine/Source/Programs/UnrealBuildTool/Platform/IOS/IOSToolChain.cs b/Engine/Source/Programs/UnrealBuildTool/Platform/IOS/IOSToolChain.cs
--- a/Engine/Source/Programs/UnrealBuildTool/Platform/IOS/IOSToolChain.cs
+++ b/Engine/Source/Programs/UnrealBuildTool/Platform/IOS/IOSToolChain.cs
@@ -618,7 +618,7 @@ namespace UnrealBuildTool
 				}
 
 				// Add the C++ source file and its included files to the prerequisite item list.
-				AddPrerequisiteSourceFile(CompileEnvironment, SourceFile, CompileAction.PrerequisiteItems);
+				CompileAction.PrerequisiteItems.Add(SourceFile);
 
 				// Add the precompiled header
 				if (CompileEnvironment.PrecompiledHeaderAction == PrecompiledHeaderAction.Include)
@@ -629,7 +629,7 @@ namespace UnrealBuildTool
 				// Upload the force included files
 				foreach(FileItem ForceIncludeFile in CompileEnvironment.ForceIncludeFiles)
 				{
-					AddPrerequisiteSourceFile(CompileEnvironment, ForceIncludeFile, CompileAction.PrerequisiteItems);
+					CompileAction.PrerequisiteItems.Add(ForceIncludeFile);
 				}
 
 				string OutputFilePath = null;
@@ -663,6 +663,15 @@ namespace UnrealBuildTool
 				// Add the source file path to the command-line.
 				FileArguments += string.Format(" \"{0}\"", SourceFile.AbsolutePath);
 
+				// Generate the included header dependency list
+				if(CompileEnvironment.bGenerateDependenciesFile)
+				{
+					FileItem DependencyListFile = FileItem.GetItemByFileReference(FileReference.Combine(OutputDir, Path.GetFileName(SourceFile.AbsolutePath) + ".d"));
+					FileArguments += string.Format(" -MD -MF\"{0}\"", DependencyListFile.AbsolutePath.Replace('\\', '/'));
+					CompileAction.DependencyListFile = DependencyListFile;
+					CompileAction.ProducedItems.Add(DependencyListFile);
+				}
+
 				string CompilerPath = Settings.Value.ToolchainDir + IOSCompiler;
 
 				string AllArgs = FilePCHArguments + Arguments + FileArguments + CompileEnvironment.AdditionalArguments;
diff --git a/Engine/Source/Programs/UnrealBuildTool/Platform/Linux/LinuxToolChain.cs b/Engine/Source/Programs/UnrealBuildTool/Platform/Linux/LinuxToolChain.cs
--- a/Engine/Source/Programs/UnrealBuildTool/Platform/Linux/LinuxToolChain.cs
+++ b/Engine/Source/Programs/UnrealBuildTool/Platform/Linux/LinuxToolChain.cs
@@ -1197,7 +1197,7 @@ namespace UnrealBuildTool
 				}
 
 				// Add the C++ source file and its included files to the prerequisite item list.
-				AddPrerequisiteSourceFile(CompileEnvironment, SourceFile, CompileAction.PrerequisiteItems);
+				CompileAction.PrerequisiteItems.Add(SourceFile);
 
 				if (CompileEnvironment.PrecompiledHeaderAction == PrecompiledHeaderAction.Create)
 				{
@@ -1229,6 +1229,15 @@ namespace UnrealBuildTool
 				// Add the source file path to the command-line.
 				FileArguments += string.Format(" \"{0}\"", SourceFile.AbsolutePath.Replace('\\', '/'));
 
+				// Generate the included header dependency list
+				if(CompileEnvironment.bGenerateDependenciesFile)
+				{
+					FileItem DependencyListFile = FileItem.GetItemByFileReference(FileReference.Combine(OutputDir, Path.GetFileName(SourceFile.AbsolutePath) + ".d"));
+					FileArguments += string.Format(" -MD -MF\"{0}\"", DependencyListFile.AbsolutePath.Replace('\\', '/'));
+					CompileAction.DependencyListFile = DependencyListFile;
+					CompileAction.ProducedItems.Add(DependencyListFile);
+				}
+
 				CompileAction.WorkingDirectory = UnrealBuildTool.EngineSourceDirectory.FullName;
 				if (!UsingClang())
 				{
diff --git a/Engine/Source/Programs/UnrealBuildTool/Platform/Mac/MacToolChain.cs b/Engine/Source/Programs/UnrealBuildTool/Platform/Mac/MacToolChain.cs
--- a/Engine/Source/Programs/UnrealBuildTool/Platform/Mac/MacToolChain.cs
+++ b/Engine/Source/Programs/UnrealBuildTool/Platform/Mac/MacToolChain.cs
@@ -466,7 +466,7 @@ namespace UnrealBuildTool
 				}
 
 				// Add the C++ source file and its included files to the prerequisite item list.
-				AddPrerequisiteSourceFile(CompileEnvironment, SourceFile, CompileAction.PrerequisiteItems);
+				CompileAction.PrerequisiteItems.Add(SourceFile);
 
 				string OutputFilePath = null;
 				if (CompileEnvironment.PrecompiledHeaderAction == PrecompiledHeaderAction.Create)
@@ -498,6 +498,15 @@ namespace UnrealBuildTool
 				// Add the source file path to the command-line.
 				FileArguments += string.Format(" \"{0}\"", SourceFile.AbsolutePath);
 
+				// Generate the included header dependency list
+				if(CompileEnvironment.bGenerateDependenciesFile)
+				{
+					FileItem DependencyListFile = FileItem.GetItemByFileReference(FileReference.Combine(OutputDir, Path.GetFileName(SourceFile.AbsolutePath) + ".d"));
+					FileArguments += string.Format(" -MD -MF\"{0}\"", DependencyListFile.AbsolutePath.Replace('\\', '/'));
+					CompileAction.DependencyListFile = DependencyListFile;
+					CompileAction.ProducedItems.Add(DependencyListFile);
+				}
+
 				string AllArgs = Arguments + FileArguments + CompileEnvironment.AdditionalArguments;
 				string CompilerPath = Settings.ToolchainDir + MacCompiler;
 				
diff --git a/Engine/Source/Programs/UnrealBuildTool/Platform/Windows/VCToolChain.cs b/Engine/Source/Programs/UnrealBuildTool/Platform/Windows/VCToolChain.cs
--- a/Engine/Source/Programs/UnrealBuildTool/Platform/Windows/VCToolChain.cs
+++ b/Engine/Source/Programs/UnrealBuildTool/Platform/Windows/VCToolChain.cs
@@ -899,7 +899,7 @@ namespace UnrealBuildTool
 				bool bIsPlainCFile = Path.GetExtension(SourceFile.AbsolutePath).ToUpperInvariant() == ".C";
 
 				// Add the C++ source file and its included files to the prerequisite item list.
-				AddPrerequisiteSourceFile(CompileEnvironment, SourceFile, CompileAction.PrerequisiteItems);
+				CompileAction.PrerequisiteItems.Add(SourceFile);
 
 				bool bEmitsObjectFile = true;
 				if (CompileEnvironment.PrecompiledHeaderAction == PrecompiledHeaderAction.Create)
@@ -1072,6 +1072,14 @@ namespace UnrealBuildTool
 					CompileAction.CommandArguments = String.Join(" ", SharedArguments.Concat(FileArguments).Concat(AdditionalArguments));
 				}
 
+				if(CompileEnvironment.bGenerateDependenciesFile)
+				{
+					CompileAction.DependencyListFile = FileItem.GetItemByFileReference(FileReference.Combine(OutputDir, SourceFile.Location.GetFileName() + ".txt"));
+					CompileAction.CommandArguments = String.Format("{0} -- {1} {2} /showIncludes", Utils.MakePathSafeToUseWithCommandLine(CompileAction.DependencyListFile.Location), Utils.MakePathSafeToUseWithCommandLine(CompileAction.CommandPath), CompileAction.CommandArguments);
+					CompileAction.CommandPath = FileReference.Combine(UnrealBuildTool.EngineDirectory, "Build", "Windows", "cl-filter", "cl-filter.exe").FullName;
+					CompileAction.ProducedItems.Add(CompileAction.DependencyListFile);
+				}
+
 				if (CompileEnvironment.PrecompiledHeaderAction == PrecompiledHeaderAction.Create)
 				{
 					Log.TraceVerbose("Creating PCH " + CompileEnvironment.PrecompiledHeaderIncludeFilename + ": \"" + CompileAction.CommandPath + "\"" + CompileAction.CommandArguments);
@@ -1174,7 +1182,7 @@ namespace UnrealBuildTool
 				CompileAction.CommandArguments = String.Join(" ", Arguments);
 
 				// Add the C++ source file and its included files to the prerequisite item list.
-				AddPrerequisiteSourceFile(CompileEnvironment, RCFile, CompileAction.PrerequisiteItems);
+				CompileAction.PrerequisiteItems.Add(RCFile);
 			}
 
 			return Result;
diff --git a/Engine/Source/Programs/UnrealBuildTool/ProjectFiles/ProjectFileGenerator.cs b/Engine/Source/Programs/UnrealBuildTool/ProjectFiles/ProjectFileGenerator.cs
--- a/Engine/Source/Programs/UnrealBuildTool/ProjectFiles/ProjectFileGenerator.cs
+++ b/Engine/Source/Programs/UnrealBuildTool/ProjectFiles/ProjectFileGenerator.cs
@@ -1700,10 +1700,6 @@ namespace UnrealBuildTool
 						NewArguments.Add("-Monolithic"); // Don't want any _API macros to be DLLIMPORT; this will cause Intellisense errors in the file that they are defined.
 						NewArguments.AddRange(Arguments);
 
-						// Clear the CachedIncludePaths field on every FileItem. It's faster to keep the cache of file items around, but we need to wipe
-						// out any state.
-						FileItem.ClearCachedIncludePaths();
-
 						try
 						{
 							const string DefaultArchitecture = "";
diff --git a/Engine/Source/Programs/UnrealBuildTool/ProjectFiles/VisualStudio/VCProject.cs b/Engine/Source/Programs/UnrealBuildTool/ProjectFiles/VisualStudio/VCProject.cs
--- a/Engine/Source/Programs/UnrealBuildTool/ProjectFiles/VisualStudio/VCProject.cs
+++ b/Engine/Source/Programs/UnrealBuildTool/ProjectFiles/VisualStudio/VCProject.cs
@@ -931,20 +931,6 @@ namespace UnrealBuildTool
 				// Add all file directories to the filters file as solution filters
 				HashSet<string> FilterDirectories = new HashSet<string>();
 				UEBuildPlatform BuildPlatform = UEBuildPlatform.GetBuildPlatform(BuildHostPlatform.Current.Platform);
-				bool bWritePerFilePCHInfo = false;
-				if (bUsePerFileIntellisense && ProjectFileFormat >= VCProjectFileFormat.VisualStudio2015)
-				{
-					string UpdateRegistryLoc = string.Format(@"HKEY_LOCAL_MACHINE\SOFTWARE\Wow6432Node\Microsoft\DevDiv\vs\Servicing\{0}\devenv", VCProjectFileGenerator.GetProjectFileToolVersionString(ProjectFileFormat));
-					object Result = Microsoft.Win32.Registry.GetValue(UpdateRegistryLoc, "UpdateVersion", null);
-					if (Result != null)
-					{
-						int UpdateVersion = 0;
-						if (Int32.TryParse(Result.ToString().Split('.').Last(), out UpdateVersion) && UpdateVersion >= 25420)
-						{
-							bWritePerFilePCHInfo = true;
-						}
-					}
-				}
 
 				foreach (AliasedFile AliasedFile in LocalAliasedFiles)
 				{
@@ -955,28 +941,7 @@ namespace UnrealBuildTool
 					}
 
 					string VCFileType = GetVCFileType(AliasedFile.FileSystemPath);
-					string PCHFileName = null;
-
-					if (bWritePerFilePCHInfo && VCFileType == "ClCompile")
-					{
-						FileReference TruePath = FileReference.Combine(ProjectFilePath.Directory, AliasedFile.FileSystemPath);
-						List <DependencyInclude> DirectlyIncludedFilenames = CPPHeaders.GetUncachedDirectIncludeDependencies(TruePath, null);
-						if (DirectlyIncludedFilenames.Count > 0)
-						{
-							PCHFileName = DirectlyIncludedFilenames[0].IncludeName;
-						}
-					}
-
-					if (!string.IsNullOrEmpty(PCHFileName))
-					{
-						VCProjectFileContent.AppendLine("    <{0} Include=\"{1}\">" + VCFileType, EscapeFileName(AliasedFile.FileSystemPath));
-						VCProjectFileContent.AppendLine("      <AdditionalOptions>$(AdditionalOptions) /Yu{0}</AdditionalOptions>", PCHFileName);
-						VCProjectFileContent.AppendLine("    </{0}>", VCFileType);
-					}
-					else
-					{
-						VCProjectFileContent.AppendLine("    <{0} Include=\"{1}\"/>", VCFileType, EscapeFileName(AliasedFile.FileSystemPath));
-					}
+					VCProjectFileContent.AppendLine("    <{0} Include=\"{1}\"/>", VCFileType, EscapeFileName(AliasedFile.FileSystemPath));
 
 					if (!String.IsNullOrWhiteSpace(AliasedFile.ProjectPath))
 					{
diff --git a/Engine/Source/Programs/UnrealBuildTool/System/Action.cs b/Engine/Source/Programs/UnrealBuildTool/System/Action.cs
--- a/Engine/Source/Programs/UnrealBuildTool/System/Action.cs
+++ b/Engine/Source/Programs/UnrealBuildTool/System/Action.cs
@@ -59,6 +59,11 @@ namespace UnrealBuildTool
 		/// </summary>
 		public List<FileItem> DeleteItems = new List<FileItem>();
 
+		/// <summary>
+		/// For C++ source files, specifies a dependency list file used to check changes to header files
+		/// </summary>
+		public FileItem DependencyListFile;
+
 		/// <summary>
 		/// Directory from which to execute the program to create produced items
 		/// </summary>
@@ -175,6 +180,7 @@ namespace UnrealBuildTool
 			PrerequisiteItems = Reader.ReadFileItemList(UniqueFileItems);
 			ProducedItems = Reader.ReadFileItemList(UniqueFileItems);
 			DeleteItems = Reader.ReadFileItemList(UniqueFileItems);
+			DependencyListFile = Reader.ReadFileItem(UniqueFileItems);
 		}
 
 		/// <summary>
@@ -198,6 +204,7 @@ namespace UnrealBuildTool
 			Writer.Write(PrerequisiteItems, UniqueFileItemToIndex);
 			Writer.Write(ProducedItems, UniqueFileItemToIndex);
 			Writer.Write(DeleteItems, UniqueFileItemToIndex);
+			Writer.Write(DependencyListFile, UniqueFileItemToIndex);
 		}
 
 		/// <summary>
diff --git a/Engine/Source/Programs/UnrealBuildTool/System/ActionGraph.cs b/Engine/Source/Programs/UnrealBuildTool/System/ActionGraph.cs
--- a/Engine/Source/Programs/UnrealBuildTool/System/ActionGraph.cs
+++ b/Engine/Source/Programs/UnrealBuildTool/System/ActionGraph.cs
@@ -7,6 +7,8 @@ using System.Diagnostics;
 using System.Linq;
 using System.Runtime.Serialization;
 using Tools.DotNETCommon;
+using System.Text;
+using System.Threading.Tasks;
 
 namespace UnrealBuildTool
 {
@@ -52,7 +54,7 @@ namespace UnrealBuildTool
 		/// <summary>
 		/// Builds a list of actions that need to be executed to produce the specified output items.
 		/// </summary>
-		public HashSet<Action> GetActionsToExecute(BuildConfiguration BuildConfiguration, Action[] PrerequisiteActions, List<UEBuildTarget> Targets, Dictionary<UEBuildTarget, CPPHeaders> TargetToHeaders, bool bIsAssemblingBuild, bool bNeedsFullCPPIncludeRescan, out Dictionary<UEBuildTarget, List<FileItem>> TargetToOutdatedPrerequisitesMap)
+		public HashSet<Action> GetActionsToExecute(BuildConfiguration BuildConfiguration, Action[] PrerequisiteActions, List<UEBuildTarget> Targets, CppDependencyCache CppDependencies, out Dictionary<UEBuildTarget, List<FileItem>> TargetToOutdatedPrerequisitesMap)
 		{
 			ITimelineEvent GetActionsToExecuteTimer = Timeline.ScopeEvent("ActionGraph.GetActionsToExecute()");
 
@@ -76,7 +78,7 @@ namespace UnrealBuildTool
 					ActionHistory History = new ActionHistory(HistoryFilename.FullName);
 					HistoryList.Add(History);
 					OpenHistoryFiles.Add(HistoryFilename);
-					GatherAllOutdatedActions(BuildConfiguration, BuildTarget, TargetToHeaders[BuildTarget], bIsAssemblingBuild, bNeedsFullCPPIncludeRescan, History, ref OutdatedActionDictionary, TargetToOutdatedPrerequisitesMap);
+					GatherAllOutdatedActions(History, OutdatedActionDictionary, CppDependencies, BuildConfiguration.bIgnoreOutdatedImportLibraries);
 				}
 			}
 
@@ -431,17 +433,13 @@ namespace UnrealBuildTool
 		/// Determines whether an action is outdated based on the modification times for its prerequisite
 		/// and produced items.
 		/// </summary>
-		/// <param name="BuildConfiguration">Build configuration options</param>
-		/// <param name="Target"></param>
-		/// <param name="Headers"></param>
 		/// <param name="RootAction">- The action being considered.</param>
-		/// <param name="bIsAssemblingBuild"></param>
-		/// <param name="bNeedsFullCPPIncludeRescan"></param>
 		/// <param name="OutdatedActionDictionary">-</param>
 		/// <param name="ActionHistory"></param>
-		/// <param name="TargetToOutdatedPrerequisitesMap"></param>
+		/// <param name="CppDependencies"></param>
+		/// <param name="bIgnoreOutdatedImportLibraries"></param>
 		/// <returns>true if outdated</returns>
-		public bool IsActionOutdated(BuildConfiguration BuildConfiguration, UEBuildTarget Target, CPPHeaders Headers, Action RootAction, bool bIsAssemblingBuild, bool bNeedsFullCPPIncludeRescan, Dictionary<Action, bool> OutdatedActionDictionary, ActionHistory ActionHistory, Dictionary<UEBuildTarget, List<FileItem>> TargetToOutdatedPrerequisitesMap)
+		public bool IsActionOutdated(Action RootAction, Dictionary<Action, bool> OutdatedActionDictionary, ActionHistory ActionHistory, CppDependencyCache CppDependencies, bool bIgnoreOutdatedImportLibraries)
 		{
 			// Only compute the outdated-ness for actions that don't aren't cached in the outdated action dictionary.
 			bool bIsOutdated = false;
@@ -452,37 +450,33 @@ namespace UnrealBuildTool
 				DateTimeOffset LastExecutionTimeUtc = DateTimeOffset.MaxValue;
 				foreach (FileItem ProducedItem in RootAction.ProducedItems)
 				{
-					// Optionally skip the action history check, as this only works for local builds
-					if (BuildConfiguration.bUseActionHistory)
+					// Check if the command-line of the action previously used to produce the item is outdated.
+					string OldProducingCommandLine = "";
+					string NewProducingCommandLine = RootAction.CommandPath + " " + RootAction.CommandArguments;
+					if (!ActionHistory.GetProducingCommandLine(ProducedItem, out OldProducingCommandLine)
+					|| !String.Equals(OldProducingCommandLine, NewProducingCommandLine, StringComparison.InvariantCultureIgnoreCase))
 					{
-						// Check if the command-line of the action previously used to produce the item is outdated.
-						string OldProducingCommandLine = "";
-						string NewProducingCommandLine = RootAction.CommandPath + " " + RootAction.CommandArguments;
-						if (!ActionHistory.GetProducingCommandLine(ProducedItem, out OldProducingCommandLine)
-						|| !String.Equals(OldProducingCommandLine, NewProducingCommandLine, StringComparison.InvariantCultureIgnoreCase))
+						if(ProducedItem.Exists)
 						{
-							if(ProducedItem.Exists)
-							{
-								Log.TraceLog(
-									"{0}: Produced item \"{1}\" was produced by outdated command-line.\n  Old command-line: {2}\n  New command-line: {3}",
-									RootAction.StatusDescription,
-									Path.GetFileName(ProducedItem.AbsolutePath),
-									OldProducingCommandLine,
-									NewProducingCommandLine
-									);
-							}
+							Log.TraceLog(
+								"{0}: Produced item \"{1}\" was produced by outdated command-line.\n  Old command-line: {2}\n  New command-line: {3}",
+								RootAction.StatusDescription,
+								Path.GetFileName(ProducedItem.AbsolutePath),
+								OldProducingCommandLine,
+								NewProducingCommandLine
+								);
+						}
 
-							bIsOutdated = true;
+						bIsOutdated = true;
 
-							// Update the command-line used to produce this item in the action history.
-							ActionHistory.SetProducingCommandLine(ProducedItem, NewProducingCommandLine);
-						}
+						// Update the command-line used to produce this item in the action history.
+						ActionHistory.SetProducingCommandLine(ProducedItem, NewProducingCommandLine);
 					}
 
 					// If the produced file doesn't exist or has zero size, consider it outdated.  The zero size check is to detect cases
 					// where aborting an earlier compile produced invalid zero-sized obj files, but that may cause actions where that's
 					// legitimate output to always be considered outdated.
-					if (ProducedItem.Exists && (RootAction.ActionType != ActionType.Compile || ProducedItem.Length > 0))
+					if (ProducedItem.Exists && (RootAction.ActionType != ActionType.Compile || ProducedItem.Length > 0 || (!ProducedItem.Location.HasExtension(".obj") && !ProducedItem.Location.HasExtension(".o"))))
 					{
 						// Use the oldest produced item's time as the last execution time.
 						if (ProducedItem.LastWriteTimeUtc < LastExecutionTimeUtc)
@@ -503,80 +497,6 @@ namespace UnrealBuildTool
 					}
 				}
 
-				Log.WriteLineIf(BuildConfiguration.bLogDetailedActionStats && !String.IsNullOrEmpty(LatestUpdatedProducedItemName),
-					LogEventType.Verbose, "{0}: Oldest produced item is {1}", RootAction.StatusDescription, LatestUpdatedProducedItemName);
-
-				bool bCheckIfIncludedFilesAreNewer = false;
-				bool bPerformExhaustiveIncludeSearchAndUpdateCache = false;
-				if (RootAction.ActionType == ActionType.Compile)
-				{
-					// Outdated targets don't need their headers scanned yet, because presumably they would already be out of dated based on already-cached
-					// includes before getting this far.  However, if we find them to be outdated after processing includes, we'll do a deep scan later
-					// on and cache all of the includes so that we have them for a quick outdatedness check the next run.
-					if (!bIsOutdated &&
-						BuildConfiguration.bUseUBTMakefiles &&
-						bIsAssemblingBuild)
-					{
-						bCheckIfIncludedFilesAreNewer = true;
-					}
-
-					// Were we asked to force an update of our cached includes BEFORE we try to build?  This may be needed if our cache can no longer
-					// be trusted and we need to fill it with perfectly valid data (even if we're in assembler only mode)
-					if (BuildConfiguration.bUseUBTMakefiles &&
-						bNeedsFullCPPIncludeRescan)
-					{
-						// This will be slow!
-						bPerformExhaustiveIncludeSearchAndUpdateCache = true;
-					}
-				}
-
-
-				if (bCheckIfIncludedFilesAreNewer || bPerformExhaustiveIncludeSearchAndUpdateCache)
-				{
-					// Scan this file for included headers that may be out of date.  Note that it's OK if we break out early because we found
-					// the action to be outdated.  For outdated actions, we kick off a separate include scan in a background thread later on to
-					// catch all of the other includes and form an exhaustive set.
-					foreach (FileItem PrerequisiteItem in RootAction.PrerequisiteItems)
-					{
-						// @todo ubtmake: Make sure we are catching RC files here too.  Anything that the toolchain would have tried it on.  Logic should match the CACHING stuff below
-						if (PrerequisiteItem.CachedIncludePaths != null)
-						{
-							List<FileItem> IncludedFileList = Headers.FindAndCacheAllIncludedFiles(PrerequisiteItem, PrerequisiteItem.CachedIncludePaths, bOnlyCachedDependencies: !bPerformExhaustiveIncludeSearchAndUpdateCache);
-							if (IncludedFileList != null)
-							{
-								foreach (FileItem IncludedFile in IncludedFileList)
-								{
-									if (IncludedFile.Exists)
-									{
-										// allow a 1 second slop for network copies
-										TimeSpan TimeDifference = IncludedFile.LastWriteTimeUtc - LastExecutionTimeUtc;
-										bool bPrerequisiteItemIsNewerThanLastExecution = TimeDifference.TotalSeconds > 1;
-										if (bPrerequisiteItemIsNewerThanLastExecution)
-										{
-											Log.TraceLog(
-												"{0}: Included file {1} is newer than the last execution of the action: {2} vs {3}",
-												RootAction.StatusDescription,
-												Path.GetFileName(IncludedFile.AbsolutePath),
-												IncludedFile.LastWriteTimeUtc.LocalDateTime,
-												LastExecutionTimeUtc.LocalDateTime
-												);
-											bIsOutdated = true;
-
-											// Don't bother checking every single include if we've found one that is out of date
-											break;
-										}
-									}
-								}
-							}
-						}
-
-						if (bIsOutdated)
-						{
-							break;
-						}
-					}
-				}
-
 				if (!bIsOutdated)
 				{
 					// Check if any of the prerequisite items are produced by outdated actions, or have changed more recently than
@@ -592,12 +512,12 @@ namespace UnrealBuildTool
 						{
 							bIsImportLibraryFile = PrerequisiteItem.AbsolutePath.EndsWith(".LIB", StringComparison.InvariantCultureIgnoreCase);
 						}
-						if (!bIsImportLibraryFile || !BuildConfiguration.bIgnoreOutdatedImportLibraries)
+						if (!bIsImportLibraryFile || !bIgnoreOutdatedImportLibraries)
 						{
 							// If the prerequisite is produced by an outdated action, then this action is outdated too.
 							if (PrerequisiteItem.ProducingAction != null)
 							{
-								if (IsActionOutdated(BuildConfiguration, Target, Headers, PrerequisiteItem.ProducingAction, bIsAssemblingBuild, bNeedsFullCPPIncludeRescan, OutdatedActionDictionary, ActionHistory, TargetToOutdatedPrerequisitesMap))
+								if (IsActionOutdated(PrerequisiteItem.ProducingAction, OutdatedActionDictionary, ActionHistory, CppDependencies, bIgnoreOutdatedImportLibraries))
 								{
 									Log.TraceLog(
 										"{0}: Prerequisite {1} is produced by outdated action.",
@@ -636,45 +556,33 @@ namespace UnrealBuildTool
 					}
 				}
 
-				// For compile actions, we have C++ files that are actually dependent on header files that could have been changed.  We only need to
-				// know about the set of header files that are included for files that are already determined to be out of date (such as if the file
-				// is missing or was modified.)  In the case that the file is out of date, we'll perform a deep scan to update our cached set of
-				// includes for this file, so that we'll be able to determine whether it is out of date next time very quickly.
-				if (BuildConfiguration.bUseUBTMakefiles)
+				// Check the dependency list
+				if(!bIsOutdated && RootAction.DependencyListFile != null)
 				{
-					Stopwatch DeepIncludeScanTimer = Stopwatch.StartNew();
-
-					// @todo ubtmake: we may be scanning more files than we need to here -- indirectly outdated files are bIsOutdated=true by this point (for example basemost includes when deeper includes are dirty)
-					if (bIsOutdated && RootAction.ActionType == ActionType.Compile)	// @todo ubtmake: Does this work with RC files?  See above too.
+					List<FileItem> DependencyFiles;
+					if(!CppDependencies.TryGetDependencies(RootAction.DependencyListFile, out DependencyFiles))
 					{
-						Log.TraceVerbose("Outdated action: {0}", RootAction.StatusDescription);
-						foreach (FileItem PrerequisiteItem in RootAction.PrerequisiteItems)
+						Log.TraceLog("{0}: Missing dependency list file \"{1}\"", RootAction.StatusDescription, RootAction.DependencyListFile);
+						bIsOutdated = true;
+					}
+					else
+					{
+						foreach(FileItem DependencyFile in DependencyFiles)
 						{
-							if (PrerequisiteItem.CachedIncludePaths != null)
-							{
-								if (!IsCPPFile(PrerequisiteItem))
-								{
-									throw new BuildException("Was only expecting C++ files to have CachedCPPEnvironments!");
-								}
-								Log.TraceVerbose("  -> DEEP include scan: {0}", PrerequisiteItem.AbsolutePath);
-
-								List<FileItem> OutdatedPrerequisites;
-								if (!TargetToOutdatedPrerequisitesMap.TryGetValue(Target, out OutdatedPrerequisites))
-								{
-									OutdatedPrerequisites = new List<FileItem>();
-									TargetToOutdatedPrerequisitesMap.Add(Target, OutdatedPrerequisites);
-								}
-
-								OutdatedPrerequisites.Add(PrerequisiteItem);
-							}
-							else if (IsCPPImplementationFile(PrerequisiteItem) || IsCPPResourceFile(PrerequisiteItem))
+							if(!DependencyFile.Exists || DependencyFile.LastWriteTimeUtc > LastExecutionTimeUtc)
 							{
-								Log.TraceVerbose("  -> WARNING: No CachedCPPEnvironment: {0}", PrerequisiteItem.AbsolutePath);
+								Log.TraceLog(
+									"{0}: Dependency {1} is newer than the last execution of the action: {2} vs {3}",
+									RootAction.StatusDescription,
+									Path.GetFileName(DependencyFile.AbsolutePath),
+									DependencyFile.LastWriteTimeUtc.LocalDateTime,
+									LastExecutionTimeUtc.LocalDateTime
+									);
+								bIsOutdated = true;
+								break;
 							}
 						}
 					}
-
-					UnrealBuildTool.TotalDeepIncludeScanTime += DeepIncludeScanTimer.Elapsed;
 				}
 
 				// Cache the outdated-ness of this action.
@@ -684,16 +592,28 @@ namespace UnrealBuildTool
 			return bIsOutdated;
 		}
 
-
 		/// <summary>
 		/// Builds a dictionary containing the actions from AllActions that are outdated by calling
 		/// IsActionOutdated.
 		/// </summary>
-		void GatherAllOutdatedActions(BuildConfiguration BuildConfiguration, UEBuildTarget Target, CPPHeaders Headers, bool bIsAssemblingBuild, bool bNeedsFullCPPIncludeRescan, ActionHistory ActionHistory, ref Dictionary<Action, bool> OutdatedActions, Dictionary<UEBuildTarget, List<FileItem>> TargetToOutdatedPrerequisitesMap)
+		void GatherAllOutdatedActions(ActionHistory ActionHistory, Dictionary<Action, bool> OutdatedActions, CppDependencyCache CppDependencies, bool bIgnoreOutdatedImportLibraries)
 		{
+			using(Timeline.ScopeEvent("Prefetching include dependencies"))
+			{
+				List<FileItem> Dependencies = new List<FileItem>();
+				foreach(Action Action in AllActions)
+				{
+					if(Action.DependencyListFile != null)
+					{
+						Dependencies.Add(Action.DependencyListFile);
+					}
+				}
+				Parallel.ForEach(Dependencies, File => { List<FileItem> Temp; CppDependencies.TryGetDependencies(File, out Temp); });
+			}
+
 			foreach (Action Action in AllActions)
 			{
-				IsActionOutdated(BuildConfiguration, Target, Headers, Action, bIsAssemblingBuild, bNeedsFullCPPIncludeRescan, OutdatedActions, ActionHistory, TargetToOutdatedPrerequisitesMap);
+				IsActionOutdated(Action, OutdatedActions, ActionHistory, CppDependencies, bIgnoreOutdatedImportLibraries);
 			}
 		}
 
diff --git a/Engine/Source/Programs/UnrealBuildTool/System/CPPHeaders.cs b/Engine/Source/Programs/UnrealBuildTool/System/CPPHeaders.cs
--- a/Engine/Source/Programs/UnrealBuildTool/System/CPPHeaders.cs
+++ b/Engine/Source/Programs/UnrealBuildTool/System/CPPHeaders.cs
@@ -15,8 +15,7 @@ namespace UnrealBuildTool
 	/// <summary>
 	/// For C++ source file items, this structure is used to cache data that will be used for include dependency scanning
 	/// </summary>
-	[Serializable]
-	class CppIncludePaths : ISerializable
+	class CppIncludePaths
 	{
 		/// <summary>
 		/// Ordered list of include paths for the module
@@ -34,11 +33,6 @@ namespace UnrealBuildTool
 		/// </summary>
 		public bool bCheckSystemHeadersForModification;
 
-		/// <summary>
-		/// Contains a mapping from filename to the full path of the header in this environment.  This is used to optimized include path lookups at runtime for any given single module.
-		/// </summary>
-		public Dictionary<string, FileItem> IncludeFileSearchDictionary = new Dictionary<string, FileItem>();
-
 		/// <summary>
 		/// Construct an empty set of include paths
 		/// </summary>
@@ -58,47 +52,6 @@ namespace UnrealBuildTool
 			SystemIncludePaths = new HashSet<DirectoryReference>(Other.SystemIncludePaths);
 			bCheckSystemHeadersForModification = Other.bCheckSystemHeadersForModification;
 		}
-
-		/// <summary>
-		/// Deserialize the include paths from the given context
-		/// </summary>
-		/// <param name="Info">Serialization info</param>
-		/// <param name="Context">Serialization context</param>
-		public CppIncludePaths(SerializationInfo Info, StreamingContext Context)
-		{
-			UserIncludePaths = new HashSet<DirectoryReference>((DirectoryReference[])Info.GetValue("ip", typeof(DirectoryReference[])));
-			SystemIncludePaths = new HashSet<DirectoryReference>((DirectoryReference[])Info.GetValue("sp", typeof(DirectoryReference[])));
-			bCheckSystemHeadersForModification = Info.GetBoolean("cs");
-		}
-
-		/// <summary>
-		/// Serialize this instance
-		/// </summary>
-		/// <param name="Info">Serialization info</param>
-		/// <param name="Context">Serialization context</param>
-		public void GetObjectData(SerializationInfo Info, StreamingContext Context)
-		{
-			Info.AddValue("ip", UserIncludePaths.ToArray());
-			Info.AddValue("sp", SystemIncludePaths.ToArray());
-			Info.AddValue("cs", bCheckSystemHeadersForModification);
-		}
-
-		/// <summary>
-		/// Given a C++ source file, returns a list of include paths we should search to resolve #includes for this path
-		/// </summary>
-		/// <param name="SourceFile">C++ source file we're going to check #includes for.</param>
-		/// <returns>Ordered list of paths to search</returns>
-		public List<DirectoryReference> GetPathsToSearch(FileReference SourceFile)
-		{
-			List<DirectoryReference> IncludePathsToSearch = new List<DirectoryReference>();
-			IncludePathsToSearch.Add(SourceFile.Directory);
-			IncludePathsToSearch.AddRange(UserIncludePaths);
-			if (bCheckSystemHeadersForModification)
-			{
-				IncludePathsToSearch.AddRange(SystemIncludePaths);
-			}
-			return IncludePathsToSearch;
-		}
 	}
 
 	class CPPHeaders
@@ -118,317 +71,15 @@ namespace UnrealBuildTool
 		/// </summary>
 		public DependencyCache IncludeDependencyCache = null;
 
-		/// <summary>
-		/// Contains a cache of include dependencies (direct and indirect), one for each target we're building.
-		/// </summary>
-		public FlatCPPIncludeDependencyCache FlatCPPIncludeDependencyCache = null;
-
-		/// <summary>
-		/// 
-		/// </summary>
-		public static int TotalFindIncludedFileCalls = 0;
 
-		/// <summary>
-		/// 
-		/// </summary>
-		public static int IncludePathSearchAttempts = 0;
-
-		/// <summary>
-		/// Whether to print performance info on include file resolution
-		/// </summary>
-		public static bool bPrintPerformanceInfo = false;
-
-		/// <summary>
-		/// A cache of the list of other files that are directly or indirectly included by a C++ file.
-		/// </summary>
-		Dictionary<FileItem, List<FileItem>> ExhaustiveIncludedFilesMap = new Dictionary<FileItem, List<FileItem>>();
-
-		/// <summary>
-		/// A cache of all files included by a C++ file, but only has files that we knew about from a previous session, loaded from a cache at startup
-		/// </summary>
-		Dictionary<FileItem, List<FileItem>> OnlyCachedIncludedFilesMap = new Dictionary<FileItem, List<FileItem>>();
-
-		/// <summary>
-		/// 
-		/// </summary>
-		public bool bUseUBTMakefiles;
-
-		/// <summary>
-		/// 
-		/// </summary>
-		public bool bUseFlatCPPIncludeDependencyCache;
-		
-		/// <summary>
-		/// 
-		/// </summary>
-		public bool bUseIncludeDependencyResolveCache;
-
-		/// <summary>
-		/// 
-		/// </summary>
-		public bool bTestIncludeDependencyResolveCache;
-
-		public CPPHeaders(FileReference ProjectFile, FileReference DependencyCacheFile, bool bUseUBTMakefiles, bool bUseFlatCPPIncludeDependencyCache, bool bUseIncludeDependencyResolveCache, bool bTestIncludeDependencyResolveCache)
+		public CPPHeaders(FileReference ProjectFile, FileReference DependencyCacheFile)
 		{
 			this.ProjectFile = ProjectFile;
 			this.DependencyCacheFile = DependencyCacheFile;
-			this.bUseUBTMakefiles = bUseUBTMakefiles;
-			this.bUseFlatCPPIncludeDependencyCache = bUseFlatCPPIncludeDependencyCache;
-			this.bUseIncludeDependencyResolveCache = bUseIncludeDependencyResolveCache;
-			this.bTestIncludeDependencyResolveCache = bTestIncludeDependencyResolveCache;
-		}
-
-		/// <summary>
-		/// Finds the header file that is referred to by a partial include filename.
-		/// </summary>
-		/// <param name="FromFile">The file containing the include directory</param>
-		/// <param name="RelativeIncludePath">path relative to the project</param>
-		/// <param name="IncludePaths">Include paths to search</param>
-		public static FileItem FindIncludedFile(FileReference FromFile, string RelativeIncludePath, CppIncludePaths IncludePaths)
-		{
-			FileItem Result = null;
-
-			++TotalFindIncludedFileCalls;
-
-			// Only search for the include file if the result hasn't been cached.
-			string InvariantPath = RelativeIncludePath.ToLowerInvariant();
-			if (!IncludePaths.IncludeFileSearchDictionary.TryGetValue(InvariantPath, out Result))
-			{
-				int SearchAttempts = 0;
-				if (Path.IsPathRooted(RelativeIncludePath))
-				{
-					FileReference Reference = FileReference.Combine(UnrealBuildTool.EngineSourceDirectory, RelativeIncludePath);
-					if (DirectoryLookupCache.FileExists(Reference))
-					{
-						Result = FileItem.GetItemByFileReference(Reference);
-					}
-					++SearchAttempts;
-				}
-				else
-				{
-					// Find the first include path that the included file exists in.
-					List<DirectoryReference> IncludePathsToSearch = IncludePaths.GetPathsToSearch(FromFile);
-					foreach (DirectoryReference IncludePath in IncludePathsToSearch)
-					{
-						++SearchAttempts;
-						FileReference FullFilePath;
-						try
-						{
-							FullFilePath = FileReference.Combine(IncludePath, RelativeIncludePath);
-						}
-						catch (ArgumentException Exception)
-						{
-							throw new BuildException(Exception, "Failed to combine null or invalid include paths.");
-						}
-						if (DirectoryLookupCache.FileExists(FullFilePath))
-						{
-							Result = FileItem.GetItemByFileReference(FullFilePath);
-							break;
-						}
-					}
-				}
-
-				IncludePathSearchAttempts += SearchAttempts;
-
-				if (bPrintPerformanceInfo)
-				{
-					// More than two search attempts indicates:
-					//		- Include path was not relative to the directory that the including file was in
-					//		- Include path was not relative to the project's base
-					if (SearchAttempts > 2)
-					{
-						Log.TraceVerbose("   Cache miss: " + RelativeIncludePath + " found after " + SearchAttempts.ToString() + " attempts: " + (Result != null ? Result.AbsolutePath : "NOT FOUND!"));
-					}
-				}
-
-				// Cache the result of the include path search.
-				IncludePaths.IncludeFileSearchDictionary.Add(InvariantPath, Result);
-			}
-
-			// @todo ubtmake: The old UBT tried to skip 'external' (STABLE) headers here.  But it didn't work.  We might want to do this though!  Skip system headers and source/thirdparty headers!
-
-			if (Result != null)
-			{
-				Log.TraceVerbose("Resolved included file \"{0}\" to: {1}", RelativeIncludePath, Result.AbsolutePath);
-			}
-			else
-			{
-				Log.TraceVerbose("Couldn't resolve included file \"{0}\"", RelativeIncludePath);
-			}
-
-			return Result;
-		}
-
-		public List<FileItem> FindAndCacheAllIncludedFiles(FileItem SourceFile, CppIncludePaths IncludePaths, bool bOnlyCachedDependencies)
-		{
-			List<FileItem> Result = null;
-
-			if (IncludePaths.IncludeFileSearchDictionary == null)
-			{
-				IncludePaths.IncludeFileSearchDictionary = new Dictionary<string, FileItem>();
-			}
-
-			if (bOnlyCachedDependencies && bUseFlatCPPIncludeDependencyCache)
-			{
-				Result = FlatCPPIncludeDependencyCache.GetDependenciesForFile(SourceFile.Location);
-				if (Result == null)
-				{
-					// Nothing cached for this file!  It is new to us.  This is the expected flow when our CPPIncludeDepencencyCache is missing.
-				}
-			}
-			else
-			{
-				// If we're doing an exhaustive include scan, make sure that we have our include dependency cache loaded and ready
-				if (!bOnlyCachedDependencies)
-				{
-					if (IncludeDependencyCache == null)
-					{
-						IncludeDependencyCache = DependencyCache.Create(DependencyCacheFile);
-					}
-				}
-
-				// Get the headers
-				Result = FindAndCacheIncludedFiles(SourceFile, IncludePaths, bOnlyCachedDependencies);
-
-				// Update cache
-				if (bUseFlatCPPIncludeDependencyCache && !bOnlyCachedDependencies)
-				{
-					List<FileReference> Dependencies = new List<FileReference>();
-					foreach (FileItem IncludedFile in Result)
-					{
-						Dependencies.Add(IncludedFile.Location);
-					}
-					FlatCPPIncludeDependencyCache.SetDependenciesForFile(SourceFile.Location, Dependencies);
-				}
-			}
-
-			return Result;
-		}
-
-		/// <summary>
-		/// Add all the files included by a source file to a set, using a cache.
-		/// </summary>
-		/// <param name="SourceFile">The file to check.</param>
-		/// <param name="IncludePaths">Include paths to search.</param>
-		/// <param name="bOnlyCachedDependencies">Whether to just return cached dependencies, or update the cache with new results.</param>
-		private List<FileItem> FindAndCacheIncludedFiles(FileItem SourceFile, CppIncludePaths IncludePaths, bool bOnlyCachedDependencies)
-		{
-			// Get the map of files to their list of includes
-			Dictionary<FileItem, List<FileItem>> FileToIncludedFiles = bOnlyCachedDependencies ? OnlyCachedIncludedFilesMap : ExhaustiveIncludedFilesMap;
-
-			// Check if the included files is in the cache. If not, we'll create it.
-			List<FileItem> IncludedFiles;
-			if(!FileToIncludedFiles.TryGetValue(SourceFile, out IncludedFiles))
-			{
-				HashSet<FileItem> VisitedFiles = new HashSet<FileItem>();
-				VisitedFiles.Add(SourceFile);
-
-				HashSet<FileItem> IncludedFilesSet = new HashSet<FileItem>();
-				FindAndCacheIncludedFilesInner(SourceFile, IncludedFilesSet, IncludePaths, bOnlyCachedDependencies, VisitedFiles);
-
-				IncludedFiles = IncludedFilesSet.ToList();
-				FileToIncludedFiles.Add(SourceFile, IncludedFiles);
-			}
-
-			return IncludedFiles;
-		}
-
-		/// <summary>
-		/// Add all the files included by a source file to a set, using a cache.
-		/// </summary>
-		/// <param name="SourceFile">The file to check.</param>
-		/// <param name="IncludedFiles">Set of included files to add to</param>
-		/// <param name="IncludePaths">Include paths to search.</param>
-		/// <param name="bOnlyCachedDependencies">Whether to just return cached dependencies, or update the cache with new results.</param>
-		/// <param name="VisitedFiles">Set of files that have already been visited. Used to prevent infinite loops between circularly dependent headers.</param>
-		private void FindAndCacheIncludedFilesInner(FileItem SourceFile, HashSet<FileItem> IncludedFiles, CppIncludePaths IncludePaths, bool bOnlyCachedDependencies, HashSet<FileItem> VisitedFiles)
-		{
-			HashSet<FileItem> DirectlyIncludedFiles = GetDirectlyIncludedFiles(SourceFile, IncludePaths, bOnlyCachedDependencies);
-			foreach (FileItem DirectlyIncludedFile in DirectlyIncludedFiles)
-			{
-				if(IncludedFiles.Add(DirectlyIncludedFile))
-				{
-					// Get the map of files to their list of includes
-					Dictionary<FileItem, List<FileItem>> FileToIncludedFiles = bOnlyCachedDependencies ? OnlyCachedIncludedFilesMap : ExhaustiveIncludedFilesMap;
-
-					// Recursively add the files included by this file
-					List<FileItem> InnerFiles;
-					if(FileToIncludedFiles.TryGetValue(DirectlyIncludedFile, out InnerFiles))
-					{
-						// We already have the include paths cached; just add them directly.
-						IncludedFiles.UnionWith(InnerFiles);
-					}
-					else if(VisitedFiles.Add(DirectlyIncludedFile))
-					{
-						// We don't have include paths cached, and this isn't a recursive call. Create a new set and add it to the cache.
-						HashSet<FileItem> InnerFilesSet = new HashSet<FileItem>();
-						FindAndCacheIncludedFilesInner(DirectlyIncludedFile, InnerFilesSet, IncludePaths, bOnlyCachedDependencies, VisitedFiles);
-						FileToIncludedFiles.Add(DirectlyIncludedFile, InnerFilesSet.ToList());
-						IncludedFiles.UnionWith(InnerFilesSet);
-					}
-					else
-					{
-						// We're already building a list of include paths for this file up the stack. Just recurse through it this time.
-						FindAndCacheIncludedFilesInner(DirectlyIncludedFile, IncludedFiles, IncludePaths, bOnlyCachedDependencies, VisitedFiles);
-					}
-				}
-			}
-		}
-
-		/// <summary>
-		/// Get a set of directly included files from the given source file, resolving their include paths to FileItem instances.
-		/// </summary>
-		/// <param name="SourceFile">The file to check.</param>
-		/// <param name="IncludePaths">Include paths to search.</param>
-		/// <param name="bOnlyCachedDependencies">Whether to just return cached dependencies, or update the cache with new results.</param>
-		/// <returns>Set of files that are included</returns>
-		private HashSet<FileItem> GetDirectlyIncludedFiles(FileItem SourceFile, CppIncludePaths IncludePaths, bool bOnlyCachedDependencies)
-		{
-			// Gather a list of names of files directly included by this C++ file.
-			List<DependencyInclude> DirectIncludes = GetDirectIncludeDependencies(SourceFile, bOnlyCachedDependencies: bOnlyCachedDependencies);
-
-			// Build a list of the unique set of files that are included by this file.
-			HashSet<FileItem> DirectlyIncludedFiles = new HashSet<FileItem>();
-			// require a for loop here because we need to keep track of the index in the list.
-			for (int DirectlyIncludedFileNameIndex = 0; DirectlyIncludedFileNameIndex < DirectIncludes.Count; ++DirectlyIncludedFileNameIndex)
-			{
-				// Resolve the included file name to an actual file.
-				DependencyInclude DirectInclude = DirectIncludes[DirectlyIncludedFileNameIndex];
-				if (!DirectInclude.HasAttemptedResolve ||
-					// ignore any preexisting resolve cache if we are not configured to use it.
-					!bUseIncludeDependencyResolveCache ||
-					// if we are testing the resolve cache, we force UBT to resolve every time to look for conflicts
-					bTestIncludeDependencyResolveCache
-					)
-				{
-					++TotalDirectIncludeResolveCacheMisses;
-
-					// search the include paths to resolve the file
-					FileItem DirectIncludeResolvedFile = CPPHeaders.FindIncludedFile(SourceFile.Location, DirectInclude.IncludeName, IncludePaths);
-					if (DirectIncludeResolvedFile != null)
-					{
-						DirectlyIncludedFiles.Add(DirectIncludeResolvedFile);
-					}
-					IncludeDependencyCache.CacheResolvedIncludeFullPath(SourceFile, DirectlyIncludedFileNameIndex, DirectIncludeResolvedFile != null ? DirectIncludeResolvedFile.Location : null, bUseIncludeDependencyResolveCache, bTestIncludeDependencyResolveCache);
-				}
-				else
-				{
-					// we might have cached an attempt to resolve the file, but couldn't actually find the file (system headers, etc).
-					if (DirectInclude.IncludeResolvedNameIfSuccessful != null)
-					{
-						DirectlyIncludedFiles.Add(FileItem.GetItemByFileReference(DirectInclude.IncludeResolvedNameIfSuccessful));
-					}
-				}
-			}
-			TotalDirectIncludeResolves += DirectIncludes.Count;
-
-			return DirectlyIncludedFiles;
 		}
 
 		public static TimeSpan DirectIncludeCacheMissesTotalTime = TimeSpan.Zero;
 		public static int TotalDirectIncludeCacheMisses = 0;
-		public static int TotalDirectIncludeResolveCacheMisses = 0;
-		public static int TotalDirectIncludeResolves = 0;
 
 
 		/// <summary>
diff --git a/Engine/Source/Programs/UnrealBuildTool/System/CppCompileEnvironment.cs b/Engine/Source/Programs/UnrealBuildTool/System/CppCompileEnvironment.cs
--- a/Engine/Source/Programs/UnrealBuildTool/System/CppCompileEnvironment.cs
+++ b/Engine/Source/Programs/UnrealBuildTool/System/CppCompileEnvironment.cs
@@ -266,6 +266,11 @@ namespace UnrealBuildTool
 		/// </summary>
 		public bool bPrintTimingInfo;
 
+		/// <summary>
+		/// Whether to output a dependencies file along with the output build products
+		/// </summary>
+		public bool bGenerateDependenciesFile = true;
+
 		/// <summary>
 		/// When enabled, allows XGE to compile pre-compiled header files on remote machines.  Otherwise, PCHs are always generated locally.
 		/// </summary>
@@ -370,6 +375,7 @@ namespace UnrealBuildTool
 			PGOFilenamePrefix = Other.PGOFilenamePrefix;
 			PGODirectory = Other.PGODirectory;
 			bPrintTimingInfo = Other.bPrintTimingInfo;
+			bGenerateDependenciesFile = Other.bGenerateDependenciesFile;
 			bAllowRemotelyCompiledPCHs = Other.bAllowRemotelyCompiledPCHs;
 			IncludePaths = new CppIncludePaths(Other.IncludePaths);
 			ForceIncludeFiles.AddRange(Other.ForceIncludeFiles);
diff --git a/Engine/Source/Programs/UnrealBuildTool/System/CppDependencyCache.cs b/Engine/Source/Programs/UnrealBuildTool/System/CppDependencyCache.cs
new file mode 100644
--- /dev/null
+++ b/Engine/Source/Programs/UnrealBuildTool/System/CppDependencyCache.cs
@@ -0,0 +1,433 @@
+// Copyright 1998-2019 Epic Games, Inc. All Rights Reserved.
+
+using System;
+using System.Collections.Concurrent;
+using System.Collections.Generic;
+using System.Diagnostics;
+using System.IO;
+using System.Linq;
+using System.Text;
+using System.Threading.Tasks;
+using Tools.DotNETCommon;
+
+namespace UnrealBuildTool
+{
+	/// <summary>
+	/// Reads the contents of C++ dependency files, and caches them for future iterations.
+	/// </summary>
+	class CppDependencyCache
+	{
+		/// <summary>
+		/// Contents of a single dependency file
+		/// </summary>
+		class DependencyInfo
+		{
+			public long LastWriteTimeUtc;
+			public List<FileItem> Files;
+
+			public DependencyInfo(long LastWriteTimeUtc, List<FileItem> Files)
+			{
+				this.LastWriteTimeUtc = LastWriteTimeUtc;
+				this.Files = Files;
+			}
+
+			public DependencyInfo(BinaryReader Reader, List<FileItem> UniqueFileItems)
+			{
+				LastWriteTimeUtc = Reader.ReadInt64();
+				Files = Reader.ReadFileItemList(UniqueFileItems);
+			}
+
+			public void Write(BinaryWriter Writer, Dictionary<FileItem, int> UniqueFileItemToIndex)
+			{
+				Writer.Write(LastWriteTimeUtc);
+				Writer.Write(Files, UniqueFileItemToIndex);
+			}
+		}
+
+		/// <summary>
+		/// Location of this dependency cache
+		/// </summary>
+		FileReference Location;
+
+		/// <summary>
+		/// Directory for files to cache dependencies for.
+		/// </summary>
+		DirectoryReference BaseDir;
+
+		/// <summary>
+		/// The parent cache.
+		/// </summary>
+		CppDependencyCache Parent;
+
+		/// <summary>
+		/// Map from file item to dependency info
+		/// </summary>
+		ConcurrentDictionary<FileItem, DependencyInfo> DependencyFileToInfo = new ConcurrentDictionary<FileItem, DependencyInfo>();
+
+		/// <summary>
+		/// Whether the cache has been modified and needs to be saved
+		/// </summary>
+		bool bModified;
+
+		/// <summary>
+		/// Static cache of all constructed dependency caches
+		/// </summary>
+		static Dictionary<FileReference, CppDependencyCache> Caches = new Dictionary<FileReference, CppDependencyCache>();
+
+		/// <summary>
+		/// Constructs a dependency cache. This method is private; call CppDependencyCache.Create() to create a cache hierarchy for a given project.
+		/// </summary>
+		/// <param name="Location">File to store the cache</param>
+		/// <param name="BaseDir">Base directory for files that this cache should store data for</param>
+		/// <param name="Parent">The parent cache to use</param>
+		private CppDependencyCache(FileReference Location, DirectoryReference BaseDir, CppDependencyCache Parent)
+		{
+			this.Location = Location;
+			this.BaseDir = BaseDir;
+			this.Parent = Parent;
+
+			if(FileReference.Exists(Location))
+			{
+				Read();
+			}
+		}
+
+		/// <summary>
+		/// Attempts to read the dependencies from the given input file
+		/// </summary>
+		/// <param name="InputFile">File to be read</param>
+		/// <param name="OutDependencyItems">Receives a list of output items</param>
+		/// <returns>True if the input file exists and the dependencies were read</returns>
+		public bool TryGetDependencies(FileItem InputFile, out List<FileItem> OutDependencyItems)
+		{
+			if(!InputFile.Exists)
+			{
+				OutDependencyItems = null;
+				return false;
+			}
+
+			return TryGetDependenciesInternal(InputFile, out OutDependencyItems);
+		}
+
+		/// <summary>
+		/// Attempts to read dependencies from the given file.
+		/// </summary>
+		/// <param name="InputFile">File to be read</param>
+		/// <param name="OutDependencyItems">Receives a list of output items</param>
+		/// <returns>True if the input file exists and the dependencies were read</returns>
+		private bool TryGetDependenciesInternal(FileItem InputFile, out List<FileItem> OutDependencyItems)
+		{
+			if(Parent != null && !InputFile.Location.IsUnderDirectory(BaseDir))
+			{
+				return Parent.TryGetDependencies(InputFile, out OutDependencyItems);
+			}
+			else
+			{
+				DependencyInfo Info;
+				if(DependencyFileToInfo.TryGetValue(InputFile, out Info) && InputFile.LastWriteTimeUtc.Ticks <= Info.LastWriteTimeUtc)
+				{
+					OutDependencyItems = Info.Files;
+					return true;
+				}
+
+				List<FileItem> DependencyItems = ReadDependenciesFile(InputFile.Location);
+				DependencyFileToInfo.TryAdd(InputFile, new DependencyInfo(InputFile.LastWriteTimeUtc.Ticks, DependencyItems));
+				bModified = true;
+
+				OutDependencyItems = DependencyItems;
+				return true;
+			}
+		}
+
+		/// <summary>
+		/// Creates a cache hierarchy for a particular target
+		/// </summary>
+		/// <param name="ProjectFile">Project file for the target being built</param>
+		/// <param name="TargetName">Name of the target</param>
+		/// <param name="Platform">Platform being built</param>
+		/// <param name="Configuration">Configuration being built</param>
+		/// <param name="TargetType">The target type</param>
+		/// <returns>Dependency cache hierarchy for the given project</returns>
+		public static CppDependencyCache CreateHierarchy(FileReference ProjectFile, string TargetName, UnrealTargetPlatform Platform, UnrealTargetConfiguration Configuration, TargetType TargetType)
+		{
+			CppDependencyCache Cache = null;
+
+			if(ProjectFile == null || !UnrealBuildTool.IsEngineInstalled())
+			{
+				string AppName;
+				if(TargetType == TargetType.Program)
+				{
+					AppName = TargetName;
+				}
+				else
+				{
+					AppName = UEBuildTarget.GetAppNameForTargetType(TargetType);
+				}
+
+				FileReference EngineCacheLocation = FileReference.Combine(UnrealBuildTool.EngineDirectory, "Intermediate", "Build", Platform.ToString(), AppName, Configuration.ToString(), "DependencyCache.dat");
+				Cache = FindOrAddCache(EngineCacheLocation, UnrealBuildTool.EngineDirectory, Cache);
+			}
+
+			if(ProjectFile != null)
+			{
+				FileReference ProjectCacheLocation = FileReference.Combine(ProjectFile.Directory, "Intermediate", "Build", Platform.ToString(), TargetName, Configuration.ToString(), "DependencyCache.dat");
+				Cache = FindOrAddCache(ProjectCacheLocation, ProjectFile.Directory, Cache);
+			}
+
+			return Cache;
+		}
+
+		/// <summary>
+		/// Reads a cache from the given location, or creates it with the given settings
+		/// </summary>
+		/// <param name="Location">File to store the cache</param>
+		/// <param name="BaseDir">Base directory for files that this cache should store data for</param>
+		/// <param name="Parent">The parent cache to use</param>
+		/// <returns>Reference to a dependency cache with the given settings</returns>
+		static CppDependencyCache FindOrAddCache(FileReference Location, DirectoryReference BaseDir, CppDependencyCache Parent)
+		{
+			lock(Caches)
+			{
+				CppDependencyCache Cache;
+				if(Caches.TryGetValue(Location, out Cache))
+				{
+					Debug.Assert(Cache.BaseDir == BaseDir);
+					Debug.Assert(Cache.Parent == Parent);
+				}
+				else
+				{
+					Cache = new CppDependencyCache(Location, BaseDir, Parent);
+					Caches.Add(Location, Cache);
+				}
+				return Cache;
+			}
+		}
+
+		/// <summary>
+		/// Save all the caches that have been modified
+		/// </summary>
+		public static void SaveAll()
+		{
+			Parallel.ForEach(Caches.Values, Cache => { if(Cache.bModified){ Cache.Write(); } });
+		}
+
+		/// <summary>
+		/// Reads data for this dependency cache from disk
+		/// </summary>
+		private void Read()
+		{
+			using(FileStream Stream = File.Open(Location.FullName, FileMode.Open, FileAccess.Read, FileShare.Read))
+			{
+				BinaryReader Reader = new BinaryReader(Stream);
+
+				long Offset = Reader.ReadInt64();
+				long OriginalOffset = Reader.BaseStream.Position;
+
+				Reader.BaseStream.Seek(Offset, SeekOrigin.Begin);
+				List<FileReference> UniqueFileReferences = Reader.ReadList(() => Reader.ReadFileReference());
+				List<FileItem> UniqueFileItems = UniqueFileReferences.Select(x => FileItem.GetItemByFileReference(x)).ToList();
+				Reader.BaseStream.Seek(OriginalOffset, SeekOrigin.Begin);
+
+				int Count = Reader.ReadInt32();
+				for(int Idx = 0; Idx < Count; Idx++)
+				{
+					FileItem File = Reader.ReadFileItem(UniqueFileItems);
+					DependencyInfo Info = new DependencyInfo(Reader, UniqueFileItems);
+					DependencyFileToInfo[File] = Info;
+				}
+			}
+		}
+
+		/// <summary>
+		/// Writes data for this dependency cache to disk
+		/// </summary>
+		private void Write()
+		{
+			DirectoryReference.CreateDirectory(Location.Directory);
+			using(FileStream Stream = File.Open(Location.FullName, FileMode.Create, FileAccess.Write, FileShare.Read))
+			{
+				BinaryWriter Writer = new BinaryWriter(Stream);
+
+				long Offset = Writer.BaseStream.Position;
+				Writer.Write((Int64)0);
+				Dictionary<FileItem, int> UniqueFileItemToIndex = new Dictionary<FileItem, int>();
+
+				Writer.Write(DependencyFileToInfo.Count);
+				foreach(KeyValuePair<FileItem, DependencyInfo> Pair in DependencyFileToInfo)
+				{
+					Writer.Write(Pair.Key, UniqueFileItemToIndex);
+					Pair.Value.Write(Writer, UniqueFileItemToIndex);
+				}
+
+				long FileItemTableOffset = Writer.BaseStream.Position;
+				FileReference[] UniqueFileReferences = new FileReference[UniqueFileItemToIndex.Count];
+				foreach(KeyValuePair<FileItem, int> Pair in UniqueFileItemToIndex)
+				{
+					UniqueFileReferences[Pair.Value] = Pair.Key.Location;
+				}
+				Writer.Write(UniqueFileReferences, x => Writer.Write(x));
+
+				Writer.BaseStream.Seek(Offset, SeekOrigin.Begin);
+				Writer.Write(FileItemTableOffset);
+				Writer.Flush();
+			}
+			bModified = false;
+		}
+
+		/// <summary>
+		/// Reads dependencies from the given file.
+		/// </summary>
+		/// <param name="InputFile">The file to read from</param>
+		/// <returns>List of included dependencies</returns>
+		static List<FileItem> ReadDependenciesFile(FileReference InputFile)
+		{
+			if(InputFile.HasExtension(".txt"))
+			{
+				string[] Lines = FileReference.ReadAllLines(InputFile);
+
+				HashSet<FileItem> DependencyItems = new HashSet<FileItem>();
+				foreach(string Line in Lines)
+				{
+					if(Line.Length > 0)
+					{
+						// Ignore *.tlh and *.tli files generated by the compiler from COM DLLs
+						if(!Line.EndsWith(".tlh", StringComparison.OrdinalIgnoreCase) && !Line.EndsWith(".tli", StringComparison.OrdinalIgnoreCase))
+						{
+							DependencyItems.Add(FileItem.GetItemByPath(Line));
+						}
+					}
+				}
+				return DependencyItems.ToList();
+			}
+			else if(InputFile.HasExtension(".d"))
+			{
+				string Text = FileReference.ReadAllText(InputFile);
+
+				List<string> Tokens = new List<string>();
+
+				StringBuilder Token = new StringBuilder();
+				for(int Idx = 0; TryReadMakefileToken(Text, ref Idx, Token); )
+				{
+					Tokens.Add(Token.ToString());
+				}
+
+				int TokenIdx = 0;
+				while(TokenIdx < Tokens.Count && Tokens[TokenIdx] == "\n")
+				{
+					TokenIdx++;
+				}
+
+				if(TokenIdx + 1 >= Tokens.Count || Tokens[TokenIdx + 1] != ":")
+				{
+					throw new BuildException("Unable to parse dependency file");
+				}
+
+				TokenIdx += 2;
+			
+				List<FileItem> NewDependencyFiles = new List<FileItem>();
+				for(; TokenIdx < Tokens.Count && Tokens[TokenIdx] != "\n"; TokenIdx++)
+				{
+					NewDependencyFiles.Add(FileItem.GetItemByPath(Tokens[TokenIdx]));
+				}
+
+				while(TokenIdx < Tokens.Count && Tokens[TokenIdx] == "\n")
+				{
+					TokenIdx++;
+				}
+
+				if(TokenIdx != Tokens.Count)
+				{
+					throw new BuildException("Unable to parse dependency file");
+				}
+
+				return NewDependencyFiles;
+			}
+			else
+			{
+				throw new BuildException("Unknown dependency list file type: {0}", InputFile);
+			}
+		}
+
+		/// <summary>
+		/// Attempts to read a single token from a makefile
+		/// </summary>
+		/// <param name="Text">Text to read from</param>
+		/// <param name="RefIdx">Current position within the file</param>
+		/// <param name="Token">Receives the token characters</param>
+		/// <returns>True if a token was read, false if the end of the buffer was reached</returns>
+		static bool TryReadMakefileToken(string Text, ref int RefIdx, StringBuilder Token)
+		{
+			Token.Clear();
+
+			int Idx = RefIdx;
+			for(;;)
+			{
+				if(Idx == Text.Length)
+				{
+					return false;
+				}
+
+				// Skip whitespace
+				while(Text[Idx] == ' ' || Text[Idx] == '\t')
+				{
+					if(++Idx == Text.Length)
+					{
+						return false;
+					}
+				}
+
+				// Colon token
+				if(Text[Idx] == ':')
+				{
+					Token.Append(':');
+					RefIdx = Idx + 1;
+					return true;
+				}
+
+				// Check for a newline
+				if(Text[Idx] == '\r' || Text[Idx] == '\n')
+				{
+					Token.Append('\n');
+					RefIdx = Idx + 1;
+					return true;
+				}
+
+				// Check for an escaped newline
+				if(Text[Idx] == '\\' && Idx + 1 < Text.Length)
+				{
+					if(Text[Idx + 1] == '\n')
+					{
+						Idx += 2;
+						continue;
+					}
+					if(Text[Idx + 1] == '\r' && Idx + 2 < Text.Length && Text[Idx + 2] == '\n')
+					{
+						Idx += 3;
+						continue;
+					}
+				}
+
+				// Read a token. Special handling for drive letters on Windows!
+				for(; Idx < Text.Length; Idx++)
+				{
+					if(Text[Idx] == ' ' || Text[Idx] == '\t' || Text[Idx] == '\r' || Text[Idx] == '\n')
+					{
+						break;
+					}
+					if(Text[Idx] == ':' && Token.Length > 1)
+					{
+						break;
+					}
+					if(Text[Idx] == '\\' && Idx + 1 < Text.Length && Text[Idx + 1] == ' ')
+					{
+						Idx++;
+					}
+					Token.Append(Text[Idx]);
+				}
+
+				RefIdx = Idx;
+				return true;
+			}
+		}
+	}
+}
diff --git a/Engine/Source/Programs/UnrealBuildTool/System/FileItem.cs b/Engine/Source/Programs/UnrealBuildTool/System/FileItem.cs
--- a/Engine/Source/Programs/UnrealBuildTool/System/FileItem.cs
+++ b/Engine/Source/Programs/UnrealBuildTool/System/FileItem.cs
@@ -28,28 +28,6 @@ namespace UnrealBuildTool
 		/// </summary>
 		public FileReference Location;
 
-		/// <summary>
-		/// For C++ file items, this stores cached information about the include paths needed in order to include header files from these C++ files.  This is part of UBT's dependency caching optimizations.
-		/// </summary>
-		public CppIncludePaths CachedIncludePaths
-		{
-			get
-			{
-				return CachedIncludePathsValue;
-			}
-			set
-			{
-				if (value != null && CachedIncludePathsValue != null && CachedIncludePathsValue != value)
-				{
-					// Uh oh.  We're clobbering our cached CompileEnvironment for this file with a different CompileEnvironment.  This means
-					// that the same source file is being compiled into more than one module.
-					throw new BuildException("File '{0}' was included by multiple modules, but with different include paths", this.Info.FullName);
-				}
-				CachedIncludePathsValue = value;
-			}
-		}
-		private CppIncludePaths CachedIncludePathsValue;
-
 		/// <summary>
 		/// The information about the file.
 		/// </summary>
@@ -103,25 +81,6 @@ namespace UnrealBuildTool
 			get { return Info.LastWriteTimeUtc; }
 		}
 
-		/// <summary>
-		/// Clears the FileItem caches.
-		/// </summary>
-		public static void ClearCaches()
-		{
-			UniqueSourceFileMap.Clear();
-		}
-
-		/// <summary>
-		/// Clears the cached include paths on every file item
-		/// </summary>
-		public static void ClearCachedIncludePaths()
-		{
-			foreach(FileItem Item in UniqueSourceFileMap.Values)
-			{
-				Item.CachedIncludePaths = null;
-			}
-		}
-
 		/// <summary>
 		/// Gets a FileItem corresponding to the given path
 		/// </summary>
diff --git a/Engine/Source/Programs/UnrealBuildTool/System/FlatCPPIncludeDepencencyCache.cs b/Engine/Source/Programs/UnrealBuildTool/System/FlatCPPIncludeDepencencyCache.cs
deleted file mode 100644
diff --git a/Engine/Source/Programs/UnrealBuildTool/System/Unity.cs b/Engine/Source/Programs/UnrealBuildTool/System/Unity.cs
--- a/Engine/Source/Programs/UnrealBuildTool/System/Unity.cs
+++ b/Engine/Source/Programs/UnrealBuildTool/System/Unity.cs
@@ -271,12 +271,6 @@ namespace UnrealBuildTool
 
 						// Compile this file as part of the unity blob
 						CPPUnityFileBuilder.AddFile(CPPFile);
-
-						// Now that the CPPFile is part of this unity file, we will no longer need to treat it like a root level prerequisite for our
-						// dependency cache, as it is now an "indirect include" from the unity file.  We'll clear out the compile environment
-						// attached to this file.  This prevents us from having to cache all of the indirect includes from these files inside our
-						// dependency cache, which speeds up iterative builds a lot!
-						CPPFile.CachedIncludePaths = null;
 					}
 				}
 
@@ -339,7 +333,6 @@ namespace UnrealBuildTool
 
 				// Write the unity file to the intermediate folder.
 				FileItem UnityCPPFile = FileItem.CreateIntermediateTextFile(UnityCPPFilePath, OutputUnityCPPWriter.ToString());
-				UnityCPPFile.CachedIncludePaths = CompileEnvironment.IncludePaths;
 				NewCPPFiles.Add(UnityCPPFile);
 			}
 
diff --git a/Engine/Source/Programs/UnrealBuildTool/ToolChain/UEToolChain.cs b/Engine/Source/Programs/UnrealBuildTool/ToolChain/UEToolChain.cs
--- a/Engine/Source/Programs/UnrealBuildTool/ToolChain/UEToolChain.cs
+++ b/Engine/Source/Programs/UnrealBuildTool/ToolChain/UEToolChain.cs
@@ -80,27 +80,6 @@ namespace UnrealBuildTool
 			return true;
 		}
 
-		protected virtual void AddPrerequisiteSourceFile(CppCompileEnvironment CompileEnvironment, FileItem SourceFile, List<FileItem> PrerequisiteItems)
-		{
-			PrerequisiteItems.Add(SourceFile);
-
-			if (!CompileEnvironment.Headers.bUseUBTMakefiles)	// In fast build iteration mode, we'll gather includes later on
-			{
-				// @todo ubtmake: What if one of the prerequisite files has become missing since it was updated in our cache? (usually, because a coder eliminated the source file)
-				//		-> Two CASES:
-				//				1) NOT WORKING: Non-unity file went away (SourceFile in this context).  That seems like an existing old use case.  Compile params or Response file should have changed?
-				//				2) WORKING: Indirect file went away (unity'd original source file or include).  This would return a file that no longer exists and adds to the prerequiteitems list
-				List<FileItem> IncludedFileList = CompileEnvironment.Headers.FindAndCacheAllIncludedFiles(SourceFile, CompileEnvironment.IncludePaths, bOnlyCachedDependencies: false);
-				if (IncludedFileList != null)
-				{
-					foreach (FileItem IncludedFile in IncludedFileList)
-					{
-						PrerequisiteItems.Add(IncludedFile);
-					}
-				}
-			}
-		}
-
 		public virtual void SetupBundleDependencies(List<UEBuildBinary> Binaries, string GameName)
 		{
 
diff --git a/Engine/Source/Programs/UnrealBuildTool/UnrealBuildTool.cs b/Engine/Source/Programs/UnrealBuildTool/UnrealBuildTool.cs
--- a/Engine/Source/Programs/UnrealBuildTool/UnrealBuildTool.cs
+++ b/Engine/Source/Programs/UnrealBuildTool/UnrealBuildTool.cs
@@ -14,11 +14,6 @@ namespace UnrealBuildTool
 {
 	static class UnrealBuildTool
 	{
-		/// <summary>
-		/// How much time was spent scanning for include dependencies for outdated C++ files
-		/// </summary>
-		public static TimeSpan TotalDeepIncludeScanTime = TimeSpan.Zero;
-
 		/// <summary>
 		/// The environment at boot time.
 		/// </summary>
diff --git a/Engine/Source/Programs/UnrealBuildTool/UnrealBuildTool.csproj b/Engine/Source/Programs/UnrealBuildTool/UnrealBuildTool.csproj
--- a/Engine/Source/Programs/UnrealBuildTool/UnrealBuildTool.csproj
+++ b/Engine/Source/Programs/UnrealBuildTool/UnrealBuildTool.csproj
@@ -39,7 +39,7 @@
     <Optimize>true</Optimize>
     <OutputPath>..\..\..\Binaries\DotNET\</OutputPath>
     <DefineConstants>TRACE</DefineConstants>
-    <Optimize>false</Optimize>
+    <Optimize>true</Optimize>
     <DebugType>pdbonly</DebugType>
     <PlatformTarget>AnyCPU</PlatformTarget>
     <ErrorReport>prompt</ErrorReport>
@@ -187,6 +187,7 @@
     <Compile Include="System\Asn.cs" />
     <Compile Include="System\BuildVersion.cs" />
     <Compile Include="System\CppCompileEnvironment.cs" />
+    <Compile Include="System\CppDependencyCache.cs" />
     <Compile Include="System\CppIncludeParser.cs" />
     <Compile Include="System\DataDrivenPlatformInfo.cs" />
     <Compile Include="System\EncryptionAndSigning.cs" />
@@ -336,7 +337,6 @@
     <Compile Include="System\BuildException.cs" />
     <Compile Include="Configuration\BuildConfiguration.cs" />
     <Compile Include="System\BuildHostPlatform.cs" />
-    <Compile Include="System\FlatCPPIncludeDepencencyCache.cs" />
     <Compile Include="System\ModuleDescriptor.cs" />
     <Compile Include="System\PluginDescriptor.cs" />
     <Compile Include="System\ProjectDescriptor.cs" />
