UBT: Add a timeline scope for enumerating editor processes for hot reload.

diff --git a/Engine/Source/Programs/UnrealBuildTool/System/HotReload.cs b/Engine/Source/Programs/UnrealBuildTool/System/HotReload.cs
--- a/Engine/Source/Programs/UnrealBuildTool/System/HotReload.cs
+++ b/Engine/Source/Programs/UnrealBuildTool/System/HotReload.cs
@@ -167,32 +167,35 @@ namespace UnrealBuildTool
 					throw new BuildException("Unknown editor filename for this platform");
 				}
 
-				BuildHostPlatform.ProcessInfo[] Processes = BuildHostPlatform.Current.GetProcesses();
-				string EditorRunsDir = Path.Combine(UnrealBuildTool.EngineDirectory.FullName, "Intermediate", "EditorRuns");
-
-				if (!Directory.Exists(EditorRunsDir))
+				using(Timeline.ScopeEvent("Finding editor processes for hot-reload"))
 				{
-					return false;
-				}
-
-				FileInfo[] EditorRunsFiles = new DirectoryInfo(EditorRunsDir).GetFiles();
+					BuildHostPlatform.ProcessInfo[] Processes = BuildHostPlatform.Current.GetProcesses();
+					string EditorRunsDir = Path.Combine(UnrealBuildTool.EngineDirectory.FullName, "Intermediate", "EditorRuns");
 
-				foreach (FileInfo File in EditorRunsFiles)
-				{
-					int PID;
-					BuildHostPlatform.ProcessInfo Proc = null;
-					if (!Int32.TryParse(File.Name, out PID) || (Proc = Processes.FirstOrDefault(P => P.PID == PID)) == default(BuildHostPlatform.ProcessInfo))
+					if (!Directory.Exists(EditorRunsDir))
 					{
-						// Delete stale files (it may happen if editor crashes).
-						File.Delete();
-						continue;
+						return false;
 					}
 
-					// Don't break here to allow clean-up of other stale files.
-					if (!bIsRunning)
+					FileInfo[] EditorRunsFiles = new DirectoryInfo(EditorRunsDir).GetFiles();
+
+					foreach (FileInfo File in EditorRunsFiles)
 					{
-						// Otherwise check if the path matches.
-						bIsRunning = new FileReference(Proc.Filename) == EditorLocation;
+						int PID;
+						BuildHostPlatform.ProcessInfo Proc = null;
+						if (!Int32.TryParse(File.Name, out PID) || (Proc = Processes.FirstOrDefault(P => P.PID == PID)) == default(BuildHostPlatform.ProcessInfo))
+						{
+							// Delete stale files (it may happen if editor crashes).
+							File.Delete();
+							continue;
+						}
+
+						// Don't break here to allow clean-up of other stale files.
+						if (!bIsRunning)
+						{
+							// Otherwise check if the path matches.
+							bIsRunning = new FileReference(Proc.Filename) == EditorLocation;
+						}
 					}
 				}
 			}
