UBT: Convert makefiles and dependency caches to use BinaryArchiveReader/BinaryArchiveWriter.

#rnx

diff --git a/Engine/Source/Programs/UnrealBuildTool/Configuration/UEBuildTarget.cs b/Engine/Source/Programs/UnrealBuildTool/Configuration/UEBuildTarget.cs
--- a/Engine/Source/Programs/UnrealBuildTool/Configuration/UEBuildTarget.cs
+++ b/Engine/Source/Programs/UnrealBuildTool/Configuration/UEBuildTarget.cs
@@ -597,48 +597,48 @@ namespace UnrealBuildTool
 			return bCompileMonolithic;	// @todo ubtmake: We need to make sure this function and similar things aren't called in assembler mode
 		}
 
-		public UEBuildTarget(BinaryReader Reader)
+		public UEBuildTarget(BinaryArchiveReader Reader)
 		{
-			TargetType = (TargetType)Reader.ReadInt32();
+			TargetType = (TargetType)Reader.ReadInt();
 			ProjectFile = Reader.ReadFileReference();
 			AppName = Reader.ReadString();
 			TargetName = Reader.ReadString();
-			bUseSharedBuildEnvironment = Reader.ReadBoolean();
-			Platform = (UnrealTargetPlatform)Reader.ReadInt32();
-			Configuration = (UnrealTargetConfiguration)Reader.ReadInt32();
+			bUseSharedBuildEnvironment = Reader.ReadBool();
+			Platform = (UnrealTargetPlatform)Reader.ReadInt();
+			Configuration = (UnrealTargetConfiguration)Reader.ReadInt();
 			Architecture = Reader.ReadString();
 			PlatformIntermediateFolder = Reader.ReadString();
 			ProjectDirectory = Reader.ReadDirectoryReference();
 			ProjectIntermediateDirectory = Reader.ReadDirectoryReference();
 			EngineIntermediateDirectory = Reader.ReadDirectoryReference();
-			bCompileMonolithic = Reader.ReadBoolean();
+			bCompileMonolithic = Reader.ReadBool();
 			ReceiptFileName = Reader.ReadFileReference();
 			PreBuildStepScripts = Reader.ReadArray(() => Reader.ReadFileReference());
 			PostBuildStepScripts = Reader.ReadArray(() => Reader.ReadFileReference());
-			bDeployAfterCompile = Reader.ReadBoolean();
-			bHasProjectScriptPlugin = Reader.ReadBoolean();
+			bDeployAfterCompile = Reader.ReadBool();
+			bHasProjectScriptPlugin = Reader.ReadBool();
 		}
 
-		public void Write(BinaryWriter Writer)
+		public void Write(BinaryArchiveWriter Writer)
 		{
-			Writer.Write((int)TargetType);
-			Writer.Write(ProjectFile);
-			Writer.Write(AppName);
-			Writer.Write(TargetName);
-			Writer.Write(bUseSharedBuildEnvironment);
-			Writer.Write((int)Platform);
-			Writer.Write((int)Configuration);
-			Writer.Write(Architecture);
-			Writer.Write(PlatformIntermediateFolder);
-			Writer.Write(ProjectDirectory);
-			Writer.Write(ProjectIntermediateDirectory);
-			Writer.Write(EngineIntermediateDirectory);
-			Writer.Write(bCompileMonolithic);
-			Writer.Write(ReceiptFileName);
-			Writer.Write(PreBuildStepScripts, x => Writer.Write(x));
-			Writer.Write(PostBuildStepScripts, x => Writer.Write(x));
-			Writer.Write(bDeployAfterCompile);
-			Writer.Write(bHasProjectScriptPlugin);
+			Writer.WriteInt((int)TargetType);
+			Writer.WriteFileReference(ProjectFile);
+			Writer.WriteString(AppName);
+			Writer.WriteString(TargetName);
+			Writer.WriteBool(bUseSharedBuildEnvironment);
+			Writer.WriteInt((int)Platform);
+			Writer.WriteInt((int)Configuration);
+			Writer.WriteString(Architecture);
+			Writer.WriteString(PlatformIntermediateFolder);
+			Writer.WriteDirectoryReference(ProjectDirectory);
+			Writer.WriteDirectoryReference(ProjectIntermediateDirectory);
+			Writer.WriteDirectoryReference(EngineIntermediateDirectory);
+			Writer.WriteBool(bCompileMonolithic);
+			Writer.WriteFileReference(ReceiptFileName);
+			Writer.WriteArray(PreBuildStepScripts, x => Writer.WriteFileReference(x));
+			Writer.WriteArray(PostBuildStepScripts, x => Writer.WriteFileReference(x));
+			Writer.WriteBool(bDeployAfterCompile);
+			Writer.WriteBool(bHasProjectScriptPlugin);
 		}
 
 		/// <summary>
diff --git a/Engine/Source/Programs/UnrealBuildTool/System/Action.cs b/Engine/Source/Programs/UnrealBuildTool/System/Action.cs
--- a/Engine/Source/Programs/UnrealBuildTool/System/Action.cs
+++ b/Engine/Source/Programs/UnrealBuildTool/System/Action.cs
@@ -162,49 +162,49 @@ namespace UnrealBuildTool
 			ActionType = InActionType;
 		}
 
-		public Action(BinaryReader Reader, List<FileItem> UniqueFileItems)
+		public Action(BinaryArchiveReader Reader)
 		{
 			ActionType = (ActionType)Reader.ReadByte();
 			WorkingDirectory = Reader.ReadString();
-			bPrintDebugInfo = Reader.ReadBoolean();
+			bPrintDebugInfo = Reader.ReadBool();
 			CommandPath = Reader.ReadString();
 			CommandArguments = Reader.ReadString();
-			CommandDescription = Reader.ReadNullable(() => Reader.ReadString());
+			CommandDescription = Reader.ReadString();
 			StatusDescription = Reader.ReadString();
-			bCanExecuteRemotely = Reader.ReadBoolean();
-			bCanExecuteRemotelyWithSNDBS = Reader.ReadBoolean();
-			bIsGCCCompiler = Reader.ReadBoolean();
-			bIsUsingPCH = Reader.ReadBoolean();
-			bShouldOutputStatusDescription = Reader.ReadBoolean();
-			bProducesImportLibrary = Reader.ReadBoolean();
-			PrerequisiteItems = Reader.ReadFileItemList(UniqueFileItems);
-			ProducedItems = Reader.ReadFileItemList(UniqueFileItems);
-			DeleteItems = Reader.ReadFileItemList(UniqueFileItems);
-			DependencyListFile = Reader.ReadFileItem(UniqueFileItems);
+			bCanExecuteRemotely = Reader.ReadBool();
+			bCanExecuteRemotelyWithSNDBS = Reader.ReadBool();
+			bIsGCCCompiler = Reader.ReadBool();
+			bIsUsingPCH = Reader.ReadBool();
+			bShouldOutputStatusDescription = Reader.ReadBool();
+			bProducesImportLibrary = Reader.ReadBool();
+			PrerequisiteItems = Reader.ReadList(() => Reader.ReadFileItem());
+			ProducedItems = Reader.ReadList(() => Reader.ReadFileItem());
+			DeleteItems = Reader.ReadList(() => Reader.ReadFileItem());
+			DependencyListFile = Reader.ReadFileItem();
 		}
 
 		/// <summary>
 		/// ISerializable: Called when serialized to report additional properties that should be saved
 		/// </summary>
-		public void Write(BinaryWriter Writer, Dictionary<FileItem, int> UniqueFileItemToIndex)
+		public void Write(BinaryArchiveWriter Writer)
 		{
-			Writer.Write((byte)ActionType);
-			Writer.Write(WorkingDirectory);
-			Writer.Write(bPrintDebugInfo);
-			Writer.Write(CommandPath);
-			Writer.Write(CommandArguments);
-			Writer.WriteNullable(CommandDescription, () => Writer.Write(CommandDescription));
-			Writer.Write(StatusDescription);
-			Writer.Write(bCanExecuteRemotely);
-			Writer.Write(bCanExecuteRemotelyWithSNDBS);
-			Writer.Write(bIsGCCCompiler);
-			Writer.Write(bIsUsingPCH);
-			Writer.Write(bShouldOutputStatusDescription);
-			Writer.Write(bProducesImportLibrary);
-			Writer.Write(PrerequisiteItems, UniqueFileItemToIndex);
-			Writer.Write(ProducedItems, UniqueFileItemToIndex);
-			Writer.Write(DeleteItems, UniqueFileItemToIndex);
-			Writer.Write(DependencyListFile, UniqueFileItemToIndex);
+			Writer.WriteByte((byte)ActionType);
+			Writer.WriteString(WorkingDirectory);
+			Writer.WriteBool(bPrintDebugInfo);
+			Writer.WriteString(CommandPath);
+			Writer.WriteString(CommandArguments);
+			Writer.WriteString(CommandDescription);
+			Writer.WriteString(StatusDescription);
+			Writer.WriteBool(bCanExecuteRemotely);
+			Writer.WriteBool(bCanExecuteRemotelyWithSNDBS);
+			Writer.WriteBool(bIsGCCCompiler);
+			Writer.WriteBool(bIsUsingPCH);
+			Writer.WriteBool(bShouldOutputStatusDescription);
+			Writer.WriteBool(bProducesImportLibrary);
+			Writer.WriteList(PrerequisiteItems, Item => Writer.WriteFileItem(Item));
+			Writer.WriteList(ProducedItems, Item => Writer.WriteFileItem(Item));
+			Writer.WriteList(DeleteItems, Item => Writer.WriteFileItem(Item));
+			Writer.WriteFileItem(DependencyListFile);
 		}
 
 		/// <summary>
diff --git a/Engine/Source/Programs/UnrealBuildTool/System/BuildPrerequisites.cs b/Engine/Source/Programs/UnrealBuildTool/System/BuildPrerequisites.cs
--- a/Engine/Source/Programs/UnrealBuildTool/System/BuildPrerequisites.cs
+++ b/Engine/Source/Programs/UnrealBuildTool/System/BuildPrerequisites.cs
@@ -44,22 +44,22 @@ namespace UnrealBuildTool
 		{
 		}
 
-		public BuildPrerequisites(BinaryReader Reader, List<FileItem> UniqueFileItems)
+		public BuildPrerequisites(BinaryArchiveReader Reader)
 		{
-			SourceDirectories = new HashSet<DirectoryReference>(Reader.ReadArray(() => Reader.ReadDirectoryReference()));
-			WorkingSet = new HashSet<FileItem>(Reader.ReadFileItemList(UniqueFileItems));
-			CandidatesForWorkingSet = new HashSet<FileItem>(Reader.ReadFileItemList(UniqueFileItems));
-			AdditionalDependencies = new HashSet<FileItem>(Reader.ReadFileItemList(UniqueFileItems));
-			UObjectModuleHeaders = Reader.ReadList(() => new UHTModuleHeaderInfo(Reader, UniqueFileItems));
+			SourceDirectories = Reader.ReadHashSet(() => Reader.ReadDirectoryReference());
+			WorkingSet = Reader.ReadHashSet(() => Reader.ReadFileItem());
+			CandidatesForWorkingSet = Reader.ReadHashSet(() => Reader.ReadFileItem());
+			AdditionalDependencies = Reader.ReadHashSet(() => Reader.ReadFileItem());
+			UObjectModuleHeaders = Reader.ReadList(() => new UHTModuleHeaderInfo(Reader));
 		}
 
-		public void Write(BinaryWriter Writer, Dictionary<FileItem, int> UniqueFileItemToIndex)
+		public void Write(BinaryArchiveWriter Writer)
 		{
-			Writer.Write(SourceDirectories.ToArray(), x => Writer.Write(x));
-			Writer.Write(WorkingSet.ToList(), UniqueFileItemToIndex);
-			Writer.Write(CandidatesForWorkingSet.ToList(), UniqueFileItemToIndex);
-			Writer.Write(AdditionalDependencies.ToList(), UniqueFileItemToIndex);
-			Writer.Write(UObjectModuleHeaders, x => x.Write(Writer, UniqueFileItemToIndex));
+			Writer.WriteHashSet(SourceDirectories, x => Writer.WriteDirectoryReference(x));
+			Writer.WriteHashSet(WorkingSet, x => Writer.WriteFileItem(x));
+			Writer.WriteHashSet(CandidatesForWorkingSet, x => Writer.WriteFileItem(x));
+			Writer.WriteHashSet(AdditionalDependencies, x => Writer.WriteFileItem(x));
+			Writer.WriteList(UObjectModuleHeaders, x => x.Write(Writer));
 		}
 	}
 }
diff --git a/Engine/Source/Programs/UnrealBuildTool/System/CppDependencyCache.cs b/Engine/Source/Programs/UnrealBuildTool/System/CppDependencyCache.cs
--- a/Engine/Source/Programs/UnrealBuildTool/System/CppDependencyCache.cs
+++ b/Engine/Source/Programs/UnrealBuildTool/System/CppDependencyCache.cs
@@ -31,19 +31,26 @@ namespace UnrealBuildTool
 				this.Files = Files;
 			}
 
-			public DependencyInfo(BinaryReader Reader, List<FileItem> UniqueFileItems)
+			public static DependencyInfo Read(BinaryArchiveReader Reader)
 			{
-				LastWriteTimeUtc = Reader.ReadInt64();
-				Files = Reader.ReadFileItemList(UniqueFileItems);
+				long LastWriteTimeUtc = Reader.ReadLong();
+				List<FileItem> Files = Reader.ReadList(() => Reader.ReadFileItem());
+
+				return new DependencyInfo(LastWriteTimeUtc, Files);
 			}
 
-			public void Write(BinaryWriter Writer, Dictionary<FileItem, int> UniqueFileItemToIndex)
+			public void Write(BinaryArchiveWriter Writer)
 			{
-				Writer.Write(LastWriteTimeUtc);
-				Writer.Write(Files, UniqueFileItemToIndex);
+				Writer.WriteLong(LastWriteTimeUtc);
+				Writer.WriteList<FileItem>(Files, File => Writer.WriteFileItem(File));
 			}
 		}
 
+		/// <summary>
+		/// The current file version
+		/// </summary>
+		public const int CurrentVersion = 1;
+
 		/// <summary>
 		/// Location of this dependency cache
 		/// </summary>
@@ -218,22 +225,21 @@ namespace UnrealBuildTool
 		{
 			using(FileStream Stream = File.Open(Location.FullName, FileMode.Open, FileAccess.Read, FileShare.Read))
 			{
-				BinaryReader Reader = new BinaryReader(Stream);
-
-				long Offset = Reader.ReadInt64();
-				long OriginalOffset = Reader.BaseStream.Position;
-
-				Reader.BaseStream.Seek(Offset, SeekOrigin.Begin);
-				List<FileReference> UniqueFileReferences = Reader.ReadList(() => Reader.ReadFileReference());
-				List<FileItem> UniqueFileItems = UniqueFileReferences.Select(x => FileItem.GetItemByFileReference(x)).ToList();
-				Reader.BaseStream.Seek(OriginalOffset, SeekOrigin.Begin);
-
-				int Count = Reader.ReadInt32();
-				for(int Idx = 0; Idx < Count; Idx++)
+				using(BinaryArchiveReader Reader = new BinaryArchiveReader(Stream))
 				{
-					FileItem File = Reader.ReadFileItem(UniqueFileItems);
-					DependencyInfo Info = new DependencyInfo(Reader, UniqueFileItems);
-					DependencyFileToInfo[File] = Info;
+					int Version = Reader.ReadInt();
+					if(Version != CurrentVersion)
+					{
+						Log.TraceLog("Unable to read dependency cache from {0}; version {1} vs current {2}", Location, Version, CurrentVersion);
+						return;
+					}
+
+					int Count = Reader.ReadInt();
+					for(int Idx = 0; Idx < Count; Idx++)
+					{
+						FileItem File = Reader.ReadFileItem();
+						DependencyFileToInfo[File] = DependencyInfo.Read(Reader);
+					}
 				}
 			}
 		}
@@ -246,30 +252,17 @@ namespace UnrealBuildTool
 			DirectoryReference.CreateDirectory(Location.Directory);
 			using(FileStream Stream = File.Open(Location.FullName, FileMode.Create, FileAccess.Write, FileShare.Read))
 			{
-				BinaryWriter Writer = new BinaryWriter(Stream);
-
-				long Offset = Writer.BaseStream.Position;
-				Writer.Write((Int64)0);
-				Dictionary<FileItem, int> UniqueFileItemToIndex = new Dictionary<FileItem, int>();
-
-				Writer.Write(DependencyFileToInfo.Count);
-				foreach(KeyValuePair<FileItem, DependencyInfo> Pair in DependencyFileToInfo)
+				using(BinaryArchiveWriter Writer = new BinaryArchiveWriter(Stream))
 				{
-					Writer.Write(Pair.Key, UniqueFileItemToIndex);
-					Pair.Value.Write(Writer, UniqueFileItemToIndex);
-				}
+					Writer.WriteInt(CurrentVersion);
 
-				long FileItemTableOffset = Writer.BaseStream.Position;
-				FileReference[] UniqueFileReferences = new FileReference[UniqueFileItemToIndex.Count];
-				foreach(KeyValuePair<FileItem, int> Pair in UniqueFileItemToIndex)
-				{
-					UniqueFileReferences[Pair.Value] = Pair.Key.Location;
+					Writer.WriteInt(DependencyFileToInfo.Count);
+					foreach(KeyValuePair<FileItem, DependencyInfo> Pair in DependencyFileToInfo)
+					{
+						Writer.WriteFileItem(Pair.Key);
+						Pair.Value.Write(Writer);
+					}
 				}
-				Writer.Write(UniqueFileReferences, x => Writer.Write(x));
-
-				Writer.BaseStream.Seek(Offset, SeekOrigin.Begin);
-				Writer.Write(FileItemTableOffset);
-				Writer.Flush();
 			}
 			bModified = false;
 		}
diff --git a/Engine/Source/Programs/UnrealBuildTool/System/ExternalExecution.cs b/Engine/Source/Programs/UnrealBuildTool/System/ExternalExecution.cs
--- a/Engine/Source/Programs/UnrealBuildTool/System/ExternalExecution.cs
+++ b/Engine/Source/Programs/UnrealBuildTool/System/ExternalExecution.cs
@@ -208,32 +208,32 @@ namespace UnrealBuildTool
 			this.bIsReadOnly = bIsReadOnly;
 		}
 
-		public UHTModuleInfo(BinaryReader Reader, List<FileItem> UniqueFileItems)
+		public UHTModuleInfo(BinaryArchiveReader Reader)
 		{
 			ModuleName = Reader.ReadString();
 			ModuleRulesFile = Reader.ReadFileReference();
 			ModuleDirectory = Reader.ReadDirectoryReference();
 			ModuleType = Reader.ReadString();
-			PublicUObjectClassesHeaders = Reader.ReadFileItemList(UniqueFileItems);
-			PublicUObjectHeaders = Reader.ReadFileItemList(UniqueFileItems);
-			PrivateUObjectHeaders = Reader.ReadFileItemList(UniqueFileItems);
+			PublicUObjectClassesHeaders = Reader.ReadList(() => Reader.ReadFileItem());
+			PublicUObjectHeaders = Reader.ReadList(() => Reader.ReadFileItem());
+			PrivateUObjectHeaders = Reader.ReadList(() => Reader.ReadFileItem());
 			GeneratedCPPFilenameBase = Reader.ReadString();
-			GeneratedCodeVersion = (EGeneratedCodeVersion)Reader.ReadInt32();
-			bIsReadOnly = Reader.ReadBoolean();
+			GeneratedCodeVersion = (EGeneratedCodeVersion)Reader.ReadInt();
+			bIsReadOnly = Reader.ReadBool();
 		}
 
-		public void Write(BinaryWriter Writer, Dictionary<FileItem, int> UniqueFileItemToIndex)
+		public void Write(BinaryArchiveWriter Writer)
 		{
-			Writer.Write(ModuleName);
-			Writer.Write(ModuleRulesFile);
-			Writer.Write(ModuleDirectory);
-			Writer.Write(ModuleType);
-			Writer.Write(PublicUObjectClassesHeaders, UniqueFileItemToIndex);
-			Writer.Write(PublicUObjectHeaders, UniqueFileItemToIndex);
-			Writer.Write(PrivateUObjectHeaders, UniqueFileItemToIndex);
-			Writer.Write(GeneratedCPPFilenameBase);
-			Writer.Write((int)GeneratedCodeVersion);
-			Writer.Write(bIsReadOnly);
+			Writer.WriteString(ModuleName);
+			Writer.WriteFileReference(ModuleRulesFile);
+			Writer.WriteDirectoryReference(ModuleDirectory);
+			Writer.WriteString(ModuleType);
+			Writer.WriteList(PublicUObjectClassesHeaders, Item => Writer.WriteFileItem(Item));
+			Writer.WriteList(PublicUObjectHeaders, Item => Writer.WriteFileItem(Item));
+			Writer.WriteList(PrivateUObjectHeaders, Item => Writer.WriteFileItem(Item));
+			Writer.WriteString(GeneratedCPPFilenameBase);
+			Writer.WriteInt((int)GeneratedCodeVersion);
+			Writer.WriteBool(bIsReadOnly);
 		}
 
 		public override string ToString()
@@ -333,16 +333,16 @@ namespace UnrealBuildTool
 			this.HeaderFiles = HeaderFiles;
 		}
 
-		public UHTModuleHeaderInfo(BinaryReader Reader, List<FileItem> UniqueFileItems)
+		public UHTModuleHeaderInfo(BinaryArchiveReader Reader)
 		{
 			SourceFolder = Reader.ReadDirectoryReference();
-			HeaderFiles = Reader.ReadFileItemList(UniqueFileItems);
+			HeaderFiles = Reader.ReadList(() => Reader.ReadFileItem());
 		}
 
-		public void Write(BinaryWriter Writer, Dictionary<FileItem, int> UniqueFileItemToIndex)
+		public void Write(BinaryArchiveWriter Writer)
 		{
-			Writer.Write(SourceFolder);
-			Writer.Write(HeaderFiles, UniqueFileItemToIndex);
+			Writer.WriteDirectoryReference(SourceFolder);
+			Writer.WriteList(HeaderFiles, Item => Writer.WriteFileItem(Item));
 		}
 	}
 
diff --git a/Engine/Source/Programs/UnrealBuildTool/System/FileItem.cs b/Engine/Source/Programs/UnrealBuildTool/System/FileItem.cs
--- a/Engine/Source/Programs/UnrealBuildTool/System/FileItem.cs
+++ b/Engine/Source/Programs/UnrealBuildTool/System/FileItem.cs
@@ -243,62 +243,29 @@ namespace UnrealBuildTool
 		}
 	}
 
+	/// <summary>
+	/// Helper functions for serialization
+	/// </summary>
 	static class FileItemExtensionMethods
 	{
-		public static FileItem ReadFileItem(this BinaryReader Reader, List<FileItem> UniqueFileItems)
-		{
-			int Idx = Reader.ReadInt32();
-			if (Idx == -1)
-			{
-				return null;
-			}
-			else
-			{
-				return UniqueFileItems[Idx];
-			}
-		}
-
-		public static List<FileItem> ReadFileItemList(this BinaryReader Reader, List<FileItem> UniqueFileItems)
-		{
-			int Length = Reader.ReadInt32();
-			if(Length == -1)
-			{
-				return null;
-			}
-
-			List<FileItem> Items = new List<FileItem>(Length);
-			for(int Idx = 0; Idx < Length; Idx++)
-			{
-				Items.Add(ReadFileItem(Reader, UniqueFileItems));
-			}
-			return Items;
-		}
-
-		public static void Write(this BinaryWriter Writer, FileItem FileItem, Dictionary<FileItem, int> UniqueFileItemToIndex)
+		/// <summary>
+		/// Read a file item from a binary archive
+		/// </summary>
+		/// <param name="Reader">Reader to serialize data from</param>
+		/// <returns>Instance of the serialized file item</returns>
+		public static FileItem ReadFileItem(this BinaryArchiveReader Reader)
 		{
-			int Index = -1;
-			if (FileItem != null && !UniqueFileItemToIndex.TryGetValue(FileItem, out Index))
-			{
-				Index = UniqueFileItemToIndex.Count;
-				UniqueFileItemToIndex.Add(FileItem, Index);
-			}
-			Writer.Write(Index);
+			return Reader.ReadObjectReference<FileItem>(() => FileItem.GetItemByFileReference(Reader.ReadFileReference()));
 		}
 
-		public static void Write(this BinaryWriter Writer, List<FileItem> FileItems, Dictionary<FileItem, int> UniqueFileItemToIndex)
+		/// <summary>
+		/// Write a file item to a binary archive
+		/// </summary>
+		/// <param name="Writer">Writer to serialize data to</param>
+		/// <param name="FileItem">File item to write</param>
+		public static void WriteFileItem(this BinaryArchiveWriter Writer, FileItem FileItem)
 		{
-			if(FileItems == null)
-			{
-				Writer.Write(-1);
-			}
-			else
-			{
-				Writer.Write(FileItems.Count);
-				for(int Idx = 0; Idx < FileItems.Count; Idx++)
-				{
-					Write(Writer, FileItems[Idx], UniqueFileItemToIndex);
-				}
-			}
+			Writer.WriteObjectReference<FileItem>(FileItem, () => Writer.WriteFileReference(FileItem.Location));
 		}
 	}
 }
diff --git a/Engine/Source/Programs/UnrealBuildTool/System/UBTMakefile.cs b/Engine/Source/Programs/UnrealBuildTool/System/UBTMakefile.cs
--- a/Engine/Source/Programs/UnrealBuildTool/System/UBTMakefile.cs
+++ b/Engine/Source/Programs/UnrealBuildTool/System/UBTMakefile.cs
@@ -72,63 +72,38 @@ namespace UnrealBuildTool
 		{
 		}
 
-		public UBTMakefile(BinaryReader Reader)
+		public UBTMakefile(BinaryArchiveReader Reader)
 		{
-			int Version = Reader.ReadInt32();
+			int Version = Reader.ReadInt();
 			if(Version != CurrentVersion)
 			{
 				throw new Exception(string.Format("Makefile version does not match - found {0}, expected: {1}", Version, CurrentVersion));
 			}
 
-			long Offset = Reader.ReadInt64();
-			long OriginalOffset = Reader.BaseStream.Position;
-
-			Reader.BaseStream.Seek(Offset, SeekOrigin.Begin);
-			List<FileReference> UniqueFileReferences = Reader.ReadList(() => Reader.ReadFileReference());
-			List<FileItem> UniqueFileItems = UniqueFileReferences.Select(x => FileItem.GetItemByFileReference(x)).ToList();
-			Reader.BaseStream.Seek(OriginalOffset, SeekOrigin.Begin);
-
-			AllActions = Reader.ReadList(() => new Action(Reader, UniqueFileItems));
+			AllActions = Reader.ReadList(() => new Action(Reader));
 			EnvironmentVariables = Reader.ReadList(() => Tuple.Create(Reader.ReadString(), Reader.ReadString()));
-			TargetNameToUObjectModules = Reader.ReadDictionary(() => Reader.ReadString(), () => Reader.ReadList(() => new UHTModuleInfo(Reader, UniqueFileItems)));
-			OutputItemsForAllTargets = Reader.ReadFileItemList(UniqueFileItems);
-			ModuleNameToOutputItems = Reader.ReadDictionary(() => Reader.ReadString(), () => Reader.ReadArray(() => Reader.ReadFileItem(UniqueFileItems)));
-			HotReloadModuleNamesForAllTargets = new HashSet<string>(Reader.ReadStringArray());
+			TargetNameToUObjectModules = Reader.ReadDictionary(() => Reader.ReadString(), () => Reader.ReadList(() => new UHTModuleInfo(Reader)));
+			OutputItemsForAllTargets = Reader.ReadList(() => Reader.ReadFileItem());
+			ModuleNameToOutputItems = Reader.ReadDictionary(() => Reader.ReadString(), () => Reader.ReadArray(() => Reader.ReadFileItem()));
+			HotReloadModuleNamesForAllTargets = Reader.ReadHashSet(() => Reader.ReadString());
 			Targets = Reader.ReadList(() => new UEBuildTarget(Reader));
-			TargetPrerequisites = Reader.ReadList(() => new BuildPrerequisites(Reader, UniqueFileItems));
-			bUseAdaptiveUnityBuild = Reader.ReadBoolean();
+			TargetPrerequisites = Reader.ReadList(() => new BuildPrerequisites(Reader));
+			bUseAdaptiveUnityBuild = Reader.ReadBool();
 		}
 
-
-		public void Write(BinaryWriter Writer)
+		public void Write(BinaryArchiveWriter Writer)
 		{
-			Writer.Write(CurrentVersion);
-
-			long Offset = Writer.BaseStream.Position;
-			Writer.Write((Int64)0);
-			Dictionary<FileItem, int> UniqueFileItemToIndex = new Dictionary<FileItem, int>();
-
-			Writer.Write(AllActions, x => x.Write(Writer, UniqueFileItemToIndex));
-			Writer.Write(EnvironmentVariables, x => { Writer.Write(x.Item1); Writer.Write(x.Item2); });
-			Writer.Write(TargetNameToUObjectModules, x => Writer.Write(x), v => Writer.Write(v, e => e.Write(Writer, UniqueFileItemToIndex)));
-			Writer.Write(OutputItemsForAllTargets, UniqueFileItemToIndex);
-			Writer.Write(ModuleNameToOutputItems, k => Writer.Write(k), v => Writer.Write(v, e => Writer.Write(e, UniqueFileItemToIndex)));
-			Writer.Write(HotReloadModuleNamesForAllTargets.ToArray(), x => Writer.Write(x));
-			Writer.Write(Targets, x => x.Write(Writer));
-			Writer.Write(TargetPrerequisites, x => x.Write(Writer, UniqueFileItemToIndex));
-			Writer.Write(bUseAdaptiveUnityBuild);
-			Writer.Flush();
-
-			long FileItemTableOffset = Writer.BaseStream.Position;
-			FileReference[] UniqueFileReferences = new FileReference[UniqueFileItemToIndex.Count];
-			foreach(KeyValuePair<FileItem, int> Pair in UniqueFileItemToIndex)
-			{
-				UniqueFileReferences[Pair.Value] = Pair.Key.Location;
-			}
-			Writer.Write(UniqueFileReferences, x => Writer.Write(x));
-
-			Writer.BaseStream.Seek(Offset, SeekOrigin.Begin);
-			Writer.Write(FileItemTableOffset);
+			Writer.WriteInt(CurrentVersion);
+
+			Writer.WriteList(AllActions, Action => Action.Write(Writer));
+			Writer.WriteList(EnvironmentVariables, x => { Writer.WriteString(x.Item1); Writer.WriteString(x.Item2); });
+			Writer.WriteDictionary(TargetNameToUObjectModules, x => Writer.WriteString(x), v => Writer.WriteList(v, e => e.Write(Writer)));
+			Writer.WriteList(OutputItemsForAllTargets, Item => Writer.WriteFileItem(Item));
+			Writer.WriteDictionary(ModuleNameToOutputItems, k => Writer.WriteString(k), v => Writer.WriteArray(v, e => Writer.WriteFileItem(e)));
+			Writer.WriteHashSet(HotReloadModuleNamesForAllTargets, x => Writer.WriteString(x));
+			Writer.WriteList(Targets, x => x.Write(Writer));
+			Writer.WriteList(TargetPrerequisites, x => x.Write(Writer));
+			Writer.WriteBool(bUseAdaptiveUnityBuild);
 			Writer.Flush();
 		}
 
@@ -169,12 +144,9 @@ namespace UnrealBuildTool
 			try
 			{
 				Directory.CreateDirectory(Path.GetDirectoryName(UBTMakefileItem.AbsolutePath));
-				using (FileStream Stream = new FileStream(UBTMakefileItem.AbsolutePath, FileMode.Create, FileAccess.Write))
+				using(BinaryArchiveWriter Writer = new BinaryArchiveWriter(UBTMakefileItem.Location))
 				{
-					using(BinaryWriter Writer = new BinaryWriter(Stream, Encoding.UTF8, true))
-					{
-						UBTMakefile.Write(Writer);
-					}
+					UBTMakefile.Write(Writer);
 				}
 			}
 			catch (Exception Ex)
@@ -284,155 +256,156 @@ namespace UnrealBuildTool
 				}
 			}
 
-			UBTMakefile LoadedUBTMakefile = null;
-
-			try
+			UBTMakefile LoadedUBTMakefile;
+			using(Timeline.ScopeEvent("Loading makefile"))
 			{
-				using(Timeline.ScopeEvent("Loading makefile"))
+				try
 				{
-					using (FileStream Stream = new FileStream(UBTMakefileInfo.FullName, FileMode.Open, FileAccess.Read))
+					using(BinaryArchiveReader Reader = new BinaryArchiveReader(new FileReference(UBTMakefileInfo)))
 					{
-						BinaryReader Reader = new BinaryReader(Stream);
 						LoadedUBTMakefile = new UBTMakefile(Reader);
 					}
 				}
-			}
-			catch (Exception Ex)
-			{
-				Log.TraceWarning("Failed to read makefile: {0}", Ex.Message);
-				Log.TraceLog("Exception: {0}", Ex.ToString());
-				ReasonNotLoaded = "couldn't read existing makefile";
-				return null;
+				catch (Exception Ex)
+				{
+					Log.TraceWarning("Failed to read makefile: {0}", Ex.Message);
+					Log.TraceLog("Exception: {0}", Ex.ToString());
+					ReasonNotLoaded = "couldn't read existing makefile";
+					return null;
+				}
 			}
 
-			if (!LoadedUBTMakefile.IsValidMakefile())
+			using(Timeline.ScopeEvent("Checking makefile validity"))
 			{
-				Log.TraceWarning("Loaded makefile appears to have invalid contents, ignoring it ({0})", UBTMakefileInfo.FullName);
-				ReasonNotLoaded = "existing makefile appears to be invalid";
-				return null;
-			}
+				if (!LoadedUBTMakefile.IsValidMakefile())
+				{
+					Log.TraceWarning("Loaded makefile appears to have invalid contents, ignoring it ({0})", UBTMakefileInfo.FullName);
+					ReasonNotLoaded = "existing makefile appears to be invalid";
+					return null;
+				}
 
-			// Check if any of the target's Build.cs files are newer than the makefile
-			foreach (UEBuildTarget Target in LoadedUBTMakefile.Targets)
-			{
-				// Check if ini files are newer. Ini files contain build settings too.
-				DirectoryReference ProjectDirectory = DirectoryReference.FromFile(ProjectFile);
-				foreach (ConfigHierarchyType IniType in (ConfigHierarchyType[])Enum.GetValues(typeof(ConfigHierarchyType)))
+				// Check if any of the target's Build.cs files are newer than the makefile
+				foreach (UEBuildTarget Target in LoadedUBTMakefile.Targets)
 				{
-					foreach (FileReference IniFilename in ConfigHierarchy.EnumerateConfigFileLocations(IniType, ProjectDirectory, Target.Platform))
+					// Check if ini files are newer. Ini files contain build settings too.
+					DirectoryReference ProjectDirectory = DirectoryReference.FromFile(ProjectFile);
+					foreach (ConfigHierarchyType IniType in (ConfigHierarchyType[])Enum.GetValues(typeof(ConfigHierarchyType)))
 					{
-						FileInfo IniFileInfo = new FileInfo(IniFilename.FullName);
-						if (UBTMakefileInfo.LastWriteTime.CompareTo(IniFileInfo.LastWriteTime) < 0)
+						foreach (FileReference IniFilename in ConfigHierarchy.EnumerateConfigFileLocations(IniType, ProjectDirectory, Target.Platform))
 						{
-							// Ini files are newer than UBTMakefile
-							ReasonNotLoaded = "ini files are newer than UBTMakefile";
-							return null;
+							FileInfo IniFileInfo = new FileInfo(IniFilename.FullName);
+							if (UBTMakefileInfo.LastWriteTime.CompareTo(IniFileInfo.LastWriteTime) < 0)
+							{
+								// Ini files are newer than UBTMakefile
+								ReasonNotLoaded = "ini files are newer than UBTMakefile";
+								return null;
+							}
 						}
 					}
 				}
-			}
 
-			foreach (BuildPrerequisites Prerequisites in LoadedUBTMakefile.TargetPrerequisites)
-			{
-				foreach(DirectoryReference SourceDirectory in Prerequisites.SourceDirectories)
+				foreach (BuildPrerequisites Prerequisites in LoadedUBTMakefile.TargetPrerequisites)
 				{
-					DirectoryInfo SourceDirectoryInfo = new DirectoryInfo(SourceDirectory.FullName);
-					if(!SourceDirectoryInfo.Exists || SourceDirectoryInfo.LastWriteTimeUtc > UBTMakefileInfo.LastWriteTimeUtc)
+					foreach(DirectoryReference SourceDirectory in Prerequisites.SourceDirectories)
 					{
-						Log.TraceLog("Timestamp of {0} ({1}) is newer than makefile ({2})", SourceDirectory, SourceDirectoryInfo.LastWriteTimeUtc, UBTMakefileInfo.LastWriteTimeUtc);
-						ReasonNotLoaded = "source directory changed";
-						return null;
+						DirectoryInfo SourceDirectoryInfo = new DirectoryInfo(SourceDirectory.FullName);
+						if(!SourceDirectoryInfo.Exists || SourceDirectoryInfo.LastWriteTimeUtc > UBTMakefileInfo.LastWriteTimeUtc)
+						{
+							Log.TraceLog("Timestamp of {0} ({1}) is newer than makefile ({2})", SourceDirectory, SourceDirectoryInfo.LastWriteTimeUtc, UBTMakefileInfo.LastWriteTimeUtc);
+							ReasonNotLoaded = "source directory changed";
+							return null;
+						}
 					}
-				}
 
-				foreach(FileItem AdditionalDependency in Prerequisites.AdditionalDependencies)
-				{
-					if (!AdditionalDependency.Exists)
-					{
-						Log.TraceLog("{0} has been deleted since makefile was built.", AdditionalDependency.Location);
-						ReasonNotLoaded = string.Format("{0} deleted", AdditionalDependency.Location.GetFileName());
-						return null;
-					}
-					if(AdditionalDependency.LastWriteTimeUtc > UBTMakefileInfo.LastWriteTime)
+					foreach(FileItem AdditionalDependency in Prerequisites.AdditionalDependencies)
 					{
-						Log.TraceLog("{0} has been modified since makefile was built.", AdditionalDependency.Location);
-						ReasonNotLoaded = string.Format("{0} modified", AdditionalDependency.Location.GetFileName());
-						return null;
+						if (!AdditionalDependency.Exists)
+						{
+							Log.TraceLog("{0} has been deleted since makefile was built.", AdditionalDependency.Location);
+							ReasonNotLoaded = string.Format("{0} deleted", AdditionalDependency.Location.GetFileName());
+							return null;
+						}
+						if(AdditionalDependency.LastWriteTimeUtc > UBTMakefileInfo.LastWriteTime)
+						{
+							Log.TraceLog("{0} has been modified since makefile was built.", AdditionalDependency.Location);
+							ReasonNotLoaded = string.Format("{0} modified", AdditionalDependency.Location.GetFileName());
+							return null;
+						}
 					}
-				}
 
-				// We do a check to see if any modules' headers have changed which have
-				// acquired or lost UHT types.  If so, which should be rare,
-				// we'll just invalidate the entire makefile and force it to be rebuilt.
-
-				// Get all H files in processed modules newer than the makefile itself
-				HashSet<FileReference> HFilesNewerThanMakefile =
-					new HashSet<FileReference>(
-						Prerequisites.UObjectModuleHeaders
-						.SelectMany(x => x.SourceFolder != null ? DirectoryReference.EnumerateFiles(x.SourceFolder, "*.h", SearchOption.AllDirectories) : Enumerable.Empty<FileReference>())
-						.Where(y => FileReference.GetLastWriteTimeUtc(y) > UBTMakefileInfo.LastWriteTimeUtc)
-						.OrderBy(z => z).Distinct()
-					);
-
-				// Get all H files in all modules processed in the last makefile build
-				HashSet<FileReference> AllUHTHeaders = new HashSet<FileReference>(Prerequisites.UObjectModuleHeaders.SelectMany(x => x.HeaderFiles).Select(x => x.Location));
-
-				// Check whether any headers have been deleted. If they have, we need to regenerate the makefile since the module might now be empty. If we don't,
-				// and the file has been moved to a different module, we may include stale generated headers.
-				foreach (FileReference FileName in AllUHTHeaders)
-				{
-					if (!FileReference.Exists(FileName))
+					// We do a check to see if any modules' headers have changed which have
+					// acquired or lost UHT types.  If so, which should be rare,
+					// we'll just invalidate the entire makefile and force it to be rebuilt.
+
+					// Get all H files in processed modules newer than the makefile itself
+					HashSet<FileReference> HFilesNewerThanMakefile =
+						new HashSet<FileReference>(
+							Prerequisites.UObjectModuleHeaders
+							.SelectMany(x => x.SourceFolder != null ? DirectoryReference.EnumerateFiles(x.SourceFolder, "*.h", SearchOption.AllDirectories) : Enumerable.Empty<FileReference>())
+							.Where(y => FileReference.GetLastWriteTimeUtc(y) > UBTMakefileInfo.LastWriteTimeUtc)
+							.OrderBy(z => z).Distinct()
+						);
+
+					// Get all H files in all modules processed in the last makefile build
+					HashSet<FileReference> AllUHTHeaders = new HashSet<FileReference>(Prerequisites.UObjectModuleHeaders.SelectMany(x => x.HeaderFiles).Select(x => x.Location));
+
+					// Check whether any headers have been deleted. If they have, we need to regenerate the makefile since the module might now be empty. If we don't,
+					// and the file has been moved to a different module, we may include stale generated headers.
+					foreach (FileReference FileName in AllUHTHeaders)
 					{
-						Log.TraceLog("File processed by UHT was deleted ({0}); invalidating makefile", FileName);
-						ReasonNotLoaded = string.Format("UHT file was deleted");
-						return null;
+						if (!FileReference.Exists(FileName))
+						{
+							Log.TraceLog("File processed by UHT was deleted ({0}); invalidating makefile", FileName);
+							ReasonNotLoaded = string.Format("UHT file was deleted");
+							return null;
+						}
 					}
-				}
 
-				// Makefile is invalid if:
-				// * There are any newer files which contain no UHT data, but were previously in the makefile
-				// * There are any newer files contain data which needs processing by UHT, but weren't not previously in the makefile
-				foreach (FileReference Filename in HFilesNewerThanMakefile)
-				{
-					bool bContainsUHTData = CPPHeaders.DoesFileContainUObjects(Filename.FullName);
-					bool bWasProcessed = AllUHTHeaders.Contains(Filename);
-					if (bContainsUHTData != bWasProcessed)
+					// Makefile is invalid if:
+					// * There are any newer files which contain no UHT data, but were previously in the makefile
+					// * There are any newer files contain data which needs processing by UHT, but weren't not previously in the makefile
+					foreach (FileReference Filename in HFilesNewerThanMakefile)
 					{
-						Log.TraceLog("{0} {1} contain UHT types and now {2} , ignoring it ({3})", Filename, bWasProcessed ? "used to" : "didn't", bWasProcessed ? "doesn't" : "does", UBTMakefileInfo.FullName);
-						ReasonNotLoaded = string.Format("new files with reflected types");
-						return null;
+						bool bContainsUHTData = CPPHeaders.DoesFileContainUObjects(Filename.FullName);
+						bool bWasProcessed = AllUHTHeaders.Contains(Filename);
+						if (bContainsUHTData != bWasProcessed)
+						{
+							Log.TraceLog("{0} {1} contain UHT types and now {2} , ignoring it ({3})", Filename, bWasProcessed ? "used to" : "didn't", bWasProcessed ? "doesn't" : "does", UBTMakefileInfo.FullName);
+							ReasonNotLoaded = string.Format("new files with reflected types");
+							return null;
+						}
 					}
 				}
-			}
 
-			// If adaptive unity build is enabled, do a check to see if there are any source files that became part of the
-			// working set since the Makefile was created (or, source files were removed from the working set.)  If anything
-			// changed, then we'll force a new Makefile to be created so that we have fresh unity build blobs.  We always
-			// want to make sure that source files in the working set are excluded from those unity blobs (for fastest possible
-			// iteration times.)
-			if (LoadedUBTMakefile.bUseAdaptiveUnityBuild)
-			{
-				foreach(BuildPrerequisites Prerequisites in LoadedUBTMakefile.TargetPrerequisites)
+				// If adaptive unity build is enabled, do a check to see if there are any source files that became part of the
+				// working set since the Makefile was created (or, source files were removed from the working set.)  If anything
+				// changed, then we'll force a new Makefile to be created so that we have fresh unity build blobs.  We always
+				// want to make sure that source files in the working set are excluded from those unity blobs (for fastest possible
+				// iteration times.)
+				if (LoadedUBTMakefile.bUseAdaptiveUnityBuild)
 				{
-					// Check if any source files in the working set no longer belong in it
-					foreach (FileItem SourceFile in Prerequisites.WorkingSet)
+					foreach(BuildPrerequisites Prerequisites in LoadedUBTMakefile.TargetPrerequisites)
 					{
-						if (!WorkingSet.Contains(SourceFile.Location) && File.GetLastWriteTimeUtc(SourceFile.AbsolutePath) > UBTMakefileInfo.LastWriteTimeUtc)
+						// Check if any source files in the working set no longer belong in it
+						foreach (FileItem SourceFile in Prerequisites.WorkingSet)
 						{
-							Log.TraceLog("{0} was part of source working set and now is not; invalidating makefile ({1})", SourceFile.AbsolutePath, UBTMakefileInfo.FullName);
-							ReasonNotLoaded = string.Format("working set of source files changed");
-							return null;
+							if (!WorkingSet.Contains(SourceFile.Location) && File.GetLastWriteTimeUtc(SourceFile.AbsolutePath) > UBTMakefileInfo.LastWriteTimeUtc)
+							{
+								Log.TraceLog("{0} was part of source working set and now is not; invalidating makefile ({1})", SourceFile.AbsolutePath, UBTMakefileInfo.FullName);
+								ReasonNotLoaded = string.Format("working set of source files changed");
+								return null;
+							}
 						}
-					}
 
-					// Check if any source files that are eligible for being in the working set have been modified
-					foreach (FileItem SourceFile in Prerequisites.CandidatesForWorkingSet)
-					{
-						if (WorkingSet.Contains(SourceFile.Location) && File.GetLastWriteTimeUtc(SourceFile.AbsolutePath) > UBTMakefileInfo.LastWriteTimeUtc)
+						// Check if any source files that are eligible for being in the working set have been modified
+						foreach (FileItem SourceFile in Prerequisites.CandidatesForWorkingSet)
 						{
-							Log.TraceLog("{0} was part of source working set and now is not; invalidating makefile ({1})", SourceFile.AbsolutePath, UBTMakefileInfo.FullName);
-							ReasonNotLoaded = string.Format("working set of source files changed");
-							return null;
+							if (WorkingSet.Contains(SourceFile.Location) && File.GetLastWriteTimeUtc(SourceFile.AbsolutePath) > UBTMakefileInfo.LastWriteTimeUtc)
+							{
+								Log.TraceLog("{0} was part of source working set and now is not; invalidating makefile ({1})", SourceFile.AbsolutePath, UBTMakefileInfo.FullName);
+								ReasonNotLoaded = string.Format("working set of source files changed");
+								return null;
+							}
 						}
 					}
 				}
