UBT: Prevent failures to read dependency cache from causing build failures.

#rnx

diff --git a/Engine/Source/Programs/UnrealBuildTool/System/CppDependencyCache.cs b/Engine/Source/Programs/UnrealBuildTool/System/CppDependencyCache.cs
--- a/Engine/Source/Programs/UnrealBuildTool/System/CppDependencyCache.cs
+++ b/Engine/Source/Programs/UnrealBuildTool/System/CppDependencyCache.cs
@@ -223,9 +223,9 @@ namespace UnrealBuildTool
 		/// </summary>
 		private void Read()
 		{
-			using(FileStream Stream = File.Open(Location.FullName, FileMode.Open, FileAccess.Read, FileShare.Read))
+			try
 			{
-				using(BinaryArchiveReader Reader = new BinaryArchiveReader(Stream))
+				using(BinaryArchiveReader Reader = new BinaryArchiveReader(Location))
 				{
 					int Version = Reader.ReadInt();
 					if(Version != CurrentVersion)
@@ -242,6 +242,11 @@ namespace UnrealBuildTool
 					}
 				}
 			}
+			catch(Exception Ex)
+			{
+				Log.TraceWarning("Unable to read {0}. See log for additional information.", Location);
+				Log.TraceLog("{0}", ExceptionUtils.FormatExceptionDetails(Ex));
+			}
 		}
 
 		/// <summary>
