UBT: Add a fast path for searching for running editors to hot-reload into on Windows.

Enumerating all processes on the system can take ~2s during every UBT invocation; only querying specific processes as necessary is significantly quicker. Mac and Linux continue to use the old path, due to custom logic for enumerating processes,

diff --git a/Engine/Source/Programs/UnrealBuildTool/System/HotReload.cs b/Engine/Source/Programs/UnrealBuildTool/System/HotReload.cs
--- a/Engine/Source/Programs/UnrealBuildTool/System/HotReload.cs
+++ b/Engine/Source/Programs/UnrealBuildTool/System/HotReload.cs
@@ -2,6 +2,7 @@
 
 using System;
 using System.Collections.Generic;
+using System.Diagnostics;
 using System.IO;
 using System.Linq;
 using System.Text;
@@ -169,32 +170,67 @@ namespace UnrealBuildTool
 
 				using(Timeline.ScopeEvent("Finding editor processes for hot-reload"))
 				{
-					BuildHostPlatform.ProcessInfo[] Processes = BuildHostPlatform.Current.GetProcesses();
-					string EditorRunsDir = Path.Combine(UnrealBuildTool.EngineDirectory.FullName, "Intermediate", "EditorRuns");
-
-					if (!Directory.Exists(EditorRunsDir))
+					DirectoryReference EditorRunsDir = DirectoryReference.Combine(UnrealBuildTool.EngineDirectory, "Intermediate", "EditorRuns");
+					if (!DirectoryReference.Exists(EditorRunsDir))
 					{
 						return false;
 					}
 
-					FileInfo[] EditorRunsFiles = new DirectoryInfo(EditorRunsDir).GetFiles();
-
-					foreach (FileInfo File in EditorRunsFiles)
+					if(BuildHostPlatform.Current.Platform == UnrealTargetPlatform.Win64)
 					{
-						int PID;
-						BuildHostPlatform.ProcessInfo Proc = null;
-						if (!Int32.TryParse(File.Name, out PID) || (Proc = Processes.FirstOrDefault(P => P.PID == PID)) == default(BuildHostPlatform.ProcessInfo))
+						foreach(FileReference EditorInstanceFile in DirectoryReference.EnumerateFiles(EditorRunsDir))
 						{
-							// Delete stale files (it may happen if editor crashes).
-							File.Delete();
-							continue;
+							int ProcessId;
+							if(!Int32.TryParse(EditorInstanceFile.GetFileName(), out ProcessId))
+							{
+								FileReference.Delete(EditorInstanceFile);
+								continue;
+							}
+
+							Process RunningProcess;
+							try
+							{
+								RunningProcess = Process.GetProcessById(ProcessId);
+							}
+							catch
+							{
+								RunningProcess = null;
+							}
+
+							if(RunningProcess == null)
+							{
+								FileReference.Delete(EditorInstanceFile);
+								continue;
+							}
+
+							if(!bIsRunning && EditorLocation == new FileReference(RunningProcess.MainModule.FileName))
+							{
+								bIsRunning = true;
+							}
 						}
+					}
+					else
+					{
+						FileInfo[] EditorRunsFiles = new DirectoryInfo(EditorRunsDir.FullName).GetFiles();
+						BuildHostPlatform.ProcessInfo[] Processes = BuildHostPlatform.Current.GetProcesses();
 
-						// Don't break here to allow clean-up of other stale files.
-						if (!bIsRunning)
+						foreach (FileInfo File in EditorRunsFiles)
 						{
-							// Otherwise check if the path matches.
-							bIsRunning = new FileReference(Proc.Filename) == EditorLocation;
+							int PID;
+							BuildHostPlatform.ProcessInfo Proc = null;
+							if (!Int32.TryParse(File.Name, out PID) || (Proc = Processes.FirstOrDefault(P => P.PID == PID)) == default(BuildHostPlatform.ProcessInfo))
+							{
+								// Delete stale files (it may happen if editor crashes).
+								File.Delete();
+								continue;
+							}
+
+							// Don't break here to allow clean-up of other stale files.
+							if (!bIsRunning)
+							{
+								// Otherwise check if the path matches.
+								bIsRunning = new FileReference(Proc.Filename) == EditorLocation;
+							}
 						}
 					}
 				}
