UBT: Rename UBTMakefile to TargetMakefile.

#rnx

diff --git a/Engine/Source/Programs/UnrealBuildTool/Modes/BuildMode.cs b/Engine/Source/Programs/UnrealBuildTool/Modes/BuildMode.cs
--- a/Engine/Source/Programs/UnrealBuildTool/Modes/BuildMode.cs
+++ b/Engine/Source/Programs/UnrealBuildTool/Modes/BuildMode.cs
@@ -226,9 +226,9 @@ namespace UnrealBuildTool
 				// Create the working set provider
 				using (ISourceFileWorkingSet WorkingSet = SourceFileWorkingSet.Create(UnrealBuildTool.RootDirectory, ProjectDirs))
 				{
-					UBTMakefile Makefile = null;
+					TargetMakefile Makefile = null;
 					{
-						FileReference UBTMakefilePath = UBTMakefile.GetUBTMakefilePath(TargetDescs);
+						FileReference UBTMakefilePath = TargetMakefile.GetUBTMakefilePath(TargetDescs);
 
 						// Make sure the gather phase is executed if we're not actually building anything
 						if (BuildConfiguration.bGenerateManifest || BuildConfiguration.bXGEExport || GeneratingActionGraph)
@@ -247,7 +247,7 @@ namespace UnrealBuildTool
 							string ReasonNotLoaded;
 							using(Timeline.ScopeEvent("UBTMakefile.LoadUBTMakefile"))
 							{
-								Makefile = UBTMakefile.Load(UBTMakefilePath, TargetDescs[0].ProjectFile, WorkingSet, out ReasonNotLoaded);
+								Makefile = TargetMakefile.Load(UBTMakefilePath, TargetDescs[0].ProjectFile, WorkingSet, out ReasonNotLoaded);
 							}
 
 							if (Makefile == null)
@@ -369,7 +369,7 @@ namespace UnrealBuildTool
 						{
 							ActionGraph.FinalizeActionGraph();
 
-							Makefile = new UBTMakefile();
+							Makefile = new TargetMakefile();
 							Makefile.AllActions = ActionGraph.AllActions;
 							Makefile.OutputItemsForAllTargets = OutputItemsForAllTargets;
 							foreach (System.Collections.DictionaryEntry EnvironmentVariable in Environment.GetEnvironmentVariables())
@@ -391,7 +391,7 @@ namespace UnrealBuildTool
 								// @todo ubtmake: Optimization: We could make 'gather + assemble' mode slightly faster by saving this while busy compiling (on our worker thread)
 								using(Timeline.ScopeEvent("UBTMakefile.SaveUBTMakefile()"))
 								{
-									UBTMakefile.SaveUBTMakefile(TargetDescs, Makefile);
+									TargetMakefile.SaveUBTMakefile(TargetDescs, Makefile);
 								}
 							}
 						}
@@ -813,7 +813,7 @@ namespace UnrealBuildTool
 		private static void InvalidateMakefiles(TargetDescriptor Target)
 		{
 			string[] MakefileNames = new string[] { "HotReloadMakefile.ubt", "Makefile.ubt" };
-			DirectoryReference BaseDir = UBTMakefile.GetUBTMakefileDirectoryPathForSingleTarget(Target);
+			DirectoryReference BaseDir = TargetMakefile.GetUBTMakefileDirectoryPathForSingleTarget(Target);
 
 			foreach (string MakefileName in MakefileNames)
 			{
diff --git a/Engine/Source/Programs/UnrealBuildTool/Modes/CleanMode.cs b/Engine/Source/Programs/UnrealBuildTool/Modes/CleanMode.cs
--- a/Engine/Source/Programs/UnrealBuildTool/Modes/CleanMode.cs
+++ b/Engine/Source/Programs/UnrealBuildTool/Modes/CleanMode.cs
@@ -128,8 +128,8 @@ namespace UnrealBuildTool
 
 				// Add all the makefiles and caches to be deleted
 				FilesToDelete.Add(DependencyCache.GetDependencyCachePathForTarget(Target.ProjectFile, Target.Platform, Target.Name));
-				FilesToDelete.Add(UBTMakefile.GetUBTMakefilePath(Target.ProjectFile, Target.Platform, Target.Configuration, Target.Name, false));
-				FilesToDelete.Add(UBTMakefile.GetUBTMakefilePath(Target.ProjectFile, Target.Platform, Target.Configuration, Target.Name, true));
+				FilesToDelete.Add(TargetMakefile.GetUBTMakefilePath(Target.ProjectFile, Target.Platform, Target.Configuration, Target.Name, false));
+				FilesToDelete.Add(TargetMakefile.GetUBTMakefilePath(Target.ProjectFile, Target.Platform, Target.Configuration, Target.Name, true));
 				FilesToDelete.Add(ActionHistory.GeneratePathForTarget(Target.ProjectFile, Target.Name, Target.Platform, Target.Architecture, (Target.BuildEnvironment == TargetBuildEnvironment.Shared)));
 
 				// Add all the intermediate folders to be deleted
diff --git a/Engine/Source/Programs/UnrealBuildTool/ProjectFiles/ProjectFileGenerator.cs b/Engine/Source/Programs/UnrealBuildTool/ProjectFiles/ProjectFileGenerator.cs
--- a/Engine/Source/Programs/UnrealBuildTool/ProjectFiles/ProjectFileGenerator.cs
+++ b/Engine/Source/Programs/UnrealBuildTool/ProjectFiles/ProjectFileGenerator.cs
@@ -1731,7 +1731,7 @@ namespace UnrealBuildTool
 							// If we're generating project files, then go ahead and wipe out the existing UBTMakefile for every target, to make sure that
 							// it gets a full dependency scan next time.
 							// NOTE: This is just a safeguard and doesn't have to be perfect.  We also check for newer project file timestamps in LoadUBTMakefile()
-							FileReference UBTMakefilePath = UBTMakefile.GetUBTMakefilePath(new List<TargetDescriptor>{ TargetDesc });
+							FileReference UBTMakefilePath = TargetMakefile.GetUBTMakefilePath(new List<TargetDescriptor>{ TargetDesc });
 							if (FileReference.Exists(UBTMakefilePath))
 							{
 								FileReference.Delete(UBTMakefilePath);
diff --git a/Engine/Source/Programs/UnrealBuildTool/System/UBTMakefile.cs b/Engine/Source/Programs/UnrealBuildTool/System/UBTMakefile.cs
--- a/Engine/Source/Programs/UnrealBuildTool/System/UBTMakefile.cs
+++ b/Engine/Source/Programs/UnrealBuildTool/System/UBTMakefile.cs
@@ -15,7 +15,7 @@ namespace UnrealBuildTool
 	/// <summary>
 	/// A special Makefile that UBT is able to create in "-gather" mode, then load in "-assemble" mode to accelerate iterative compiling and linking
 	/// </summary>
-	class UBTMakefile
+	class TargetMakefile
 	{
 		/// <summary>
 		/// The version number to write
@@ -68,11 +68,11 @@ namespace UnrealBuildTool
 		/// </summary>
 		public List<BuildPrerequisites> TargetPrerequisites;
 
-		public UBTMakefile()
+		public TargetMakefile()
 		{
 		}
 
-		public UBTMakefile(BinaryArchiveReader Reader)
+		public TargetMakefile(BinaryArchiveReader Reader)
 		{
 			int Version = Reader.ReadInt();
 			if(Version != CurrentVersion)
@@ -128,7 +128,7 @@ namespace UnrealBuildTool
 		/// </summary>
 		/// <param name="TargetDescs">List of targets.  Order is not important</param>
 		/// <param name="UBTMakefile">The UBT makefile</param>
-		public static void SaveUBTMakefile(List<TargetDescriptor> TargetDescs, UBTMakefile UBTMakefile)
+		public static void SaveUBTMakefile(List<TargetDescriptor> TargetDescs, TargetMakefile UBTMakefile)
 		{
 			if (!UBTMakefile.IsValidMakefile())
 			{
@@ -164,7 +164,7 @@ namespace UnrealBuildTool
 		/// <param name="ReasonNotLoaded">If the function returns null, this string will contain the reason why</param>
 		/// <param name="WorkingSet">Interface to query which source files are in the working set</param>
 		/// <returns>The loaded makefile, or null if it failed for some reason.  On failure, the 'ReasonNotLoaded' variable will contain information about why</returns>
-		public static UBTMakefile Load(FileReference MakefilePath, FileReference ProjectFile, ISourceFileWorkingSet WorkingSet, out string ReasonNotLoaded)
+		public static TargetMakefile Load(FileReference MakefilePath, FileReference ProjectFile, ISourceFileWorkingSet WorkingSet, out string ReasonNotLoaded)
 		{
 			// Check the directory timestamp on the project files directory.  If the user has generated project files more
 			// recently than the UBTMakefile, then we need to consider the file to be out of date
@@ -256,14 +256,14 @@ namespace UnrealBuildTool
 				}
 			}
 
-			UBTMakefile LoadedUBTMakefile;
+			TargetMakefile LoadedUBTMakefile;
 			using(Timeline.ScopeEvent("Loading makefile"))
 			{
 				try
 				{
 					using(BinaryArchiveReader Reader = new BinaryArchiveReader(new FileReference(UBTMakefileInfo)))
 					{
-						LoadedUBTMakefile = new UBTMakefile(Reader);
+						LoadedUBTMakefile = new TargetMakefile(Reader);
 					}
 				}
 				catch (Exception Ex)
