UBT: Rename UBTMakefile.cs to TargetMakefile.cs.

#rnx

diff --git a/Engine/Source/Programs/UnrealBuildTool/System/TargetMakefile.cs b/Engine/Source/Programs/UnrealBuildTool/System/TargetMakefile.cs
new file mode 100644
--- /dev/null
+++ b/Engine/Source/Programs/UnrealBuildTool/System/TargetMakefile.cs
@@ -0,0 +1,524 @@
+// Copyright 1998-2019 Epic Games, Inc. All Rights Reserved.
+using System;
+using System.Collections.Generic;
+using System.IO;
+using System.Linq;
+using System.Reflection;
+using System.Runtime.Serialization;
+using System.Runtime.Serialization.Formatters.Binary;
+using System.Text;
+using System.Threading.Tasks;
+using Tools.DotNETCommon;
+
+namespace UnrealBuildTool
+{
+	/// <summary>
+	/// A special Makefile that UBT is able to create in "-gather" mode, then load in "-assemble" mode to accelerate iterative compiling and linking
+	/// </summary>
+	class TargetMakefile
+	{
+		/// <summary>
+		/// The version number to write
+		/// </summary>
+		public const int CurrentVersion = 8;
+
+		/// <summary>
+		/// Every action in the action graph
+		/// </summary>
+		public List<Action> AllActions;
+
+		/// <summary>
+		/// Environment variables that we'll need in order to invoke the platform's compiler and linker
+		/// </summary>
+		// @todo ubtmake: Really we want to allow a different set of environment variables for every Action.  This would allow for targets on multiple platforms to be built in a single assembling phase.  We'd only have unique variables for each platform that has actions, so we'd want to make sure we only store the minimum set.
+		public readonly List<Tuple<string, string>> EnvironmentVariables = new List<Tuple<string, string>>();
+
+		/// <summary>
+		/// Maps each target to a list of UObject module info structures
+		/// </summary>
+		public Dictionary<string, List<UHTModuleInfo>> TargetNameToUObjectModules;
+
+		/// <summary>
+		/// The final output items for all target
+		/// </summary>
+		public List<FileItem> OutputItemsForAllTargets;
+
+		/// <summary>
+		/// Maps module names to output items
+		/// </summary>
+		public Dictionary<string, FileItem[]> ModuleNameToOutputItems;
+
+		/// <summary>
+		/// List of game module names, for hot-reload
+		/// </summary>
+		public HashSet<string> HotReloadModuleNamesForAllTargets;
+
+		/// <summary>
+		/// List of targets being built
+		/// </summary>
+		public List<UEBuildTarget> Targets;
+
+		/// <summary>
+		/// Whether adaptive unity build is enabled for any of these targets
+		/// </summary>
+		public bool bUseAdaptiveUnityBuild;
+
+		/// <summary>
+		/// List of predicates for the action graph to be valid.
+		/// </summary>
+		public List<BuildPrerequisites> TargetPrerequisites;
+
+		public TargetMakefile()
+		{
+		}
+
+		public TargetMakefile(BinaryArchiveReader Reader)
+		{
+			int Version = Reader.ReadInt();
+			if(Version != CurrentVersion)
+			{
+				throw new Exception(string.Format("Makefile version does not match - found {0}, expected: {1}", Version, CurrentVersion));
+			}
+
+			AllActions = Reader.ReadList(() => new Action(Reader));
+			EnvironmentVariables = Reader.ReadList(() => Tuple.Create(Reader.ReadString(), Reader.ReadString()));
+			TargetNameToUObjectModules = Reader.ReadDictionary(() => Reader.ReadString(), () => Reader.ReadList(() => new UHTModuleInfo(Reader)));
+			OutputItemsForAllTargets = Reader.ReadList(() => Reader.ReadFileItem());
+			ModuleNameToOutputItems = Reader.ReadDictionary(() => Reader.ReadString(), () => Reader.ReadArray(() => Reader.ReadFileItem()));
+			HotReloadModuleNamesForAllTargets = Reader.ReadHashSet(() => Reader.ReadString());
+			Targets = Reader.ReadList(() => new UEBuildTarget(Reader));
+			TargetPrerequisites = Reader.ReadList(() => new BuildPrerequisites(Reader));
+			bUseAdaptiveUnityBuild = Reader.ReadBool();
+		}
+
+		public void Write(BinaryArchiveWriter Writer)
+		{
+			Writer.WriteInt(CurrentVersion);
+
+			Writer.WriteList(AllActions, Action => Action.Write(Writer));
+			Writer.WriteList(EnvironmentVariables, x => { Writer.WriteString(x.Item1); Writer.WriteString(x.Item2); });
+			Writer.WriteDictionary(TargetNameToUObjectModules, x => Writer.WriteString(x), v => Writer.WriteList(v, e => e.Write(Writer)));
+			Writer.WriteList(OutputItemsForAllTargets, Item => Writer.WriteFileItem(Item));
+			Writer.WriteDictionary(ModuleNameToOutputItems, k => Writer.WriteString(k), v => Writer.WriteArray(v, e => Writer.WriteFileItem(e)));
+			Writer.WriteHashSet(HotReloadModuleNamesForAllTargets, x => Writer.WriteString(x));
+			Writer.WriteList(Targets, x => x.Write(Writer));
+			Writer.WriteList(TargetPrerequisites, x => x.Write(Writer));
+			Writer.WriteBool(bUseAdaptiveUnityBuild);
+			Writer.Flush();
+		}
+
+
+		/// <returns> True if this makefile's contents look valid.  Called after loading the file to make sure it is legit.</returns>
+		public bool IsValidMakefile()
+		{
+			return
+				AllActions != null && AllActions.Count > 0 &&
+				EnvironmentVariables != null &&
+				TargetNameToUObjectModules != null && TargetNameToUObjectModules.Count > 0 &&
+				OutputItemsForAllTargets != null && 
+				ModuleNameToOutputItems != null && 
+				HotReloadModuleNamesForAllTargets != null &&
+				Targets != null && Targets.Count > 0 &&
+				TargetPrerequisites != null;
+		}
+
+
+		/// <summary>
+		/// Saves a UBTMakefile to disk
+		/// </summary>
+		/// <param name="TargetDescs">List of targets.  Order is not important</param>
+		/// <param name="UBTMakefile">The UBT makefile</param>
+		public static void SaveUBTMakefile(List<TargetDescriptor> TargetDescs, TargetMakefile UBTMakefile)
+		{
+			if (!UBTMakefile.IsValidMakefile())
+			{
+				throw new BuildException("Can't save a makefile that has invalid contents.  See UBTMakefile.IsValidMakefile()");
+			}
+
+			FileItem UBTMakefileItem = FileItem.GetItemByFileReference(GetUBTMakefilePath(TargetDescs));
+
+			// @todo ubtmake: Optimization: The UBTMakefile saved for game projects is upwards of 9 MB.  We should try to shrink its content if possible
+			// @todo ubtmake: Optimization: C# Serialization may be too slow for these big Makefiles.  Loading these files often shows up as the slower part of the assembling phase.
+
+			// Serialize the cache to disk.
+			try
+			{
+				Directory.CreateDirectory(Path.GetDirectoryName(UBTMakefileItem.AbsolutePath));
+				using(BinaryArchiveWriter Writer = new BinaryArchiveWriter(UBTMakefileItem.Location))
+				{
+					UBTMakefile.Write(Writer);
+				}
+			}
+			catch (Exception Ex)
+			{
+				Log.TraceError("Failed to write makefile: {0}", Ex.Message);
+			}
+		}
+
+
+		/// <summary>
+		/// Loads a UBTMakefile from disk
+		/// </summary>
+		/// <param name="MakefilePath">Path to the makefile to load</param>
+		/// <param name="ProjectFile">Path to the project file</param>
+		/// <param name="ReasonNotLoaded">If the function returns null, this string will contain the reason why</param>
+		/// <param name="WorkingSet">Interface to query which source files are in the working set</param>
+		/// <returns>The loaded makefile, or null if it failed for some reason.  On failure, the 'ReasonNotLoaded' variable will contain information about why</returns>
+		public static TargetMakefile Load(FileReference MakefilePath, FileReference ProjectFile, ISourceFileWorkingSet WorkingSet, out string ReasonNotLoaded)
+		{
+			// Check the directory timestamp on the project files directory.  If the user has generated project files more
+			// recently than the UBTMakefile, then we need to consider the file to be out of date
+			FileInfo UBTMakefileInfo = new FileInfo(MakefilePath.FullName);
+			if (!UBTMakefileInfo.Exists)
+			{
+				// UBTMakefile doesn't even exist, so we won't bother loading it
+				ReasonNotLoaded = "no existing makefile";
+				return null;
+			}
+
+			// Check the build version
+			FileInfo BuildVersionFileInfo = new FileInfo(BuildVersion.GetDefaultFileName().FullName);
+			if (BuildVersionFileInfo.Exists && UBTMakefileInfo.LastWriteTime.CompareTo(BuildVersionFileInfo.LastWriteTime) < 0)
+			{
+				Log.TraceLog("Existing makefile is older than Build.version, ignoring it");
+				ReasonNotLoaded = "Build.version is newer";
+				return null;
+			}
+
+			// @todo ubtmake: This will only work if the directory timestamp actually changes with every single GPF.  Force delete existing files before creating new ones?  Eh... really we probably just want to delete + create a file in that folder
+			//			-> UPDATE: Seems to work OK right now though on Windows platform, maybe due to GUID changes
+			// @todo ubtmake: Some platforms may not save any files into this folder.  We should delete + generate a "touch" file to force the directory timestamp to be updated (or just check the timestamp file itself.  We could put it ANYWHERE, actually)
+
+			// Installed Build doesn't need to check engine projects for outdatedness
+			if (!UnrealBuildTool.IsEngineInstalled())
+			{
+				if (DirectoryReference.Exists(ProjectFileGenerator.IntermediateProjectFilesPath))
+				{
+					DateTime EngineProjectFilesLastUpdateTime = new FileInfo(ProjectFileGenerator.ProjectTimestampFile).LastWriteTime;
+					if (UBTMakefileInfo.LastWriteTime.CompareTo(EngineProjectFilesLastUpdateTime) < 0)
+					{
+						// Engine project files are newer than UBTMakefile
+						Log.TraceLog("Existing makefile is older than generated engine project files, ignoring it");
+						ReasonNotLoaded = "project files are newer";
+						return null;
+					}
+				}
+			}
+
+			// Check the game project directory too
+			if (ProjectFile != null)
+			{
+				string ProjectFilename = ProjectFile.FullName;
+				FileInfo ProjectFileInfo = new FileInfo(ProjectFilename);
+				if (!ProjectFileInfo.Exists || UBTMakefileInfo.LastWriteTime.CompareTo(ProjectFileInfo.LastWriteTime) < 0)
+				{
+					// .uproject file is newer than UBTMakefile
+					Log.TraceLog("Makefile is older than .uproject file, ignoring it");
+					ReasonNotLoaded = ".uproject file is newer";
+					return null;
+				}
+
+				DirectoryReference MasterProjectRelativePath = ProjectFile.Directory;
+				string GameIntermediateProjectFilesPath = Path.Combine(MasterProjectRelativePath.FullName, "Intermediate", "ProjectFiles");
+				if (Directory.Exists(GameIntermediateProjectFilesPath))
+				{
+					DateTime GameProjectFilesLastUpdateTime = new DirectoryInfo(GameIntermediateProjectFilesPath).LastWriteTime;
+					if (UBTMakefileInfo.LastWriteTime.CompareTo(GameProjectFilesLastUpdateTime) < 0)
+					{
+						// Game project files are newer than UBTMakefile
+						Log.TraceLog("Makefile is older than generated game project files, ignoring it");
+						ReasonNotLoaded = "game project files are newer";
+						return null;
+					}
+				}
+			}
+
+			// Check to see if UnrealBuildTool.exe was compiled more recently than the UBTMakefile
+			DateTime UnrealBuildToolTimestamp = new FileInfo(Assembly.GetExecutingAssembly().Location).LastWriteTime;
+			if (UBTMakefileInfo.LastWriteTime.CompareTo(UnrealBuildToolTimestamp) < 0)
+			{
+				// UnrealBuildTool.exe was compiled more recently than the UBTMakefile
+				Log.TraceLog("Makefile is older than UnrealBuildTool.exe, ignoring it");
+				ReasonNotLoaded = "UnrealBuildTool.exe is newer";
+				return null;
+			}
+
+			// Check to see if any BuildConfiguration files have changed since the last build
+			List<XmlConfig.InputFile> InputFiles = XmlConfig.FindInputFiles();
+			foreach (XmlConfig.InputFile InputFile in InputFiles)
+			{
+				FileInfo InputFileInfo = new FileInfo(InputFile.Location.FullName);
+				if (InputFileInfo.LastWriteTime > UBTMakefileInfo.LastWriteTime)
+				{
+					Log.TraceLog("Makefile is older than BuildConfiguration.xml, ignoring it");
+					ReasonNotLoaded = "BuildConfiguration.xml is newer";
+					return null;
+				}
+			}
+
+			TargetMakefile LoadedUBTMakefile;
+			using(Timeline.ScopeEvent("Loading makefile"))
+			{
+				try
+				{
+					using(BinaryArchiveReader Reader = new BinaryArchiveReader(new FileReference(UBTMakefileInfo)))
+					{
+						LoadedUBTMakefile = new TargetMakefile(Reader);
+					}
+				}
+				catch (Exception Ex)
+				{
+					Log.TraceWarning("Failed to read makefile: {0}", Ex.Message);
+					Log.TraceLog("Exception: {0}", Ex.ToString());
+					ReasonNotLoaded = "couldn't read existing makefile";
+					return null;
+				}
+			}
+
+			using(Timeline.ScopeEvent("Checking makefile validity"))
+			{
+				if (!LoadedUBTMakefile.IsValidMakefile())
+				{
+					Log.TraceWarning("Loaded makefile appears to have invalid contents, ignoring it ({0})", UBTMakefileInfo.FullName);
+					ReasonNotLoaded = "existing makefile appears to be invalid";
+					return null;
+				}
+
+				// Check if any of the target's Build.cs files are newer than the makefile
+				foreach (UEBuildTarget Target in LoadedUBTMakefile.Targets)
+				{
+					// Check if ini files are newer. Ini files contain build settings too.
+					DirectoryReference ProjectDirectory = DirectoryReference.FromFile(ProjectFile);
+					foreach (ConfigHierarchyType IniType in (ConfigHierarchyType[])Enum.GetValues(typeof(ConfigHierarchyType)))
+					{
+						foreach (FileReference IniFilename in ConfigHierarchy.EnumerateConfigFileLocations(IniType, ProjectDirectory, Target.Platform))
+						{
+							FileInfo IniFileInfo = new FileInfo(IniFilename.FullName);
+							if (UBTMakefileInfo.LastWriteTime.CompareTo(IniFileInfo.LastWriteTime) < 0)
+							{
+								// Ini files are newer than UBTMakefile
+								ReasonNotLoaded = "ini files are newer than UBTMakefile";
+								return null;
+							}
+						}
+					}
+				}
+
+				foreach (BuildPrerequisites Prerequisites in LoadedUBTMakefile.TargetPrerequisites)
+				{
+					foreach(DirectoryReference SourceDirectory in Prerequisites.SourceDirectories)
+					{
+						DirectoryInfo SourceDirectoryInfo = new DirectoryInfo(SourceDirectory.FullName);
+						if(!SourceDirectoryInfo.Exists || SourceDirectoryInfo.LastWriteTimeUtc > UBTMakefileInfo.LastWriteTimeUtc)
+						{
+							Log.TraceLog("Timestamp of {0} ({1}) is newer than makefile ({2})", SourceDirectory, SourceDirectoryInfo.LastWriteTimeUtc, UBTMakefileInfo.LastWriteTimeUtc);
+							ReasonNotLoaded = "source directory changed";
+							return null;
+						}
+					}
+
+					foreach(FileItem AdditionalDependency in Prerequisites.AdditionalDependencies)
+					{
+						if (!AdditionalDependency.Exists)
+						{
+							Log.TraceLog("{0} has been deleted since makefile was built.", AdditionalDependency.Location);
+							ReasonNotLoaded = string.Format("{0} deleted", AdditionalDependency.Location.GetFileName());
+							return null;
+						}
+						if(AdditionalDependency.LastWriteTimeUtc > UBTMakefileInfo.LastWriteTime)
+						{
+							Log.TraceLog("{0} has been modified since makefile was built.", AdditionalDependency.Location);
+							ReasonNotLoaded = string.Format("{0} modified", AdditionalDependency.Location.GetFileName());
+							return null;
+						}
+					}
+
+					// We do a check to see if any modules' headers have changed which have
+					// acquired or lost UHT types.  If so, which should be rare,
+					// we'll just invalidate the entire makefile and force it to be rebuilt.
+
+					// Get all H files in processed modules newer than the makefile itself
+					HashSet<FileReference> HFilesNewerThanMakefile =
+						new HashSet<FileReference>(
+							Prerequisites.UObjectModuleHeaders
+							.SelectMany(x => x.SourceFolder != null ? DirectoryReference.EnumerateFiles(x.SourceFolder, "*.h", SearchOption.AllDirectories) : Enumerable.Empty<FileReference>())
+							.Where(y => FileReference.GetLastWriteTimeUtc(y) > UBTMakefileInfo.LastWriteTimeUtc)
+							.OrderBy(z => z).Distinct()
+						);
+
+					// Get all H files in all modules processed in the last makefile build
+					HashSet<FileReference> AllUHTHeaders = new HashSet<FileReference>(Prerequisites.UObjectModuleHeaders.SelectMany(x => x.HeaderFiles).Select(x => x.Location));
+
+					// Check whether any headers have been deleted. If they have, we need to regenerate the makefile since the module might now be empty. If we don't,
+					// and the file has been moved to a different module, we may include stale generated headers.
+					foreach (FileReference FileName in AllUHTHeaders)
+					{
+						if (!FileReference.Exists(FileName))
+						{
+							Log.TraceLog("File processed by UHT was deleted ({0}); invalidating makefile", FileName);
+							ReasonNotLoaded = string.Format("UHT file was deleted");
+							return null;
+						}
+					}
+
+					// Makefile is invalid if:
+					// * There are any newer files which contain no UHT data, but were previously in the makefile
+					// * There are any newer files contain data which needs processing by UHT, but weren't not previously in the makefile
+					foreach (FileReference Filename in HFilesNewerThanMakefile)
+					{
+						bool bContainsUHTData = CPPHeaders.DoesFileContainUObjects(Filename.FullName);
+						bool bWasProcessed = AllUHTHeaders.Contains(Filename);
+						if (bContainsUHTData != bWasProcessed)
+						{
+							Log.TraceLog("{0} {1} contain UHT types and now {2} , ignoring it ({3})", Filename, bWasProcessed ? "used to" : "didn't", bWasProcessed ? "doesn't" : "does", UBTMakefileInfo.FullName);
+							ReasonNotLoaded = string.Format("new files with reflected types");
+							return null;
+						}
+					}
+				}
+
+				// If adaptive unity build is enabled, do a check to see if there are any source files that became part of the
+				// working set since the Makefile was created (or, source files were removed from the working set.)  If anything
+				// changed, then we'll force a new Makefile to be created so that we have fresh unity build blobs.  We always
+				// want to make sure that source files in the working set are excluded from those unity blobs (for fastest possible
+				// iteration times.)
+				if (LoadedUBTMakefile.bUseAdaptiveUnityBuild)
+				{
+					foreach(BuildPrerequisites Prerequisites in LoadedUBTMakefile.TargetPrerequisites)
+					{
+						// Check if any source files in the working set no longer belong in it
+						foreach (FileItem SourceFile in Prerequisites.WorkingSet)
+						{
+							if (!WorkingSet.Contains(SourceFile.Location) && File.GetLastWriteTimeUtc(SourceFile.AbsolutePath) > UBTMakefileInfo.LastWriteTimeUtc)
+							{
+								Log.TraceLog("{0} was part of source working set and now is not; invalidating makefile ({1})", SourceFile.AbsolutePath, UBTMakefileInfo.FullName);
+								ReasonNotLoaded = string.Format("working set of source files changed");
+								return null;
+							}
+						}
+
+						// Check if any source files that are eligible for being in the working set have been modified
+						foreach (FileItem SourceFile in Prerequisites.CandidatesForWorkingSet)
+						{
+							if (WorkingSet.Contains(SourceFile.Location) && File.GetLastWriteTimeUtc(SourceFile.AbsolutePath) > UBTMakefileInfo.LastWriteTimeUtc)
+							{
+								Log.TraceLog("{0} was part of source working set and now is not; invalidating makefile ({1})", SourceFile.AbsolutePath, UBTMakefileInfo.FullName);
+								ReasonNotLoaded = string.Format("working set of source files changed");
+								return null;
+							}
+						}
+					}
+				}
+			}
+
+			ReasonNotLoaded = null;
+			return LoadedUBTMakefile;
+		}
+
+
+		/// <summary>
+		/// Gets the file path for a UBTMakefile
+		/// </summary>
+		/// <param name="TargetDescs">List of targets.  Order is not important</param>
+		/// <returns>UBTMakefile path</returns>
+		public static FileReference GetUBTMakefilePath(List<TargetDescriptor> TargetDescs)
+		{
+			FileReference UBTMakefilePath;
+
+			if (TargetDescs.Count == 1)
+			{
+				TargetDescriptor TargetDesc = TargetDescs[0];
+
+				string UBTMakefileName = "Makefile.ubt";
+
+				UBTMakefilePath = FileReference.Combine(GetUBTMakefileDirectoryPathForSingleTarget(TargetDesc), UBTMakefileName);
+			}
+			else
+			{
+				// For Makefiles that contain multiple targets, we'll make up a file name that contains all of the targets, their
+				// configurations and platforms, and save it into the base intermediate folder
+				string TargetCollectionName = MakeTargetCollectionName(TargetDescs);
+
+				TargetDescriptor DescriptorWithProject = TargetDescs.FirstOrDefault(x => x.ProjectFile != null);
+
+				DirectoryReference ProjectIntermediatePath;
+				if (DescriptorWithProject != null)
+				{
+					ProjectIntermediatePath = DirectoryReference.Combine(DescriptorWithProject.ProjectFile.Directory, "Intermediate", "Build");
+				}
+				else
+				{
+					ProjectIntermediatePath = DirectoryReference.Combine(UnrealBuildTool.EngineDirectory, "Intermediate", "Build");
+				}
+
+				// @todo ubtmake: The TargetCollectionName string could be really long if there is more than one target!  Hash it?
+				UBTMakefilePath = FileReference.Combine(ProjectIntermediatePath, TargetCollectionName + ".ubt");
+			}
+
+			return UBTMakefilePath;
+		}
+
+		/// <summary>
+		/// Gets the file path for a UBTMakefile for single target.
+		/// </summary>
+		/// <param name="Target">The target.</param>
+		/// <returns>UBTMakefile path</returns>
+		public static DirectoryReference GetUBTMakefileDirectoryPathForSingleTarget(TargetDescriptor Target)
+		{
+			return GetUBTMakefileDirectory(Target.ProjectFile, Target.Platform, Target.Configuration, Target.Name);
+		}
+
+		public static DirectoryReference GetUBTMakefileDirectory(FileReference ProjectFile, UnrealTargetPlatform Platform, UnrealTargetConfiguration Configuration, string TargetName)
+		{
+			// If there's only one target, just save the UBTMakefile in the target's build intermediate directory
+			// under a folder for that target (and platform/config combo.)
+			if (ProjectFile != null)
+			{
+				return DirectoryReference.Combine(ProjectFile.Directory, "Intermediate", "Build", Platform.ToString(), TargetName, Configuration.ToString());
+			}
+			else
+			{
+				return DirectoryReference.Combine(UnrealBuildTool.EngineDirectory, "Intermediate", "Build", Platform.ToString(), TargetName, Configuration.ToString());
+			}
+		}
+
+		public static FileReference GetUBTMakefilePath(FileReference ProjectFile, UnrealTargetPlatform Platform, UnrealTargetConfiguration Configuration, string TargetName, bool bForHotReload)
+		{
+			DirectoryReference BaseDir = GetUBTMakefileDirectory(ProjectFile, Platform, Configuration, TargetName);
+			return FileReference.Combine(BaseDir, bForHotReload ? "HotReloadMakefile.ubt" : "Makefile.ubt");
+		}
+
+		/// <summary>
+		/// Makes up a name for a set of targets that we can use for file or directory names
+		/// </summary>
+		/// <param name="TargetDescs">List of targets.  Order is not important</param>
+		/// <returns>The name to use</returns>
+		public static string MakeTargetCollectionName(List<TargetDescriptor> TargetDescs)
+		{
+			if (TargetDescs.Count == 0)
+			{
+				throw new BuildException("Expecting at least one Target to be passed to MakeTargetCollectionName");
+			}
+
+			List<TargetDescriptor> SortedTargets = new List<TargetDescriptor>();
+			SortedTargets.AddRange(TargetDescs);
+			SortedTargets.Sort((x, y) => { return x.Name.CompareTo(y.Name); });
+
+			// Figure out what to call our action graph based on everything we're building
+			StringBuilder TargetCollectionName = new StringBuilder();
+			foreach (TargetDescriptor Target in SortedTargets)
+			{
+				if (TargetCollectionName.Length > 0)
+				{
+					TargetCollectionName.Append("_");
+				}
+
+				// @todo ubtmake: Should we also have the platform Architecture in this string?
+				TargetCollectionName.Append(Target.Name + "-" + Target.Platform.ToString() + "-" + Target.Configuration.ToString());
+			}
+
+			return TargetCollectionName.ToString();
+		}
+
+	}
+}
diff --git a/Engine/Source/Programs/UnrealBuildTool/System/UBTMakefile.cs b/Engine/Source/Programs/UnrealBuildTool/System/UBTMakefile.cs
deleted file mode 100644
diff --git a/Engine/Source/Programs/UnrealBuildTool/UnrealBuildTool.csproj b/Engine/Source/Programs/UnrealBuildTool/UnrealBuildTool.csproj
--- a/Engine/Source/Programs/UnrealBuildTool/UnrealBuildTool.csproj
+++ b/Engine/Source/Programs/UnrealBuildTool/UnrealBuildTool.csproj
@@ -223,10 +223,10 @@
     <Compile Include="System\StartupTraceListener.cs" />
     <Compile Include="System\SupportedConfigurationsAttribute.cs" />
     <Compile Include="System\SupportedPlatformsAttribute.cs" />
+    <Compile Include="System\TargetMakefile.cs" />
     <Compile Include="System\TargetReceipt.cs" />
     <Compile Include="System\Timeline.cs" />
     <Compile Include="System\ToolMode.cs" />
-    <Compile Include="System\UBTMakefile.cs" />
     <Compile Include="System\VersionNumber.cs" />
     <Compile Include="System\XmlConfig.cs" />
     <Compile Include="System\XmlConfigData.cs" />
