UBT: Remove the bAdaptiveNonUnity flag from the makefile. We don't need to explicitly check for this; the working set and candidate working set will be empty if it's not set.

#rnx

diff --git a/Engine/Source/Programs/UnrealBuildTool/Modes/BuildMode.cs b/Engine/Source/Programs/UnrealBuildTool/Modes/BuildMode.cs
--- a/Engine/Source/Programs/UnrealBuildTool/Modes/BuildMode.cs
+++ b/Engine/Source/Programs/UnrealBuildTool/Modes/BuildMode.cs
@@ -380,7 +380,6 @@ namespace UnrealBuildTool
 							Makefile.ModuleNameToOutputItems = ModuleNameToOutputItems;
 							Makefile.HotReloadModuleNamesForAllTargets = HotReloadModuleNamesForAllTargets;
 							Makefile.Targets = Targets;
-							Makefile.bUseAdaptiveUnityBuild = Targets.Any(x => x.Rules.bUseAdaptiveUnityBuild);
 							Makefile.TargetPrerequisites = TargetPrerequisites;
 
 							if (BuildConfiguration.bUseUBTMakefiles)
diff --git a/Engine/Source/Programs/UnrealBuildTool/System/TargetMakefile.cs b/Engine/Source/Programs/UnrealBuildTool/System/TargetMakefile.cs
--- a/Engine/Source/Programs/UnrealBuildTool/System/TargetMakefile.cs
+++ b/Engine/Source/Programs/UnrealBuildTool/System/TargetMakefile.cs
@@ -58,11 +58,6 @@ namespace UnrealBuildTool
 		/// </summary>
 		public List<UEBuildTarget> Targets;
 
-		/// <summary>
-		/// Whether adaptive unity build is enabled for any of these targets
-		/// </summary>
-		public bool bUseAdaptiveUnityBuild;
-
 		/// <summary>
 		/// List of predicates for the action graph to be valid.
 		/// </summary>
@@ -88,7 +83,6 @@ namespace UnrealBuildTool
 			HotReloadModuleNamesForAllTargets = Reader.ReadHashSet(() => Reader.ReadString());
 			Targets = Reader.ReadList(() => new UEBuildTarget(Reader));
 			TargetPrerequisites = Reader.ReadList(() => new BuildPrerequisites(Reader));
-			bUseAdaptiveUnityBuild = Reader.ReadBool();
 		}
 
 		public void Write(BinaryArchiveWriter Writer)
@@ -103,7 +97,6 @@ namespace UnrealBuildTool
 			Writer.WriteHashSet(HotReloadModuleNamesForAllTargets, x => Writer.WriteString(x));
 			Writer.WriteList(Targets, x => x.Write(Writer));
 			Writer.WriteList(TargetPrerequisites, x => x.Write(Writer));
-			Writer.WriteBool(bUseAdaptiveUnityBuild);
 			Writer.Flush();
 		}
 
@@ -382,30 +375,27 @@ namespace UnrealBuildTool
 				// changed, then we'll force a new Makefile to be created so that we have fresh unity build blobs.  We always
 				// want to make sure that source files in the working set are excluded from those unity blobs (for fastest possible
 				// iteration times.)
-				if (LoadedUBTMakefile.bUseAdaptiveUnityBuild)
+				foreach(BuildPrerequisites Prerequisites in LoadedUBTMakefile.TargetPrerequisites)
 				{
-					foreach(BuildPrerequisites Prerequisites in LoadedUBTMakefile.TargetPrerequisites)
+					// Check if any source files in the working set no longer belong in it
+					foreach (FileItem SourceFile in Prerequisites.WorkingSet)
 					{
-						// Check if any source files in the working set no longer belong in it
-						foreach (FileItem SourceFile in Prerequisites.WorkingSet)
+						if (!WorkingSet.Contains(SourceFile.Location) && File.GetLastWriteTimeUtc(SourceFile.AbsolutePath) > UBTMakefileInfo.LastWriteTimeUtc)
 						{
-							if (!WorkingSet.Contains(SourceFile.Location) && File.GetLastWriteTimeUtc(SourceFile.AbsolutePath) > UBTMakefileInfo.LastWriteTimeUtc)
-							{
-								Log.TraceLog("{0} was part of source working set and now is not; invalidating makefile ({1})", SourceFile.AbsolutePath, UBTMakefileInfo.FullName);
-								ReasonNotLoaded = string.Format("working set of source files changed");
-								return null;
-							}
+							Log.TraceLog("{0} was part of source working set and now is not; invalidating makefile ({1})", SourceFile.AbsolutePath, UBTMakefileInfo.FullName);
+							ReasonNotLoaded = string.Format("working set of source files changed");
+							return null;
 						}
+					}
 
-						// Check if any source files that are eligible for being in the working set have been modified
-						foreach (FileItem SourceFile in Prerequisites.CandidatesForWorkingSet)
+					// Check if any source files that are eligible for being in the working set have been modified
+					foreach (FileItem SourceFile in Prerequisites.CandidatesForWorkingSet)
+					{
+						if (WorkingSet.Contains(SourceFile.Location) && File.GetLastWriteTimeUtc(SourceFile.AbsolutePath) > UBTMakefileInfo.LastWriteTimeUtc)
 						{
-							if (WorkingSet.Contains(SourceFile.Location) && File.GetLastWriteTimeUtc(SourceFile.AbsolutePath) > UBTMakefileInfo.LastWriteTimeUtc)
-							{
-								Log.TraceLog("{0} was part of source working set and now is not; invalidating makefile ({1})", SourceFile.AbsolutePath, UBTMakefileInfo.FullName);
-								ReasonNotLoaded = string.Format("working set of source files changed");
-								return null;
-							}
+							Log.TraceLog("{0} was part of source working set and now is not; invalidating makefile ({1})", SourceFile.AbsolutePath, UBTMakefileInfo.FullName);
+							ReasonNotLoaded = string.Format("working set of source files changed");
+							return null;
 						}
 					}
 				}
