UBT: Add support for building multiple targets in parallel.

* A makefile is now constructed for each target, and the outdated actions are combined into a single graph before execution.
* Makefiles are used in all build modes, regardless of whether they are saved to disk, and represent the vector by which a target's build actions are passed to the executor.
* The -JsonExport= argument has been removed, and a separate JsonExport mode has been added (invoke with eg. "-Mode=JsonExport UE4Editor Win64 Development")
* Multiple action history files are kept open at one time, and a hierarchy is used for each target.

diff --git a/Engine/Source/Programs/UnrealBuildTool/Configuration/BuildConfiguration.cs b/Engine/Source/Programs/UnrealBuildTool/Configuration/BuildConfiguration.cs
--- a/Engine/Source/Programs/UnrealBuildTool/Configuration/BuildConfiguration.cs
+++ b/Engine/Source/Programs/UnrealBuildTool/Configuration/BuildConfiguration.cs
@@ -63,13 +63,6 @@ namespace UnrealBuildTool
 		[XmlConfigFile]
 		public bool bAllowSNDBS = true;
 
-		/// <summary>
-		/// Whether we should export a JSON file containing detailed target information.
-		/// </summary>
-		[XmlConfigFile]
-		[CommandLine("-JsonExport")]
-		public string JsonExportFile = null;
-
 		/// <summary>
 		/// Skip building; just do setup and terminate.
 		/// </summary>
diff --git a/Engine/Source/Programs/UnrealBuildTool/Configuration/UEBuildTarget.cs b/Engine/Source/Programs/UnrealBuildTool/Configuration/UEBuildTarget.cs
--- a/Engine/Source/Programs/UnrealBuildTool/Configuration/UEBuildTarget.cs
+++ b/Engine/Source/Programs/UnrealBuildTool/Configuration/UEBuildTarget.cs
@@ -439,13 +439,11 @@ namespace UnrealBuildTool
 		/// <summary>
 		/// The target rules
 		/// </summary>
-		[NonSerialized]
 		public ReadOnlyTargetRules Rules;
 
 		/// <summary>
 		/// The rules assembly to use when searching for modules
 		/// </summary>
-		[NonSerialized]
 		public RulesAssembly RulesAssembly;
 
 		/// <summary>
@@ -456,7 +454,6 @@ namespace UnrealBuildTool
 		/// <summary>
 		/// The project descriptor for this target
 		/// </summary>
-		[NonSerialized]
 		public ProjectDescriptor ProjectDescriptor;
 
 		/// <summary>
@@ -522,13 +519,11 @@ namespace UnrealBuildTool
 		/// <summary>
 		/// All plugins which are built for this target
 		/// </summary>
-		[NonSerialized]
 		public List<UEBuildPlugin> BuildPlugins;
 
 		/// <summary>
 		/// All plugin dependencies for this target. This differs from the list of plugins that is built for Launcher, where we build everything, but link in only the enabled plugins.
 		/// </summary>
-		[NonSerialized]
 		public List<UEBuildPlugin> EnabledPlugins;
 
 		/// <summary>
@@ -539,13 +534,11 @@ namespace UnrealBuildTool
 		/// <summary>
 		/// All application binaries; may include binaries not built by this target.
 		/// </summary>
-		[NonSerialized]
 		public List<UEBuildBinary> Binaries = new List<UEBuildBinary>();
 
 		/// <summary>
 		/// Kept to determine the correct module parsing order when filtering modules.
 		/// </summary>
-		[NonSerialized]
 		protected List<UEBuildBinary> NonFilteredModules = new List<UEBuildBinary>();
 
 		/// <summary>
@@ -556,7 +549,6 @@ namespace UnrealBuildTool
 		/// <summary>
 		/// Used to keep track of all modules by name.
 		/// </summary>
-		[NonSerialized]
 		private Dictionary<string, UEBuildModule> Modules = new Dictionary<string, UEBuildModule>(StringComparer.InvariantCultureIgnoreCase);
 
 		/// <summary>
@@ -576,7 +568,7 @@ namespace UnrealBuildTool
 		/// <summary>
 		/// List of scripts to run before building
 		/// </summary>
-		FileReference[] PreBuildStepScripts;
+		public FileReference[] PreBuildStepScripts;
 
 		/// <summary>
 		/// List of scripts to run after building
@@ -597,50 +589,6 @@ namespace UnrealBuildTool
 			return bCompileMonolithic;	// @todo ubtmake: We need to make sure this function and similar things aren't called in assembler mode
 		}
 
-		public UEBuildTarget(BinaryArchiveReader Reader)
-		{
-			TargetType = (TargetType)Reader.ReadInt();
-			ProjectFile = Reader.ReadFileReference();
-			AppName = Reader.ReadString();
-			TargetName = Reader.ReadString();
-			bUseSharedBuildEnvironment = Reader.ReadBool();
-			Platform = (UnrealTargetPlatform)Reader.ReadInt();
-			Configuration = (UnrealTargetConfiguration)Reader.ReadInt();
-			Architecture = Reader.ReadString();
-			PlatformIntermediateFolder = Reader.ReadString();
-			ProjectDirectory = Reader.ReadDirectoryReference();
-			ProjectIntermediateDirectory = Reader.ReadDirectoryReference();
-			EngineIntermediateDirectory = Reader.ReadDirectoryReference();
-			bCompileMonolithic = Reader.ReadBool();
-			ReceiptFileName = Reader.ReadFileReference();
-			PreBuildStepScripts = Reader.ReadArray(() => Reader.ReadFileReference());
-			PostBuildStepScripts = Reader.ReadArray(() => Reader.ReadFileReference());
-			bDeployAfterCompile = Reader.ReadBool();
-			bHasProjectScriptPlugin = Reader.ReadBool();
-		}
-
-		public void Write(BinaryArchiveWriter Writer)
-		{
-			Writer.WriteInt((int)TargetType);
-			Writer.WriteFileReference(ProjectFile);
-			Writer.WriteString(AppName);
-			Writer.WriteString(TargetName);
-			Writer.WriteBool(bUseSharedBuildEnvironment);
-			Writer.WriteInt((int)Platform);
-			Writer.WriteInt((int)Configuration);
-			Writer.WriteString(Architecture);
-			Writer.WriteString(PlatformIntermediateFolder);
-			Writer.WriteDirectoryReference(ProjectDirectory);
-			Writer.WriteDirectoryReference(ProjectIntermediateDirectory);
-			Writer.WriteDirectoryReference(EngineIntermediateDirectory);
-			Writer.WriteBool(bCompileMonolithic);
-			Writer.WriteFileReference(ReceiptFileName);
-			Writer.WriteArray(PreBuildStepScripts, x => Writer.WriteFileReference(x));
-			Writer.WriteArray(PostBuildStepScripts, x => Writer.WriteFileReference(x));
-			Writer.WriteBool(bDeployAfterCompile);
-			Writer.WriteBool(bHasProjectScriptPlugin);
-		}
-
 		/// <summary>
 		/// Constructor.
 		/// </summary>
@@ -1705,10 +1653,11 @@ namespace UnrealBuildTool
 		/// <summary>
 		/// Export the definition of this target to a JSON file
 		/// </summary>
-		/// <param name="FileName">File to write to</param>
-		public void ExportJson(string FileName)
+		/// <param name="OutputFile">File to write to</param>
+		public void ExportJson(FileReference OutputFile)
 		{
-			using (JsonWriter Writer = new JsonWriter(FileName))
+			DirectoryReference.CreateDirectory(OutputFile.Directory);
+			using (JsonWriter Writer = new JsonWriter(OutputFile))
 			{
 				Writer.WriteObjectStart();
 
diff --git a/Engine/Source/Programs/UnrealBuildTool/Executors/Distcc.cs b/Engine/Source/Programs/UnrealBuildTool/Executors/Distcc.cs
--- a/Engine/Source/Programs/UnrealBuildTool/Executors/Distcc.cs
+++ b/Engine/Source/Programs/UnrealBuildTool/Executors/Distcc.cs
@@ -265,31 +265,28 @@ namespace UnrealBuildTool
 									// Determine whether there are any prerequisites of the action that are outdated.
 									bool bHasOutdatedPrerequisites = false;
 									bool bHasFailedPrerequisites = false;
-									foreach (FileItem PrerequisiteItem in Action.PrerequisiteItems)
+									foreach(Action PrerequisiteAction in Action.PrerequisiteActions)
 									{
-										if (PrerequisiteItem.ProducingAction != null && Actions.Contains(PrerequisiteItem.ProducingAction))
+										ActionThread PrerequisiteProcess = null;
+										bool bFoundPrerequisiteProcess = ActionThreadDictionary.TryGetValue(PrerequisiteAction, out PrerequisiteProcess);
+										if (bFoundPrerequisiteProcess == true)
 										{
-											ActionThread PrerequisiteProcess = null;
-											bool bFoundPrerequisiteProcess = ActionThreadDictionary.TryGetValue(PrerequisiteItem.ProducingAction, out PrerequisiteProcess);
-											if (bFoundPrerequisiteProcess == true)
+											if (PrerequisiteProcess == null)
 											{
-												if (PrerequisiteProcess == null)
-												{
-													bHasFailedPrerequisites = true;
-												}
-												else if (PrerequisiteProcess.bComplete == false)
-												{
-													bHasOutdatedPrerequisites = true;
-												}
-												else if (PrerequisiteProcess.ExitCode != 0)
-												{
-													bHasFailedPrerequisites = true;
-												}
+												bHasFailedPrerequisites = true;
 											}
-											else
+											else if (PrerequisiteProcess.bComplete == false)
 											{
 												bHasOutdatedPrerequisites = true;
 											}
+											else if (PrerequisiteProcess.ExitCode != 0)
+											{
+												bHasFailedPrerequisites = true;
+											}
+										}
+										else
+										{
+											bHasOutdatedPrerequisites = true;
 										}
 									}
 
diff --git a/Engine/Source/Programs/UnrealBuildTool/Executors/LocalExecutor.cs b/Engine/Source/Programs/UnrealBuildTool/Executors/LocalExecutor.cs
--- a/Engine/Source/Programs/UnrealBuildTool/Executors/LocalExecutor.cs
+++ b/Engine/Source/Programs/UnrealBuildTool/Executors/LocalExecutor.cs
@@ -378,12 +378,12 @@ namespace UnrealBuildTool
 								// Determine whether there are any prerequisites of the action that are outdated.
 								bool bHasOutdatedPrerequisites = false;
 								bool bHasFailedPrerequisites = false;
-								foreach (FileItem PrerequisiteItem in Action.PrerequisiteItems)
+								foreach (Action PrerequisiteAction in Action.PrerequisiteActions)
 								{
-									if (PrerequisiteItem.ProducingAction != null && Actions.Contains(PrerequisiteItem.ProducingAction))
+									if (Actions.Contains(PrerequisiteAction))
 									{
 										ActionThread PrerequisiteProcess = null;
-										bool bFoundPrerequisiteProcess = ActionThreadDictionary.TryGetValue(PrerequisiteItem.ProducingAction, out PrerequisiteProcess);
+										bool bFoundPrerequisiteProcess = ActionThreadDictionary.TryGetValue(PrerequisiteAction, out PrerequisiteProcess);
 										if (bFoundPrerequisiteProcess == true)
 										{
 											if (PrerequisiteProcess == null)
diff --git a/Engine/Source/Programs/UnrealBuildTool/Executors/SNDBS.cs b/Engine/Source/Programs/UnrealBuildTool/Executors/SNDBS.cs
--- a/Engine/Source/Programs/UnrealBuildTool/Executors/SNDBS.cs
+++ b/Engine/Source/Programs/UnrealBuildTool/Executors/SNDBS.cs
@@ -1,4 +1,4 @@
-﻿// Copyright 1998-2019 Epic Games, Inc. All Rights Reserved.
+// Copyright 1998-2019 Epic Games, Inc. All Rights Reserved.
 
 using System;
 using System.Collections.Generic;
@@ -113,12 +113,12 @@ namespace UnrealBuildTool
 							// Determine whether there are any prerequisites of the action that are outdated.
 							bool bHasOutdatedPrerequisites = false;
 							bool bHasFailedPrerequisites = false;
-							foreach (FileItem PrerequisiteItem in Action.PrerequisiteItems)
+							foreach (Action PrerequisiteAction in Action.PrerequisiteActions)
 							{
-								if (PrerequisiteItem.ProducingAction != null && InLocalActions.Contains(PrerequisiteItem.ProducingAction))
+								if (InLocalActions.Contains(PrerequisiteAction))
 								{
 									ActionThread PrerequisiteProcess = null;
-									bool bFoundPrerequisiteProcess = InActionThreadDictionary.TryGetValue(PrerequisiteItem.ProducingAction, out PrerequisiteProcess);
+									bool bFoundPrerequisiteProcess = InActionThreadDictionary.TryGetValue(PrerequisiteAction, out PrerequisiteProcess);
 									if (bFoundPrerequisiteProcess == true)
 									{
 										if (PrerequisiteProcess == null)
@@ -190,12 +190,12 @@ namespace UnrealBuildTool
 					// Determine whether there are any prerequisites of the action that are outdated.
 					bool bHasOutdatedPrerequisites = false;
 					bool bHasFailedPrerequisites = false;
-					foreach (FileItem PrerequisiteItem in Action.PrerequisiteItems)
+					foreach (Action PrerequisiteAction in Action.PrerequisiteActions)
 					{
-						if (PrerequisiteItem.ProducingAction != null && InActions.Contains(PrerequisiteItem.ProducingAction))
+						if (InActions.Contains(PrerequisiteAction))
 						{
 							ActionThread PrerequisiteProcess = null;
-							bool bFoundPrerequisiteProcess = InActionThreadDictionary.TryGetValue(PrerequisiteItem.ProducingAction, out PrerequisiteProcess);
+							bool bFoundPrerequisiteProcess = InActionThreadDictionary.TryGetValue(PrerequisiteAction, out PrerequisiteProcess);
 							if (bFoundPrerequisiteProcess == true)
 							{
 								if (PrerequisiteProcess == null)
diff --git a/Engine/Source/Programs/UnrealBuildTool/Executors/XGE.cs b/Engine/Source/Programs/UnrealBuildTool/Executors/XGE.cs
--- a/Engine/Source/Programs/UnrealBuildTool/Executors/XGE.cs
+++ b/Engine/Source/Programs/UnrealBuildTool/Executors/XGE.cs
@@ -218,11 +218,11 @@ namespace UnrealBuildTool
 				for (int ActionIndex = 0; ActionIndex < InActions.Count; ActionIndex++)
 				{
 					Action Action = InActions[ActionIndex];
-					foreach (FileItem Item in Action.PrerequisiteItems)
+					foreach (Action PrerequisiteAction in Action.PrerequisiteActions)
 					{
-						if (Item.ProducingAction != null && InActions.Contains(Item.ProducingAction))
+						if (InActions.Contains(PrerequisiteAction))
 						{
-							int DepIndex = InActions.IndexOf(Item.ProducingAction);
+							int DepIndex = InActions.IndexOf(PrerequisiteAction);
 							if (DepIndex > ActionIndex)
 							{
 								NumSortErrors++;
@@ -241,16 +241,16 @@ namespace UnrealBuildTool
 							continue;
 						}
 						Action Action = InActions[ActionIndex];
-						foreach (FileItem Item in Action.PrerequisiteItems)
+						foreach(Action PrerequisiteAction in Action.PrerequisiteActions)
 						{
-							if (Item.ProducingAction != null && InActions.Contains(Item.ProducingAction))
+							if (InActions.Contains(PrerequisiteAction))
 							{
-								int DepIndex = InActions.IndexOf(Item.ProducingAction);
+								int DepIndex = InActions.IndexOf(PrerequisiteAction);
 								if (UsedActions.Contains(DepIndex))
 								{
 									continue;
 								}
-								Actions.Add(Item.ProducingAction);
+								Actions.Add(PrerequisiteAction);
 								UsedActions.Add(DepIndex);
 							}
 						}
@@ -260,14 +260,14 @@ namespace UnrealBuildTool
 					for (int ActionIndex = 0; ActionIndex < Actions.Count; ActionIndex++)
 					{
 						Action Action = Actions[ActionIndex];
-						foreach (FileItem Item in Action.PrerequisiteItems)
+						foreach(Action PrerequisiteAction in Action.PrerequisiteActions)
 						{
-							if (Item.ProducingAction != null && Actions.Contains(Item.ProducingAction))
+							if (Actions.Contains(PrerequisiteAction))
 							{
-								int DepIndex = Actions.IndexOf(Item.ProducingAction);
+								int DepIndex = Actions.IndexOf(PrerequisiteAction);
 								if (DepIndex > ActionIndex)
 								{
-									throw new BuildException("Action is not topologically sorted.\n  {0} {1}\nDependency\n  {2} {3}", Action.CommandPath, Action.CommandArguments, Item.ProducingAction.CommandPath, Item.ProducingAction.CommandArguments);
+									throw new BuildException("Action is not topologically sorted.\n  {0} {1}\nDependency\n  {2} {3}", Action.CommandPath, Action.CommandArguments, PrerequisiteAction.CommandPath, PrerequisiteAction.CommandArguments);
 								}
 							}
 						}
@@ -339,6 +339,10 @@ namespace UnrealBuildTool
 				{
 					ToolElement.SetAttribute("OutputPrefix", OutputPrefix);
 				}
+				if(Action.GroupNames.Count > 0)
+				{
+					ToolElement.SetAttribute("GroupPrefix", String.Format("** For {0}", String.Join(" + ", Action.GroupNames)));
+				}
 
 				// When running on Windows, differentiate between .exe and batch files.
 				// Those (.bat, .cmd) need to be run via cmd /c or shellexecute, 
@@ -417,11 +421,11 @@ namespace UnrealBuildTool
 
 				// Create a semi-colon separated list of the other tasks this task depends on the results of.
 				List<string> DependencyNames = new List<string>();
-				foreach (FileItem Item in Action.PrerequisiteItems)
+				foreach(Action PrerequisiteAction in Action.PrerequisiteActions)
 				{
-					if (Item.ProducingAction != null && Actions.Contains(Item.ProducingAction))
+					if (Actions.Contains(PrerequisiteAction))
 					{
-						DependencyNames.Add(string.Format("Action{0}", Actions.IndexOf(Item.ProducingAction)));
+						DependencyNames.Add(string.Format("Action{0}", Actions.IndexOf(PrerequisiteAction)));
 					}
 				}
 
diff --git a/Engine/Source/Programs/UnrealBuildTool/Modes/BuildMode.cs b/Engine/Source/Programs/UnrealBuildTool/Modes/BuildMode.cs
--- a/Engine/Source/Programs/UnrealBuildTool/Modes/BuildMode.cs
+++ b/Engine/Source/Programs/UnrealBuildTool/Modes/BuildMode.cs
@@ -93,18 +93,12 @@ namespace UnrealBuildTool
 				JunkDeleter.DeleteJunk();
 			}
 
-			bool bSuccess = true;
-
 			Stopwatch BuildTimer = Stopwatch.StartNew();
 
 			// Reset global configurations
-			ActionGraph ActionGraph = new ActionGraph();
-
 			string ExecutorName = "Unknown";
-			ECompilationResult BuildResult = ECompilationResult.Succeeded;
 
-			List<UEBuildTarget> Targets = null;
-			Dictionary<UEBuildTarget, CPPHeaders> TargetToHeaders = new Dictionary<UEBuildTarget, CPPHeaders>();
+			List<CPPHeaders> HeadersList = new List<CPPHeaders>();
 
 			Stopwatch ExecutorTimer = new Stopwatch();
 
@@ -117,16 +111,16 @@ namespace UnrealBuildTool
 
 				bool bSkipRulesCompile = Arguments.Any(x => x.Equals("-skiprulescompile", StringComparison.InvariantCultureIgnoreCase));
 
-				List<TargetDescriptor> TargetDescs = new List<TargetDescriptor>();
+				List<TargetDescriptor> TargetDescriptors = new List<TargetDescriptor>();
 				using(Timeline.ScopeEvent("TargetDescriptor.ParseCommandLine()"))
 				{
-					TargetDescs.AddRange(TargetDescriptor.ParseCommandLine(Arguments, BuildConfiguration.bUsePrecompiled, bSkipRulesCompile));
+					TargetDescriptors.AddRange(TargetDescriptor.ParseCommandLine(Arguments, BuildConfiguration.bUsePrecompiled, bSkipRulesCompile));
 				}
 
 				// Handle remote builds
-				for(int Idx = 0; Idx < TargetDescs.Count; Idx++)
+				for(int Idx = 0; Idx < TargetDescriptors.Count; Idx++)
 				{
-					TargetDescriptor TargetDesc = TargetDescs[Idx];
+					TargetDescriptor TargetDesc = TargetDescriptors[Idx];
 					if(RemoteMac.HandlesTargetPlatform(TargetDesc.Platform))
 					{
 						FileReference RemoteLogFile = null;
@@ -141,8 +135,8 @@ namespace UnrealBuildTool
 							return (int)ECompilationResult.Unknown;
 						}
 
-						TargetDescs.RemoveAt(Idx--);
-						if(TargetDescs.Count == 0)
+						TargetDescriptors.RemoveAt(Idx--);
+						if(TargetDescriptors.Count == 0)
 						{
 							return (int)ECompilationResult.Succeeded;
 						}
@@ -152,70 +146,20 @@ namespace UnrealBuildTool
 				if (Arguments.Any(x => x.Equals("-InvalidateMakefilesOnly", StringComparison.InvariantCultureIgnoreCase)))
 				{
 					Log.TraceInformation("Invalidating makefiles only in this run.");
-					if (TargetDescs.Count != 1)
+					foreach(TargetDescriptor TargetDesc in TargetDescriptors)
 					{
-						Log.TraceError("You have to provide one target name for makefile invalidation.");
-						return (int)ECompilationResult.OtherCompilationError;
+						FileReference MakefileLocation = TargetMakefile.GetLocation(TargetDesc.ProjectFile, TargetDesc.Name, TargetDesc.Platform, TargetDesc.Configuration);
+						if(FileReference.Exists(MakefileLocation))
+						{ 
+							FileReference.Delete(MakefileLocation);
+						}
 					}
-
-					InvalidateMakefiles(TargetDescs[0]);
 					return (int)ECompilationResult.Succeeded;
 				}
 
-				// True if we should gather module dependencies for building in this run.  If this is false, then we'll expect to be loading information from disk about
-				// the target's modules before we'll be able to build anything.  One or both of IsGatheringBuild or IsAssemblingBuild must be true.
-				bool bIsGatheringBuild = true;
-
-				// True if we should go ahead and actually compile and link in this run.  If this is false, no actions will be executed and we'll stop after gathering
-				// and saving information about dependencies.  One or both of IsGatheringBuild or IsAssemblingBuild must be true.
-				bool bIsAssemblingBuild = true;
-
-				// If we were asked to enable fast build iteration, we want the 'gather' phase to default to off (unless it is overridden below
-				// using a command-line option.)
-				if (BuildConfiguration.bUseUBTMakefiles)
-				{
-					bIsGatheringBuild = false;
-				}
-
-				foreach (string Argument in Arguments)
-				{
-					string LowercaseArg = Argument.ToLowerInvariant();
-					if (LowercaseArg == "-gather")
-					{
-						bIsGatheringBuild = true;
-					}
-					else if (LowercaseArg == "-nogather")
-					{
-						bIsGatheringBuild = false;
-					}
-					else if (LowercaseArg == "-gatheronly")
-					{
-						bIsGatheringBuild = true;
-						bIsAssemblingBuild = false;
-					}
-					else if (LowercaseArg == "-assemble")
-					{
-						bIsAssemblingBuild = true;
-					}
-					else if (LowercaseArg == "-noassemble")
-					{
-						bIsAssemblingBuild = false;
-					}
-					else if (LowercaseArg == "-assembleonly")
-					{
-						bIsGatheringBuild = false;
-						bIsAssemblingBuild = true;
-					}
-				}
-
-				if (!bIsGatheringBuild && !bIsAssemblingBuild)
-				{
-					throw new BuildException("UnrealBuildTool: At least one of either IsGatheringBuild or IsAssemblingBuild must be true.  Did you pass '-NoGather' with '-NoAssemble'?");
-				}
-
 				// Get a set of all the project directories.
 				HashSet<DirectoryReference> ProjectDirs = new HashSet<DirectoryReference>();
-				foreach(TargetDescriptor TargetDesc in TargetDescs)
+				foreach(TargetDescriptor TargetDesc in TargetDescriptors)
 				{
 					if(TargetDesc.ProjectFile != null)
 					{
@@ -226,558 +170,245 @@ namespace UnrealBuildTool
 				// Create the working set provider
 				using (ISourceFileWorkingSet WorkingSet = SourceFileWorkingSet.Create(UnrealBuildTool.RootDirectory, ProjectDirs))
 				{
-					TargetMakefile Makefile = null;
+					// Create a makefile for each target
+					TargetMakefile[] Makefiles = new TargetMakefile[TargetDescriptors.Count];
+					for(int TargetIdx = 0; TargetIdx < TargetDescriptors.Count; TargetIdx++)
 					{
-						FileReference UBTMakefilePath = TargetMakefile.GetUBTMakefilePath(TargetDescs);
+						TargetDescriptor TargetDesc = TargetDescriptors[TargetIdx];
 
-						// Make sure the gather phase is executed if we're not actually building anything
-						if (BuildConfiguration.bGenerateManifest || BuildConfiguration.bXGEExport || GeneratingActionGraph)
-						{
-							bIsGatheringBuild = true;
-						}
+						// Get the path to the makefile for this target
+						FileReference MakefileLocation = TargetMakefile.GetLocation(TargetDesc.ProjectFile, TargetDesc.Name, TargetDesc.Platform, TargetDesc.Configuration);
 
-						// Were we asked to run in 'assembler only' mode?  If so, let's check to see if that's even possible by seeing if
-						// we have a valid UBTMakefile already saved to disk, ready for us to load.
-						if (bIsAssemblingBuild && !bIsGatheringBuild)
+						// Try to load an existing makefile
+						TargetMakefile Makefile = null;
+						if(BuildConfiguration.bUseUBTMakefiles)
 						{
-							// @todo ubtmake: Mildly terrified of BuildConfiguration/UEBuildConfiguration globals that were set during the Prepare phase but are not set now.  We may need to save/load all of these, otherwise
-							//		we'll need to call SetupGlobalEnvironment on all of the targets (maybe other stuff, too.  See PreBuildStep())
-
-							// Try to load the UBTMakefile.  It will only be loaded if it has valid content and is not determined to be out of date.    
-							string ReasonNotLoaded;
 							using(Timeline.ScopeEvent("UBTMakefile.LoadUBTMakefile"))
 							{
-								Makefile = TargetMakefile.Load(UBTMakefilePath, TargetDescs[0].ProjectFile, WorkingSet, out ReasonNotLoaded);
-							}
-
-							if (Makefile == null)
-							{
-								// If the Makefile couldn't be loaded, then we're not going to be able to continue in "assembler only" mode.  We'll do both
-								// a 'gather' and 'assemble' in the same run.  This will take a while longer, but subsequent runs will be fast!
-								bIsGatheringBuild = true;
-
-								Log.TraceInformation("Creating makefile for {0}{1} ({2})",
-									TargetDescs[0].Name,
-									TargetDescs.Count > 1 ? (" (and " + (TargetDescs.Count - 1).ToString() + " more)") : "",
-									ReasonNotLoaded);
+								string ReasonNotLoaded;
+								Makefile = TargetMakefile.Load(MakefileLocation, TargetDesc.ProjectFile, TargetDesc.Platform, WorkingSet, out ReasonNotLoaded);
+								if (Makefile == null)
+								{
+									Log.TraceInformation("Creating makefile for {0} ({1})", TargetDesc.Name, ReasonNotLoaded);
+								}
 							}
 						}
-					}
 
-
-					if (Makefile != null && !bIsGatheringBuild && bIsAssemblingBuild)
-					{
-						// If we've loaded a makefile, then we can fill target information from this file!
-						Targets = Makefile.Targets;
-					}
-					else
-					{
-						Targets = new List<UEBuildTarget>();
-						foreach (TargetDescriptor TargetDesc in TargetDescs)
+						// If we couldn't load a makefile, create a new one
+						if(Makefile == null)
 						{
-							using(Timeline.ScopeEvent("UEBuildTarget.CreateTarget()"))
-							{
-								UEBuildTarget Target = UEBuildTarget.CreateTarget(TargetDesc, bSkipRulesCompile, BuildConfiguration.SingleFileToCompile != null, BuildConfiguration.bUsePrecompiled);
-								Targets.Add(Target);
-							}
-						}
-					}
-
-					// Build action lists for all passed in targets.
-					List<FileItem> OutputItemsForAllTargets = new List<FileItem>();
-					Dictionary<string, FileItem[]> ModuleNameToOutputItems = new Dictionary<string, FileItem[]>(StringComparer.OrdinalIgnoreCase);
-					Dictionary<string, List<UHTModuleInfo>> TargetNameToUObjectModules = new Dictionary<string, List<UHTModuleInfo>>(StringComparer.InvariantCultureIgnoreCase);
-					HashSet<string> HotReloadModuleNamesForAllTargets = new HashSet<string>(StringComparer.OrdinalIgnoreCase);
-					List<BuildPrerequisites> TargetPrerequisites = new List<BuildPrerequisites>();
-					foreach (UEBuildTarget Target in Targets)
-					{
-						// Create the header cache for this target
-						FileReference DependencyCacheFile = DependencyCache.GetDependencyCachePathForTarget(Target.ProjectFile, Target.Platform, Target.TargetName);
-						CPPHeaders Headers = new CPPHeaders(Target.ProjectFile, DependencyCacheFile);
-						TargetToHeaders[Target] = Headers;
+							// Create the target
+							UEBuildTarget Target = UEBuildTarget.CreateTarget(TargetDesc, bSkipRulesCompile, BuildConfiguration.SingleFileToCompile != null, BuildConfiguration.bUsePrecompiled);
 
-						// Make sure the appropriate executor is selected
-						UEBuildPlatform BuildPlatform = UEBuildPlatform.GetBuildPlatform(Target.Platform);
-						BuildConfiguration.bAllowXGE &= BuildPlatform.CanUseXGE();
-						BuildConfiguration.bAllowDistcc &= BuildPlatform.CanUseDistcc();
-						BuildConfiguration.bAllowSNDBS &= BuildPlatform.CanUseSNDBS();
+							// Create the header cache for this target
+							FileReference DependencyCacheFile = DependencyCache.GetDependencyCachePathForTarget(Target.ProjectFile, Target.Platform, Target.TargetName);
+							CPPHeaders Headers = new CPPHeaders(Target.ProjectFile, DependencyCacheFile);
+							HeadersList.Add(Headers);
 
-						// When in 'assembler only' mode, we'll load this cache later on a worker thread.  It takes a long time to load!
-						if (!(!bIsGatheringBuild && bIsAssemblingBuild))
-						{
-							// Load the direct include dependency cache.
+							// Read the dependency cache
 							using(Timeline.ScopeEvent("DependencyCache.Create()"))
 							{
 								Headers.IncludeDependencyCache = DependencyCache.Create(DependencyCache.GetDependencyCachePathForTarget(Target.ProjectFile, Target.Platform, Target.TargetName));
 							}
-						}
 
-						if (bIsGatheringBuild)
-						{
-							List<FileItem> TargetOutputItems = new List<FileItem>();
-							List<UHTModuleInfo> TargetUObjectModules = new List<UHTModuleInfo>();
-							Dictionary<string, FileItem[]> TargetModuleNameToOutputItems = new Dictionary<string, FileItem[]>(StringComparer.OrdinalIgnoreCase);
+							// Build the target
+							const bool bIsAssemblingBuild = true;
+
+							List<FileItem> OutputItems = new List<FileItem>();
+							Dictionary<string, FileItem[]> ModuleNameToOutputItems = new Dictionary<string, FileItem[]>();
+							List<UHTModuleInfo> UObjectModules = new List<UHTModuleInfo>();
+							ActionGraph ActionGraph = new ActionGraph();
 							BuildPrerequisites Prerequisites = new BuildPrerequisites();
 
-							BuildResult = Target.Build(BuildConfiguration, TargetToHeaders[Target], TargetOutputItems, TargetModuleNameToOutputItems, TargetUObjectModules, WorkingSet, ActionGraph, Prerequisites, bIsAssemblingBuild);
-							
+							ECompilationResult BuildResult = Target.Build(BuildConfiguration, Headers, OutputItems, ModuleNameToOutputItems, UObjectModules, WorkingSet, ActionGraph, Prerequisites, bIsAssemblingBuild);
 							if (BuildResult != ECompilationResult.Succeeded)
 							{
-								break;
-							}
-
-							// Export a JSON dump of this target
-							if (BuildConfiguration.JsonExportFile != null)
-							{
-								Directory.CreateDirectory(Path.GetDirectoryName(Path.GetFullPath(BuildConfiguration.JsonExportFile)));
-								Target.ExportJson(BuildConfiguration.JsonExportFile);
+								return (int)BuildResult;
 							}
 
-							OutputItemsForAllTargets.AddRange(TargetOutputItems);
-
-							TargetPrerequisites.Add(Prerequisites);
-
-							// Update mapping of the target name to the list of UObject modules in this target
-							TargetNameToUObjectModules[Target.TargetName] = TargetUObjectModules;
-
-							// Merge the module name to output items map
-							foreach(KeyValuePair<string, FileItem[]> Pair in TargetModuleNameToOutputItems)
-							{
-								FileItem[] Items;
-								if(ModuleNameToOutputItems.TryGetValue(Pair.Key, out Items))
-								{
-									ModuleNameToOutputItems[Pair.Key] = Items.Concat(Pair.Value).ToArray();
-								}
-								else
-								{
-									ModuleNameToOutputItems[Pair.Key] = Pair.Value;
-								}
-							}
-
-							// Get the game module names
-							HotReloadModuleNamesForAllTargets.UnionWith(Target.GetHotReloadModuleNames());
-						}
-					}
-
-					if (BuildResult == ECompilationResult.Succeeded && !BuildConfiguration.bSkipBuild &&
-						(
-							(BuildConfiguration.bXGEExport && BuildConfiguration.bGenerateManifest) ||
-							(!GeneratingActionGraph && !BuildConfiguration.bGenerateManifest)
-						))
-					{
-						if (bIsGatheringBuild)
-						{
-							ActionGraph.FinalizeActionGraph();
-
-							Makefile = new TargetMakefile();
-							Makefile.AllActions = ActionGraph.AllActions;
-							Makefile.OutputItemsForAllTargets = OutputItemsForAllTargets;
+							// Create the makefile
+							Makefile = new TargetMakefile(Target.TargetType);
+							Makefile.ReceiptFile = Target.ReceiptFileName;
+							Makefile.Actions = ActionGraph.Actions;
+							Makefile.OutputItems = OutputItems;
+							Makefile.ModuleNameToOutputItems = ModuleNameToOutputItems;
+							Makefile.UObjectModules = UObjectModules;
+							Makefile.Prerequisites = Prerequisites;
+							Makefile.HotReloadModuleNames = Target.GetHotReloadModuleNames();
+							Makefile.Prerequisites = Prerequisites;
+							Makefile.ProjectIntermediateDirectory = Target.ProjectIntermediateDirectory;
+							Makefile.bDeployAfterCompile = Target.bDeployAfterCompile;
+							Makefile.bHasProjectScriptPlugin = Target.bHasProjectScriptPlugin;
+							Makefile.PreBuildScripts = Target.PreBuildStepScripts;
+
+							// Save the environment variables
 							foreach (System.Collections.DictionaryEntry EnvironmentVariable in Environment.GetEnvironmentVariables())
 							{
 								Makefile.EnvironmentVariables.Add(Tuple.Create((string)EnvironmentVariable.Key, (string)EnvironmentVariable.Value));
 							}
-							Makefile.TargetNameToUObjectModules = TargetNameToUObjectModules;
-							Makefile.ModuleNameToOutputItems = ModuleNameToOutputItems;
-							Makefile.HotReloadModuleNamesForAllTargets = HotReloadModuleNamesForAllTargets;
-							Makefile.Targets = Targets;
-							Makefile.TargetPrerequisites = TargetPrerequisites;
 
-							if (BuildConfiguration.bUseUBTMakefiles)
+							// Save the makefile for next time
+							using(Timeline.ScopeEvent("UBTMakefile.SaveUBTMakefile()"))
 							{
-								// We've been told to prepare to build, so let's go ahead and save out our action graph so that we can use in a later invocation 
-								// to assemble the build.  Even if we are configured to assemble the build in this same invocation, we want to save out the
-								// Makefile so that it can be used on subsequent 'assemble only' runs, for the fastest possible iteration times
-								// @todo ubtmake: Optimization: We could make 'gather + assemble' mode slightly faster by saving this while busy compiling (on our worker thread)
-								using(Timeline.ScopeEvent("UBTMakefile.SaveUBTMakefile()"))
-								{
-									TargetMakefile.SaveUBTMakefile(TargetDescs, Makefile);
-								}
+								TargetMakefile.SaveUBTMakefile(TargetDesc, Makefile);
 							}
 						}
-
-						HashSet<string> OnlyModuleNames = new HashSet<string>(StringComparer.OrdinalIgnoreCase);
-						foreach(string Argument in Arguments)
+						else
 						{
-							const string ModuleArgument = "-Module=";
-							if(Argument.StartsWith(ModuleArgument, StringComparison.OrdinalIgnoreCase))
+							// Restore the environment variables
+							foreach (Tuple<string, string> EnvironmentVariable in Makefile.EnvironmentVariables)
 							{
-								OnlyModuleNames.Add(Argument.Substring(ModuleArgument.Length));
+								Environment.SetEnvironmentVariable(EnvironmentVariable.Item1, EnvironmentVariable.Item2);
 							}
-						}
 
-						Dictionary<string, int> HotReloadModuleNameToSuffix = new Dictionary<string, int>(StringComparer.OrdinalIgnoreCase);
-						foreach(string Argument in Arguments)
-						{
-							const string ModuleWithSuffixArgument = "-ModuleWithSuffix=";
-							if(Argument.StartsWith(ModuleWithSuffixArgument, StringComparison.OrdinalIgnoreCase))
+							// Execute all the pre-build steps
+							if (!BuildConfiguration.bXGEExport)
 							{
-								string Value = Argument.Substring(ModuleWithSuffixArgument.Length);
-
-								int SuffixIdx = Value.LastIndexOf(',');
-								if(SuffixIdx == -1)
+								if (!Utils.ExecuteCustomBuildSteps(Makefile.PreBuildScripts))
 								{
-									throw new BuildException("Missing suffix argument from -ModuleWithSuffix=Name,Suffix");
+									return (int)ECompilationResult.OtherCompilationError;
 								}
+							}
 
-								string ModuleName = Value.Substring(0, SuffixIdx);
+							// If the target needs UHT to be run, we'll go ahead and do that now
+							if (Makefile.UObjectModules.Count > 0)
+							{
+								const bool bIsGatheringBuild = false;
+								const bool bIsAssemblingBuild = true;
 
-								int Suffix;
-								if(!Int32.TryParse(Value.Substring(SuffixIdx + 1), out Suffix))
+								FileReference ModuleInfoFileName = FileReference.Combine(Makefile.ProjectIntermediateDirectory, TargetDesc.Name + ".uhtmanifest");
+								ECompilationResult UHTResult = ExternalExecution.ExecuteHeaderToolIfNecessary(BuildConfiguration, TargetDesc.ProjectFile, TargetDesc.Name, Makefile.TargetType, Makefile.bHasProjectScriptPlugin, UObjectModules: Makefile.UObjectModules, ModuleInfoFileName: ModuleInfoFileName, bIsGatheringBuild: bIsGatheringBuild, bIsAssemblingBuild: bIsAssemblingBuild);
+								if(!UHTResult.Succeeded())
 								{
-									throw new BuildException("Suffix for modules must be an integer");
+									Log.TraceInformation("UnrealHeaderTool failed for target '" + TargetDesc.Name + "' (platform: " + TargetDesc.Platform.ToString() + ", module info: " + ModuleInfoFileName + ").");
+									return (int)UHTResult;
 								}
-
-								HotReloadModuleNameToSuffix[ModuleName] = Suffix;
 							}
 						}
 
-						HotReloadMode HotReloadMode = HotReloadMode.Disabled;
-						if(Arguments.Any(x => x.Equals("-ForceHotReload", StringComparison.InvariantCultureIgnoreCase)))
+						// Add the makefile to the list to build
+						Makefiles[TargetIdx] = Makefile;
+					}
+
+					// Execute the build
+					if(!BuildConfiguration.bSkipBuild)
+					{
+						// Make sure that none of the actions conflict with any other (producing output files differently, etc...)
+						if(!ActionGraph.CheckForConflicts(Makefiles.SelectMany(x => x.Actions)))
 						{
-							HotReloadMode = HotReloadMode.FromIDE;
+							Log.TraceInformation("Check log for additional details.");
+							return (int)ECompilationResult.OtherCompilationError;
 						}
-						else if (TargetDescs.Count == 1 && !Arguments.Any(x => x.Equals("-NoHotReload", StringComparison.InvariantCultureIgnoreCase)))
+
+						// Find all the actions to be executed
+						HashSet<Action>[] ActionsToExecute = new HashSet<Action>[TargetDescriptors.Count];
+						for(int TargetIdx = 0; TargetIdx < TargetDescriptors.Count; TargetIdx++)
 						{
-							if (BuildConfiguration.bAllowHotReloadFromIDE && HotReload.ShouldDoHotReloadFromIDE(BuildConfiguration, TargetDescs[0]))
-							{
-								HotReloadMode = HotReloadMode.FromIDE;
-							}
-							else if (HotReloadModuleNameToSuffix.Count > 0 && TargetDescs[0].ForeignPlugin == null)
-							{
-								HotReloadMode = HotReloadMode.FromEditor;
-							}
+							ActionsToExecute[TargetIdx] = GetActionsForTarget(BuildConfiguration, TargetDescriptors[TargetIdx], Makefiles[TargetIdx]);
 						}
-						TargetDescriptor HotReloadTargetDesc = (HotReloadMode != HotReloadMode.Disabled) ? TargetDescs[0] : null;
 
-						if (bIsAssemblingBuild)
+						// If there are multiple targets being built, make sure the actions for each one 
+						List<Action> MergedActionsToExecute = new List<Action>();
+						if(TargetDescriptors.Count == 1)
+						{
+							// Just take the actions from the first target
+							MergedActionsToExecute.AddRange(ActionsToExecute[0]);
+						}
+						else
 						{
-							// If we didn't build the graph in this session, then we'll need to load a cached one
-							if (!bIsGatheringBuild && BuildResult.Succeeded())
+							// Set of all output items. Knowing that there are no conflicts in produced items, we use this to eliminate duplicate actions.
+							Dictionary<FileItem, Action> OutputItemToProducingAction = new Dictionary<FileItem, Action>();
+							for(int TargetIdx = 0; TargetIdx < TargetDescriptors.Count; TargetIdx++)
 							{
-								ActionGraph.AllActions = Makefile.AllActions;
-								ActionGraph.LinkActionsAndItems();
-
-								OutputItemsForAllTargets = Makefile.OutputItemsForAllTargets;
-
-								ModuleNameToOutputItems = Makefile.ModuleNameToOutputItems;
-
-								HotReloadModuleNamesForAllTargets = Makefile.HotReloadModuleNamesForAllTargets;
-
-								foreach (Tuple<string, string> EnvironmentVariable in Makefile.EnvironmentVariables)
-								{
-									// @todo ubtmake: There may be some variables we do NOT want to clobber.
-									Environment.SetEnvironmentVariable(EnvironmentVariable.Item1, EnvironmentVariable.Item2);
-								}
-
-								// Execute all the pre-build steps
-								if (!BuildConfiguration.bXGEExport)
-								{
-									foreach (UEBuildTarget Target in Targets)
-									{
-										if (!Target.ExecuteCustomPreBuildSteps())
-										{
-											BuildResult = ECompilationResult.OtherCompilationError;
-											break;
-										}
-									}
-								}
-
-								// If any of the targets need UHT to be run, we'll go ahead and do that now
-								if (BuildResult.Succeeded())
+								string GroupPrefix = String.Format("{0}-{1}-{2}", TargetDescriptors[TargetIdx].Name, TargetDescriptors[TargetIdx].Platform, TargetDescriptors[TargetIdx].Configuration);
+								foreach(Action Action in ActionsToExecute[TargetIdx])
 								{
-									foreach (UEBuildTarget Target in Targets)
+									Action ExistingAction;
+									if(!OutputItemToProducingAction.TryGetValue(Action.ProducedItems[0], out ExistingAction))
 									{
-										List<UHTModuleInfo> TargetUObjectModules;
-										if (Makefile.TargetNameToUObjectModules.TryGetValue(Target.TargetName, out TargetUObjectModules))
-										{
-											if (TargetUObjectModules.Count > 0)
-											{
-												// Execute the header tool
-												FileReference ModuleInfoFileName = FileReference.Combine(Target.ProjectIntermediateDirectory, Target.TargetName + ".uhtmanifest");
-												ECompilationResult UHTResult = ExternalExecution.ExecuteHeaderToolIfNecessary(BuildConfiguration, Target.ProjectFile, Target.TargetName, Target.TargetType, Target.bHasProjectScriptPlugin,  UObjectModules: TargetUObjectModules, ModuleInfoFileName: ModuleInfoFileName, bIsGatheringBuild: bIsGatheringBuild, bIsAssemblingBuild: bIsAssemblingBuild);
-												if(!UHTResult.Succeeded())
-												{
-													Log.TraceInformation("UnrealHeaderTool failed for target '" + Target.TargetName + "' (platform: " + Target.Platform.ToString() + ", module info: " + ModuleInfoFileName + ").");
-													BuildResult = UHTResult;
-													break;
-												}
-											}
-										}
+										OutputItemToProducingAction[Action.ProducedItems[0]] = Action;
+										ExistingAction = Action;
 									}
+									ExistingAction.GroupNames.Add(GroupPrefix);
 								}
 							}
+							MergedActionsToExecute.AddRange(OutputItemToProducingAction.Values);
+						}
 
-							if (BuildResult.Succeeded())
-							{
-								// Get the output items we're trying to build right now
-								List<FileItem> OutputItems = OutputItemsForAllTargets;
-
-								// Filter them if we have any specific modules specified
-								if (OnlyModuleNames.Count > 0)
-								{
-									OutputItems = new List<FileItem>();
-									foreach(string OnlyModuleName in OnlyModuleNames)
-									{
-										FileItem[] OutputItemsForModule;
-										if(!ModuleNameToOutputItems.TryGetValue(OnlyModuleName, out OutputItemsForModule))
-										{
-											throw new BuildException("Unable to find output items for module '{0}'", OnlyModuleName);
-										}
-										OutputItems.AddRange(OutputItemsForModule);
-									}
-								}
-
-								// For now simply treat all object files as the root target.
-								Action[] PrerequisiteActions;
-								{
-									HashSet<Action> PrerequisiteActionsSet = new HashSet<Action>();
-									foreach (FileItem OutputItem in OutputItems)
-									{
-										ActionGraph.GatherPrerequisiteActions(OutputItem, ref PrerequisiteActionsSet);
-									}
-									PrerequisiteActions = PrerequisiteActionsSet.ToArray();
-								}
-
-								// Apply the previous hot reload state
-								HotReloadState HotReloadState = null;
-								if(HotReloadMode == HotReloadMode.Disabled)
-								{
-									// Delete any previous state files; they are no longer valid
-									foreach(TargetDescriptor TargetDesc in TargetDescs)
-									{
-										FileReference HotReloadStateFile = HotReloadState.GetLocation(TargetDesc.ProjectFile, TargetDesc.Name, TargetDesc.Platform, TargetDesc.Configuration, TargetDesc.Architecture);
-										HotReload.DeleteTemporaryFiles(HotReloadStateFile);
-									}
-								}
-								else
-								{
-									// Read the previous state file and apply it to the action graph
-									FileReference HotReloadStateFile = HotReloadState.GetLocation(TargetDescs[0].ProjectFile, TargetDescs[0].Name, TargetDescs[0].Platform, TargetDescs[0].Configuration, TargetDescs[0].Architecture);
-									if(FileReference.Exists(HotReloadStateFile))
-									{
-										HotReloadState = HotReloadState.Load(HotReloadStateFile);
-									}
-									else
-									{
-										HotReloadState = new HotReloadState();
-									}
-
-									// Update the action graph to produce these new files
-									HotReload.PatchActionGraph(PrerequisiteActions, HotReloadState.OriginalFileToHotReloadFile);
-
-									// Update the module to output file mapping
-									foreach(string HotReloadModuleName in HotReloadModuleNamesForAllTargets)
-									{
-										FileItem[] ModuleOutputItems = ModuleNameToOutputItems[HotReloadModuleName];
-										for(int Idx = 0; Idx < ModuleOutputItems.Length; Idx++)
-										{
-											FileReference NewLocation;
-											if(HotReloadState.OriginalFileToHotReloadFile.TryGetValue(ModuleOutputItems[Idx].Location, out NewLocation))
-											{
-												ModuleOutputItems[Idx] = FileItem.GetItemByFileReference(NewLocation);
-											}
-										}
-									}
-
-									// If we want a specific suffix on any modules, apply that now. We'll track the outputs later, but the suffix has to be forced  (and is always out of date if it doesn't exist).
-									if(HotReloadModuleNameToSuffix.Count > 0)
-									{
-										Dictionary<FileReference, FileReference> OldLocationToNewLocation = new Dictionary<FileReference, FileReference>();
-										foreach(string HotReloadModuleName in HotReloadModuleNamesForAllTargets)
-										{
-											int ModuleSuffix;
-											if(HotReloadModuleNameToSuffix.TryGetValue(HotReloadModuleName, out ModuleSuffix))
-											{
-												FileItem[] ModuleOutputItems = ModuleNameToOutputItems[HotReloadModuleName];
-												foreach(FileItem ModuleOutputItem in ModuleOutputItems)
-												{
-													FileReference OldLocation = ModuleOutputItem.Location;
-													FileReference NewLocation = HotReload.ReplaceSuffix(OldLocation, ModuleSuffix);
-													OldLocationToNewLocation[OldLocation] = NewLocation;
-												}
-											}
-										}
-										HotReload.PatchActionGraph(PrerequisiteActions, OldLocationToNewLocation);
-									}
-								}
-
-								// Create the dependencies cache
-								CppDependencyCache CppDependencies;
-								using(Timeline.ScopeEvent("Reading dependency cache"))
-								{
-									CppDependencies = CppDependencyCache.CreateHierarchy(Targets[0].ProjectFile, Targets[0].TargetName, Targets[0].Platform, Targets[0].Configuration, Targets[0].TargetType);
-								}
-
-								// Plan the actions to execute for the build.
-								Dictionary<UEBuildTarget, List<FileItem>> TargetToOutdatedPrerequisitesMap;
-								HashSet<Action> ActionsToExecute = ActionGraph.GetActionsToExecute(BuildConfiguration, PrerequisiteActions, Targets, CppDependencies, out TargetToOutdatedPrerequisitesMap);
-
-								// Patch action history for hot reload when running in assembler mode.  In assembler mode, the suffix on the output file will be
-								// the same for every invocation on that makefile, but we need a new suffix each time.
-								if (HotReloadMode != HotReloadMode.Disabled)
-								{
-									// For all the hot-reloadable modules that may need a unique suffix appended, build a mapping from output item to all the output items in that module. We can't 
-									// apply a suffix to one without applying a suffix to all of them.
-									Dictionary<FileItem, FileItem[]> HotReloadItemToDependentItems = new Dictionary<FileItem, FileItem[]>();
-									foreach(string HotReloadModuleName in HotReloadModuleNamesForAllTargets)
-									{
-										int ModuleSuffix;
-										if(!HotReloadModuleNameToSuffix.TryGetValue(HotReloadModuleName, out ModuleSuffix) || ModuleSuffix == -1)
-										{
-											FileItem[] ModuleOutputItems;
-											if(ModuleNameToOutputItems.TryGetValue(HotReloadModuleName, out ModuleOutputItems))
-											{
-												foreach(FileItem ModuleOutputItem in ModuleOutputItems)
-												{
-													HotReloadItemToDependentItems[ModuleOutputItem] = ModuleOutputItems;
-												}
-											}
-										} 
-									}
-
-									// Expand the list of actions to execute to include everything that references any files with a new suffix. Unlike a regular build, we can't ignore
-									// dependencies on import libraries under the assumption that a header would change if the API changes, because the dependency will be on a different DLL.
-									HashSet<FileItem> FilesRequiringSuffix = new HashSet<FileItem>(ActionsToExecute.SelectMany(x => x.ProducedItems).Where(x => HotReloadItemToDependentItems.ContainsKey(x)));
-									for(int LastNumFilesWithNewSuffix = 0; FilesRequiringSuffix.Count > LastNumFilesWithNewSuffix;)
-									{
-										LastNumFilesWithNewSuffix = FilesRequiringSuffix.Count;
-										foreach(Action PrerequisiteAction in PrerequisiteActions)
-										{
-											if(!ActionsToExecute.Contains(PrerequisiteAction))
-											{
-												foreach(FileItem ProducedItem in PrerequisiteAction.ProducedItems)
-												{
-													FileItem[] DependentItems;
-													if(HotReloadItemToDependentItems.TryGetValue(ProducedItem, out DependentItems))
-													{
-														ActionsToExecute.Add(PrerequisiteAction);
-														FilesRequiringSuffix.UnionWith(DependentItems);
-													}
-												}
-											}
-										}
-									}
-
-									// Build a list of file mappings
-									Dictionary<FileReference, FileReference> OldLocationToNewLocation = new Dictionary<FileReference, FileReference>();
-									foreach(FileItem FileRequiringSuffix in FilesRequiringSuffix)
-									{
-										FileReference OldLocation = FileRequiringSuffix.Location;
-										FileReference NewLocation = HotReload.ReplaceSuffix(OldLocation, HotReloadState.NextSuffix);
-										OldLocationToNewLocation[OldLocation] = NewLocation;
-									}
-
-									// Update the action graph with these new paths
-									HotReload.PatchActionGraph(PrerequisiteActions, OldLocationToNewLocation);
-
-									// Get a new list of actions to execute now that the graph has been modified
-									ActionsToExecute = ActionGraph.GetActionsToExecute(BuildConfiguration, PrerequisiteActions, Targets, CppDependencies, out TargetToOutdatedPrerequisitesMap);
-
-									// Build a mapping of all file items to their original
-									Dictionary<FileReference, FileReference> HotReloadFileToOriginalFile = new Dictionary<FileReference, FileReference>();
-									foreach(KeyValuePair<FileReference, FileReference> Pair in HotReloadState.OriginalFileToHotReloadFile)
-									{
-										HotReloadFileToOriginalFile[Pair.Value] = Pair.Key;
-									}
-									foreach(KeyValuePair<FileReference, FileReference> Pair in OldLocationToNewLocation)
-									{
-										FileReference OriginalLocation;
-										if(!HotReloadFileToOriginalFile.TryGetValue(Pair.Key, out OriginalLocation))
-										{
-											OriginalLocation = Pair.Key;
-										}
-										HotReloadFileToOriginalFile[Pair.Value] = OriginalLocation;
-									}
-
-									// Now filter out all the hot reload files and update the state
-									foreach(Action Action in ActionsToExecute)
-									{
-										foreach(FileItem ProducedItem in Action.ProducedItems)
-										{
-											FileReference OriginalLocation;
-											if(HotReloadFileToOriginalFile.TryGetValue(ProducedItem.Location, out OriginalLocation))
-											{
-												HotReloadState.OriginalFileToHotReloadFile[OriginalLocation] = ProducedItem.Location;
-												HotReloadState.TemporaryFiles.Add(ProducedItem.Location);
-											}
-										}
-									}
+						// Link all the actions together
+						ActionGraph.Link(MergedActionsToExecute);
 
-									// Increment the suffix for the next iteration
-									if(ActionsToExecute.Count > 0)
-									{
-										HotReloadState.NextSuffix++;
-									}
+						// Make sure the appropriate executor is selected
+						foreach(TargetDescriptor TargetDescriptor in TargetDescriptors)
+						{
+							UEBuildPlatform BuildPlatform = UEBuildPlatform.GetBuildPlatform(TargetDescriptor.Platform);
+							BuildConfiguration.bAllowXGE &= BuildPlatform.CanUseXGE();
+							BuildConfiguration.bAllowDistcc &= BuildPlatform.CanUseDistcc();
+							BuildConfiguration.bAllowSNDBS &= BuildPlatform.CanUseSNDBS();
+						}
 
-									// Save the new state
-									FileReference HotReloadStateFile = HotReloadState.GetLocation(TargetDescs[0].ProjectFile, TargetDescs[0].Name, TargetDescs[0].Platform, TargetDescs[0].Configuration, TargetDescs[0].Architecture);
-									HotReloadState.Save(HotReloadStateFile);
-								}
+						// Delete produced items that are outdated.
+						ActionGraph.DeleteOutdatedProducedItems(MergedActionsToExecute);
 
-								// Display some stats to the user.
-								Log.TraceVerbose(
-										"{0} actions, {1} outdated and requested actions",
-										ActionGraph.AllActions.Count,
-										ActionsToExecute.Count
-										);
+						// Save all the action histories now that files have been removed. We have to do this after deleting produced items to ensure that any
+						// items created during the build don't have the wrong command line.
+						ActionHistory.SaveAll();
 
-								// Execute the actions.
-								ExecutorTimer.Start();
-								if (BuildConfiguration.bXGEExport)
-								{
-									using(Timeline.ScopeEvent("XGE.ExportActions()"))
-									{
-										XGE.ExportActions(ActionsToExecute.ToList());
-									}
-									bSuccess = true;
-								}
-								else
-								{
-									using(Timeline.ScopeEvent("ActionGraph.ExecuteActions()"))
-									{
-										bSuccess = ActionGraph.ExecuteActions(BuildConfiguration, ActionsToExecute.ToList(), out ExecutorName);
-									}
-								}
-								ExecutorTimer.Stop();
+						// Create directories for the outdated produced items.
+						ActionGraph.CreateDirectoriesForProducedItems(MergedActionsToExecute);
 
-								// if the build succeeded, write the receipts and do any needed syncing
-								if (!bSuccess)
+						// Execute the actions.
+						ExecutorTimer.Start();
+						if (BuildConfiguration.bXGEExport)
+						{
+							using(Timeline.ScopeEvent("XGE.ExportActions()"))
+							{
+								XGE.ExportActions(MergedActionsToExecute);
+							}
+						}
+						else
+						{
+							using(Timeline.ScopeEvent("ActionGraph.ExecuteActions()"))
+							{
+								if(!ActionGraph.ExecuteActions(BuildConfiguration, MergedActionsToExecute, out ExecutorName))
 								{
-									BuildResult = ECompilationResult.OtherCompilationError;
+									return (int)ECompilationResult.OtherCompilationError;
 								}
 							}
+						}
+						ExecutorTimer.Stop();
 
-							// Run the deployment steps
-							if (BuildResult.Succeeded()
-								&& BuildConfiguration.SingleFileToCompile == null
-								&& !BuildConfiguration.bGenerateManifest
-								&& !BuildConfiguration.bXGEExport)
+						// Run the deployment steps
+						if (BuildConfiguration.SingleFileToCompile == null
+							&& !BuildConfiguration.bGenerateManifest
+							&& !BuildConfiguration.bXGEExport)
+						{
+							foreach(TargetMakefile Makefile in Makefiles)
 							{
-								foreach (UEBuildTarget Target in Targets)
+								if (Makefile.bDeployAfterCompile)
 								{
-									if (Target.bDeployAfterCompile && HotReloadModuleNameToSuffix.Count == 0)
-									{
-										Log.WriteLine(LogEventType.Console, "Deploying {0} {1} {2}...", Target.TargetName, Target.Platform, Target.Configuration);
-										TargetReceipt Receipt = TargetReceipt.Read(Target.ReceiptFileName);
-										UEBuildPlatform.GetBuildPlatform(Target.Platform).Deploy(Receipt);
-									}
+									TargetReceipt Receipt = TargetReceipt.Read(Makefile.ReceiptFile);
+									Log.TraceInformation("Deploying {0} {1} {2}...", Receipt.TargetName, Receipt.Platform, Receipt.Configuration);
+									UEBuildPlatform.GetBuildPlatform(Receipt.Platform).Deploy(Receipt);
 								}
 							}
 						}
 					}
 				}
+				return (int)ECompilationResult.Succeeded;
 			}
 			catch (Exception Ex)
 			{
 				Log.WriteException(Ex, (LogFile == null) ? null : LogFile.FullName);
-				BuildResult = ECompilationResult.OtherCompilationError;
+				return (int)ECompilationResult.OtherCompilationError;
 			}
 			finally
 			{
+				// Save the dependency caches
 				CppDependencyCache.SaveAll();
 
 				// Save the include dependency cache.
-				foreach (CPPHeaders Headers in TargetToHeaders.Values)
+				foreach (CPPHeaders Headers in HeadersList)
 				{
 					// NOTE: It's very important that we save the include cache, even if a build exception was thrown (compile error, etc), because we need to make sure that
 					//    any C++ include dependencies that we computed for out of date source files are saved.  Remember, the build may fail *after* some build products
@@ -802,26 +433,267 @@ namespace UnrealBuildTool
 				// Print some performance info
 				Log.TraceLog("DirectIncludes cache miss time: {0}s ({1} misses)", CPPHeaders.DirectIncludeCacheMissesTotalTime, CPPHeaders.TotalDirectIncludeCacheMisses);
 			}
-			return (int)BuildResult;
 		}
 
 		/// <summary>
-		/// Invalidates makefiles for given target.
+		/// Determine what needs to be built for a target
 		/// </summary>
-		/// <param name="Target">Target</param>
-		private static void InvalidateMakefiles(TargetDescriptor Target)
+		/// <param name="BuildConfiguration">The build configuration</param>
+		/// <param name="TargetDescriptor">Target being build</param>
+		/// <param name="Makefile">Makefile generated for this target</param>
+		/// <returns>Set of actions to execute</returns>
+		static HashSet<Action> GetActionsForTarget(BuildConfiguration BuildConfiguration, TargetDescriptor TargetDescriptor, TargetMakefile Makefile)
 		{
-			string[] MakefileNames = new string[] { "HotReloadMakefile.ubt", "Makefile.ubt" };
-			DirectoryReference BaseDir = TargetMakefile.GetUBTMakefileDirectoryPathForSingleTarget(Target);
+			// Create the action graph
+			ActionGraph.Link(Makefile.Actions);
+
+			// Parse the list of hot-reload module names
+			Dictionary<string, int> HotReloadModuleNameToSuffix = new Dictionary<string, int>(StringComparer.OrdinalIgnoreCase);
+			foreach(string ModuleWithSuffix in TargetDescriptor.AdditionalArguments.GetValues("-ModuleWithSuffix="))
+			{
+				int SuffixIdx = ModuleWithSuffix.LastIndexOf(',');
+				if(SuffixIdx == -1)
+				{
+					throw new BuildException("Missing suffix argument from -ModuleWithSuffix=Name,Suffix");
+				}
+
+				string ModuleName = ModuleWithSuffix.Substring(0, SuffixIdx);
+
+				int Suffix;
+				if(!Int32.TryParse(ModuleWithSuffix.Substring(SuffixIdx + 1), out Suffix))
+				{
+					throw new BuildException("Suffix for modules must be an integer");
+				}
+
+				HotReloadModuleNameToSuffix[ModuleName] = Suffix;
+			}
+
+			// Get the hot-reload mode
+			HotReloadMode HotReloadMode = HotReloadMode.Disabled;
+			if(TargetDescriptor.AdditionalArguments.Any(x => x.Equals("-ForceHotReload", StringComparison.InvariantCultureIgnoreCase)))
+			{
+				HotReloadMode = HotReloadMode.FromIDE;
+			}
+			else if (!TargetDescriptor.AdditionalArguments.Any(x => x.Equals("-NoHotReload", StringComparison.InvariantCultureIgnoreCase)))
+			{
+				if (BuildConfiguration.bAllowHotReloadFromIDE && HotReload.ShouldDoHotReloadFromIDE(BuildConfiguration, TargetDescriptor))
+				{
+					HotReloadMode = HotReloadMode.FromIDE;
+				}
+				else if (HotReloadModuleNameToSuffix.Count > 0 && TargetDescriptor.ForeignPlugin == null)
+				{
+					HotReloadMode = HotReloadMode.FromEditor;
+				}
+			}
+
+			// Get the output items we're trying to build
+			List<FileItem> TargetOutputItems = Makefile.OutputItems;
+
+			// Parse the list of modules to build
+			HashSet<string> OnlyModuleNames = new HashSet<string>(TargetDescriptor.AdditionalArguments.GetValues("-Module="));
+			if (OnlyModuleNames.Count > 0)
+			{
+				TargetOutputItems = new List<FileItem>();
+				foreach(string OnlyModuleName in OnlyModuleNames)
+				{
+					FileItem[] OutputItemsForModule;
+					if(!Makefile.ModuleNameToOutputItems.TryGetValue(OnlyModuleName, out OutputItemsForModule))
+					{
+						throw new BuildException("Unable to find output items for module '{0}'", OnlyModuleName);
+					}
+					TargetOutputItems.AddRange(OutputItemsForModule);
+				}
+			}
+
+			// Get the root prerequisite actions
+			List<Action> PrerequisiteActions = ActionGraph.GatherPrerequisiteActions(Makefile.Actions, new HashSet<FileItem>(TargetOutputItems));
+
+			// Get the path to the hot reload state file for this target
+			FileReference HotReloadStateFile = global::UnrealBuildTool.HotReloadState.GetLocation(TargetDescriptor.ProjectFile, TargetDescriptor.Name, TargetDescriptor.Platform, TargetDescriptor.Configuration, TargetDescriptor.Architecture);
+
+			// Apply the previous hot reload state
+			HotReloadState HotReloadState = null;
+			if(HotReloadMode == HotReloadMode.Disabled)
+			{
+				// Delete the previous state file
+				HotReload.DeleteTemporaryFiles(HotReloadStateFile);
+			}
+			else
+			{
+				// Read the previous state file and apply it to the action graph
+				if(FileReference.Exists(HotReloadStateFile))
+				{
+					HotReloadState = HotReloadState.Load(HotReloadStateFile);
+				}
+				else
+				{
+					HotReloadState = new HotReloadState();
+				}
+
+				// Update the action graph to produce these new files
+				HotReload.PatchActionGraph(PrerequisiteActions, HotReloadState.OriginalFileToHotReloadFile);
+
+				// Update the module to output file mapping
+				foreach(string HotReloadModuleName in Makefile.HotReloadModuleNames)
+				{
+					FileItem[] ModuleOutputItems = Makefile.ModuleNameToOutputItems[HotReloadModuleName];
+					for(int Idx = 0; Idx < ModuleOutputItems.Length; Idx++)
+					{
+						FileReference NewLocation;
+						if(HotReloadState.OriginalFileToHotReloadFile.TryGetValue(ModuleOutputItems[Idx].Location, out NewLocation))
+						{
+							ModuleOutputItems[Idx] = FileItem.GetItemByFileReference(NewLocation);
+						}
+					}
+				}
+
+				// If we want a specific suffix on any modules, apply that now. We'll track the outputs later, but the suffix has to be forced  (and is always out of date if it doesn't exist).
+				if(HotReloadModuleNameToSuffix.Count > 0)
+				{
+					Dictionary<FileReference, FileReference> OldLocationToNewLocation = new Dictionary<FileReference, FileReference>();
+					foreach(string HotReloadModuleName in Makefile.HotReloadModuleNames)
+					{
+						int ModuleSuffix;
+						if(HotReloadModuleNameToSuffix.TryGetValue(HotReloadModuleName, out ModuleSuffix))
+						{
+							FileItem[] ModuleOutputItems = Makefile.ModuleNameToOutputItems[HotReloadModuleName];
+							foreach(FileItem ModuleOutputItem in ModuleOutputItems)
+							{
+								FileReference OldLocation = ModuleOutputItem.Location;
+								FileReference NewLocation = HotReload.ReplaceSuffix(OldLocation, ModuleSuffix);
+								OldLocationToNewLocation[OldLocation] = NewLocation;
+							}
+						}
+					}
+					HotReload.PatchActionGraph(PrerequisiteActions, OldLocationToNewLocation);
+				}
+
+				// Re-link the action graph
+				ActionGraph.Link(PrerequisiteActions);
+			}
+
+			// Create the dependencies cache
+			CppDependencyCache CppDependencies;
+			using(Timeline.ScopeEvent("Reading dependency cache"))
+			{
+				CppDependencies = CppDependencyCache.CreateHierarchy(TargetDescriptor.ProjectFile, TargetDescriptor.Name, TargetDescriptor.Platform, TargetDescriptor.Configuration, Makefile.TargetType);
+			}
+
+			// Create the action history
+			ActionHistory History;
+			using(Timeline.ScopeEvent("Reading action history"))
+			{
+				History = ActionHistory.CreateHierarchy(TargetDescriptor.ProjectFile, TargetDescriptor.Name, TargetDescriptor.Platform, Makefile.TargetType);
+			}
+
+			// Plan the actions to execute for the build.
+			HashSet<Action> TargetActionsToExecute = ActionGraph.GetActionsToExecute(BuildConfiguration, Makefile.Actions, PrerequisiteActions, CppDependencies, History);
 
-			foreach (string MakefileName in MakefileNames)
+			// Patch action history for hot reload when running in assembler mode.  In assembler mode, the suffix on the output file will be
+			// the same for every invocation on that makefile, but we need a new suffix each time.
+			if (HotReloadMode != HotReloadMode.Disabled)
 			{
-				FileReference MakefileRef = FileReference.Combine(BaseDir, MakefileName);
-				if (FileReference.Exists(MakefileRef))
+				// For all the hot-reloadable modules that may need a unique suffix appended, build a mapping from output item to all the output items in that module. We can't 
+				// apply a suffix to one without applying a suffix to all of them.
+				Dictionary<FileItem, FileItem[]> HotReloadItemToDependentItems = new Dictionary<FileItem, FileItem[]>();
+				foreach(string HotReloadModuleName in Makefile.HotReloadModuleNames)
 				{
-					FileReference.Delete(MakefileRef);
+					int ModuleSuffix;
+					if(!HotReloadModuleNameToSuffix.TryGetValue(HotReloadModuleName, out ModuleSuffix) || ModuleSuffix == -1)
+					{
+						FileItem[] ModuleOutputItems;
+						if(Makefile.ModuleNameToOutputItems.TryGetValue(HotReloadModuleName, out ModuleOutputItems))
+						{
+							foreach(FileItem ModuleOutputItem in ModuleOutputItems)
+							{
+								HotReloadItemToDependentItems[ModuleOutputItem] = ModuleOutputItems;
+							}
+						}
+					} 
 				}
+
+				// Expand the list of actions to execute to include everything that references any files with a new suffix. Unlike a regular build, we can't ignore
+				// dependencies on import libraries under the assumption that a header would change if the API changes, because the dependency will be on a different DLL.
+				HashSet<FileItem> FilesRequiringSuffix = new HashSet<FileItem>(TargetActionsToExecute.SelectMany(x => x.ProducedItems).Where(x => HotReloadItemToDependentItems.ContainsKey(x)));
+				for(int LastNumFilesWithNewSuffix = 0; FilesRequiringSuffix.Count > LastNumFilesWithNewSuffix;)
+				{
+					LastNumFilesWithNewSuffix = FilesRequiringSuffix.Count;
+					foreach(Action PrerequisiteAction in PrerequisiteActions)
+					{
+						if(!TargetActionsToExecute.Contains(PrerequisiteAction))
+						{
+							foreach(FileItem ProducedItem in PrerequisiteAction.ProducedItems)
+							{
+								FileItem[] DependentItems;
+								if(HotReloadItemToDependentItems.TryGetValue(ProducedItem, out DependentItems))
+								{
+									TargetActionsToExecute.Add(PrerequisiteAction);
+									FilesRequiringSuffix.UnionWith(DependentItems);
+								}
+							}
+						}
+					}
+				}
+
+				// Build a list of file mappings
+				Dictionary<FileReference, FileReference> OldLocationToNewLocation = new Dictionary<FileReference, FileReference>();
+				foreach(FileItem FileRequiringSuffix in FilesRequiringSuffix)
+				{
+					FileReference OldLocation = FileRequiringSuffix.Location;
+					FileReference NewLocation = HotReload.ReplaceSuffix(OldLocation, HotReloadState.NextSuffix);
+					OldLocationToNewLocation[OldLocation] = NewLocation;
+				}
+
+				// Update the action graph with these new paths
+				HotReload.PatchActionGraph(PrerequisiteActions, OldLocationToNewLocation);
+
+				// Get a new list of actions to execute now that the graph has been modified
+				TargetActionsToExecute = ActionGraph.GetActionsToExecute(BuildConfiguration, Makefile.Actions, PrerequisiteActions, CppDependencies, History);
+
+				// Build a mapping of all file items to their original
+				Dictionary<FileReference, FileReference> HotReloadFileToOriginalFile = new Dictionary<FileReference, FileReference>();
+				foreach(KeyValuePair<FileReference, FileReference> Pair in HotReloadState.OriginalFileToHotReloadFile)
+				{
+					HotReloadFileToOriginalFile[Pair.Value] = Pair.Key;
+				}
+				foreach(KeyValuePair<FileReference, FileReference> Pair in OldLocationToNewLocation)
+				{
+					FileReference OriginalLocation;
+					if(!HotReloadFileToOriginalFile.TryGetValue(Pair.Key, out OriginalLocation))
+					{
+						OriginalLocation = Pair.Key;
+					}
+					HotReloadFileToOriginalFile[Pair.Value] = OriginalLocation;
+				}
+
+				// Now filter out all the hot reload files and update the state
+				foreach(Action Action in TargetActionsToExecute)
+				{
+					foreach(FileItem ProducedItem in Action.ProducedItems)
+					{
+						FileReference OriginalLocation;
+						if(HotReloadFileToOriginalFile.TryGetValue(ProducedItem.Location, out OriginalLocation))
+						{
+							HotReloadState.OriginalFileToHotReloadFile[OriginalLocation] = ProducedItem.Location;
+							HotReloadState.TemporaryFiles.Add(ProducedItem.Location);
+						}
+					}
+				}
+
+				// Increment the suffix for the next iteration
+				if(TargetActionsToExecute.Count > 0)
+				{
+					HotReloadState.NextSuffix++;
+				}
+
+				// Save the new state
+				HotReloadState.Save(HotReloadStateFile);
+
+				// Prevent this target from deploying
+				Makefile.bDeployAfterCompile = false;
 			}
+
+			return TargetActionsToExecute;
 		}
 	}
 }
diff --git a/Engine/Source/Programs/UnrealBuildTool/Modes/CleanMode.cs b/Engine/Source/Programs/UnrealBuildTool/Modes/CleanMode.cs
--- a/Engine/Source/Programs/UnrealBuildTool/Modes/CleanMode.cs
+++ b/Engine/Source/Programs/UnrealBuildTool/Modes/CleanMode.cs
@@ -128,9 +128,8 @@ namespace UnrealBuildTool
 
 				// Add all the makefiles and caches to be deleted
 				FilesToDelete.Add(DependencyCache.GetDependencyCachePathForTarget(Target.ProjectFile, Target.Platform, Target.Name));
-				FilesToDelete.Add(TargetMakefile.GetUBTMakefilePath(Target.ProjectFile, Target.Platform, Target.Configuration, Target.Name, false));
-				FilesToDelete.Add(TargetMakefile.GetUBTMakefilePath(Target.ProjectFile, Target.Platform, Target.Configuration, Target.Name, true));
-				FilesToDelete.Add(ActionHistory.GeneratePathForTarget(Target.ProjectFile, Target.Name, Target.Platform, Target.Architecture, (Target.BuildEnvironment == TargetBuildEnvironment.Shared)));
+				FilesToDelete.Add(TargetMakefile.GetLocation(Target.ProjectFile, Target.Name, Target.Platform, Target.Configuration));
+				FilesToDelete.UnionWith(ActionHistory.GetFilesToClean(Target.ProjectFile, Target.Name, Target.Platform, Target.Type));
 
 				// Add all the intermediate folders to be deleted
 				foreach (DirectoryReference BaseDir in BaseDirs)
diff --git a/Engine/Source/Programs/UnrealBuildTool/Modes/JsonExportMode.cs b/Engine/Source/Programs/UnrealBuildTool/Modes/JsonExportMode.cs
new file mode 100644
--- /dev/null
+++ b/Engine/Source/Programs/UnrealBuildTool/Modes/JsonExportMode.cs
@@ -0,0 +1,45 @@
+// Copyright 1998-2019 Epic Games, Inc. All Rights Reserved.
+
+using System;
+using System.Collections.Generic;
+using System.Linq;
+using System.Text;
+using System.Threading.Tasks;
+using Tools.DotNETCommon;
+
+namespace UnrealBuildTool
+{
+	/// <summary>
+	/// Exports a target as a JSON file
+	/// </summary>
+	[ToolMode("JsonExport", ToolModeOptions.XmlConfig | ToolModeOptions.BuildPlatforms | ToolModeOptions.SingleInstance)]
+	class JsonExportMode : ToolMode
+	{
+		/// <summary>
+		/// Execute this command
+		/// </summary>
+		/// <param name="Arguments">Command line arguments</param>
+		/// <returns>Exit code (always zero)</returns>
+		public override int Execute(CommandLineArguments Arguments)
+		{
+			List<TargetDescriptor> TargetDescriptors = TargetDescriptor.ParseCommandLine(Arguments, false, false);
+			foreach(TargetDescriptor TargetDescriptor in TargetDescriptors)
+			{
+				// Create the target
+				UEBuildTarget Target = UEBuildTarget.CreateTarget(TargetDescriptor, false, false, false);
+
+				// Get the output file
+				FileReference OutputFile = TargetDescriptor.AdditionalArguments.GetFileReferenceOrDefault("-OutputFile=", null);
+				if(OutputFile == null)
+				{
+					OutputFile = Target.ReceiptFileName.ChangeExtension(".json");
+				}
+
+				// Write the output file
+				Log.TraceInformation("Writing {0}...", OutputFile);
+				Target.ExportJson(OutputFile);
+			}
+			return 0;
+		}
+	}
+}
diff --git a/Engine/Source/Programs/UnrealBuildTool/ProjectFiles/ProjectFileGenerator.cs b/Engine/Source/Programs/UnrealBuildTool/ProjectFiles/ProjectFileGenerator.cs
--- a/Engine/Source/Programs/UnrealBuildTool/ProjectFiles/ProjectFileGenerator.cs
+++ b/Engine/Source/Programs/UnrealBuildTool/ProjectFiles/ProjectFileGenerator.cs
@@ -1731,10 +1731,10 @@ namespace UnrealBuildTool
 							// If we're generating project files, then go ahead and wipe out the existing UBTMakefile for every target, to make sure that
 							// it gets a full dependency scan next time.
 							// NOTE: This is just a safeguard and doesn't have to be perfect.  We also check for newer project file timestamps in LoadUBTMakefile()
-							FileReference UBTMakefilePath = TargetMakefile.GetUBTMakefilePath(new List<TargetDescriptor>{ TargetDesc });
-							if (FileReference.Exists(UBTMakefilePath))
+							FileReference MakefileLocation = TargetMakefile.GetLocation(TargetDesc.ProjectFile, TargetDesc.Name, TargetDesc.Platform, TargetDesc.Configuration);
+							if (FileReference.Exists(MakefileLocation))
 							{
-								FileReference.Delete(UBTMakefilePath);
+								FileReference.Delete(MakefileLocation);
 							}
 						}
 						catch(Exception Ex)
diff --git a/Engine/Source/Programs/UnrealBuildTool/System/Action.cs b/Engine/Source/Programs/UnrealBuildTool/System/Action.cs
--- a/Engine/Source/Programs/UnrealBuildTool/System/Action.cs
+++ b/Engine/Source/Programs/UnrealBuildTool/System/Action.cs
@@ -64,6 +64,11 @@ namespace UnrealBuildTool
 		/// </summary>
 		public FileItem DependencyListFile;
 
+		/// <summary>
+		/// Set of other actions that this action depends on. This set is built when the action graph is linked.
+		/// </summary>
+		public HashSet<Action> PrerequisiteActions;
+
 		/// <summary>
 		/// Directory from which to execute the program to create produced items
 		/// </summary>
@@ -94,6 +99,11 @@ namespace UnrealBuildTool
 		/// </summary>
 		public string StatusDescription = "...";
 
+		/// <summary>
+		/// If set, will be output whenever the group differs to the last executed action. Set when executing multiple targets at once.
+		/// </summary>
+		public List<string> GroupNames = new List<string>();
+
 		/// <summary>
 		/// True if this action is allowed to be run on a remote machine when a distributed build system is being used, such as XGE
 		/// </summary>
@@ -136,11 +146,6 @@ namespace UnrealBuildTool
 		/// </summary>
 		public int NumTotalDependentActions = 0;
 
-		/// <summary>
-		/// Relative cost of producing items for this action.
-		/// </summary>
-		public long RelativeCost = 0;
-
 
 		///
 		/// Assembly only (not serialized)
@@ -207,6 +212,81 @@ namespace UnrealBuildTool
 			Writer.WriteFileItem(DependencyListFile);
 		}
 
+		/// <summary>
+		/// Finds conflicts betwee two actions, and prints them to the log
+		/// </summary>
+		/// <param name="Other">Other action to compare to.</param>
+		/// <returns>True if any conflicts were found, false otherwise.</returns>
+		public bool CheckForConflicts(Action Other)
+		{
+			bool bResult = true;
+			if(ActionType != Other.ActionType)
+			{
+				LogConflict("action type is different", ActionType.ToString(), Other.ActionType.ToString());
+				bResult = false;
+			}
+			if(!Enumerable.SequenceEqual(PrerequisiteItems, Other.PrerequisiteItems))
+			{
+				LogConflict("prerequisites are different", String.Join(", ", PrerequisiteItems.Select(x => x.Location)), String.Join(", ", Other.PrerequisiteItems.Select(x => x.Location)));
+				bResult = false;
+			}
+			if(!Enumerable.SequenceEqual(DeleteItems, Other.DeleteItems))
+			{
+				LogConflict("deleted items are different", String.Join(", ", DeleteItems.Select(x => x.Location)), String.Join(", ", Other.DeleteItems.Select(x => x.Location)));
+				bResult = false;
+			}
+			if(DependencyListFile != Other.DependencyListFile)
+			{
+				LogConflict("dependency list is different", (DependencyListFile == null)? "(none)" : DependencyListFile.AbsolutePath, (Other.DependencyListFile == null)? "(none)" : Other.DependencyListFile.AbsolutePath);
+				bResult = false;
+			}
+			if(WorkingDirectory != Other.WorkingDirectory)
+			{
+				LogConflict("working directory is different", WorkingDirectory, Other.WorkingDirectory);
+				bResult = false;
+			}
+			if(CommandPath != Other.CommandPath)
+			{
+				LogConflict("command path is different", CommandPath, Other.CommandPath);
+				bResult = false;
+			}
+			if(CommandArguments != Other.CommandArguments)
+			{
+				LogConflict("command arguments are different", CommandArguments, Other.CommandArguments);
+				bResult = false;
+			}
+			return bResult;
+		}
+
+		/// <summary>
+		/// Adds the description of a merge error to an output message
+		/// </summary>
+		/// <param name="Description">Description of the difference</param>
+		/// <param name="OldValue">Previous value for the field</param>
+		/// <param name="NewValue">Conflicting value for the field</param>
+		void LogConflict(string Description, string OldValue, string NewValue)
+		{
+			Log.TraceError("Unable to merge actions producing {0}: {1}", ProducedItems[0].Location.GetFileName(), Description);
+			Log.TraceLog("  Previous: {0}", OldValue);
+			Log.TraceLog("  Conflict: {0}", NewValue);
+		}
+
+		/// <summary>
+		/// Increment the number of dependents, recursively
+		/// </summary>
+		/// <param name="VisitedActions">Set of visited actions</param>
+		public void IncrementDependentCount(HashSet<Action> VisitedActions)
+		{
+			if(VisitedActions.Add(this))
+			{
+				NumTotalDependentActions++;
+				foreach(Action PrerequisiteAction in PrerequisiteActions)
+				{
+					PrerequisiteAction.IncrementDependentCount(VisitedActions);
+				}
+			}
+		}
+
 		/// <summary>
 		/// Compares two actions based on total number of dependent items, descending.
 		/// </summary>
@@ -219,12 +299,7 @@ namespace UnrealBuildTool
 			{
 				return Math.Sign(B.NumTotalDependentActions - A.NumTotalDependentActions);
 			}
-			// Secondary sort criteria is relative cost.
-			if (B.RelativeCost != A.RelativeCost)
-			{
-				return Math.Sign(B.RelativeCost - A.RelativeCost);
-			}
-			// Tertiary sort criteria is number of pre-requisites.
+			// Secondary sort criteria is number of pre-requisites.
 			else
 			{
 				return Math.Sign(B.PrerequisiteItems.Count - A.PrerequisiteItems.Count);
diff --git a/Engine/Source/Programs/UnrealBuildTool/System/ActionGraph.cs b/Engine/Source/Programs/UnrealBuildTool/System/ActionGraph.cs
--- a/Engine/Source/Programs/UnrealBuildTool/System/ActionGraph.cs
+++ b/Engine/Source/Programs/UnrealBuildTool/System/ActionGraph.cs
@@ -17,17 +17,16 @@ namespace UnrealBuildTool
 		/// <summary>
 		/// List of all the actions
 		/// </summary>
-		public List<Action> AllActions = new List<Action>();
+		public List<Action> Actions = new List<Action>();
 
 		public ActionGraph()
 		{
-			XmlConfig.ApplyTo(this);
 		}
 
 		public Action Add(ActionType Type)
 		{
 			Action NewAction = new Action(Type);
-			AllActions.Add(NewAction);
+			Actions.Add(NewAction);
 			return NewAction;
 		}
 
@@ -39,22 +38,76 @@ namespace UnrealBuildTool
 			return NewAction;
 		}
 
-		public void FinalizeActionGraph()
+		/// <summary>
+		/// Links the actions together and sets up their dependencies
+		/// </summary>
+		/// <param name="Actions">List of actions in the graph</param>
+		public static void Link(List<Action> Actions)
 		{
-			// Link producing actions to the items they produce.
-			LinkActionsAndItems();
+			// Build a map from item to its producing action
+			Dictionary<FileItem, Action> ItemToProducingAction = new Dictionary<FileItem, Action>();
+			foreach (Action Action in Actions)
+			{
+				foreach (FileItem ProducedItem in Action.ProducedItems)
+				{
+					ItemToProducingAction[ProducedItem] = Action;
+				}
+			}
+
+			// Check for cycles
+			DetectActionGraphCycles(Actions, ItemToProducingAction);
 
-			// Detect cycles in the action graph.
-			DetectActionGraphCycles();
+			// Use this map to add all the prerequisite actions
+			foreach (Action Action in Actions)
+			{
+				Action.PrerequisiteActions = new HashSet<Action>();
+				foreach(FileItem PrerequisiteItem in Action.PrerequisiteItems)
+				{
+					Action PrerequisiteAction;
+					if(ItemToProducingAction.TryGetValue(PrerequisiteItem, out PrerequisiteAction))
+					{
+						Action.PrerequisiteActions.Add(PrerequisiteAction);
+					}
+				}
+			}
 
-			// Sort action list by "cost" in descending order to improve parallelism.
-			SortActionList();
+			// Sort the action graph
+			SortActionList(Actions);
+		}
+
+		/// <summary>
+		/// Checks a set of actions for conflicts (ie. different actions producing the same output items)
+		/// </summary>
+		/// <param name="Actions">The set of actions to check</param>
+		/// <returns>True if the actions are ok, false if there are conflicts</returns>
+		public static bool CheckForConflicts(IEnumerable<Action> Actions)
+		{
+			bool bResult = true;
+
+			Dictionary<FileItem, Action> ItemToProducingAction = new Dictionary<FileItem, Action>();
+			foreach(Action Action in Actions)
+			{
+				foreach(FileItem ProducedItem in Action.ProducedItems)
+				{
+					Action ExistingAction;
+					if(ItemToProducingAction.TryGetValue(ProducedItem, out ExistingAction))
+					{
+						bResult &= ExistingAction.CheckForConflicts(Action);
+					}
+					else
+					{
+						ItemToProducingAction.Add(ProducedItem, Action);
+					}
+				}
+			}
+
+			return bResult;
 		}
 
 		/// <summary>
 		/// Builds a list of actions that need to be executed to produce the specified output items.
 		/// </summary>
-		public HashSet<Action> GetActionsToExecute(BuildConfiguration BuildConfiguration, Action[] PrerequisiteActions, List<UEBuildTarget> Targets, CppDependencyCache CppDependencies, out Dictionary<UEBuildTarget, List<FileItem>> TargetToOutdatedPrerequisitesMap)
+		public static HashSet<Action> GetActionsToExecute(BuildConfiguration BuildConfiguration, List<Action> Actions, List<Action> PrerequisiteActions, CppDependencyCache CppDependencies, ActionHistory History)
 		{
 			ITimelineEvent GetActionsToExecuteTimer = Timeline.ScopeEvent("ActionGraph.GetActionsToExecute()");
 
@@ -67,55 +120,25 @@ namespace UnrealBuildTool
 
 			// For all targets, build a set of all actions that are outdated.
 			Dictionary<Action, bool> OutdatedActionDictionary = new Dictionary<Action, bool>();
-			List<ActionHistory> HistoryList = new List<ActionHistory>();
-			HashSet<FileReference> OpenHistoryFiles = new HashSet<FileReference>();
-			TargetToOutdatedPrerequisitesMap = new Dictionary<UEBuildTarget, List<FileItem>>();
-			foreach (UEBuildTarget BuildTarget in Targets)	// @todo ubtmake: Optimization: Ideally we don't even need to know about targets for ubtmake -- everything would come from the files
-			{
-				FileReference HistoryFilename = ActionHistory.GeneratePathForTarget(BuildTarget.ProjectFile, BuildTarget.TargetName, BuildTarget.Platform, BuildTarget.Architecture, BuildTarget.bUseSharedBuildEnvironment);
-				if (!OpenHistoryFiles.Contains(HistoryFilename))		// @todo ubtmake: Optimization: We should be able to move the command-line outdatedness and build product deletion over to the 'gather' phase, as the command-lines won't change between assembler runs
-				{
-					ActionHistory History = new ActionHistory(HistoryFilename.FullName);
-					HistoryList.Add(History);
-					OpenHistoryFiles.Add(HistoryFilename);
-					GatherAllOutdatedActions(History, OutdatedActionDictionary, CppDependencies, BuildConfiguration.bIgnoreOutdatedImportLibraries);
-				}
-			}
+			GatherAllOutdatedActions(Actions, History, OutdatedActionDictionary, CppDependencies, BuildConfiguration.bIgnoreOutdatedImportLibraries);
 
 			// If we're only compiling a single file, we should always compile and should never link.
 			if (BuildConfiguration.SingleFileToCompile != null)
 			{
 				// Never do anything but compile the target file
-				AllActions.RemoveAll(x => x.ActionType != ActionType.Compile);
+				Actions.RemoveAll(x => x.ActionType != ActionType.Compile);
 
 				// Check all of the leftover compilation actions for the one we want... that one is always outdated.
 				FileItem SingleFileToCompile = FileItem.GetItemByFileReference(BuildConfiguration.SingleFileToCompile);
-				foreach (Action Action in AllActions)
+				foreach (Action Action in Actions)
 				{
 					bool bIsSingleFileAction = Action.PrerequisiteItems.Contains(SingleFileToCompile);
 					OutdatedActionDictionary[Action] = bIsSingleFileAction;
 				}
-
-				// Don't save the history of a single file compilation.
-				HistoryList.Clear();
-			}
-
-			// Delete produced items that are outdated.
-			DeleteOutdatedProducedItems(OutdatedActionDictionary);
-
-			// Save the action history.
-			// This must happen after deleting outdated produced items to ensure that the action history on disk doesn't have
-			// command-lines that don't match the produced items on disk.
-			foreach (ActionHistory TargetHistory in HistoryList)
-			{
-				TargetHistory.Save();
 			}
 
-			// Create directories for the outdated produced items.
-			CreateDirectoriesForProducedItems(OutdatedActionDictionary);
-
 			// Build a list of actions that are both needed for this target and outdated.
-			HashSet<Action> ActionsToExecute = new HashSet<Action>(AllActions.Where(Action => Action.CommandPath != null && IsActionOutdatedMap.ContainsKey(Action) && OutdatedActionDictionary[Action]));
+			HashSet<Action> ActionsToExecute = new HashSet<Action>(Actions.Where(Action => Action.CommandPath != null && IsActionOutdatedMap.ContainsKey(Action) && OutdatedActionDictionary[Action]));
 
 			GetActionsToExecuteTimer.Finish();
 
@@ -177,7 +200,6 @@ namespace UnrealBuildTool
 					}
 				}
 			}
-			// Nothing to execute.
 			else
 			{
 				ExecutorName = "NoActionsToExecute";
@@ -188,109 +210,30 @@ namespace UnrealBuildTool
 		}
 
 		/// <summary>
-		/// Links actions with their prerequisite and produced items into an action graph.
+		/// Sorts the action list for improved parallelism with local execution.
 		/// </summary>
-		public void LinkActionsAndItems()
-		{
-			foreach (Action Action in AllActions)
-			{
-				foreach (FileItem ProducedItem in Action.ProducedItems)
-				{
-					ProducedItem.ProducingAction = Action;
-				}
-			}
-		}
-		static string SplitFilename(string Filename, out string PlatformSuffix, out string ConfigSuffix, out string ProducedItemExtension)
+		static void SortActionList(List<Action> Actions)
 		{
-			string WorkingString = Filename;
-			ProducedItemExtension = Path.GetExtension(WorkingString);
-			if (!WorkingString.EndsWith(ProducedItemExtension))
+			// Clear the current dependent count
+			foreach(Action Action in Actions)
 			{
-				throw new BuildException("Bogus extension");
+				Action.NumTotalDependentActions = 0;
 			}
-			WorkingString = WorkingString.Substring(0, WorkingString.Length - ProducedItemExtension.Length);
 
-			ConfigSuffix = "";
-			foreach (UnrealTargetConfiguration CurConfig in Enum.GetValues(typeof(UnrealTargetConfiguration)))
-			{
-				if (CurConfig != UnrealTargetConfiguration.Unknown)
-				{
-					string Test = "-" + CurConfig;
-					if (WorkingString.EndsWith(Test))
-					{
-						WorkingString = WorkingString.Substring(0, WorkingString.Length - Test.Length);
-						ConfigSuffix = Test;
-						break;
-					}
-				}
-			}
-			PlatformSuffix = "";
-			foreach (object CurPlatform in Enum.GetValues(typeof(UnrealTargetPlatform)))
+			// Increment all the dependencies
+			foreach(Action Action in Actions)
 			{
-				string Test = "-" + CurPlatform;
-				if (WorkingString.EndsWith(Test))
-				{
-					WorkingString = WorkingString.Substring(0, WorkingString.Length - Test.Length);
-					PlatformSuffix = Test;
-					break;
-				}
+				Action.IncrementDependentCount(new HashSet<Action>());
 			}
-			return WorkingString;
-		}
 
-		/// <summary>
-		/// Sorts the action list for improved parallelism with local execution.
-		/// </summary>
-		public void SortActionList()
-		{
-			// Mapping from action to a list of actions that directly or indirectly depend on it (up to a certain depth).
-			Dictionary<Action, HashSet<Action>> ActionToDependentActionsMap = new Dictionary<Action, HashSet<Action>>();
-			// Perform multiple passes over all actions to propagate dependencies.
-			const int MaxDepth = 5;
-			for (int Pass = 0; Pass < MaxDepth; Pass++)
-			{
-				foreach (Action DependendAction in AllActions)
-				{
-					foreach (FileItem PrerequisiteItem in DependendAction.PrerequisiteItems)
-					{
-						Action PrerequisiteAction = PrerequisiteItem.ProducingAction;
-						if (PrerequisiteAction != null)
-						{
-							HashSet<Action> DependentActions = null;
-							if (ActionToDependentActionsMap.ContainsKey(PrerequisiteAction))
-							{
-								DependentActions = ActionToDependentActionsMap[PrerequisiteAction];
-							}
-							else
-							{
-								DependentActions = new HashSet<Action>();
-								ActionToDependentActionsMap[PrerequisiteAction] = DependentActions;
-							}
-							// Add dependent action...
-							DependentActions.Add(DependendAction);
-							// ... and all actions depending on it.
-							if (ActionToDependentActionsMap.ContainsKey(DependendAction))
-							{
-								DependentActions.UnionWith(ActionToDependentActionsMap[DependendAction]);
-							}
-						}
-					}
-				}
-
-			}
-			// At this point we have a list of dependent actions for each action, up to MaxDepth layers deep.
-			foreach (KeyValuePair<Action, HashSet<Action>> ActionMap in ActionToDependentActionsMap)
-			{
-				ActionMap.Key.NumTotalDependentActions = ActionMap.Value.Count;
-			}
 			// Sort actions by number of actions depending on them, descending. Secondary sort criteria is file size.
-			AllActions.Sort(Action.Compare);
+			Actions.Sort(Action.Compare);
 		}
 
 		/// <summary>
 		/// Checks for cycles in the action graph.
 		/// </summary>
-		void DetectActionGraphCycles()
+		static void DetectActionGraphCycles(List<Action> Actions, Dictionary<FileItem, Action> ItemToProducingAction)
 		{
 			// Starting with actions that only depend on non-produced items, iteratively expand a set of actions that are only dependent on
 			// non-cyclical actions.
@@ -300,7 +243,7 @@ namespace UnrealBuildTool
 			{
 				bool bFoundNewNonCyclicalAction = false;
 
-				foreach (Action Action in AllActions)
+				foreach (Action Action in Actions)
 				{
 					if (!ActionIsNonCyclical.ContainsKey(Action))
 					{
@@ -308,9 +251,10 @@ namespace UnrealBuildTool
 						bool bActionOnlyDependsOnNonCyclicalActions = true;
 						foreach (FileItem PrerequisiteItem in Action.PrerequisiteItems)
 						{
-							if (PrerequisiteItem.ProducingAction != null)
+							Action ProducingAction;
+							if (ItemToProducingAction.TryGetValue(PrerequisiteItem, out ProducingAction))
 							{
-								if (!ActionIsNonCyclical.ContainsKey(PrerequisiteItem.ProducingAction))
+								if (!ActionIsNonCyclical.ContainsKey(ProducingAction))
 								{
 									bActionOnlyDependsOnNonCyclicalActions = false;
 									if (!CyclicActions.ContainsKey(Action))
@@ -319,9 +263,9 @@ namespace UnrealBuildTool
 									}
 
 									List<Action> CyclicPrereq = CyclicActions[Action];
-									if (!CyclicPrereq.Contains(PrerequisiteItem.ProducingAction))
+									if (!CyclicPrereq.Contains(ProducingAction))
 									{
-										CyclicPrereq.Add(PrerequisiteItem.ProducingAction);
+										CyclicPrereq.Add(ProducingAction);
 									}
 								}
 							}
@@ -349,18 +293,18 @@ namespace UnrealBuildTool
 			}
 
 			// If there are any cyclical actions, throw an exception.
-			if (ActionIsNonCyclical.Count < AllActions.Count)
+			if (ActionIsNonCyclical.Count < Actions.Count)
 			{
 				// Find the index of each action
 				Dictionary<Action, int> ActionToIndex = new Dictionary<Action, int>();
-				for(int Idx = 0; Idx < AllActions.Count; Idx++)
+				for(int Idx = 0; Idx < Actions.Count; Idx++)
 				{
-					ActionToIndex[AllActions[Idx]] = Idx;
+					ActionToIndex[Actions[Idx]] = Idx;
 				}
 
 				// Describe the cyclical actions.
 				string CycleDescription = "";
-				foreach (Action Action in AllActions)
+				foreach (Action Action in Actions)
 				{
 					if (!ActionIsNonCyclical.ContainsKey(Action))
 					{
@@ -409,22 +353,34 @@ namespace UnrealBuildTool
 		/// <summary>
 		/// Determines the full set of actions that must be built to produce an item.
 		/// </summary>
-		/// <param name="OutputItem">- The item to be built.</param>
-		/// <param name="PrerequisiteActions">- The actions that must be built and the root action are</param>
-		public void GatherPrerequisiteActions(
-			FileItem OutputItem,
-			ref HashSet<Action> PrerequisiteActions
-			)
+		/// <param name="Actions">All the actions in the graph</param>
+		/// <param name="OutputItems">Set of output items to be built</param>
+		/// <returns>Set of prerequisite actions</returns>
+		public static List<Action> GatherPrerequisiteActions(List<Action> Actions, HashSet<FileItem> OutputItems)
 		{
-			if (OutputItem != null && OutputItem.ProducingAction != null)
+			HashSet<Action> PrerequisiteActions = new HashSet<Action>();
+			foreach(Action Action in Actions)
 			{
-				if (!PrerequisiteActions.Contains(OutputItem.ProducingAction))
+				if(Action.ProducedItems.Any(x => OutputItems.Contains(x)))
 				{
-					PrerequisiteActions.Add(OutputItem.ProducingAction);
-					foreach (FileItem PrerequisiteItem in OutputItem.ProducingAction.PrerequisiteItems)
-					{
-						GatherPrerequisiteActions(PrerequisiteItem, ref PrerequisiteActions);
-					}
+					GatherPrerequisiteActions(Action, PrerequisiteActions);
+				}
+			}
+			return PrerequisiteActions.ToList();
+		}
+
+		/// <summary>
+		/// Determines the full set of actions that must be built to produce an item.
+		/// </summary>
+		/// <param name="Action">The root action to scan</param>
+		/// <param name="PrerequisiteActions">Set of prerequisite actions</param>
+		private static void GatherPrerequisiteActions(Action Action, HashSet<Action> PrerequisiteActions)
+		{
+			if(PrerequisiteActions.Add(Action))
+			{
+				foreach(Action PrerequisiteAction in Action.PrerequisiteActions)
+				{
+					GatherPrerequisiteActions(PrerequisiteAction, PrerequisiteActions);
 				}
 			}
 		}
@@ -439,7 +395,7 @@ namespace UnrealBuildTool
 		/// <param name="CppDependencies"></param>
 		/// <param name="bIgnoreOutdatedImportLibraries"></param>
 		/// <returns>true if outdated</returns>
-		public bool IsActionOutdated(Action RootAction, Dictionary<Action, bool> OutdatedActionDictionary, ActionHistory ActionHistory, CppDependencyCache CppDependencies, bool bIgnoreOutdatedImportLibraries)
+		public static bool IsActionOutdated(Action RootAction, Dictionary<Action, bool> OutdatedActionDictionary, ActionHistory ActionHistory, CppDependencyCache CppDependencies, bool bIgnoreOutdatedImportLibraries)
 		{
 			// Only compute the outdated-ness for actions that don't aren't cached in the outdated action dictionary.
 			bool bIsOutdated = false;
@@ -453,7 +409,7 @@ namespace UnrealBuildTool
 					// Check if the command-line of the action previously used to produce the item is outdated.
 					string OldProducingCommandLine = "";
 					string NewProducingCommandLine = RootAction.CommandPath + " " + RootAction.CommandArguments;
-					if (!ActionHistory.GetProducingCommandLine(ProducedItem, out OldProducingCommandLine)
+					if (!ActionHistory.TryGetProducingCommandLine(ProducedItem, out OldProducingCommandLine)
 					|| !String.Equals(OldProducingCommandLine, NewProducingCommandLine, StringComparison.InvariantCultureIgnoreCase))
 					{
 						if(ProducedItem.Exists)
@@ -497,59 +453,41 @@ namespace UnrealBuildTool
 					}
 				}
 
+				// Check if any of the prerequisite actions are out of date
 				if (!bIsOutdated)
 				{
-					// Check if any of the prerequisite items are produced by outdated actions, or have changed more recently than
-					// the oldest produced item.
-					foreach (FileItem PrerequisiteItem in RootAction.PrerequisiteItems)
+					foreach (Action PrerequisiteAction in RootAction.PrerequisiteActions)
 					{
-						// Only check for outdated import libraries if we were configured to do so.  Often, a changed import library
-						// won't affect a dependency unless a public header file was also changed, in which case we would be forced
-						// to recompile anyway.  This just allows for faster iteration when working on a subsystem in a DLL, as we
-						// won't have to wait for dependent targets to be relinked after each change.
-						bool bIsImportLibraryFile = false;
-						if (PrerequisiteItem.ProducingAction != null && PrerequisiteItem.ProducingAction.bProducesImportLibrary)
-						{
-							bIsImportLibraryFile = PrerequisiteItem.AbsolutePath.EndsWith(".LIB", StringComparison.InvariantCultureIgnoreCase);
-						}
-						if (!bIsImportLibraryFile || !bIgnoreOutdatedImportLibraries)
+						if (IsActionOutdated(PrerequisiteAction, OutdatedActionDictionary, ActionHistory, CppDependencies, bIgnoreOutdatedImportLibraries))
 						{
-							// If the prerequisite is produced by an outdated action, then this action is outdated too.
-							if (PrerequisiteItem.ProducingAction != null)
+							// Only check for outdated import libraries if we were configured to do so.  Often, a changed import library
+							// won't affect a dependency unless a public header file was also changed, in which case we would be forced
+							// to recompile anyway.  This just allows for faster iteration when working on a subsystem in a DLL, as we
+							// won't have to wait for dependent targets to be relinked after each change.
+							if(!bIgnoreOutdatedImportLibraries || !IsImportLibraryDependency(RootAction, PrerequisiteAction))
 							{
-								if (IsActionOutdated(PrerequisiteItem.ProducingAction, OutdatedActionDictionary, ActionHistory, CppDependencies, bIgnoreOutdatedImportLibraries))
-								{
-									Log.TraceLog(
-										"{0}: Prerequisite {1} is produced by outdated action.",
-										RootAction.StatusDescription,
-										Path.GetFileName(PrerequisiteItem.AbsolutePath)
-										);
-									bIsOutdated = true;
-								}
-							}
-
-							if (PrerequisiteItem.Exists)
-							{
-								// allow a 1 second slop for network copies
-								TimeSpan TimeDifference = PrerequisiteItem.LastWriteTimeUtc - LastExecutionTimeUtc;
-								bool bPrerequisiteItemIsNewerThanLastExecution = TimeDifference.TotalSeconds > 1;
-								if (bPrerequisiteItemIsNewerThanLastExecution)
-								{
-									Log.TraceLog(
-										"{0}: Prerequisite {1} is newer than the last execution of the action: {2} vs {3}",
-										RootAction.StatusDescription,
-										Path.GetFileName(PrerequisiteItem.AbsolutePath),
-										PrerequisiteItem.LastWriteTimeUtc.LocalDateTime,
-										LastExecutionTimeUtc.LocalDateTime
-										);
-									bIsOutdated = true;
-								}
+								Log.TraceLog("{0}: Prerequisite {1} is produced by outdated action.", RootAction.StatusDescription, PrerequisiteAction.StatusDescription);
+								bIsOutdated = true;
+								break;
 							}
+						}
+					}
+				} 
 
-							// GatherAllOutdatedActions will ensure all actions are checked for outdated-ness, so we don't need to recurse with
-							// all this action's prerequisites once we've determined it's outdated.
-							if (bIsOutdated)
+				// Check if any prerequisite item has a newer timestamp than the last execution time of this action
+				if(!bIsOutdated)
+				{
+					foreach (FileItem PrerequisiteItem in RootAction.PrerequisiteItems)
+					{
+						if (PrerequisiteItem.Exists)
+						{
+							// allow a 1 second slop for network copies
+							TimeSpan TimeDifference = PrerequisiteItem.LastWriteTimeUtc - LastExecutionTimeUtc;
+							bool bPrerequisiteItemIsNewerThanLastExecution = TimeDifference.TotalSeconds > 1;
+							if (bPrerequisiteItemIsNewerThanLastExecution)
 							{
+								Log.TraceLog("{0}: Prerequisite {1} is newer than the last execution of the action: {2} vs {3}", RootAction.StatusDescription, Path.GetFileName(PrerequisiteItem.AbsolutePath), PrerequisiteItem.LastWriteTimeUtc.LocalDateTime, LastExecutionTimeUtc.LocalDateTime);
+								bIsOutdated = true;
 								break;
 							}
 						}
@@ -592,16 +530,34 @@ namespace UnrealBuildTool
 			return bIsOutdated;
 		}
 
+		/// <summary>
+		/// Determines if the dependency between two actions is only for an import library
+		/// </summary>
+		/// <param name="RootAction">The action to check</param>
+		/// <param name="PrerequisiteAction">The action that it depends on</param>
+		/// <returns>True if the only dependency between two actions is for an import library</returns>
+		static bool IsImportLibraryDependency(Action RootAction, Action PrerequisiteAction)
+		{
+			if(PrerequisiteAction.bProducesImportLibrary)
+			{
+				return PrerequisiteAction.ProducedItems.Any(x => !x.Location.HasExtension(".lib") && RootAction.PrerequisiteItems.Contains(x));
+			}
+			else
+			{
+				return false;
+			}
+		}
+
 		/// <summary>
 		/// Builds a dictionary containing the actions from AllActions that are outdated by calling
 		/// IsActionOutdated.
 		/// </summary>
-		void GatherAllOutdatedActions(ActionHistory ActionHistory, Dictionary<Action, bool> OutdatedActions, CppDependencyCache CppDependencies, bool bIgnoreOutdatedImportLibraries)
+		static void GatherAllOutdatedActions(List<Action> Actions, ActionHistory ActionHistory, Dictionary<Action, bool> OutdatedActions, CppDependencyCache CppDependencies, bool bIgnoreOutdatedImportLibraries)
 		{
 			using(Timeline.ScopeEvent("Prefetching include dependencies"))
 			{
 				List<FileItem> Dependencies = new List<FileItem>();
-				foreach(Action Action in AllActions)
+				foreach(Action Action in Actions)
 				{
 					if(Action.DependencyListFile != null)
 					{
@@ -611,7 +567,7 @@ namespace UnrealBuildTool
 				Parallel.ForEach(Dependencies, File => { List<FileItem> Temp; CppDependencies.TryGetDependencies(File, out Temp); });
 			}
 
-			foreach (Action Action in AllActions)
+			foreach (Action Action in Actions)
 			{
 				IsActionOutdated(Action, OutdatedActions, ActionHistory, CppDependencies, bIgnoreOutdatedImportLibraries);
 			}
@@ -620,21 +576,17 @@ namespace UnrealBuildTool
 		/// <summary>
 		/// Deletes all the items produced by actions in the provided outdated action dictionary.
 		/// </summary>
-		/// <param name="OutdatedActionDictionary">Dictionary of outdated actions</param>
-		static void DeleteOutdatedProducedItems(Dictionary<Action, bool> OutdatedActionDictionary)
+		/// <param name="OutdatedActions">List of outdated actions</param>
+		public static void DeleteOutdatedProducedItems(List<Action> OutdatedActions)
 		{
-			foreach (KeyValuePair<Action, bool> OutdatedActionInfo in OutdatedActionDictionary)
+			foreach(Action OutdatedAction in OutdatedActions)
 			{
-				if (OutdatedActionInfo.Value)
+				foreach (FileItem DeleteItem in OutdatedAction.DeleteItems)
 				{
-					Action OutdatedAction = OutdatedActionInfo.Key;
-					foreach (FileItem DeleteItem in OutdatedActionInfo.Key.DeleteItems)
+					if (DeleteItem.Exists)
 					{
-						if (DeleteItem.Exists)
-						{
-							Log.TraceLog("Deleting outdated item: {0}", DeleteItem.AbsolutePath);
-							DeleteItem.Delete();
-						}
+						Log.TraceLog("Deleting outdated item: {0}", DeleteItem.AbsolutePath);
+						DeleteItem.Delete();
 					}
 				}
 			}
@@ -644,72 +596,23 @@ namespace UnrealBuildTool
 		/// Creates directories for all the items produced by actions in the provided outdated action
 		/// dictionary.
 		/// </summary>
-		static void CreateDirectoriesForProducedItems(Dictionary<Action, bool> OutdatedActionDictionary)
+		public static void CreateDirectoriesForProducedItems(List<Action> OutdatedActions)
 		{
-			foreach (KeyValuePair<Action, bool> OutdatedActionInfo in OutdatedActionDictionary)
+			HashSet<DirectoryReference> OutputDirectories = new HashSet<DirectoryReference>();
+			foreach(Action OutdatedAction in OutdatedActions)
 			{
-				if (OutdatedActionInfo.Value)
+				foreach(FileItem ProducedItem in OutdatedAction.ProducedItems)
 				{
-					foreach (FileItem ProducedItem in OutdatedActionInfo.Key.ProducedItems)
-					{
-						string DirectoryPath = Path.GetDirectoryName(ProducedItem.AbsolutePath);
-						if (!Directory.Exists(DirectoryPath))
-						{
-							Log.TraceVerbose("Creating directory for produced item: {0}", DirectoryPath);
-							Directory.CreateDirectory(DirectoryPath);
-						}
-					}
+					OutputDirectories.Add(ProducedItem.Location.Directory);
+				}
+			}
+			foreach(DirectoryReference OutputDirectory in OutputDirectories)
+			{
+				if(!DirectoryReference.Exists(OutputDirectory))
+				{
+					DirectoryReference.CreateDirectory(OutputDirectory);
 				}
 			}
-		}
-
-
-
-		/// <summary>
-		/// Checks if the specified file is a C++ source implementation file (e.g., .cpp)
-		/// </summary>
-		/// <param name="FileItem">The file to check</param>
-		/// <returns>True if this is a C++ source file</returns>
-		private static bool IsCPPImplementationFile(FileItem FileItem)
-		{
-			return (FileItem.AbsolutePath.EndsWith(".cpp", StringComparison.InvariantCultureIgnoreCase) ||
-					FileItem.AbsolutePath.EndsWith(".c", StringComparison.InvariantCultureIgnoreCase) ||
-					FileItem.AbsolutePath.EndsWith(".mm", StringComparison.InvariantCultureIgnoreCase));
-		}
-
-
-		/// <summary>
-		/// Checks if the specified file is a C++ source header file (e.g., .h or .inl)
-		/// </summary>
-		/// <param name="FileItem">The file to check</param>
-		/// <returns>True if this is a C++ source file</returns>
-		private static bool IsCPPIncludeFile(FileItem FileItem)
-		{
-			return (FileItem.AbsolutePath.EndsWith(".h", StringComparison.InvariantCultureIgnoreCase) ||
-					FileItem.AbsolutePath.EndsWith(".hpp", StringComparison.InvariantCultureIgnoreCase) ||
-					FileItem.AbsolutePath.EndsWith(".inl", StringComparison.InvariantCultureIgnoreCase));
-		}
-
-
-		/// <summary>
-		/// Checks if the specified file is a C++ resource file (e.g., .rc)
-		/// </summary>
-		/// <param name="FileItem">The file to check</param>
-		/// <returns>True if this is a C++ source file</returns>
-		private static bool IsCPPResourceFile(FileItem FileItem)
-		{
-			return (FileItem.AbsolutePath.EndsWith(".rc", StringComparison.InvariantCultureIgnoreCase) || FileItem.AbsolutePath.EndsWith(".rc2", StringComparison.InvariantCultureIgnoreCase));
-		}
-
-
-		/// <summary>
-		/// Checks if the specified file is a C++ source file
-		/// </summary>
-		/// <param name="FileItem">The file to check</param>
-		/// <returns>True if this is a C++ source file</returns>
-		private static bool IsCPPFile(FileItem FileItem)
-		{
-			return IsCPPImplementationFile(FileItem) || IsCPPIncludeFile(FileItem) || IsCPPResourceFile(FileItem);
 		}
 	}
 }
diff --git a/Engine/Source/Programs/UnrealBuildTool/System/ActionHistory.cs b/Engine/Source/Programs/UnrealBuildTool/System/ActionHistory.cs
--- a/Engine/Source/Programs/UnrealBuildTool/System/ActionHistory.cs
+++ b/Engine/Source/Programs/UnrealBuildTool/System/ActionHistory.cs
@@ -14,115 +14,255 @@ namespace UnrealBuildTool
 	/// </summary>
 	class ActionHistory
 	{
+		/// <summary>
+		/// Version number to check
+		/// </summary>
+		const int CurrentVersion = 1;
+
 		/// <summary>
 		/// Path to store the cache data to.
 		/// </summary>
-		private string FilePath;
+		FileReference Location;
+
+		/// <summary>
+		/// Base directory for output files. Any files under this directory will have their command lines stored in this object, otherwise the parent will be used.
+		/// </summary>
+		DirectoryReference BaseDirectory;
+
+		/// <summary>
+		/// The parent action history. Any files not under this base directory will use this.
+		/// </summary>
+		ActionHistory Parent;
 
 		/// <summary>
 		/// The command lines used to produce files, keyed by the absolute file paths.
 		/// </summary>
-		private Dictionary<string, string> ProducedItemToPreviousActionCommandLine;
+		Dictionary<FileItem, string> OutputItemToCommandLine = new Dictionary<FileItem, string>();
 
 		/// <summary>
 		/// Whether the dependency cache is dirty and needs to be saved.
 		/// </summary>
-		private bool bIsDirty;
+		bool bModified;
+
+		/// <summary>
+		/// Static cache of all loaded action history files
+		/// </summary>
+		static Dictionary<FileReference, ActionHistory> LoadedFiles = new Dictionary<FileReference, ActionHistory>();
 
-		public ActionHistory(string InFilePath)
+		/// <summary>
+		/// Constructor
+		/// </summary>
+		/// <param name="Location">File to store this history in</param>
+		/// <param name="BaseDirectory">Base directory for files to track</param>
+		/// <param name="Parent">The parent action history</param>
+		public ActionHistory(FileReference Location, DirectoryReference BaseDirectory, ActionHistory Parent)
 		{
-			FilePath = Path.GetFullPath(InFilePath);
+			this.Location = Location;
+			this.BaseDirectory = BaseDirectory;
+			this.Parent = Parent;
+
+			if(FileReference.Exists(Location))
+			{
+				Load();
+			}
+		}
 
-			bool bFoundCache = false;
-			if (File.Exists(FilePath) == true)
+		/// <summary>
+		/// Attempts to load this action history from disk
+		/// </summary>
+		void Load()
+		{
+			try
 			{
-				try
+				using(BinaryArchiveReader Reader = new BinaryArchiveReader(Location))
 				{
-					// Deserialize the history from disk if the file exists.
-					using (FileStream Stream = new FileStream(FilePath, FileMode.Open, FileAccess.Read))
+					int Version = Reader.ReadInt();
+					if(Version != CurrentVersion)
 					{
-						BinaryFormatter Formatter = new BinaryFormatter();
-						ProducedItemToPreviousActionCommandLine = Formatter.Deserialize(Stream) as Dictionary<string, string>;
+						Log.TraceLog("Unable to read action history from {0}; version {1} vs current {2}", Location, Version, CurrentVersion);
+						return;
 					}
 
-					bFoundCache = true;
-				}
-				catch (Exception)
-				{
-					// If this fails for any reason just reset. History will be created later.
-					ProducedItemToPreviousActionCommandLine = null;
+					OutputItemToCommandLine = Reader.ReadDictionary(() => Reader.ReadFileItem(), () => Reader.ReadString());
 				}
 			}
+			catch(Exception Ex)
+			{
+				Log.TraceWarning("Unable to read {0}. See log for additional information.", Location);
+				Log.TraceLog("{0}", ExceptionUtils.FormatExceptionDetails(Ex));
+			}
+		}
 
-			if (!bFoundCache)
+		/// <summary>
+		/// Saves this action history to disk
+		/// </summary>
+		void Save()
+		{
+			DirectoryReference.CreateDirectory(Location.Directory);
+			using(BinaryArchiveWriter Writer = new BinaryArchiveWriter(Location))
 			{
-				// Otherwise create a fresh history.
-				ProducedItemToPreviousActionCommandLine = new Dictionary<string, string>();
+				Writer.WriteInt(CurrentVersion);
+				Writer.WriteDictionary(OutputItemToCommandLine, Key => Writer.WriteFileItem(Key), Value => Writer.WriteString(Value));
 			}
+			bModified = false;
+		}
 
-			bIsDirty = false;
+		/// <summary>
+		/// Gets the producing command line for the given file
+		/// </summary>
+		/// <param name="File">The output file to look for</param>
+		/// <param name="CommandLine">Receives the command line used to produce this file</param>
+		/// <returns>True if the output item exists</returns>
+		public bool TryGetProducingCommandLine(FileItem File, out string CommandLine)
+		{
+			if(File.Location.IsUnderDirectory(BaseDirectory) || Parent == null)
+			{
+				return OutputItemToCommandLine.TryGetValue(File, out CommandLine);
+			}
+			else
+			{
+				return Parent.TryGetProducingCommandLine(File, out CommandLine);
+			}
 		}
 
-		public void Save()
+		/// <summary>
+		/// Updates the producing command line for the given file
+		/// </summary>
+		/// <param name="File">The output file to update</param>
+		/// <param name="CommandLine">The command line used to produce this file</param>
+		public void SetProducingCommandLine(FileItem File, string CommandLine)
 		{
-			// Only save if we've made changes to it since load.
-			if (bIsDirty)
+			if(File.Location.IsUnderDirectory(BaseDirectory) || Parent == null)
 			{
-				// Serialize the cache to disk.
-				try
-				{
-					Directory.CreateDirectory(Path.GetDirectoryName(FilePath));
-					using (FileStream Stream = new FileStream(FilePath, FileMode.Create, FileAccess.Write))
-					{
-						BinaryFormatter Formatter = new BinaryFormatter();
-						Formatter.Serialize(Stream, ProducedItemToPreviousActionCommandLine);
-					}
-				}
-				catch (Exception Ex)
-				{
-					throw new BuildException(Ex, "Failed to write dependency cache.");
-				}
+				OutputItemToCommandLine[File] = CommandLine;
+				bModified = true;
+			}
+			else
+			{
+				Parent.SetProducingCommandLine(File, CommandLine);
 			}
 		}
 
-		public bool GetProducingCommandLine(FileItem File, out string Result)
+		/// <summary>
+		/// Gets the location for the engine action history
+		/// </summary>
+		/// <param name="TargetName">Target name being built</param>
+		/// <param name="Platform">The platform being built</param>
+		/// <param name="TargetType">Type of the target being built</param>
+		/// <returns>Path to the engine action history for this target</returns>
+		public static FileReference GetEngineLocation(string TargetName, UnrealTargetPlatform Platform, TargetType TargetType)
 		{
-			return ProducedItemToPreviousActionCommandLine.TryGetValue(File.AbsolutePath.ToUpperInvariant(), out Result);
+			string AppName;
+			if(TargetType == TargetType.Program)
+			{
+				AppName = TargetName;
+			}
+			else
+			{
+				AppName = UEBuildTarget.GetAppNameForTargetType(TargetType);
+			}
+
+			return FileReference.Combine(UnrealBuildTool.EngineDirectory, "Intermediate", "Build", Platform.ToString(), AppName, "ActionHistory.dat");
 		}
 
-		public void SetProducingCommandLine(FileItem File, string CommandLine)
+		/// <summary>
+		/// Gets the location of the project action history
+		/// </summary>
+		/// <param name="ProjectFile">Path to the project file</param>
+		/// <param name="Platform">Platform being built</param>
+		/// <param name="TargetName">Name of the target being built</param>
+		/// <returns>Path to the project action history</returns>
+		public static FileReference GetProjectLocation(FileReference ProjectFile, string TargetName, UnrealTargetPlatform Platform)
 		{
-			ProducedItemToPreviousActionCommandLine[File.AbsolutePath.ToUpperInvariant()] = CommandLine;
-			bIsDirty = true;
+			return FileReference.Combine(ProjectFile.Directory, "Intermediate", "Build", Platform.ToString(), TargetName, "ActionHistory.dat");
 		}
 
 		/// <summary>
-		/// Generates a full path to action history file for the specified target.
+		/// Creates a hierarchy of action history stores for a particular target
 		/// </summary>
-		public static FileReference GeneratePathForTarget(FileReference ProjectFile, string TargetName, UnrealTargetPlatform Platform, string Architecture, bool bUseSharedBuildEnvironment)
+		/// <param name="ProjectFile">Project file for the target being built</param>
+		/// <param name="TargetName">Name of the target</param>
+		/// <param name="Platform">Platform being built</param>
+		/// <param name="TargetType">The target type</param>
+		/// <returns>Dependency cache hierarchy for the given project</returns>
+		public static ActionHistory CreateHierarchy(FileReference ProjectFile, string TargetName, UnrealTargetPlatform Platform, TargetType TargetType)
 		{
-			if(bUseSharedBuildEnvironment)
+			ActionHistory History = null;
+
+			if(ProjectFile == null || !UnrealBuildTool.IsEngineInstalled())
+			{
+				FileReference EngineCacheLocation = GetEngineLocation(TargetName, Platform, TargetType);
+				History = FindOrAddHistory(EngineCacheLocation, UnrealBuildTool.EngineDirectory, History);
+			}
+
+			if(ProjectFile != null)
 			{
-				if(UnrealBuildTool.IsEngineInstalled() && ProjectFile != null)
+				FileReference ProjectCacheLocation = GetProjectLocation(ProjectFile, TargetName, Platform);
+				History = FindOrAddHistory(ProjectCacheLocation, ProjectFile.Directory, History);
+			}
+
+			return History;
+		}
+
+		/// <summary>
+		/// Enumerates all the locations of action history files for the given target
+		/// </summary>
+		/// <param name="ProjectFile">Project file for the target being built</param>
+		/// <param name="TargetName">Name of the target</param>
+		/// <param name="Platform">Platform being built</param>
+		/// <param name="TargetType">The target type</param>
+		/// <returns>Dependency cache hierarchy for the given project</returns>
+		public static IEnumerable<FileReference> GetFilesToClean(FileReference ProjectFile, string TargetName, UnrealTargetPlatform Platform, TargetType TargetType)
+		{
+			if(ProjectFile == null || !UnrealBuildTool.IsEngineInstalled())
+			{
+				yield return GetEngineLocation(TargetName, Platform, TargetType);
+			}
+			if(ProjectFile != null)
+			{
+				yield return GetProjectLocation(ProjectFile, TargetName, Platform);
+			}
+		}
+
+		/// <summary>
+		/// Reads a cache from the given location, or creates it with the given settings
+		/// </summary>
+		/// <param name="Location">File to store the cache</param>
+		/// <param name="BaseDirectory">Base directory for files that this cache should store data for</param>
+		/// <param name="Parent">The parent cache to use</param>
+		/// <returns>Reference to a dependency cache with the given settings</returns>
+		static ActionHistory FindOrAddHistory(FileReference Location, DirectoryReference BaseDirectory, ActionHistory Parent)
+		{
+			lock(LoadedFiles)
+			{
+				ActionHistory History;
+				if(LoadedFiles.TryGetValue(Location, out History))
 				{
-					return FileReference.Combine(ProjectFile.Directory, "Intermediate", "Build", "ActionHistory.bin");
+					Debug.Assert(History.BaseDirectory == BaseDirectory);
+					Debug.Assert(History.Parent == Parent);
 				}
 				else
 				{
-					return FileReference.Combine(UnrealBuildTool.EngineDirectory, "Intermediate", "Build", "ActionHistory.bin");
+					History = new ActionHistory(Location, BaseDirectory, Parent);
+					LoadedFiles.Add(Location, History);
 				}
+				return History;
 			}
-			else
-			{
-				string PlatformIntermediateFolder = UEBuildTarget.GetPlatformIntermediateFolder(Platform, Architecture);
+		}
 
-				if(ProjectFile != null)
-				{
-					return FileReference.Combine(ProjectFile.Directory, PlatformIntermediateFolder, TargetName, "ActionHistory.bin");
-				}
-				else
+		/// <summary>
+		/// Save all the loaded action histories
+		/// </summary>
+		public static void SaveAll()
+		{
+			lock(LoadedFiles)
+			{
+				foreach(ActionHistory History in LoadedFiles.Values)
 				{
-					return FileReference.Combine(UnrealBuildTool.EngineDirectory, PlatformIntermediateFolder, TargetName, "ActionHistory.bin");
+					if(History.bModified)
+					{
+						History.Save();
+					}
 				}
 			}
 		}
diff --git a/Engine/Source/Programs/UnrealBuildTool/System/FileItem.cs b/Engine/Source/Programs/UnrealBuildTool/System/FileItem.cs
--- a/Engine/Source/Programs/UnrealBuildTool/System/FileItem.cs
+++ b/Engine/Source/Programs/UnrealBuildTool/System/FileItem.cs
@@ -18,11 +18,6 @@ namespace UnrealBuildTool
 	/// </summary>
 	class FileItem
 	{
-		/// <summary>
-		/// The action that produces the file.
-		/// </summary>
-		public Action ProducingAction = null;
-
 		/// <summary>
 		/// The file reference
 		/// </summary>
diff --git a/Engine/Source/Programs/UnrealBuildTool/System/HotReload.cs b/Engine/Source/Programs/UnrealBuildTool/System/HotReload.cs
--- a/Engine/Source/Programs/UnrealBuildTool/System/HotReload.cs
+++ b/Engine/Source/Programs/UnrealBuildTool/System/HotReload.cs
@@ -463,9 +463,6 @@ namespace UnrealBuildTool
 							FileItem NewPrerequisiteItem = FileItem.GetItemByPath(NewPrerequisiteItemFilePath);
 							Action.PrerequisiteItems[ItemIndex] = NewPrerequisiteItem;
 
-							// Copy the other important settings from the original file item
-							NewPrerequisiteItem.ProducingAction = OriginalPrerequisiteItem.ProducingAction;
-
 							// Keep track of it so we can fix up dependencies in a second pass afterwards
 							AffectedOriginalFileItemAndNewFileItemMap.Add(OriginalPrerequisiteItem, NewPrerequisiteItem);
 
@@ -490,7 +487,6 @@ namespace UnrealBuildTool
 				// Go ahead and replace all occurrences of our file name in the command-line (ignoring extensions)
 				Action.CommandArguments = ReplaceBaseFileName(Action.CommandArguments, OriginalFileNameWithoutExtension, NewFileNameWithoutExtension);
 
-
 				// Update this action's list of produced items too
 				for (int ItemIndex = 0; ItemIndex < Action.ProducedItems.Count; ++ItemIndex)
 				{
@@ -503,9 +499,6 @@ namespace UnrealBuildTool
 						FileItem NewProducedItem = FileItem.GetItemByPath(NewProducedItemFilePath);
 						Action.ProducedItems[ItemIndex] = NewProducedItem;
 
-						// Copy the other important settings from the original file item
-						NewProducedItem.ProducingAction = OriginalProducedItem.ProducingAction;
-
 						// Keep track of it so we can fix up dependencies in a second pass afterwards
 						AffectedOriginalFileItemAndNewFileItemMap.Add(OriginalProducedItem, NewProducedItem);
 					}
@@ -514,7 +507,6 @@ namespace UnrealBuildTool
 				// The status description of the item has the file name, so we'll update it too
 				Action.StatusDescription = ReplaceBaseFileName(Action.StatusDescription, OriginalFileNameWithoutExtension, NewFileNameWithoutExtension);
 
-
 				// Keep track of the file names, so we can fix up response files afterwards.
 				if(!OriginalFileNameAndNewFileNameList_NoExtensions.ContainsKey(OriginalFileNameWithoutExtension))
 				{
diff --git a/Engine/Source/Programs/UnrealBuildTool/System/TargetMakefile.cs b/Engine/Source/Programs/UnrealBuildTool/System/TargetMakefile.cs
--- a/Engine/Source/Programs/UnrealBuildTool/System/TargetMakefile.cs
+++ b/Engine/Source/Programs/UnrealBuildTool/System/TargetMakefile.cs
@@ -20,12 +20,42 @@ namespace UnrealBuildTool
 		/// <summary>
 		/// The version number to write
 		/// </summary>
-		public const int CurrentVersion = 8;
+		public const int CurrentVersion = 9;
+
+		/// <summary>
+		/// Path to the receipt file for this target
+		/// </summary>
+		public FileReference ReceiptFile;
+
+		/// <summary>
+		/// The project intermediate directory
+		/// </summary>
+		public DirectoryReference ProjectIntermediateDirectory;
+
+		/// <summary>
+		/// Type of the target
+		/// </summary>
+		public TargetType TargetType;
+
+		/// <summary>
+		/// Whether the target should be deployed after being built
+		/// </summary>
+		public bool bDeployAfterCompile;
+
+		/// <summary>
+		/// Whether the project has a script plugin. UHT needs to know this to detect which manifest to use for checking out-of-datedness.
+		/// </summary>
+		public bool bHasProjectScriptPlugin;
+
+		/// <summary>
+		/// Scripts which should be run before building anything
+		/// </summary>
+		public FileReference[] PreBuildScripts;
 
 		/// <summary>
 		/// Every action in the action graph
 		/// </summary>
-		public List<Action> AllActions;
+		public List<Action> Actions;
 
 		/// <summary>
 		/// Environment variables that we'll need in order to invoke the platform's compiler and linker
@@ -36,12 +66,12 @@ namespace UnrealBuildTool
 		/// <summary>
 		/// Maps each target to a list of UObject module info structures
 		/// </summary>
-		public Dictionary<string, List<UHTModuleInfo>> TargetNameToUObjectModules;
+		public List<UHTModuleInfo> UObjectModules;
 
 		/// <summary>
 		/// The final output items for all target
 		/// </summary>
-		public List<FileItem> OutputItemsForAllTargets;
+		public List<FileItem> OutputItems;
 
 		/// <summary>
 		/// Maps module names to output items
@@ -51,20 +81,16 @@ namespace UnrealBuildTool
 		/// <summary>
 		/// List of game module names, for hot-reload
 		/// </summary>
-		public HashSet<string> HotReloadModuleNamesForAllTargets;
-
-		/// <summary>
-		/// List of targets being built
-		/// </summary>
-		public List<UEBuildTarget> Targets;
+		public HashSet<string> HotReloadModuleNames;
 
 		/// <summary>
 		/// List of predicates for the action graph to be valid.
 		/// </summary>
-		public List<BuildPrerequisites> TargetPrerequisites;
+		public BuildPrerequisites Prerequisites;
 
-		public TargetMakefile()
+		public TargetMakefile(TargetType TargetType)
 		{
+			this.TargetType = TargetType;
 		}
 
 		public TargetMakefile(BinaryArchiveReader Reader)
@@ -75,60 +101,67 @@ namespace UnrealBuildTool
 				throw new Exception(string.Format("Makefile version does not match - found {0}, expected: {1}", Version, CurrentVersion));
 			}
 
-			AllActions = Reader.ReadList(() => new Action(Reader));
+			ReceiptFile = Reader.ReadFileReference();
+			ProjectIntermediateDirectory = Reader.ReadDirectoryReference();
+			TargetType = (TargetType)Reader.ReadInt();
+			bDeployAfterCompile = Reader.ReadBool();
+			bHasProjectScriptPlugin = Reader.ReadBool();
+			PreBuildScripts = Reader.ReadArray(() => Reader.ReadFileReference());
+			Actions = Reader.ReadList(() => new Action(Reader));
 			EnvironmentVariables = Reader.ReadList(() => Tuple.Create(Reader.ReadString(), Reader.ReadString()));
-			TargetNameToUObjectModules = Reader.ReadDictionary(() => Reader.ReadString(), () => Reader.ReadList(() => new UHTModuleInfo(Reader)));
-			OutputItemsForAllTargets = Reader.ReadList(() => Reader.ReadFileItem());
+			UObjectModules = Reader.ReadList(() => new UHTModuleInfo(Reader));
+			OutputItems = Reader.ReadList(() => Reader.ReadFileItem());
 			ModuleNameToOutputItems = Reader.ReadDictionary(() => Reader.ReadString(), () => Reader.ReadArray(() => Reader.ReadFileItem()));
-			HotReloadModuleNamesForAllTargets = Reader.ReadHashSet(() => Reader.ReadString());
-			Targets = Reader.ReadList(() => new UEBuildTarget(Reader));
-			TargetPrerequisites = Reader.ReadList(() => new BuildPrerequisites(Reader));
+			HotReloadModuleNames = Reader.ReadHashSet(() => Reader.ReadString());
+			Prerequisites = new BuildPrerequisites(Reader);
 		}
 
 		public void Write(BinaryArchiveWriter Writer)
 		{
 			Writer.WriteInt(CurrentVersion);
 
-			Writer.WriteList(AllActions, Action => Action.Write(Writer));
+			Writer.WriteFileReference(ReceiptFile);
+			Writer.WriteDirectoryReference(ProjectIntermediateDirectory);
+			Writer.WriteInt((int)TargetType);
+			Writer.WriteBool(bDeployAfterCompile);
+			Writer.WriteBool(bHasProjectScriptPlugin);
+			Writer.WriteArray(PreBuildScripts, Item => Writer.WriteFileReference(Item));
+			Writer.WriteList(Actions, Action => Action.Write(Writer));
 			Writer.WriteList(EnvironmentVariables, x => { Writer.WriteString(x.Item1); Writer.WriteString(x.Item2); });
-			Writer.WriteDictionary(TargetNameToUObjectModules, x => Writer.WriteString(x), v => Writer.WriteList(v, e => e.Write(Writer)));
-			Writer.WriteList(OutputItemsForAllTargets, Item => Writer.WriteFileItem(Item));
+			Writer.WriteList(UObjectModules, e => e.Write(Writer));
+			Writer.WriteList(OutputItems, Item => Writer.WriteFileItem(Item));
 			Writer.WriteDictionary(ModuleNameToOutputItems, k => Writer.WriteString(k), v => Writer.WriteArray(v, e => Writer.WriteFileItem(e)));
-			Writer.WriteHashSet(HotReloadModuleNamesForAllTargets, x => Writer.WriteString(x));
-			Writer.WriteList(Targets, x => x.Write(Writer));
-			Writer.WriteList(TargetPrerequisites, x => x.Write(Writer));
-			Writer.Flush();
+			Writer.WriteHashSet(HotReloadModuleNames, x => Writer.WriteString(x));
+            Prerequisites.Write(Writer);
 		}
 
-
 		/// <returns> True if this makefile's contents look valid.  Called after loading the file to make sure it is legit.</returns>
 		public bool IsValidMakefile()
 		{
 			return
-				AllActions != null && AllActions.Count > 0 &&
+				Actions != null && Actions.Count > 0 &&
 				EnvironmentVariables != null &&
-				TargetNameToUObjectModules != null && TargetNameToUObjectModules.Count > 0 &&
-				OutputItemsForAllTargets != null && 
+				UObjectModules != null && 
+				OutputItems != null && 
 				ModuleNameToOutputItems != null && 
-				HotReloadModuleNamesForAllTargets != null &&
-				Targets != null && Targets.Count > 0 &&
-				TargetPrerequisites != null;
+				HotReloadModuleNames != null &&
+				Prerequisites != null;
 		}
 
 
 		/// <summary>
 		/// Saves a UBTMakefile to disk
 		/// </summary>
-		/// <param name="TargetDescs">List of targets.  Order is not important</param>
+		/// <param name="TargetDesc">Target being built</param>
 		/// <param name="UBTMakefile">The UBT makefile</param>
-		public static void SaveUBTMakefile(List<TargetDescriptor> TargetDescs, TargetMakefile UBTMakefile)
+		public static void SaveUBTMakefile(TargetDescriptor TargetDesc, TargetMakefile UBTMakefile)
 		{
 			if (!UBTMakefile.IsValidMakefile())
 			{
 				throw new BuildException("Can't save a makefile that has invalid contents.  See UBTMakefile.IsValidMakefile()");
 			}
 
-			FileItem UBTMakefileItem = FileItem.GetItemByFileReference(GetUBTMakefilePath(TargetDescs));
+			FileItem UBTMakefileItem = FileItem.GetItemByFileReference(GetLocation(TargetDesc.ProjectFile, TargetDesc.Name, TargetDesc.Platform, TargetDesc.Configuration));
 
 			// @todo ubtmake: Optimization: The UBTMakefile saved for game projects is upwards of 9 MB.  We should try to shrink its content if possible
 			// @todo ubtmake: Optimization: C# Serialization may be too slow for these big Makefiles.  Loading these files often shows up as the slower part of the assembling phase.
@@ -154,10 +187,11 @@ namespace UnrealBuildTool
 		/// </summary>
 		/// <param name="MakefilePath">Path to the makefile to load</param>
 		/// <param name="ProjectFile">Path to the project file</param>
+		/// <param name="Platform">Platform for this makefile</param>
 		/// <param name="ReasonNotLoaded">If the function returns null, this string will contain the reason why</param>
 		/// <param name="WorkingSet">Interface to query which source files are in the working set</param>
 		/// <returns>The loaded makefile, or null if it failed for some reason.  On failure, the 'ReasonNotLoaded' variable will contain information about why</returns>
-		public static TargetMakefile Load(FileReference MakefilePath, FileReference ProjectFile, ISourceFileWorkingSet WorkingSet, out string ReasonNotLoaded)
+		public static TargetMakefile Load(FileReference MakefilePath, FileReference ProjectFile, UnrealTargetPlatform Platform, ISourceFileWorkingSet WorkingSet, out string ReasonNotLoaded)
 		{
 			// Check the directory timestamp on the project files directory.  If the user has generated project files more
 			// recently than the UBTMakefile, then we need to consider the file to be out of date
@@ -277,96 +311,91 @@ namespace UnrealBuildTool
 					return null;
 				}
 
-				// Check if any of the target's Build.cs files are newer than the makefile
-				foreach (UEBuildTarget Target in LoadedUBTMakefile.Targets)
+				// Check if ini files are newer. Ini files contain build settings too.
+				DirectoryReference ProjectDirectory = DirectoryReference.FromFile(ProjectFile);
+				foreach (ConfigHierarchyType IniType in (ConfigHierarchyType[])Enum.GetValues(typeof(ConfigHierarchyType)))
 				{
-					// Check if ini files are newer. Ini files contain build settings too.
-					DirectoryReference ProjectDirectory = DirectoryReference.FromFile(ProjectFile);
-					foreach (ConfigHierarchyType IniType in (ConfigHierarchyType[])Enum.GetValues(typeof(ConfigHierarchyType)))
+					foreach (FileReference IniFilename in ConfigHierarchy.EnumerateConfigFileLocations(IniType, ProjectDirectory, Platform))
 					{
-						foreach (FileReference IniFilename in ConfigHierarchy.EnumerateConfigFileLocations(IniType, ProjectDirectory, Target.Platform))
+						FileInfo IniFileInfo = new FileInfo(IniFilename.FullName);
+						if (UBTMakefileInfo.LastWriteTime.CompareTo(IniFileInfo.LastWriteTime) < 0)
 						{
-							FileInfo IniFileInfo = new FileInfo(IniFilename.FullName);
-							if (UBTMakefileInfo.LastWriteTime.CompareTo(IniFileInfo.LastWriteTime) < 0)
-							{
-								// Ini files are newer than UBTMakefile
-								ReasonNotLoaded = "ini files are newer than UBTMakefile";
-								return null;
-							}
+							// Ini files are newer than UBTMakefile
+							ReasonNotLoaded = "ini files are newer than UBTMakefile";
+							return null;
 						}
 					}
 				}
 
-				foreach (BuildPrerequisites Prerequisites in LoadedUBTMakefile.TargetPrerequisites)
+				BuildPrerequisites Prerequisites = LoadedUBTMakefile.Prerequisites;
+
+				foreach(DirectoryReference SourceDirectory in Prerequisites.SourceDirectories)
 				{
-					foreach(DirectoryReference SourceDirectory in Prerequisites.SourceDirectories)
+					DirectoryInfo SourceDirectoryInfo = new DirectoryInfo(SourceDirectory.FullName);
+					if(!SourceDirectoryInfo.Exists || SourceDirectoryInfo.LastWriteTimeUtc > UBTMakefileInfo.LastWriteTimeUtc)
 					{
-						DirectoryInfo SourceDirectoryInfo = new DirectoryInfo(SourceDirectory.FullName);
-						if(!SourceDirectoryInfo.Exists || SourceDirectoryInfo.LastWriteTimeUtc > UBTMakefileInfo.LastWriteTimeUtc)
-						{
-							Log.TraceLog("Timestamp of {0} ({1}) is newer than makefile ({2})", SourceDirectory, SourceDirectoryInfo.LastWriteTimeUtc, UBTMakefileInfo.LastWriteTimeUtc);
-							ReasonNotLoaded = "source directory changed";
-							return null;
-						}
+						Log.TraceLog("Timestamp of {0} ({1}) is newer than makefile ({2})", SourceDirectory, SourceDirectoryInfo.LastWriteTimeUtc, UBTMakefileInfo.LastWriteTimeUtc);
+						ReasonNotLoaded = "source directory changed";
+						return null;
 					}
+				}
 
-					foreach(FileItem AdditionalDependency in Prerequisites.AdditionalDependencies)
+				foreach(FileItem AdditionalDependency in Prerequisites.AdditionalDependencies)
+				{
+					if (!AdditionalDependency.Exists)
 					{
-						if (!AdditionalDependency.Exists)
-						{
-							Log.TraceLog("{0} has been deleted since makefile was built.", AdditionalDependency.Location);
-							ReasonNotLoaded = string.Format("{0} deleted", AdditionalDependency.Location.GetFileName());
-							return null;
-						}
-						if(AdditionalDependency.LastWriteTimeUtc > UBTMakefileInfo.LastWriteTime)
-						{
-							Log.TraceLog("{0} has been modified since makefile was built.", AdditionalDependency.Location);
-							ReasonNotLoaded = string.Format("{0} modified", AdditionalDependency.Location.GetFileName());
-							return null;
-						}
+						Log.TraceLog("{0} has been deleted since makefile was built.", AdditionalDependency.Location);
+						ReasonNotLoaded = string.Format("{0} deleted", AdditionalDependency.Location.GetFileName());
+						return null;
 					}
+					if(AdditionalDependency.LastWriteTimeUtc > UBTMakefileInfo.LastWriteTime)
+					{
+						Log.TraceLog("{0} has been modified since makefile was built.", AdditionalDependency.Location);
+						ReasonNotLoaded = string.Format("{0} modified", AdditionalDependency.Location.GetFileName());
+						return null;
+					}
+				}
 
-					// We do a check to see if any modules' headers have changed which have
-					// acquired or lost UHT types.  If so, which should be rare,
-					// we'll just invalidate the entire makefile and force it to be rebuilt.
-
-					// Get all H files in processed modules newer than the makefile itself
-					HashSet<FileReference> HFilesNewerThanMakefile =
-						new HashSet<FileReference>(
-							Prerequisites.UObjectModuleHeaders
-							.SelectMany(x => x.SourceFolder != null ? DirectoryReference.EnumerateFiles(x.SourceFolder, "*.h", SearchOption.AllDirectories) : Enumerable.Empty<FileReference>())
-							.Where(y => FileReference.GetLastWriteTimeUtc(y) > UBTMakefileInfo.LastWriteTimeUtc)
-							.OrderBy(z => z).Distinct()
-						);
-
-					// Get all H files in all modules processed in the last makefile build
-					HashSet<FileReference> AllUHTHeaders = new HashSet<FileReference>(Prerequisites.UObjectModuleHeaders.SelectMany(x => x.HeaderFiles).Select(x => x.Location));
-
-					// Check whether any headers have been deleted. If they have, we need to regenerate the makefile since the module might now be empty. If we don't,
-					// and the file has been moved to a different module, we may include stale generated headers.
-					foreach (FileReference FileName in AllUHTHeaders)
+				// We do a check to see if any modules' headers have changed which have
+				// acquired or lost UHT types.  If so, which should be rare,
+				// we'll just invalidate the entire makefile and force it to be rebuilt.
+
+				// Get all H files in processed modules newer than the makefile itself
+				HashSet<FileReference> HFilesNewerThanMakefile =
+					new HashSet<FileReference>(
+						Prerequisites.UObjectModuleHeaders
+						.SelectMany(x => x.SourceFolder != null ? DirectoryReference.EnumerateFiles(x.SourceFolder, "*.h", SearchOption.AllDirectories) : Enumerable.Empty<FileReference>())
+						.Where(y => FileReference.GetLastWriteTimeUtc(y) > UBTMakefileInfo.LastWriteTimeUtc)
+						.OrderBy(z => z).Distinct()
+					);
+
+				// Get all H files in all modules processed in the last makefile build
+				HashSet<FileReference> AllUHTHeaders = new HashSet<FileReference>(Prerequisites.UObjectModuleHeaders.SelectMany(x => x.HeaderFiles).Select(x => x.Location));
+
+				// Check whether any headers have been deleted. If they have, we need to regenerate the makefile since the module might now be empty. If we don't,
+				// and the file has been moved to a different module, we may include stale generated headers.
+				foreach (FileReference FileName in AllUHTHeaders)
+				{
+					if (!FileReference.Exists(FileName))
 					{
-						if (!FileReference.Exists(FileName))
-						{
-							Log.TraceLog("File processed by UHT was deleted ({0}); invalidating makefile", FileName);
-							ReasonNotLoaded = string.Format("UHT file was deleted");
-							return null;
-						}
+						Log.TraceLog("File processed by UHT was deleted ({0}); invalidating makefile", FileName);
+						ReasonNotLoaded = string.Format("UHT file was deleted");
+						return null;
 					}
+				}
 
-					// Makefile is invalid if:
-					// * There are any newer files which contain no UHT data, but were previously in the makefile
-					// * There are any newer files contain data which needs processing by UHT, but weren't not previously in the makefile
-					foreach (FileReference Filename in HFilesNewerThanMakefile)
+				// Makefile is invalid if:
+				// * There are any newer files which contain no UHT data, but were previously in the makefile
+				// * There are any newer files contain data which needs processing by UHT, but weren't not previously in the makefile
+				foreach (FileReference Filename in HFilesNewerThanMakefile)
+				{
+					bool bContainsUHTData = CPPHeaders.DoesFileContainUObjects(Filename.FullName);
+					bool bWasProcessed = AllUHTHeaders.Contains(Filename);
+					if (bContainsUHTData != bWasProcessed)
 					{
-						bool bContainsUHTData = CPPHeaders.DoesFileContainUObjects(Filename.FullName);
-						bool bWasProcessed = AllUHTHeaders.Contains(Filename);
-						if (bContainsUHTData != bWasProcessed)
-						{
-							Log.TraceLog("{0} {1} contain UHT types and now {2} , ignoring it ({3})", Filename, bWasProcessed ? "used to" : "didn't", bWasProcessed ? "doesn't" : "does", UBTMakefileInfo.FullName);
-							ReasonNotLoaded = string.Format("new files with reflected types");
-							return null;
-						}
+						Log.TraceLog("{0} {1} contain UHT types and now {2} , ignoring it ({3})", Filename, bWasProcessed ? "used to" : "didn't", bWasProcessed ? "doesn't" : "does", UBTMakefileInfo.FullName);
+						ReasonNotLoaded = string.Format("new files with reflected types");
+						return null;
 					}
 				}
 
@@ -375,28 +404,26 @@ namespace UnrealBuildTool
 				// changed, then we'll force a new Makefile to be created so that we have fresh unity build blobs.  We always
 				// want to make sure that source files in the working set are excluded from those unity blobs (for fastest possible
 				// iteration times.)
-				foreach(BuildPrerequisites Prerequisites in LoadedUBTMakefile.TargetPrerequisites)
+
+				// Check if any source files in the working set no longer belong in it
+				foreach (FileItem SourceFile in Prerequisites.WorkingSet)
 				{
-					// Check if any source files in the working set no longer belong in it
-					foreach (FileItem SourceFile in Prerequisites.WorkingSet)
+					if (!WorkingSet.Contains(SourceFile.Location) && File.GetLastWriteTimeUtc(SourceFile.AbsolutePath) > UBTMakefileInfo.LastWriteTimeUtc)
 					{
-						if (!WorkingSet.Contains(SourceFile.Location) && File.GetLastWriteTimeUtc(SourceFile.AbsolutePath) > UBTMakefileInfo.LastWriteTimeUtc)
-						{
-							Log.TraceLog("{0} was part of source working set and now is not; invalidating makefile ({1})", SourceFile.AbsolutePath, UBTMakefileInfo.FullName);
-							ReasonNotLoaded = string.Format("working set of source files changed");
-							return null;
-						}
+						Log.TraceLog("{0} was part of source working set and now is not; invalidating makefile ({1})", SourceFile.AbsolutePath, UBTMakefileInfo.FullName);
+						ReasonNotLoaded = string.Format("working set of source files changed");
+						return null;
 					}
+				}
 
-					// Check if any source files that are eligible for being in the working set have been modified
-					foreach (FileItem SourceFile in Prerequisites.CandidatesForWorkingSet)
+				// Check if any source files that are eligible for being in the working set have been modified
+				foreach (FileItem SourceFile in Prerequisites.CandidatesForWorkingSet)
+				{
+					if (WorkingSet.Contains(SourceFile.Location) && File.GetLastWriteTimeUtc(SourceFile.AbsolutePath) > UBTMakefileInfo.LastWriteTimeUtc)
 					{
-						if (WorkingSet.Contains(SourceFile.Location) && File.GetLastWriteTimeUtc(SourceFile.AbsolutePath) > UBTMakefileInfo.LastWriteTimeUtc)
-						{
-							Log.TraceLog("{0} was part of source working set and now is not; invalidating makefile ({1})", SourceFile.AbsolutePath, UBTMakefileInfo.FullName);
-							ReasonNotLoaded = string.Format("working set of source files changed");
-							return null;
-						}
+						Log.TraceLog("{0} was part of source working set and now is not; invalidating makefile ({1})", SourceFile.AbsolutePath, UBTMakefileInfo.FullName);
+						ReasonNotLoaded = string.Format("working set of source files changed");
+						return null;
 					}
 				}
 			}
@@ -405,110 +432,18 @@ namespace UnrealBuildTool
 			return LoadedUBTMakefile;
 		}
 
-
-		/// <summary>
-		/// Gets the file path for a UBTMakefile
-		/// </summary>
-		/// <param name="TargetDescs">List of targets.  Order is not important</param>
-		/// <returns>UBTMakefile path</returns>
-		public static FileReference GetUBTMakefilePath(List<TargetDescriptor> TargetDescs)
-		{
-			FileReference UBTMakefilePath;
-
-			if (TargetDescs.Count == 1)
-			{
-				TargetDescriptor TargetDesc = TargetDescs[0];
-
-				string UBTMakefileName = "Makefile.ubt";
-
-				UBTMakefilePath = FileReference.Combine(GetUBTMakefileDirectoryPathForSingleTarget(TargetDesc), UBTMakefileName);
-			}
-			else
-			{
-				// For Makefiles that contain multiple targets, we'll make up a file name that contains all of the targets, their
-				// configurations and platforms, and save it into the base intermediate folder
-				string TargetCollectionName = MakeTargetCollectionName(TargetDescs);
-
-				TargetDescriptor DescriptorWithProject = TargetDescs.FirstOrDefault(x => x.ProjectFile != null);
-
-				DirectoryReference ProjectIntermediatePath;
-				if (DescriptorWithProject != null)
-				{
-					ProjectIntermediatePath = DirectoryReference.Combine(DescriptorWithProject.ProjectFile.Directory, "Intermediate", "Build");
-				}
-				else
-				{
-					ProjectIntermediatePath = DirectoryReference.Combine(UnrealBuildTool.EngineDirectory, "Intermediate", "Build");
-				}
-
-				// @todo ubtmake: The TargetCollectionName string could be really long if there is more than one target!  Hash it?
-				UBTMakefilePath = FileReference.Combine(ProjectIntermediatePath, TargetCollectionName + ".ubt");
-			}
-
-			return UBTMakefilePath;
-		}
-
 		/// <summary>
-		/// Gets the file path for a UBTMakefile for single target.
+		/// Gets the location of the makefile for particular target
 		/// </summary>
-		/// <param name="Target">The target.</param>
-		/// <returns>UBTMakefile path</returns>
-		public static DirectoryReference GetUBTMakefileDirectoryPathForSingleTarget(TargetDescriptor Target)
+		/// <param name="ProjectFile">Project file for the build</param>
+		/// <param name="TargetName">Name of the target being built</param>
+		/// <param name="Platform">The platform that the target is being built for</param>
+		/// <param name="Configuration">The configuration being built</param>
+		/// <returns>Path to the makefile</returns>
+		public static FileReference GetLocation(FileReference ProjectFile, string TargetName, UnrealTargetPlatform Platform, UnrealTargetConfiguration Configuration)
 		{
-			return GetUBTMakefileDirectory(Target.ProjectFile, Target.Platform, Target.Configuration, Target.Name);
+			DirectoryReference BaseDirectory = DirectoryReference.FromFile(ProjectFile) ?? UnrealBuildTool.EngineDirectory;
+			return FileReference.Combine(BaseDirectory, "Intermediate", "Build", Platform.ToString(), TargetName, Configuration.ToString(), "Makefile.ubt");
 		}
-
-		public static DirectoryReference GetUBTMakefileDirectory(FileReference ProjectFile, UnrealTargetPlatform Platform, UnrealTargetConfiguration Configuration, string TargetName)
-		{
-			// If there's only one target, just save the UBTMakefile in the target's build intermediate directory
-			// under a folder for that target (and platform/config combo.)
-			if (ProjectFile != null)
-			{
-				return DirectoryReference.Combine(ProjectFile.Directory, "Intermediate", "Build", Platform.ToString(), TargetName, Configuration.ToString());
-			}
-			else
-			{
-				return DirectoryReference.Combine(UnrealBuildTool.EngineDirectory, "Intermediate", "Build", Platform.ToString(), TargetName, Configuration.ToString());
-			}
-		}
-
-		public static FileReference GetUBTMakefilePath(FileReference ProjectFile, UnrealTargetPlatform Platform, UnrealTargetConfiguration Configuration, string TargetName, bool bForHotReload)
-		{
-			DirectoryReference BaseDir = GetUBTMakefileDirectory(ProjectFile, Platform, Configuration, TargetName);
-			return FileReference.Combine(BaseDir, bForHotReload ? "HotReloadMakefile.ubt" : "Makefile.ubt");
-		}
-
-		/// <summary>
-		/// Makes up a name for a set of targets that we can use for file or directory names
-		/// </summary>
-		/// <param name="TargetDescs">List of targets.  Order is not important</param>
-		/// <returns>The name to use</returns>
-		public static string MakeTargetCollectionName(List<TargetDescriptor> TargetDescs)
-		{
-			if (TargetDescs.Count == 0)
-			{
-				throw new BuildException("Expecting at least one Target to be passed to MakeTargetCollectionName");
-			}
-
-			List<TargetDescriptor> SortedTargets = new List<TargetDescriptor>();
-			SortedTargets.AddRange(TargetDescs);
-			SortedTargets.Sort((x, y) => { return x.Name.CompareTo(y.Name); });
-
-			// Figure out what to call our action graph based on everything we're building
-			StringBuilder TargetCollectionName = new StringBuilder();
-			foreach (TargetDescriptor Target in SortedTargets)
-			{
-				if (TargetCollectionName.Length > 0)
-				{
-					TargetCollectionName.Append("_");
-				}
-
-				// @todo ubtmake: Should we also have the platform Architecture in this string?
-				TargetCollectionName.Append(Target.Name + "-" + Target.Platform.ToString() + "-" + Target.Configuration.ToString());
-			}
-
-			return TargetCollectionName.ToString();
-		}
-
 	}
 }
diff --git a/Engine/Source/Programs/UnrealBuildTool/UnrealBuildTool.csproj b/Engine/Source/Programs/UnrealBuildTool/UnrealBuildTool.csproj
--- a/Engine/Source/Programs/UnrealBuildTool/UnrealBuildTool.csproj
+++ b/Engine/Source/Programs/UnrealBuildTool/UnrealBuildTool.csproj
@@ -143,6 +143,7 @@
     <Compile Include="Modes\BuildMode.cs" />
     <Compile Include="Modes\CleanMode.cs" />
     <Compile Include="Modes\DeployMode.cs" />
+    <Compile Include="Modes\JsonExportMode.cs" />
     <Compile Include="Modes\GenerateProjectFilesMode.cs" />
     <Compile Include="Modes\SetupPlatformsMode.cs" />
     <Compile Include="Modes\ValidatePlatformsMode.cs" />
