UBT: Move handling of single file compiles into the build mode.

#rnx

diff --git a/Engine/Source/Programs/UnrealBuildTool/Modes/BuildMode.cs b/Engine/Source/Programs/UnrealBuildTool/Modes/BuildMode.cs
--- a/Engine/Source/Programs/UnrealBuildTool/Modes/BuildMode.cs
+++ b/Engine/Source/Programs/UnrealBuildTool/Modes/BuildMode.cs
@@ -505,6 +505,13 @@ namespace UnrealBuildTool
 				}
 			}
 
+			// If we're just compiling a single file, set the target items to be all the derived items
+			if(BuildConfiguration.SingleFileToCompile != null)
+			{
+				FileItem FileToCompile = FileItem.GetItemByFileReference(BuildConfiguration.SingleFileToCompile);
+				TargetOutputItems = Makefile.Actions.Where(x => x.PrerequisiteItems.Contains(FileToCompile)).SelectMany(x => x.ProducedItems).ToList();
+			}
+
 			// Get the root prerequisite actions
 			List<Action> PrerequisiteActions = ActionGraph.GatherPrerequisiteActions(Makefile.Actions, new HashSet<FileItem>(TargetOutputItems));
 
@@ -587,8 +594,16 @@ namespace UnrealBuildTool
 			}
 
 			// Plan the actions to execute for the build.
-			HashSet<Action> TargetActionsToExecute = ActionGraph.GetActionsToExecute(BuildConfiguration, Makefile.Actions, PrerequisiteActions, CppDependencies, History);
-
+			HashSet<Action> TargetActionsToExecute;
+			if (BuildConfiguration.SingleFileToCompile == null)
+			{
+				TargetActionsToExecute = ActionGraph.GetActionsToExecute(Makefile.Actions, PrerequisiteActions, CppDependencies, History, BuildConfiguration.bIgnoreOutdatedImportLibraries);
+			}
+			else
+			{
+				TargetActionsToExecute = new HashSet<Action>(PrerequisiteActions);
+			}
+			
 			// Patch action history for hot reload when running in assembler mode.  In assembler mode, the suffix on the output file will be
 			// the same for every invocation on that makefile, but we need a new suffix each time.
 			if (HotReloadMode != HotReloadMode.Disabled)
@@ -648,7 +663,7 @@ namespace UnrealBuildTool
 				HotReload.PatchActionGraph(PrerequisiteActions, OldLocationToNewLocation);
 
 				// Get a new list of actions to execute now that the graph has been modified
-				TargetActionsToExecute = ActionGraph.GetActionsToExecute(BuildConfiguration, Makefile.Actions, PrerequisiteActions, CppDependencies, History);
+				TargetActionsToExecute = ActionGraph.GetActionsToExecute(Makefile.Actions, PrerequisiteActions, CppDependencies, History, BuildConfiguration.bIgnoreOutdatedImportLibraries);
 
 				// Build a mapping of all file items to their original
 				Dictionary<FileReference, FileReference> HotReloadFileToOriginalFile = new Dictionary<FileReference, FileReference>();
diff --git a/Engine/Source/Programs/UnrealBuildTool/System/ActionGraph.cs b/Engine/Source/Programs/UnrealBuildTool/System/ActionGraph.cs
--- a/Engine/Source/Programs/UnrealBuildTool/System/ActionGraph.cs
+++ b/Engine/Source/Programs/UnrealBuildTool/System/ActionGraph.cs
@@ -107,7 +107,7 @@ namespace UnrealBuildTool
 		/// <summary>
 		/// Builds a list of actions that need to be executed to produce the specified output items.
 		/// </summary>
-		public static HashSet<Action> GetActionsToExecute(BuildConfiguration BuildConfiguration, List<Action> Actions, List<Action> PrerequisiteActions, CppDependencyCache CppDependencies, ActionHistory History)
+		public static HashSet<Action> GetActionsToExecute(List<Action> Actions, List<Action> PrerequisiteActions, CppDependencyCache CppDependencies, ActionHistory History, bool bIgnoreOutdatedImportLibraries)
 		{
 			ITimelineEvent GetActionsToExecuteTimer = Timeline.ScopeEvent("ActionGraph.GetActionsToExecute()");
 
@@ -120,22 +120,7 @@ namespace UnrealBuildTool
 
 			// For all targets, build a set of all actions that are outdated.
 			Dictionary<Action, bool> OutdatedActionDictionary = new Dictionary<Action, bool>();
-			GatherAllOutdatedActions(Actions, History, OutdatedActionDictionary, CppDependencies, BuildConfiguration.bIgnoreOutdatedImportLibraries);
-
-			// If we're only compiling a single file, we should always compile and should never link.
-			if (BuildConfiguration.SingleFileToCompile != null)
-			{
-				// Never do anything but compile the target file
-				Actions.RemoveAll(x => x.ActionType != ActionType.Compile);
-
-				// Check all of the leftover compilation actions for the one we want... that one is always outdated.
-				FileItem SingleFileToCompile = FileItem.GetItemByFileReference(BuildConfiguration.SingleFileToCompile);
-				foreach (Action Action in Actions)
-				{
-					bool bIsSingleFileAction = Action.PrerequisiteItems.Contains(SingleFileToCompile);
-					OutdatedActionDictionary[Action] = bIsSingleFileAction;
-				}
-			}
+			GatherAllOutdatedActions(Actions, History, OutdatedActionDictionary, CppDependencies, bIgnoreOutdatedImportLibraries);
 
 			// Build a list of actions that are both needed for this target and outdated.
 			HashSet<Action> ActionsToExecute = new HashSet<Action>(Actions.Where(Action => Action.CommandPath != null && IsActionOutdatedMap.ContainsKey(Action) && OutdatedActionDictionary[Action]));
