UBT: Convert ActionGraph to a static class, and pass a list of actions through the toolchains.

#rnx

diff --git a/Engine/Source/Programs/UnrealBuildTool/Configuration/UEBuildBinary.cs b/Engine/Source/Programs/UnrealBuildTool/Configuration/UEBuildBinary.cs
--- a/Engine/Source/Programs/UnrealBuildTool/Configuration/UEBuildBinary.cs
+++ b/Engine/Source/Programs/UnrealBuildTool/Configuration/UEBuildBinary.cs
@@ -160,9 +160,9 @@ namespace UnrealBuildTool
 		/// <param name="WorkingSet">The working set of source files</param>
 		/// <param name="Prerequisites">Predicates for the generated actions to be valid</param>
 		/// <param name="ExeDir">Directory containing the output executable</param>
-		/// <param name="ActionGraph">Graph to add build actions to</param>
+		/// <param name="Actions">List of actions to be executed. Additional actions will be added to this list.</param>
 		/// <returns>Set of built products</returns>
-		public List<FileItem> Build(ReadOnlyTargetRules Target, UEToolChain ToolChain, CppCompileEnvironment CompileEnvironment, LinkEnvironment LinkEnvironment, List<PrecompiledHeaderTemplate> SharedPCHs, ISourceFileWorkingSet WorkingSet, BuildPrerequisites Prerequisites, DirectoryReference ExeDir, ActionGraph ActionGraph)
+		public List<FileItem> Build(ReadOnlyTargetRules Target, UEToolChain ToolChain, CppCompileEnvironment CompileEnvironment, LinkEnvironment LinkEnvironment, List<PrecompiledHeaderTemplate> SharedPCHs, ISourceFileWorkingSet WorkingSet, BuildPrerequisites Prerequisites, DirectoryReference ExeDir, List<Action> Actions)
 		{
 			// Return nothing if we're using precompiled binaries. If we're not linking, we might want just one module to be compiled (eg. a foreign plugin), so allow any actions to run.
 			if (bUsePrecompiled && !Target.bDisableLinking)
@@ -171,7 +171,7 @@ namespace UnrealBuildTool
 			}
 
 			// Setup linking environment.
-			LinkEnvironment BinaryLinkEnvironment = SetupBinaryLinkEnvironment(Target, ToolChain, LinkEnvironment, CompileEnvironment, SharedPCHs, WorkingSet, Prerequisites, ExeDir, ActionGraph);
+			LinkEnvironment BinaryLinkEnvironment = SetupBinaryLinkEnvironment(Target, ToolChain, LinkEnvironment, CompileEnvironment, SharedPCHs, WorkingSet, Prerequisites, ExeDir, Actions);
 
 			// If we're generating projects, we only need include paths and definitions, there is no need to run the linking logic.
 			if (ProjectFileGenerator.bGenerateProjectFiles)
@@ -195,12 +195,12 @@ namespace UnrealBuildTool
 				if (BinaryLinkEnvironment.Platform != CppPlatform.Mac && BinaryLinkEnvironment.Platform != CppPlatform.Linux)
 				{
 					// Create the import library.
-					OutputFiles.AddRange(ToolChain.LinkAllFiles(BinaryLinkEnvironment, true, ActionGraph));
+					OutputFiles.AddRange(ToolChain.LinkAllFiles(BinaryLinkEnvironment, true, Actions));
 				}
 			}
 
 			// Link the binary.
-			FileItem[] Executables = ToolChain.LinkAllFiles(BinaryLinkEnvironment, false, ActionGraph);
+			FileItem[] Executables = ToolChain.LinkAllFiles(BinaryLinkEnvironment, false, Actions);
 			OutputFiles.AddRange(Executables);
 
 			// Produce additional console app if requested
@@ -213,12 +213,12 @@ namespace UnrealBuildTool
 				ConsoleAppLinkEvironment.OutputFilePaths = ConsoleAppLinkEvironment.OutputFilePaths.Select(Path => GetAdditionalConsoleAppPath(Path)).ToList();
 
 				// Link the console app executable
-				OutputFiles.AddRange(ToolChain.LinkAllFiles(ConsoleAppLinkEvironment, false, ActionGraph));
+				OutputFiles.AddRange(ToolChain.LinkAllFiles(ConsoleAppLinkEvironment, false, Actions));
 			}
 
 			foreach (FileItem Executable in Executables)
 			{
-				OutputFiles.AddRange(ToolChain.PostBuild(Executable, BinaryLinkEnvironment, ActionGraph));
+				OutputFiles.AddRange(ToolChain.PostBuild(Executable, BinaryLinkEnvironment, Actions));
 			}
 
 			return OutputFiles;
@@ -650,7 +650,7 @@ namespace UnrealBuildTool
 			return BinaryCompileEnvironment;
 		}
 
-		private LinkEnvironment SetupBinaryLinkEnvironment(ReadOnlyTargetRules Target, UEToolChain ToolChain, LinkEnvironment LinkEnvironment, CppCompileEnvironment CompileEnvironment, List<PrecompiledHeaderTemplate> SharedPCHs, ISourceFileWorkingSet WorkingSet, BuildPrerequisites Prerequisites, DirectoryReference ExeDir, ActionGraph ActionGraph)
+		private LinkEnvironment SetupBinaryLinkEnvironment(ReadOnlyTargetRules Target, UEToolChain ToolChain, LinkEnvironment LinkEnvironment, CppCompileEnvironment CompileEnvironment, List<PrecompiledHeaderTemplate> SharedPCHs, ISourceFileWorkingSet WorkingSet, BuildPrerequisites Prerequisites, DirectoryReference ExeDir, List<Action> Actions)
 		{
 			LinkEnvironment BinaryLinkEnvironment = new LinkEnvironment(LinkEnvironment);
 			HashSet<UEBuildModule> LinkEnvironmentVisitedModules = new HashSet<UEBuildModule>();
@@ -670,7 +670,7 @@ namespace UnrealBuildTool
 				{
 					// Compile each module.
 					Log.TraceVerbose("Compile module: " + Module.Name);
-					LinkInputFiles = Module.Compile(Target, ToolChain, BinaryCompileEnvironment, SharedPCHs, WorkingSet, Prerequisites, ActionGraph);
+					LinkInputFiles = Module.Compile(Target, ToolChain, BinaryCompileEnvironment, SharedPCHs, WorkingSet, Prerequisites, Actions);
 
 					// NOTE: Because of 'Shared PCHs', in monolithic builds the same PCH file may appear as a link input
 					// multiple times for a single binary.  We'll check for that here, and only add it once.  This avoids
@@ -746,7 +746,7 @@ namespace UnrealBuildTool
 
 						// Otherwise compile the default resource file per-binary, so that it gets the correct ORIGINAL_FILE_NAME macro.
 						FileItem DefaultResourceFile = FileItem.GetItemByFileReference(FileReference.Combine(UnrealBuildTool.EngineDirectory, "Build", "Windows", "Default.rc2"));
-						CPPOutput DefaultResourceOutput = ToolChain.CompileRCFiles(ResourceCompileEnvironment, new List<FileItem> { DefaultResourceFile }, ResourceIntermediateDirectory, ActionGraph);
+						CPPOutput DefaultResourceOutput = ToolChain.CompileRCFiles(ResourceCompileEnvironment, new List<FileItem> { DefaultResourceFile }, ResourceIntermediateDirectory, Actions);
 						BinaryLinkEnvironment.InputFiles.AddRange(DefaultResourceOutput.ObjectFiles);
 					}
 				}
diff --git a/Engine/Source/Programs/UnrealBuildTool/Configuration/UEBuildModule.cs b/Engine/Source/Programs/UnrealBuildTool/Configuration/UEBuildModule.cs
--- a/Engine/Source/Programs/UnrealBuildTool/Configuration/UEBuildModule.cs
+++ b/Engine/Source/Programs/UnrealBuildTool/Configuration/UEBuildModule.cs
@@ -725,7 +725,7 @@ namespace UnrealBuildTool
 		/// <summary>
 		/// Compiles the module, and returns a list of files output by the compiler.
 		/// </summary>
-		public abstract List<FileItem> Compile(ReadOnlyTargetRules Target, UEToolChain ToolChain, CppCompileEnvironment CompileEnvironment, List<PrecompiledHeaderTemplate> SharedPCHModules, ISourceFileWorkingSet WorkingSet, BuildPrerequisites Predicates, ActionGraph ActionGraph);
+		public abstract List<FileItem> Compile(ReadOnlyTargetRules Target, UEToolChain ToolChain, CppCompileEnvironment CompileEnvironment, List<PrecompiledHeaderTemplate> SharedPCHModules, ISourceFileWorkingSet WorkingSet, BuildPrerequisites Predicates, List<Action> Actions);
 
 		// Object interface.
 		public override string ToString()
diff --git a/Engine/Source/Programs/UnrealBuildTool/Configuration/UEBuildModuleCPP.cs b/Engine/Source/Programs/UnrealBuildTool/Configuration/UEBuildModuleCPP.cs
--- a/Engine/Source/Programs/UnrealBuildTool/Configuration/UEBuildModuleCPP.cs
+++ b/Engine/Source/Programs/UnrealBuildTool/Configuration/UEBuildModuleCPP.cs
@@ -294,7 +294,7 @@ namespace UnrealBuildTool
 		}
 
 		// UEBuildModule interface.
-		public override List<FileItem> Compile(ReadOnlyTargetRules Target, UEToolChain ToolChain, CppCompileEnvironment BinaryCompileEnvironment, List<PrecompiledHeaderTemplate> SharedPCHs, ISourceFileWorkingSet WorkingSet, BuildPrerequisites Prerequisites, ActionGraph ActionGraph)
+		public override List<FileItem> Compile(ReadOnlyTargetRules Target, UEToolChain ToolChain, CppCompileEnvironment BinaryCompileEnvironment, List<PrecompiledHeaderTemplate> SharedPCHs, ISourceFileWorkingSet WorkingSet, BuildPrerequisites Prerequisites, List<Action> Actions)
 		{
 			UEBuildPlatform BuildPlatform = UEBuildPlatform.GetBuildPlatformForCPPTargetPlatform(BinaryCompileEnvironment.Platform);
 
@@ -401,7 +401,7 @@ namespace UnrealBuildTool
 				// If this module doesn't need a shared PCH, configure that
 				if(Rules.PrivatePCHHeaderFile != null && (Rules.PCHUsage == ModuleRules.PCHUsageMode.NoSharedPCHs || Rules.PCHUsage == ModuleRules.PCHUsageMode.UseExplicitOrSharedPCHs))
 				{
-					PrecompiledHeaderInstance Instance = CreatePrivatePCH(ToolChain, FileItem.GetItemByFileReference(FileReference.Combine(ModuleDirectory, Rules.PrivatePCHHeaderFile)), CompileEnvironment, ActionGraph);
+					PrecompiledHeaderInstance Instance = CreatePrivatePCH(ToolChain, FileItem.GetItemByFileReference(FileReference.Combine(ModuleDirectory, Rules.PrivatePCHHeaderFile)), CompileEnvironment, Actions);
 
 					CompileEnvironment = new CppCompileEnvironment(CompileEnvironment);
 					CompileEnvironment.Definitions.Clear();
@@ -423,7 +423,7 @@ namespace UnrealBuildTool
 					PrecompiledHeaderTemplate Template = SharedPCHs.FirstOrDefault(x => ReferencedModules.Contains(x.Module));
 					if(Template != null && Template.IsValidFor(CompileEnvironment))
 					{
-						PrecompiledHeaderInstance Instance = FindOrCreateSharedPCH(ToolChain, Template, ModuleCompileEnvironment.bOptimizeCode, ModuleCompileEnvironment.bUseRTTI, ModuleCompileEnvironment.bEnableExceptions, ActionGraph);
+						PrecompiledHeaderInstance Instance = FindOrCreateSharedPCH(ToolChain, Template, ModuleCompileEnvironment.bOptimizeCode, ModuleCompileEnvironment.bUseRTTI, ModuleCompileEnvironment.bEnableExceptions, Actions);
 
 						FileReference PrivateDefinitionsFile = FileReference.Combine(IntermediateDirectory, String.Format("Definitions.{0}.h", Name));
 
@@ -467,11 +467,11 @@ namespace UnrealBuildTool
 			if (bModuleUsesUnityBuild)
 			{
 				CPPFilesToCompile = Unity.GenerateUnityCPPs(Target, CPPFilesToCompile, CompileEnvironment, WorkingSet, Prerequisites, Rules.ShortName ?? Name, IntermediateDirectory);
-				LinkInputFiles.AddRange(CompileUnityFilesWithToolChain(Target, ToolChain, CompileEnvironment, ModuleCompileEnvironment, CPPFilesToCompile, ActionGraph).ObjectFiles);
+				LinkInputFiles.AddRange(CompileUnityFilesWithToolChain(Target, ToolChain, CompileEnvironment, ModuleCompileEnvironment, CPPFilesToCompile, Actions).ObjectFiles);
 			}
 			else
 			{
-				LinkInputFiles.AddRange(ToolChain.CompileCPPFiles(CompileEnvironment, CPPFilesToCompile, IntermediateDirectory, Name, ActionGraph).ObjectFiles);
+				LinkInputFiles.AddRange(ToolChain.CompileCPPFiles(CompileEnvironment, CPPFilesToCompile, IntermediateDirectory, Name, Actions).ObjectFiles);
 			}
 
 			// Compile all the generated CPP files
@@ -514,11 +514,11 @@ namespace UnrealBuildTool
 					if (bModuleUsesUnityBuild)
 					{
 						GeneratedFileItems = Unity.GenerateUnityCPPs(Target, GeneratedFileItems, GeneratedCPPCompileEnvironment, WorkingSet, Prerequisites, (Rules.ShortName ?? Name) + ".gen", IntermediateDirectory);
-						LinkInputFiles.AddRange(CompileUnityFilesWithToolChain(Target, ToolChain, GeneratedCPPCompileEnvironment, ModuleCompileEnvironment, GeneratedFileItems, ActionGraph).ObjectFiles);
+						LinkInputFiles.AddRange(CompileUnityFilesWithToolChain(Target, ToolChain, GeneratedCPPCompileEnvironment, ModuleCompileEnvironment, GeneratedFileItems, Actions).ObjectFiles);
 					}
 					else
 					{
-						LinkInputFiles.AddRange(ToolChain.CompileCPPFiles(GeneratedCPPCompileEnvironment, GeneratedFileItems, IntermediateDirectory, Name, ActionGraph).ObjectFiles);
+						LinkInputFiles.AddRange(ToolChain.CompileCPPFiles(GeneratedCPPCompileEnvironment, GeneratedFileItems, IntermediateDirectory, Name, Actions).ObjectFiles);
 					}
 				}
 			}
@@ -526,19 +526,19 @@ namespace UnrealBuildTool
 			// Compile C files directly. Do not use a PCH here, because a C++ PCH is not compatible with C source files.
 			if(SourceFilesToBuild.CFiles.Count > 0)
 			{
-				LinkInputFiles.AddRange(ToolChain.CompileCPPFiles(ModuleCompileEnvironment, SourceFilesToBuild.CFiles, IntermediateDirectory, Name, ActionGraph).ObjectFiles);
+				LinkInputFiles.AddRange(ToolChain.CompileCPPFiles(ModuleCompileEnvironment, SourceFilesToBuild.CFiles, IntermediateDirectory, Name, Actions).ObjectFiles);
 			}
 
 			// Compile CC files directly.
 			if(SourceFilesToBuild.CCFiles.Count > 0)
 			{
-				LinkInputFiles.AddRange(ToolChain.CompileCPPFiles(CompileEnvironment, SourceFilesToBuild.CCFiles, IntermediateDirectory, Name, ActionGraph).ObjectFiles);
+				LinkInputFiles.AddRange(ToolChain.CompileCPPFiles(CompileEnvironment, SourceFilesToBuild.CCFiles, IntermediateDirectory, Name, Actions).ObjectFiles);
 			}
 
 			// Compile MM files directly.
 			if(SourceFilesToBuild.MMFiles.Count > 0)
 			{
-				LinkInputFiles.AddRange(ToolChain.CompileCPPFiles(CompileEnvironment, SourceFilesToBuild.MMFiles, IntermediateDirectory, Name, ActionGraph).ObjectFiles);
+				LinkInputFiles.AddRange(ToolChain.CompileCPPFiles(CompileEnvironment, SourceFilesToBuild.MMFiles, IntermediateDirectory, Name, Actions).ObjectFiles);
 			}
 
 			// Compile RC files. The resource compiler does not work with response files, and using the regular compile environment can easily result in the 
@@ -548,7 +548,7 @@ namespace UnrealBuildTool
 			{
 				CppCompileEnvironment ResourceCompileEnvironment = new CppCompileEnvironment(BinaryCompileEnvironment);
 				WindowsPlatform.SetupResourceCompileEnvironment(ResourceCompileEnvironment, Target);
-				LinkInputFiles.AddRange(ToolChain.CompileRCFiles(ResourceCompileEnvironment, SourceFilesToBuild.RCFiles, IntermediateDirectory, ActionGraph).ObjectFiles);
+				LinkInputFiles.AddRange(ToolChain.CompileRCFiles(ResourceCompileEnvironment, SourceFilesToBuild.RCFiles, IntermediateDirectory, Actions).ObjectFiles);
 			}
 
 			// Write the compiled manifest
@@ -594,9 +594,9 @@ namespace UnrealBuildTool
 		/// <param name="ToolChain">The toolchain to generate the PCH</param>
 		/// <param name="HeaderFile"></param>
 		/// <param name="ModuleCompileEnvironment"></param>
-		/// <param name="ActionGraph">Graph containing build actions</param>
+		/// <param name="Actions">List of actions to be executed. Additional actions will be added to this list.</param>
 		/// <returns>The created PCH instance.</returns>
-		private PrecompiledHeaderInstance CreatePrivatePCH(UEToolChain ToolChain, FileItem HeaderFile, CppCompileEnvironment ModuleCompileEnvironment, ActionGraph ActionGraph)
+		private PrecompiledHeaderInstance CreatePrivatePCH(UEToolChain ToolChain, FileItem HeaderFile, CppCompileEnvironment ModuleCompileEnvironment, List<Action> Actions)
 		{
 			// Create the wrapper file, which sets all the definitions needed to compile it
 			FileReference WrapperLocation = FileReference.Combine(IntermediateDirectory, String.Format("PCH.{0}.h", Name));
@@ -610,7 +610,7 @@ namespace UnrealBuildTool
 			CompileEnvironment.bOptimizeCode = ModuleCompileEnvironment.bOptimizeCode;
 
 			// Create the action to compile the PCH file.
-			CPPOutput Output = ToolChain.CompileCPPFiles(CompileEnvironment, new List<FileItem>() { WrapperFile }, IntermediateDirectory, Name, ActionGraph);
+			CPPOutput Output = ToolChain.CompileCPPFiles(CompileEnvironment, new List<FileItem>() { WrapperFile }, IntermediateDirectory, Name, Actions);
 			return new PrecompiledHeaderInstance(WrapperFile, CompileEnvironment.bOptimizeCode, CompileEnvironment.bUseRTTI, CompileEnvironment.bEnableExceptions, Output);
 		}
 
@@ -622,9 +622,9 @@ namespace UnrealBuildTool
 		/// <param name="bOptimizeCode">Whether optimization should be enabled for this PCH</param>
 		/// <param name="bUseRTTI">Whether to enable RTTI for this PCH</param>
 		/// <param name="bEnableExceptions">Whether to enable exceptions for this PCH</param>
-		/// <param name="ActionGraph">Graph containing build actions</param>
+		/// <param name="Actions">List of actions to be executed. Additional actions will be added to this list.</param>
 		/// <returns>Instance of a PCH</returns>
-		public PrecompiledHeaderInstance FindOrCreateSharedPCH(UEToolChain ToolChain, PrecompiledHeaderTemplate Template, bool bOptimizeCode, bool bUseRTTI, bool bEnableExceptions, ActionGraph ActionGraph)
+		public PrecompiledHeaderInstance FindOrCreateSharedPCH(UEToolChain ToolChain, PrecompiledHeaderTemplate Template, bool bOptimizeCode, bool bUseRTTI, bool bEnableExceptions, List<Action> Actions)
 		{
 			PrecompiledHeaderInstance Instance = Template.Instances.Find(x => x.bOptimizeCode == bOptimizeCode && x.bUseRTTI == bUseRTTI && x.bEnableExceptions == bEnableExceptions);
 			if(Instance == null)
@@ -679,7 +679,7 @@ namespace UnrealBuildTool
 				CompileEnvironment.bEnableExceptions = bEnableExceptions;
 
 				// Create the PCH
-				CPPOutput Output = ToolChain.CompileCPPFiles(CompileEnvironment, new List<FileItem>() { WrapperFile }, Template.OutputDir, "Shared", ActionGraph);
+				CPPOutput Output = ToolChain.CompileCPPFiles(CompileEnvironment, new List<FileItem>() { WrapperFile }, Template.OutputDir, "Shared", Actions);
 				Instance = new PrecompiledHeaderInstance(WrapperFile, bOptimizeCode, bUseRTTI, bEnableExceptions, Output);
 				Template.Instances.Add(Instance);
 			}
@@ -689,7 +689,7 @@ namespace UnrealBuildTool
 		/// <summary>
 		/// Compiles the provided CPP unity files. Will
 		/// </summary>
-		private CPPOutput CompileUnityFilesWithToolChain(ReadOnlyTargetRules Target, UEToolChain ToolChain, CppCompileEnvironment CompileEnvironment, CppCompileEnvironment ModuleCompileEnvironment, List<FileItem> SourceFiles, ActionGraph ActionGraph)
+		private CPPOutput CompileUnityFilesWithToolChain(ReadOnlyTargetRules Target, UEToolChain ToolChain, CppCompileEnvironment CompileEnvironment, CppCompileEnvironment ModuleCompileEnvironment, List<FileItem> SourceFiles, List<Action> Actions)
 		{
 			List<FileItem> NormalFiles = new List<FileItem>();
 			List<FileItem> AdaptiveFiles = new List<FileItem>();
@@ -720,7 +720,7 @@ namespace UnrealBuildTool
 
 			if (NormalFiles.Count > 0)
 			{
-				OutputFiles = ToolChain.CompileCPPFiles(CompileEnvironment, NormalFiles, IntermediateDirectory, Name, ActionGraph);
+				OutputFiles = ToolChain.CompileCPPFiles(CompileEnvironment, NormalFiles, IntermediateDirectory, Name, Actions);
 			}
 
 			if (AdaptiveFiles.Count > 0)
@@ -740,19 +740,19 @@ namespace UnrealBuildTool
 				CPPOutput AdaptiveOutput;
 				if(Target.bAdaptiveUnityCreatesDedicatedPCH)
 				{
-					AdaptiveOutput = CompileAdaptiveNonUnityFilesWithDedicatedPCH(ToolChain, AdaptiveUnityEnvironment, AdaptiveFiles, IntermediateDirectory, Name, ActionGraph);
+					AdaptiveOutput = CompileAdaptiveNonUnityFilesWithDedicatedPCH(ToolChain, AdaptiveUnityEnvironment, AdaptiveFiles, IntermediateDirectory, Name, Actions);
 				}
 				else if(bAdaptiveUnityDisablesPCH)
 				{
-					AdaptiveOutput = CompileAdaptiveNonUnityFilesWithoutPCH(ToolChain, AdaptiveUnityEnvironment, AdaptiveFiles, IntermediateDirectory, Name, ActionGraph);
+					AdaptiveOutput = CompileAdaptiveNonUnityFilesWithoutPCH(ToolChain, AdaptiveUnityEnvironment, AdaptiveFiles, IntermediateDirectory, Name, Actions);
 				}
 				else if(AdaptiveUnityEnvironment.bOptimizeCode != CompileEnvironment.bOptimizeCode || AdaptiveUnityEnvironment.bSupportEditAndContinue != CompileEnvironment.bSupportEditAndContinue)
 				{
-					AdaptiveOutput = CompileAdaptiveNonUnityFiles(ToolChain, AdaptiveUnityEnvironment, AdaptiveFiles, IntermediateDirectory, Name, ActionGraph);
+					AdaptiveOutput = CompileAdaptiveNonUnityFiles(ToolChain, AdaptiveUnityEnvironment, AdaptiveFiles, IntermediateDirectory, Name, Actions);
 				}
 				else
 				{
-					AdaptiveOutput = CompileAdaptiveNonUnityFiles(ToolChain, CompileEnvironment, AdaptiveFiles, IntermediateDirectory, Name, ActionGraph);
+					AdaptiveOutput = CompileAdaptiveNonUnityFiles(ToolChain, CompileEnvironment, AdaptiveFiles, IntermediateDirectory, Name, Actions);
 				}
 
 				// Merge output
@@ -763,13 +763,13 @@ namespace UnrealBuildTool
 			return OutputFiles;
 		}
 
-		static CPPOutput CompileAdaptiveNonUnityFiles(UEToolChain ToolChain, CppCompileEnvironment CompileEnvironment, List<FileItem> Files, DirectoryReference IntermediateDirectory, string ModuleName, ActionGraph ActionGraph)
+		static CPPOutput CompileAdaptiveNonUnityFiles(UEToolChain ToolChain, CppCompileEnvironment CompileEnvironment, List<FileItem> Files, DirectoryReference IntermediateDirectory, string ModuleName, List<Action> Actions)
 		{
 			// Compile the files
-			return ToolChain.CompileCPPFiles(CompileEnvironment, Files, IntermediateDirectory, ModuleName, ActionGraph);
+			return ToolChain.CompileCPPFiles(CompileEnvironment, Files, IntermediateDirectory, ModuleName, Actions);
 		}
 
-		static CPPOutput CompileAdaptiveNonUnityFilesWithoutPCH(UEToolChain ToolChain, CppCompileEnvironment CompileEnvironment, List<FileItem> Files, DirectoryReference IntermediateDirectory, string ModuleName, ActionGraph ActionGraph)
+		static CPPOutput CompileAdaptiveNonUnityFilesWithoutPCH(UEToolChain ToolChain, CppCompileEnvironment CompileEnvironment, List<FileItem> Files, DirectoryReference IntermediateDirectory, string ModuleName, List<Action> Actions)
 		{
 			// Disable precompiled headers
 			CompileEnvironment.PrecompiledHeaderAction = PrecompiledHeaderAction.None;
@@ -778,10 +778,10 @@ namespace UnrealBuildTool
 			CreateHeaderForDefinitions(CompileEnvironment, IntermediateDirectory, "Adaptive");
 
 			// Compile the files
-			return ToolChain.CompileCPPFiles(CompileEnvironment, Files, IntermediateDirectory, ModuleName, ActionGraph);
+			return ToolChain.CompileCPPFiles(CompileEnvironment, Files, IntermediateDirectory, ModuleName, Actions);
 		}
 
-		static CPPOutput CompileAdaptiveNonUnityFilesWithDedicatedPCH(UEToolChain ToolChain, CppCompileEnvironment CompileEnvironment, List<FileItem> Files, DirectoryReference IntermediateDirectory, string ModuleName, ActionGraph ActionGraph)
+		static CPPOutput CompileAdaptiveNonUnityFilesWithDedicatedPCH(UEToolChain ToolChain, CppCompileEnvironment CompileEnvironment, List<FileItem> Files, DirectoryReference IntermediateDirectory, string ModuleName, List<Action> Actions)
 		{
 			CPPOutput Output = new CPPOutput();
 			foreach(FileItem File in Files)
@@ -812,7 +812,7 @@ namespace UnrealBuildTool
 				PchEnvironment.PrecompiledHeaderIncludeFilename = DedicatedPchFile.Location;
 
 				// Create the action to compile the PCH file.
-				CPPOutput PchOutput = ToolChain.CompileCPPFiles(PchEnvironment, new List<FileItem>() { DedicatedPchFile }, IntermediateDirectory, ModuleName, ActionGraph);
+				CPPOutput PchOutput = ToolChain.CompileCPPFiles(PchEnvironment, new List<FileItem>() { DedicatedPchFile }, IntermediateDirectory, ModuleName, Actions);
 				Output.ObjectFiles.AddRange(PchOutput.ObjectFiles);
 
 				// Create a new C++ environment to compile the original file
@@ -823,7 +823,7 @@ namespace UnrealBuildTool
 				FileEnvironment.PrecompiledHeaderFile = PchOutput.PrecompiledHeaderFile;
 
 				// Create the action to compile the PCH file.
-				CPPOutput FileOutput = ToolChain.CompileCPPFiles(FileEnvironment, new List<FileItem>() { File }, IntermediateDirectory, ModuleName, ActionGraph);
+				CPPOutput FileOutput = ToolChain.CompileCPPFiles(FileEnvironment, new List<FileItem>() { File }, IntermediateDirectory, ModuleName, Actions);
 				Output.ObjectFiles.AddRange(FileOutput.ObjectFiles);
 			}
 			return Output;
diff --git a/Engine/Source/Programs/UnrealBuildTool/Configuration/UEBuildModuleExternal.cs b/Engine/Source/Programs/UnrealBuildTool/Configuration/UEBuildModuleExternal.cs
--- a/Engine/Source/Programs/UnrealBuildTool/Configuration/UEBuildModuleExternal.cs
+++ b/Engine/Source/Programs/UnrealBuildTool/Configuration/UEBuildModuleExternal.cs
@@ -20,7 +20,7 @@ namespace UnrealBuildTool
 		}
 
 		// UEBuildModule interface.
-		public override List<FileItem> Compile(ReadOnlyTargetRules Target, UEToolChain ToolChain, CppCompileEnvironment CompileEnvironment, List<PrecompiledHeaderTemplate> SharedPCHs, ISourceFileWorkingSet WorkingSet, BuildPrerequisites Prerequisites, ActionGraph ActionGraph)
+		public override List<FileItem> Compile(ReadOnlyTargetRules Target, UEToolChain ToolChain, CppCompileEnvironment CompileEnvironment, List<PrecompiledHeaderTemplate> SharedPCHs, ISourceFileWorkingSet WorkingSet, BuildPrerequisites Prerequisites, List<Action> Actions)
 		{
 			return new List<FileItem>();
 		}
diff --git a/Engine/Source/Programs/UnrealBuildTool/Configuration/UEBuildTarget.cs b/Engine/Source/Programs/UnrealBuildTool/Configuration/UEBuildTarget.cs
--- a/Engine/Source/Programs/UnrealBuildTool/Configuration/UEBuildTarget.cs
+++ b/Engine/Source/Programs/UnrealBuildTool/Configuration/UEBuildTarget.cs
@@ -1084,7 +1084,7 @@ namespace UnrealBuildTool
 		/// <summary>
 		/// Builds the target, appending list of output files and returns building result.
 		/// </summary>
-		public ECompilationResult Build(BuildConfiguration BuildConfiguration, CPPHeaders Headers, List<FileItem> OutputItems, Dictionary<string, FileItem[]> ModuleNameToOutputItems, List<UHTModuleInfo> UObjectModules, ISourceFileWorkingSet WorkingSet, ActionGraph ActionGraph, BuildPrerequisites Prerequisites, bool bIsAssemblingBuild)
+		public ECompilationResult Build(BuildConfiguration BuildConfiguration, CPPHeaders Headers, List<FileItem> OutputItems, Dictionary<string, FileItem[]> ModuleNameToOutputItems, List<UHTModuleInfo> UObjectModules, ISourceFileWorkingSet WorkingSet, List<Action> Actions, BuildPrerequisites Prerequisites, bool bIsAssemblingBuild)
 		{
 			CppPlatform CppPlatform = UEBuildPlatform.GetBuildPlatform(Platform).DefaultCppPlatform;
 			CppConfiguration CppConfiguration = GetCppConfiguration(Configuration);
@@ -1301,7 +1301,7 @@ namespace UnrealBuildTool
 
 						WindowsPlatform.SetupResourceCompileEnvironment(DefaultResourceCompileEnvironment, Rules);
 
-						CPPOutput DefaultResourceOutput = TargetToolChain.CompileRCFiles(DefaultResourceCompileEnvironment, new List<FileItem> { DefaultResourceFile }, EngineIntermediateDirectory, ActionGraph);
+						CPPOutput DefaultResourceOutput = TargetToolChain.CompileRCFiles(DefaultResourceCompileEnvironment, new List<FileItem> { DefaultResourceFile }, EngineIntermediateDirectory, Actions);
 						GlobalLinkEnvironment.DefaultResourceFiles.AddRange(DefaultResourceOutput.ObjectFiles);
 					}
 				}
@@ -1315,7 +1315,7 @@ namespace UnrealBuildTool
 			}
 			foreach (UEBuildBinary Binary in Binaries)
 			{
-				List<FileItem> BinaryOutputItems = Binary.Build(Rules, TargetToolChain, GlobalCompileEnvironment, GlobalLinkEnvironment, SharedPCHs, WorkingSet, Prerequisites, ExeDir, ActionGraph);
+				List<FileItem> BinaryOutputItems = Binary.Build(Rules, TargetToolChain, GlobalCompileEnvironment, GlobalLinkEnvironment, SharedPCHs, WorkingSet, Prerequisites, ExeDir, Actions);
 				if(!bCompileMonolithic)
 				{
 					ModuleNameToOutputItems[Binary.PrimaryModule.Name] = BinaryOutputItems.ToArray();
@@ -1334,7 +1334,7 @@ namespace UnrealBuildTool
 			{
 				if(!UnrealBuildTool.IsFileInstalled(Pair.Key))
 				{
-					OutputItems.Add(CreateCopyAction(Pair.Value, Pair.Key, ActionGraph));
+					OutputItems.Add(CreateCopyAction(Pair.Value, Pair.Key, Actions));
 				}
 			}
 
@@ -1345,7 +1345,7 @@ namespace UnrealBuildTool
 			}
 
 			// Allow the toolchain to modify the final output items
-			TargetToolChain.FinalizeOutput(Rules, OutputItems, ActionGraph);
+			TargetToolChain.FinalizeOutput(Rules, OutputItems, Actions);
 
 			// Get all the regular build products
 			List<KeyValuePair<FileReference, BuildProductType>> BuildProducts = new List<KeyValuePair<FileReference, BuildProductType>>();
@@ -1420,19 +1420,20 @@ namespace UnrealBuildTool
 				BinaryFormatterUtils.Save(MetadataTargetFile, MetadataTargetInfo);
 
 				StringBuilder WriteMetadataArguments = new StringBuilder();
-				WriteMetadataArguments.AppendFormat("-Mode=WriteMetadata -Input={0}", Utils.MakePathSafeToUseWithCommandLine(MetadataTargetFile));
+				WriteMetadataArguments.AppendFormat("-Input={0}", Utils.MakePathSafeToUseWithCommandLine(MetadataTargetFile));
 				if(Rules.bNoManifestChanges)
 				{
 					WriteMetadataArguments.Append(" -NoManifestChanges");
 				}
 
-				Action WriteMetadataAction = ActionGraph.AddRecursiveCall(ActionType.WriteMetadata, WriteMetadataArguments.ToString());
+				Action WriteMetadataAction = Action.CreateRecursiveAction<WriteMetadataMode>(ActionType.WriteMetadata, WriteMetadataArguments.ToString());
 				WriteMetadataAction.WorkingDirectory = UnrealBuildTool.EngineSourceDirectory.FullName;
 				WriteMetadataAction.StatusDescription = ReceiptFileName.GetFileName();
 				WriteMetadataAction.bCanExecuteRemotely = false;
 				WriteMetadataAction.PrerequisiteItems.AddRange(OutputItems);
 				WriteMetadataAction.PrerequisiteItems.Add(FileItem.GetItemByPath(Assembly.GetExecutingAssembly().Location));
 				WriteMetadataAction.ProducedItems.Add(FileItem.GetItemByFileReference(ReceiptFileName));
+				Actions.Add(WriteMetadataAction);
 
 				OutputItems.AddRange(WriteMetadataAction.ProducedItems);
 
@@ -1441,7 +1442,7 @@ namespace UnrealBuildTool
 				{
 					FileReference OutputFile = new FileReference(PostBuildStepScript.FullName + ".ran");
 
-					Action PostBuildStepAction = ActionGraph.Add(ActionType.PostBuildStep);
+					Action PostBuildStepAction = new Action(ActionType.PostBuildStep);
 					if(BuildHostPlatform.Current.Platform == UnrealTargetPlatform.Win64)
 					{
 						PostBuildStepAction.CommandPath = "cmd.exe";
@@ -1457,6 +1458,7 @@ namespace UnrealBuildTool
 					PostBuildStepAction.bCanExecuteRemotely = false;
 					PostBuildStepAction.PrerequisiteItems.Add(FileItem.GetItemByFileReference(ReceiptFileName));
 					PostBuildStepAction.ProducedItems.Add(FileItem.GetItemByFileReference(OutputFile));
+					Actions.Add(PostBuildStepAction);
 
 					OutputItems.AddRange(PostBuildStepAction.ProducedItems);
 				}
@@ -1524,14 +1526,14 @@ namespace UnrealBuildTool
 		/// </summary>
 		/// <param name="SourceFile">The source file location</param>
 		/// <param name="TargetFile">The target file location</param>
-		/// <param name="ActionGraph">The action graph</param>
+		/// <param name="Actions">List of actions to be executed. Additional actions will be added to this list.</param>
 		/// <returns>File item for the output file</returns>
-		static FileItem CreateCopyAction(FileReference SourceFile, FileReference TargetFile, ActionGraph ActionGraph)
+		static FileItem CreateCopyAction(FileReference SourceFile, FileReference TargetFile, List<Action> Actions)
 		{
 			FileItem SourceFileItem = FileItem.GetItemByFileReference(SourceFile);
 			FileItem TargetFileItem = FileItem.GetItemByFileReference(TargetFile);
 
-			Action CopyAction = ActionGraph.Add(ActionType.BuildProject);
+			Action CopyAction = new Action(ActionType.BuildProject);
 			CopyAction.CommandDescription = "Copy";
 			if(BuildHostPlatform.Current.Platform == UnrealTargetPlatform.Win64)
 			{
@@ -1549,6 +1551,7 @@ namespace UnrealBuildTool
 			CopyAction.DeleteItems.Add(TargetFileItem);
 			CopyAction.StatusDescription = TargetFileItem.Location.GetFileName();
 			CopyAction.bCanExecuteRemotely = false;
+			Actions.Add(CopyAction);
 
 			return TargetFileItem;
 		}
diff --git a/Engine/Source/Programs/UnrealBuildTool/Modes/BuildMode.cs b/Engine/Source/Programs/UnrealBuildTool/Modes/BuildMode.cs
--- a/Engine/Source/Programs/UnrealBuildTool/Modes/BuildMode.cs
+++ b/Engine/Source/Programs/UnrealBuildTool/Modes/BuildMode.cs
@@ -217,10 +217,10 @@ namespace UnrealBuildTool
 							List<FileItem> OutputItems = new List<FileItem>();
 							Dictionary<string, FileItem[]> ModuleNameToOutputItems = new Dictionary<string, FileItem[]>();
 							List<UHTModuleInfo> UObjectModules = new List<UHTModuleInfo>();
-							ActionGraph ActionGraph = new ActionGraph();
+							List<Action> Actions = new List<Action>();
 							BuildPrerequisites Prerequisites = new BuildPrerequisites();
 
-							ECompilationResult BuildResult = Target.Build(BuildConfiguration, Headers, OutputItems, ModuleNameToOutputItems, UObjectModules, WorkingSet, ActionGraph, Prerequisites, bIsAssemblingBuild);
+							ECompilationResult BuildResult = Target.Build(BuildConfiguration, Headers, OutputItems, ModuleNameToOutputItems, UObjectModules, WorkingSet, Actions, Prerequisites, bIsAssemblingBuild);
 							if (BuildResult != ECompilationResult.Succeeded)
 							{
 								return (int)BuildResult;
@@ -229,7 +229,7 @@ namespace UnrealBuildTool
 							// Create the makefile
 							Makefile = new TargetMakefile(Target.TargetType);
 							Makefile.ReceiptFile = Target.ReceiptFileName;
-							Makefile.Actions = ActionGraph.Actions;
+							Makefile.Actions = Actions;
 							Makefile.OutputItems = OutputItems;
 							Makefile.ModuleNameToOutputItems = ModuleNameToOutputItems;
 							Makefile.UObjectModules = UObjectModules;
diff --git a/Engine/Source/Programs/UnrealBuildTool/Platform/Android/AndroidToolChain.cs b/Engine/Source/Programs/UnrealBuildTool/Platform/Android/AndroidToolChain.cs
--- a/Engine/Source/Programs/UnrealBuildTool/Platform/Android/AndroidToolChain.cs
+++ b/Engine/Source/Programs/UnrealBuildTool/Platform/Android/AndroidToolChain.cs
@@ -1186,7 +1186,7 @@ namespace UnrealBuildTool
 
 		static private bool bHasPrintedApiLevel = false;
 		static private bool bHasHandledLaunchModule = false;
-		public override CPPOutput CompileCPPFiles(CppCompileEnvironment CompileEnvironment, List<FileItem> InputFiles, DirectoryReference OutputDir, string ModuleName, ActionGraph ActionGraph)
+		public override CPPOutput CompileCPPFiles(CppCompileEnvironment CompileEnvironment, List<FileItem> InputFiles, DirectoryReference OutputDir, string ModuleName, List<Action> Actions)
 		{
 			if (Arches.Count == 0)
 			{
@@ -1321,7 +1321,7 @@ namespace UnrealBuildTool
 
 					foreach (FileItem SourceFile in InputFiles)
 					{
-						Action CompileAction = ActionGraph.Add(ActionType.Compile);
+						Action CompileAction = new Action(ActionType.Compile);
 						CompileAction.PrerequisiteItems.AddRange(CompileEnvironment.ForceIncludeFiles);
 
 						string FileArguments = "";
@@ -1481,6 +1481,8 @@ namespace UnrealBuildTool
 						CompileAction.bCanExecuteRemotely =
 							CompileEnvironment.PrecompiledHeaderAction != PrecompiledHeaderAction.Create ||
 							CompileEnvironment.bAllowRemotelyCompiledPCHs;
+
+						Actions.Add(CompileAction);
 					}
 				}
 			}
@@ -1488,7 +1490,7 @@ namespace UnrealBuildTool
 			return Result;
 		}
 
-		public override FileItem LinkFiles(LinkEnvironment LinkEnvironment, bool bBuildImportLibraryOnly, ActionGraph ActionGraph)
+		public override FileItem LinkFiles(LinkEnvironment LinkEnvironment, bool bBuildImportLibraryOnly, List<Action> Actions)
 		{
 			return null;
 		}
@@ -1518,7 +1520,7 @@ namespace UnrealBuildTool
 			return Pathname;
 		}
 
-		public override FileItem[] LinkAllFiles(LinkEnvironment LinkEnvironment, bool bBuildImportLibraryOnly, ActionGraph ActionGraph)
+		public override FileItem[] LinkAllFiles(LinkEnvironment LinkEnvironment, bool bBuildImportLibraryOnly, List<Action> Actions)
 		{
 			List<FileItem> Outputs = new List<FileItem>();
 
@@ -1545,7 +1547,7 @@ namespace UnrealBuildTool
 					}
 
 					// Create an action that invokes the linker.
-					Action LinkAction = ActionGraph.Add(ActionType.Link);
+					Action LinkAction = new Action(ActionType.Link);
 					LinkAction.WorkingDirectory = UnrealBuildTool.EngineSourceDirectory.FullName;
 
 					if (LinkEnvironment.bIsBuildingLibrary)
@@ -1698,6 +1700,7 @@ namespace UnrealBuildTool
 							LinkAction.CommandDescription = "Link";
 						}
 					}
+					Actions.Add(LinkAction);
 
 					// Windows can run into an issue with too long of a commandline when clang tries to call ld to link.
 					// To work around this we call clang to just get the command it would execute and generate a
diff --git a/Engine/Source/Programs/UnrealBuildTool/Platform/HTML5/HTML5ToolChain.cs b/Engine/Source/Programs/UnrealBuildTool/Platform/HTML5/HTML5ToolChain.cs
--- a/Engine/Source/Programs/UnrealBuildTool/Platform/HTML5/HTML5ToolChain.cs
+++ b/Engine/Source/Programs/UnrealBuildTool/Platform/HTML5/HTML5ToolChain.cs
@@ -339,7 +339,7 @@ namespace UnrealBuildTool
 			}
 		}
 
-		public override CPPOutput CompileCPPFiles(CppCompileEnvironment CompileEnvironment, List<FileItem> InputFiles, DirectoryReference OutputDir, string ModuleName, ActionGraph ActionGraph)
+		public override CPPOutput CompileCPPFiles(CppCompileEnvironment CompileEnvironment, List<FileItem> InputFiles, DirectoryReference OutputDir, string ModuleName, List<Action> Actions)
 		{
 			string Arguments = GetCLArguments_Global(CompileEnvironment);
 
@@ -375,7 +375,7 @@ namespace UnrealBuildTool
 
 			foreach (FileItem SourceFile in InputFiles)
 			{
-				Action CompileAction = ActionGraph.Add(ActionType.Compile);
+				Action CompileAction = new Action(ActionType.Compile);
 				CompileAction.CommandDescription = "Compile";
 				CompileAction.PrerequisiteItems.AddRange(CompileEnvironment.ForceIncludeFiles);
 //				CompileAction.bPrintDebugInfo = true;
@@ -432,24 +432,26 @@ namespace UnrealBuildTool
 				CompileAction.bCanExecuteRemotely =
 					CompileEnvironment.PrecompiledHeaderAction != PrecompiledHeaderAction.Create ||
 					CompileEnvironment.bAllowRemotelyCompiledPCHs;
+
+				Actions.Add(CompileAction);
 			}
 
 			return Result;
 		}
 
-		public override CPPOutput CompileRCFiles(CppCompileEnvironment CompileEnvironment, List<FileItem> InputFiles, DirectoryReference OutputDir, ActionGraph ActionGraph)
+		public override CPPOutput CompileRCFiles(CppCompileEnvironment CompileEnvironment, List<FileItem> InputFiles, DirectoryReference OutputDir, List<Action> Actions)
 		{
 			CPPOutput Result = new CPPOutput();
 
 			return Result;
 		}
 
-		public override FileItem LinkFiles(LinkEnvironment LinkEnvironment, bool bBuildImportLibraryOnly, ActionGraph ActionGraph)
+		public override FileItem LinkFiles(LinkEnvironment LinkEnvironment, bool bBuildImportLibraryOnly, List<Action> Actions)
 		{
 			FileItem OutputFile;
 
 			// Make the final javascript file
-			Action LinkAction = ActionGraph.Add(ActionType.Link);
+			Action LinkAction = new Action(ActionType.Link);
 			LinkAction.CommandDescription = "Link";
 //			LinkAction.bPrintDebugInfo = true;
 
@@ -535,6 +537,7 @@ namespace UnrealBuildTool
 
 			LinkAction.CommandArguments += string.Format(" @\"{0}\"", ResponseFileName);
 			LinkAction.PrerequisiteItems.Add(ResponseFileItem);
+			Actions.Add(LinkAction);
 
 			return OutputFile;
 		}
diff --git a/Engine/Source/Programs/UnrealBuildTool/Platform/IOS/IOSToolChain.cs b/Engine/Source/Programs/UnrealBuildTool/Platform/IOS/IOSToolChain.cs
--- a/Engine/Source/Programs/UnrealBuildTool/Platform/IOS/IOSToolChain.cs
+++ b/Engine/Source/Programs/UnrealBuildTool/Platform/IOS/IOSToolChain.cs
@@ -548,7 +548,7 @@ namespace UnrealBuildTool
 			return Result;
 		}
 
-		public override CPPOutput CompileCPPFiles(CppCompileEnvironment CompileEnvironment, List<FileItem> InputFiles, DirectoryReference OutputDir, string ModuleName, ActionGraph ActionGraph)
+		public override CPPOutput CompileCPPFiles(CppCompileEnvironment CompileEnvironment, List<FileItem> InputFiles, DirectoryReference OutputDir, string ModuleName, List<Action> Actions)
 		{
 			string Arguments = GetCompileArguments_Global(CompileEnvironment);
 			string PCHArguments = "";
@@ -580,7 +580,7 @@ namespace UnrealBuildTool
 			// Create a compile action for each source file.
 			foreach (FileItem SourceFile in InputFiles)
 			{
-				Action CompileAction = ActionGraph.Add(ActionType.Compile);
+				Action CompileAction = new Action(ActionType.Compile);
 				string FilePCHArguments = "";
 				string FileArguments = "";
 				string Extension = Path.GetExtension(SourceFile.AbsolutePath).ToUpperInvariant();
@@ -710,21 +710,23 @@ namespace UnrealBuildTool
 				{
 					if(Framework.ZipFile != null)
 					{
-						FileItem ExtractedTokenFile = ExtractFramework(Framework, ActionGraph);
+						FileItem ExtractedTokenFile = ExtractFramework(Framework, Actions);
 						CompileAction.PrerequisiteItems.Add(ExtractedTokenFile);
 					}
 				}
+
+				Actions.Add(CompileAction);
 			}
 			return Result;
 		}
 
-		public override FileItem LinkFiles(LinkEnvironment LinkEnvironment, bool bBuildImportLibraryOnly, ActionGraph ActionGraph)
+		public override FileItem LinkFiles(LinkEnvironment LinkEnvironment, bool bBuildImportLibraryOnly, List<Action> Actions)
 		{
 			string LinkerPath = Settings.Value.ToolchainDir +
 				(LinkEnvironment.bIsBuildingLibrary ? IOSArchiver : IOSLinker);
 
 			// Create an action that invokes the linker.
-			Action LinkAction = ActionGraph.Add(ActionType.Link);
+			Action LinkAction = new Action(ActionType.Link);
 
 			// RPC utility parameters are in terms of the Mac side
 			LinkAction.WorkingDirectory = GetMacDevSrcRoot();
@@ -765,7 +767,7 @@ namespace UnrealBuildTool
 			{
 				if(Framework.ZipFile != null)
 				{
-					FileItem ExtractedTokenFile = ExtractFramework(Framework, ActionGraph);
+					FileItem ExtractedTokenFile = ExtractFramework(Framework, Actions);
 					LinkAction.PrerequisiteItems.Add(ExtractedTokenFile);
 				}
 			}
@@ -844,6 +846,7 @@ namespace UnrealBuildTool
 				// This is not a shipping build so no need to delete the output file since symbols will not have been stripped from it.
 				LinkAction.CommandArguments = string.Format("-c '{0} {1}'", LinkerPath, LinkCommandArguments);
 			}
+			Actions.Add(LinkAction);
 
 			return OutputFile;
 		}
@@ -912,8 +915,8 @@ namespace UnrealBuildTool
         /// Generates debug info for a given executable
         /// </summary>
         /// <param name="Executable">FileItem describing the executable to generate debug info for</param>
-		/// <param name="ActionGraph"></param>
-        public FileItem GenerateDebugInfo(FileItem Executable, ActionGraph ActionGraph)
+		/// <param name="Actions">List of actions to be executed. Additional actions will be added to this list.</param>
+        public FileItem GenerateDebugInfo(FileItem Executable, List<Action> Actions)
 		{
             // Make a file item for the source and destination files
 			string FullDestPathRoot = Path.Combine(Path.GetDirectoryName(Executable.AbsolutePath), Path.GetFileName(Executable.AbsolutePath) + ".dSYM");
@@ -922,7 +925,7 @@ namespace UnrealBuildTool
             FileItem ZipOutputFile = FileItem.GetItemByPath(FullDestPathRoot + ".zip");
 
             // Make the compile action
-            Action GenDebugAction = ActionGraph.Add(ActionType.GenerateDebugInfo);
+            Action GenDebugAction = new Action(ActionType.GenerateDebugInfo);
 
 			GenDebugAction.WorkingDirectory = GetMacDevSrcRoot();
 			GenDebugAction.CommandPath = "sh";
@@ -948,6 +951,7 @@ namespace UnrealBuildTool
             GenDebugAction.ProducedItems.Add(OutputFile);
             GenDebugAction.StatusDescription = GenDebugAction.CommandArguments;// string.Format("Generating debug info for {0}", Path.GetFileName(Executable.AbsolutePath));
 			GenDebugAction.bCanExecuteRemotely = false;
+			Actions.Add(GenDebugAction);
 
 			return (ProjectSettings.bGeneratedSYMBundle ? ZipOutputFile : OutputFile);
 		}
@@ -956,8 +960,8 @@ namespace UnrealBuildTool
         /// Generates pseudo pdb info for a given executable
         /// </summary>
         /// <param name="Executable">FileItem describing the executable to generate debug info for</param>
-		/// <param name="ActionGraph"></param>
-        public FileItem GeneratePseudoPDB(FileItem Executable, ActionGraph ActionGraph)
+		/// <param name="Actions">List of actions to be executed. Additional actions will be added to this list.</param>
+        public FileItem GeneratePseudoPDB(FileItem Executable, List<Action> Actions)
         {
             // Make a file item for the source and destination files
             string FullDestPathRoot = Path.Combine(Path.GetDirectoryName(Executable.AbsolutePath), Path.GetFileName(Executable.AbsolutePath) + ".udebugsymbols");
@@ -971,7 +975,7 @@ namespace UnrealBuildTool
             FileItem OutputFile = FileItem.GetItemByPath(FullDestPathRoot);
 
             // Make the compile action
-            Action GenDebugAction = ActionGraph.Add(ActionType.GenerateDebugInfo);
+            Action GenDebugAction = new Action(ActionType.GenerateDebugInfo);
             GenDebugAction.WorkingDirectory = GetMacDevEngineRoot() + "/Binaries/Mac/";
 
             GenDebugAction.CommandPath = "sh";
@@ -984,6 +988,7 @@ namespace UnrealBuildTool
             GenDebugAction.ProducedItems.Add(OutputFile);
             GenDebugAction.StatusDescription = GenDebugAction.CommandArguments;// string.Format("Generating debug info for {0}", Path.GetFileName(Executable.AbsolutePath));
             GenDebugAction.bCanExecuteRemotely = false;
+			Actions.Add(GenDebugAction);
 
             return OutputFile;
         }
@@ -1102,14 +1107,14 @@ namespace UnrealBuildTool
 			}
 		}
 
-		FileItem ExtractFramework(UEBuildFramework Framework, ActionGraph ActionGraph)
+		FileItem ExtractFramework(UEBuildFramework Framework, List<Action> Actions)
 		{
 			if(Framework.ZipFile == null)
 			{
 				throw new BuildException("Unable to extract framework '{0}' - no zip file specified", Framework.Name);
 			}
 			if(Framework.ExtractedTokenFile == null)
-				{
+			{
 				FileItem InputFile = FileItem.GetItemByFileReference(Framework.ZipFile);
 				Framework.ExtractedTokenFile = FileItem.GetItemByFileReference(new FileReference(Framework.OutputDirectory.FullName + ".extracted"));
 
@@ -1123,7 +1128,7 @@ namespace UnrealBuildTool
 
 				FileItem ExtractScriptFileItem = FileItem.CreateIntermediateTextFile(new FileReference(Framework.OutputDirectory.FullName + ".sh"), ExtractScript.ToString());
 
-				Action UnzipAction = ActionGraph.Add(ActionType.BuildProject);
+				Action UnzipAction = new Action(ActionType.BuildProject);
 				UnzipAction.CommandPath = "/bin/sh";
 				UnzipAction.CommandArguments = Utils.EscapeShellArgument(ExtractScriptFileItem.AbsolutePath);
 				UnzipAction.WorkingDirectory = UnrealBuildTool.EngineDirectory.FullName;
@@ -1133,6 +1138,7 @@ namespace UnrealBuildTool
 				UnzipAction.DeleteItems.Add(Framework.ExtractedTokenFile);
 				UnzipAction.StatusDescription = String.Format("Unzipping: {0} -> {1}", Framework.ZipFile, Framework.OutputDirectory);
 				UnzipAction.bCanExecuteRemotely = false;
+				Actions.Add(UnzipAction);
 			}
 			return Framework.ExtractedTokenFile;
         }
@@ -1282,9 +1288,9 @@ namespace UnrealBuildTool
 			return new DirectoryReference(ResourcesDir);
         }
 
-        public override ICollection<FileItem> PostBuild(FileItem Executable, LinkEnvironment BinaryLinkEnvironment, ActionGraph ActionGraph)
+        public override ICollection<FileItem> PostBuild(FileItem Executable, LinkEnvironment BinaryLinkEnvironment, List<Action> Actions)
         {
-            List<FileItem> OutputFiles = new List<FileItem>(base.PostBuild(Executable, BinaryLinkEnvironment, ActionGraph));
+            List<FileItem> OutputFiles = new List<FileItem>(base.PostBuild(Executable, BinaryLinkEnvironment, Actions));
 
             if (BinaryLinkEnvironment.bIsBuildingLibrary)
             {
@@ -1294,10 +1300,10 @@ namespace UnrealBuildTool
             // For IOS/tvOS, generate the dSYM file if the config file is set to do so
 			if (ProjectSettings.bGeneratedSYMFile == true || ProjectSettings.bGeneratedSYMBundle == true || BinaryLinkEnvironment.bUsePDBFiles == true)
             {
-                OutputFiles.Add(GenerateDebugInfo(Executable, ActionGraph));
+                OutputFiles.Add(GenerateDebugInfo(Executable, Actions));
                 if (ProjectSettings.bGenerateCrashReportSymbols || Target.bUseMallocProfiler)
                 {
-                    OutputFiles.Add(GeneratePseudoPDB(Executable, ActionGraph));
+                    OutputFiles.Add(GeneratePseudoPDB(Executable, Actions));
                 }
             }
 
@@ -1306,7 +1312,7 @@ namespace UnrealBuildTool
 			{
 				FileItem StripCompleteFile = FileItem.GetItemByFileReference(FileReference.Combine(BinaryLinkEnvironment.IntermediateDirectory, Executable.Location.GetFileName() + ".stripped"));
 
-				Action StripAction = ActionGraph.Add(ActionType.CreateAppBundle);
+				Action StripAction = new Action(ActionType.CreateAppBundle);
 				StripAction.WorkingDirectory = GetMacDevSrcRoot();
 				StripAction.CommandPath = "sh";
 				StripAction.CommandArguments = String.Format("-c '\"{0}strip\" \"{1}\" && touch \"{2}\"'", Settings.Value.ToolchainDir, Executable.Location, StripCompleteFile);
@@ -1315,6 +1321,7 @@ namespace UnrealBuildTool
 				StripAction.ProducedItems.Add(StripCompleteFile);
 				StripAction.StatusDescription = String.Format("Stripping symbols from {0}", Executable.AbsolutePath);
 				StripAction.bCanExecuteRemotely = false;
+				Actions.Add(StripAction);
 
 				OutputFiles.Add(StripCompleteFile);
 			}
@@ -1329,7 +1336,7 @@ namespace UnrealBuildTool
 				FileItem AssetCatalogFile = FileItem.GetItemByFileReference(GetAssetCatalogFile(CppPlatform, Executable.Location));
 
 				// Make the compile action
-				Action CompileAssetAction = ActionGraph.Add(ActionType.CreateAppBundle);
+				Action CompileAssetAction = new Action(ActionType.CreateAppBundle);
 				CompileAssetAction.WorkingDirectory = GetMacDevSrcRoot();
 				CompileAssetAction.CommandPath = "/usr/bin/xcrun";
 				CompileAssetAction.CommandArguments = GetAssetCatalogArgs(CppPlatform, ResourcesDir.FullName, Path.GetDirectoryName(AssetCatalogFile.AbsolutePath));
@@ -1338,6 +1345,7 @@ namespace UnrealBuildTool
 				CompileAssetAction.DeleteItems.Add(AssetCatalogFile);
 				CompileAssetAction.StatusDescription = CompileAssetAction.CommandArguments;// string.Format("Generating debug info for {0}", Path.GetFileName(Executable.AbsolutePath));
 				CompileAssetAction.bCanExecuteRemotely = false;
+				Actions.Add(CompileAssetAction);
 
 				// Add it to the output files so it's always built
 				OutputFiles.Add(AssetCatalogFile);
@@ -1367,9 +1375,9 @@ namespace UnrealBuildTool
 				FileReference PostBuildSyncFile = FileReference.Combine(BinaryLinkEnvironment.IntermediateDirectory, "PostBuildSync.dat");
 				BinaryFormatterUtils.Save(PostBuildSyncFile, PostBuildSyncTarget);
 
-				string PostBuildSyncArguments = String.Format("-Mode=IOSPostBuildSync -Input=\"{0}\" -XmlConfigCache=\"{1}\"", PostBuildSyncFile, XmlConfig.CacheFile);
+				string PostBuildSyncArguments = String.Format("-Input=\"{0}\" -XmlConfigCache=\"{1}\"", PostBuildSyncFile, XmlConfig.CacheFile);
 
-				Action PostBuildSyncAction = ActionGraph.AddRecursiveCall(ActionType.CreateAppBundle, PostBuildSyncArguments);
+				Action PostBuildSyncAction = Action.CreateRecursiveAction<IOSPostBuildSyncMode>(ActionType.CreateAppBundle, PostBuildSyncArguments);
 				PostBuildSyncAction.WorkingDirectory = UnrealBuildTool.EngineSourceDirectory.FullName;
 				PostBuildSyncAction.PrerequisiteItems.Add(Executable);
 				PostBuildSyncAction.PrerequisiteItems.AddRange(OutputFiles);
@@ -1377,6 +1385,7 @@ namespace UnrealBuildTool
 				PostBuildSyncAction.DeleteItems.AddRange(PostBuildSyncAction.ProducedItems);
 				PostBuildSyncAction.StatusDescription = "Executing PostBuildSync";
 				PostBuildSyncAction.bCanExecuteRemotely = false;
+				Actions.Add(PostBuildSyncAction);
 
 				OutputFiles.AddRange(PostBuildSyncAction.ProducedItems);
 			}
diff --git a/Engine/Source/Programs/UnrealBuildTool/Platform/Linux/LinuxToolChain.cs b/Engine/Source/Programs/UnrealBuildTool/Platform/Linux/LinuxToolChain.cs
--- a/Engine/Source/Programs/UnrealBuildTool/Platform/Linux/LinuxToolChain.cs
+++ b/Engine/Source/Programs/UnrealBuildTool/Platform/Linux/LinuxToolChain.cs
@@ -1097,7 +1097,7 @@ namespace UnrealBuildTool
 			return false;
 		}
 
-		public override CPPOutput CompileCPPFiles(CppCompileEnvironment CompileEnvironment, List<FileItem> InputFiles, DirectoryReference OutputDir, string ModuleName, ActionGraph ActionGraph)
+		public override CPPOutput CompileCPPFiles(CppCompileEnvironment CompileEnvironment, List<FileItem> InputFiles, DirectoryReference OutputDir, string ModuleName, List<Action> Actions)
 		{
 			string Arguments = GetCLArguments_Global(CompileEnvironment);
 			string PCHArguments = "";
@@ -1161,7 +1161,7 @@ namespace UnrealBuildTool
 			CPPOutput Result = new CPPOutput();
 			foreach (FileItem SourceFile in InputFiles)
 			{
-				Action CompileAction = ActionGraph.Add(ActionType.Compile);
+				Action CompileAction = new Action(ActionType.Compile);
 				CompileAction.PrerequisiteItems.AddRange(CompileEnvironment.ForceIncludeFiles);
 
 				string FileArguments = "";
@@ -1266,6 +1266,8 @@ namespace UnrealBuildTool
 				CompileAction.bCanExecuteRemotely =
 					CompileEnvironment.PrecompiledHeaderAction != PrecompiledHeaderAction.Create ||
 					CompileEnvironment.bAllowRemotelyCompiledPCHs;
+
+				Actions.Add(CompileAction);
 			}
 
 			return Result;
@@ -1279,10 +1281,10 @@ namespace UnrealBuildTool
 		/// <summary>
 		/// Creates an action to archive all the .o files into single .a file
 		/// </summary>
-		public FileItem CreateArchiveAndIndex(LinkEnvironment LinkEnvironment, ActionGraph ActionGraph)
+		public FileItem CreateArchiveAndIndex(LinkEnvironment LinkEnvironment, List<Action> Actions)
 		{
 			// Create an archive action
-			Action ArchiveAction = ActionGraph.Add(ActionType.Link);
+			Action ArchiveAction = new Action(ActionType.Link);
 			ArchiveAction.WorkingDirectory = UnrealBuildTool.EngineSourceDirectory.FullName;
 			bool bUsingSh = BuildHostPlatform.Current.Platform != UnrealTargetPlatform.Win64 && BuildHostPlatform.Current.Platform != UnrealTargetPlatform.Win32;
 			if (bUsingSh)
@@ -1345,11 +1347,12 @@ namespace UnrealBuildTool
 
 			// Only execute linking on the local PC.
 			ArchiveAction.bCanExecuteRemotely = false;
+			Actions.Add(ArchiveAction);
 
 			return OutputFile;
 		}
 
-		public FileItem FixDependencies(LinkEnvironment LinkEnvironment, FileItem Executable, ActionGraph ActionGraph)
+		public FileItem FixDependencies(LinkEnvironment LinkEnvironment, FileItem Executable, List<Action> Actions)
 		{
 			if (bUseFixdeps)
 			{
@@ -1367,7 +1370,7 @@ namespace UnrealBuildTool
 
 				FileItem FixDepsScript = FileItem.GetItemByFileReference(FileReference.Combine(LinkEnvironment.LocalShadowDirectory, ScriptName));
 
-				Action PostLinkAction = ActionGraph.Add(ActionType.Link);
+				Action PostLinkAction = new Action(ActionType.Link);
 				PostLinkAction.WorkingDirectory = UnrealBuildTool.EngineSourceDirectory.FullName;
 				PostLinkAction.CommandPath = ShellBinary;
 				PostLinkAction.StatusDescription = string.Format("{0}", Path.GetFileName(Executable.AbsolutePath));
@@ -1398,6 +1401,7 @@ namespace UnrealBuildTool
 				System.Console.WriteLine("{0} {1}", PostLinkAction.CommandPath, PostLinkAction.CommandArguments);
 
 				PostLinkAction.ProducedItems.Add(OutputFile);
+				Actions.Add(PostLinkAction);
 				return OutputFile;
 			}
 			else
@@ -1413,7 +1417,7 @@ namespace UnrealBuildTool
 		}
 
 
-		public override FileItem LinkFiles(LinkEnvironment LinkEnvironment, bool bBuildImportLibraryOnly, ActionGraph ActionGraph)
+		public override FileItem LinkFiles(LinkEnvironment LinkEnvironment, bool bBuildImportLibraryOnly, List<Action> Actions)
 		{
 			Debug.Assert(!bBuildImportLibraryOnly);
 
@@ -1426,11 +1430,11 @@ namespace UnrealBuildTool
 
 			if (LinkEnvironment.bIsBuildingLibrary || bBuildImportLibraryOnly)
 			{
-				return CreateArchiveAndIndex(LinkEnvironment, ActionGraph);
+				return CreateArchiveAndIndex(LinkEnvironment, Actions);
 			}
 
 			// Create an action that invokes the linker.
-			Action LinkAction = ActionGraph.Add(ActionType.Link);
+			Action LinkAction = new Action(ActionType.Link);
 			LinkAction.WorkingDirectory = UnrealBuildTool.EngineSourceDirectory.FullName;
 
 			string LinkCommandString;
@@ -1718,6 +1722,7 @@ namespace UnrealBuildTool
 
 			// Only execute linking on the local PC.
 			LinkAction.bCanExecuteRemotely = false;
+			Actions.Add(LinkAction);
 
 			// Prepare a script that will run later, once all shared libraries and the executable
 			// are created. This script will be called by action created in FixDependencies()
@@ -1786,7 +1791,7 @@ namespace UnrealBuildTool
 				else
 				{
 					// Create the action to relink the library. This actions does not overwrite the source file so it can be executed in parallel
-					Action RelinkAction = ActionGraph.Add(ActionType.Link);
+					Action RelinkAction = new Action(ActionType.Link);
 					RelinkAction.WorkingDirectory = LinkAction.WorkingDirectory;
 					RelinkAction.StatusDescription = LinkAction.StatusDescription;
 					RelinkAction.CommandDescription = "Relink";
@@ -1872,6 +1877,7 @@ namespace UnrealBuildTool
 
 					RelinkAction.CommandPath = ShellBinary;
 					RelinkAction.CommandArguments = ExecuteSwitch + " \"" + RelinkScriptFullPath + "\"";
+					Actions.Add(RelinkAction);
 				}
 			}
 			return OutputFile;
@@ -1888,9 +1894,9 @@ namespace UnrealBuildTool
 			}
 		}
 
-		public override ICollection<FileItem> PostBuild(FileItem Executable, LinkEnvironment BinaryLinkEnvironment, ActionGraph ActionGraph)
+		public override ICollection<FileItem> PostBuild(FileItem Executable, LinkEnvironment BinaryLinkEnvironment, List<Action> Actions)
 		{
-			ICollection<FileItem> OutputFiles = base.PostBuild(Executable, BinaryLinkEnvironment, ActionGraph);
+			ICollection<FileItem> OutputFiles = base.PostBuild(Executable, BinaryLinkEnvironment, Actions);
 
 			if (bUseFixdeps)
 			{
@@ -1899,7 +1905,7 @@ namespace UnrealBuildTool
 					return OutputFiles;
 				}
 
-				FileItem FixDepsOutputFile = FixDependencies(BinaryLinkEnvironment, Executable, ActionGraph);
+				FileItem FixDepsOutputFile = FixDependencies(BinaryLinkEnvironment, Executable, Actions);
 				if (FixDepsOutputFile != null)
 				{
 					OutputFiles.Add(FixDepsOutputFile);
diff --git a/Engine/Source/Programs/UnrealBuildTool/Platform/Mac/MacToolChain.cs b/Engine/Source/Programs/UnrealBuildTool/Platform/Mac/MacToolChain.cs
--- a/Engine/Source/Programs/UnrealBuildTool/Platform/Mac/MacToolChain.cs
+++ b/Engine/Source/Programs/UnrealBuildTool/Platform/Mac/MacToolChain.cs
@@ -384,7 +384,7 @@ namespace UnrealBuildTool
 			return Result;
 		}
 
-		public override CPPOutput CompileCPPFiles(CppCompileEnvironment CompileEnvironment, List<FileItem> InputFiles, DirectoryReference OutputDir, string ModuleName, ActionGraph ActionGraph)
+		public override CPPOutput CompileCPPFiles(CppCompileEnvironment CompileEnvironment, List<FileItem> InputFiles, DirectoryReference OutputDir, string ModuleName, List<Action> Actions)
 		{
 			StringBuilder Arguments = new StringBuilder();
 			StringBuilder PCHArguments = new StringBuilder();
@@ -427,7 +427,7 @@ namespace UnrealBuildTool
 			// Create a compile action for each source file.
 			foreach (FileItem SourceFile in InputFiles)
 			{
-				Action CompileAction = ActionGraph.Add(ActionType.Compile);
+				Action CompileAction = new Action(ActionType.Compile);
 				CompileAction.PrerequisiteItems.AddRange(CompileEnvironment.ForceIncludeFiles);
 
 				string FileArguments = "";
@@ -528,6 +528,7 @@ namespace UnrealBuildTool
 				// We're already distributing the command by execution on Mac.
 				CompileAction.bCanExecuteRemotely = false;
 				CompileAction.bShouldOutputStatusDescription = true;
+				Actions.Add(CompileAction);
 			}
 			return Result;
 		}
@@ -631,12 +632,12 @@ namespace UnrealBuildTool
 			}
 		}
 
-		public override FileItem LinkFiles(LinkEnvironment LinkEnvironment, bool bBuildImportLibraryOnly, ActionGraph ActionGraph)
+		public override FileItem LinkFiles(LinkEnvironment LinkEnvironment, bool bBuildImportLibraryOnly, List<Action> Actions)
 		{
 			bool bIsBuildingLibrary = LinkEnvironment.bIsBuildingLibrary || bBuildImportLibraryOnly;
 
 			// Create an action that invokes the linker.
-			Action LinkAction = ActionGraph.Add(ActionType.Link);
+			Action LinkAction = new Action(ActionType.Link);
 
 			LinkAction.WorkingDirectory = GetMacDevSrcRoot();
 			LinkAction.CommandPath = "/bin/sh";
@@ -848,6 +849,8 @@ namespace UnrealBuildTool
 
 			LinkAction.ProducedItems.Add(OutputFile);
 
+			Actions.Add(LinkAction);
+
 			if (!DirectoryReference.Exists(LinkEnvironment.IntermediateDirectory))
 			{
 				return OutputFile;
@@ -1035,9 +1038,9 @@ namespace UnrealBuildTool
 			return String.Format("rsync --checksum \"{0}\" \"{1}\"", SourceFile, TargetFile);
 		}
 
-		FileItem FixDylibDependencies(LinkEnvironment LinkEnvironment, FileItem Executable, ActionGraph ActionGraph)
+		FileItem FixDylibDependencies(LinkEnvironment LinkEnvironment, FileItem Executable, List<Action> Actions)
 		{
-			Action LinkAction = ActionGraph.Add(ActionType.Link);
+			Action LinkAction = new Action(ActionType.Link);
 			LinkAction.WorkingDirectory = UnrealBuildTool.EngineSourceDirectory.FullName;
 			LinkAction.CommandPath = "/bin/sh";
 			LinkAction.CommandDescription = "";
@@ -1067,6 +1070,7 @@ namespace UnrealBuildTool
 			LinkAction.CommandArguments += "'";
 
 			LinkAction.ProducedItems.Add(OutputFile);
+			Actions.Add(LinkAction);
 
 			return OutputFile;
 		}
@@ -1076,8 +1080,8 @@ namespace UnrealBuildTool
 		/// </summary>
 		/// <param name="MachOBinary">FileItem describing the executable or dylib to generate debug info for</param>
 		/// <param name="LinkEnvironment"></param>
-		/// <param name="ActionGraph"></param>
-		public FileItem GenerateDebugInfo(FileItem MachOBinary, LinkEnvironment LinkEnvironment, ActionGraph ActionGraph)
+		/// <param name="Actions">List of actions to be executed. Additional actions will be added to this list.</param>
+		public FileItem GenerateDebugInfo(FileItem MachOBinary, LinkEnvironment LinkEnvironment, List<Action> Actions)
 		{
 			string BinaryPath = MachOBinary.AbsolutePath;
 			if (BinaryPath.Contains(".app"))
@@ -1102,7 +1106,7 @@ namespace UnrealBuildTool
 			}
 
 			// Make the compile action
-			Action GenDebugAction = ActionGraph.Add(ActionType.GenerateDebugInfo);
+			Action GenDebugAction = new Action(ActionType.GenerateDebugInfo);
 			GenDebugAction.WorkingDirectory = GetMacDevSrcRoot();
 			GenDebugAction.CommandPath = "sh";
 
@@ -1122,6 +1126,7 @@ namespace UnrealBuildTool
 			GenDebugAction.CommandDescription = "";
 			GenDebugAction.StatusDescription = "Generating " + Path.GetFileName(BinaryPath);
 			GenDebugAction.bCanExecuteRemotely = false;
+			Actions.Add(GenDebugAction);
 
 			return OutputFile;
 		}
@@ -1132,15 +1137,15 @@ namespace UnrealBuildTool
 		/// <param name="LinkEnvironment"></param>
 		/// <param name="Executable">FileItem describing the executable to generate app bundle for</param>
 		/// <param name="FixDylibOutputFile"></param>
-		/// <param name="ActionGraph"></param>
-		FileItem FinalizeAppBundle(LinkEnvironment LinkEnvironment, FileItem Executable, FileItem FixDylibOutputFile, ActionGraph ActionGraph)
+		/// <param name="Actions">List of actions to be executed. Additional actions will be added to this list.</param>
+		FileItem FinalizeAppBundle(LinkEnvironment LinkEnvironment, FileItem Executable, FileItem FixDylibOutputFile, List<Action> Actions)
 		{
 			// Make a file item for the source and destination files
 			string FullDestPath = Executable.AbsolutePath.Substring(0, Executable.AbsolutePath.IndexOf(".app") + 4);
 			FileItem DestFile = FileItem.GetItemByPath(FullDestPath);
 
 			// Make the compile action
-			Action FinalizeAppBundleAction = ActionGraph.Add(ActionType.CreateAppBundle);
+			Action FinalizeAppBundleAction = new Action(ActionType.CreateAppBundle);
 			FinalizeAppBundleAction.WorkingDirectory = GetMacDevSrcRoot(); // Path.GetFullPath(".");
 			FinalizeAppBundleAction.CommandPath = "/bin/sh";
 			FinalizeAppBundleAction.CommandDescription = "";
@@ -1153,13 +1158,14 @@ namespace UnrealBuildTool
 			FinalizeAppBundleAction.ProducedItems.Add(DestFile);
 			FinalizeAppBundleAction.StatusDescription = string.Format("Finalizing app bundle: {0}.app", Path.GetFileName(Executable.AbsolutePath));
 			FinalizeAppBundleAction.bCanExecuteRemotely = false;
+			Actions.Add(FinalizeAppBundleAction);
 
 			return DestFile;
 		}
 
-		FileItem CopyBundleResource(UEBuildBundleResource Resource, FileItem Executable, DirectoryReference BundleDirectory, ActionGraph ActionGraph)
+		FileItem CopyBundleResource(UEBuildBundleResource Resource, FileItem Executable, DirectoryReference BundleDirectory, List<Action> Actions)
 		{
-			Action CopyAction = ActionGraph.Add(ActionType.CreateAppBundle);
+			Action CopyAction = new Action(ActionType.CreateAppBundle);
 			CopyAction.WorkingDirectory = GetMacDevSrcRoot(); // Path.GetFullPath(".");
 			CopyAction.CommandPath = "/bin/sh";
 			CopyAction.CommandDescription = "";
@@ -1176,6 +1182,7 @@ namespace UnrealBuildTool
 			CopyAction.bShouldOutputStatusDescription = Resource.bShouldLog;
 			CopyAction.StatusDescription = string.Format("Copying {0} to app bundle", Path.GetFileName(Resource.ResourcePath));
 			CopyAction.bCanExecuteRemotely = false;
+			Actions.Add(CopyAction);
 
 			return TargetItem;
 		}
@@ -1289,9 +1296,9 @@ namespace UnrealBuildTool
 			}
 		}
 
-		public override ICollection<FileItem> PostBuild(FileItem Executable, LinkEnvironment BinaryLinkEnvironment, ActionGraph ActionGraph)
+		public override ICollection<FileItem> PostBuild(FileItem Executable, LinkEnvironment BinaryLinkEnvironment, List<Action> Actions)
 		{
-			ICollection<FileItem> OutputFiles = base.PostBuild(Executable, BinaryLinkEnvironment, ActionGraph);
+			ICollection<FileItem> OutputFiles = base.PostBuild(Executable, BinaryLinkEnvironment, Actions);
 
 			if (BinaryLinkEnvironment.bIsBuildingLibrary)
 			{
@@ -1302,7 +1309,7 @@ namespace UnrealBuildTool
 			{
 				foreach (UEBuildBundleResource Resource in BinaryLinkEnvironment.AdditionalBundleResources)
 				{
-					OutputFiles.Add(CopyBundleResource(Resource, Executable, BinaryLinkEnvironment.BundleDirectory, ActionGraph));
+					OutputFiles.Add(CopyBundleResource(Resource, Executable, BinaryLinkEnvironment.BundleDirectory, Actions));
 				}
 			}
 
@@ -1312,7 +1319,7 @@ namespace UnrealBuildTool
 				// We want dsyms to be created after all dylib dependencies are fixed. If FixDylibDependencies action was not created yet, save the info for later.
 				if (FixDylibOutputFile != null)
 				{
-					OutputFiles.Add(GenerateDebugInfo(Executable, BinaryLinkEnvironment, ActionGraph));
+					OutputFiles.Add(GenerateDebugInfo(Executable, BinaryLinkEnvironment, Actions));
 				}
 				else
 				{
@@ -1325,17 +1332,17 @@ namespace UnrealBuildTool
 				return OutputFiles;
 			}
 
-			FixDylibOutputFile = FixDylibDependencies(BinaryLinkEnvironment, Executable, ActionGraph);
+			FixDylibOutputFile = FixDylibDependencies(BinaryLinkEnvironment, Executable, Actions);
 			OutputFiles.Add(FixDylibOutputFile);
 			if (!BinaryLinkEnvironment.bIsBuildingConsoleApplication)
 			{
-				OutputFiles.Add(FinalizeAppBundle(BinaryLinkEnvironment, Executable, FixDylibOutputFile, ActionGraph));
+				OutputFiles.Add(FinalizeAppBundle(BinaryLinkEnvironment, Executable, FixDylibOutputFile, Actions));
 			}
 
 			// Add dsyms that we couldn't add before FixDylibDependencies action was created
 			foreach (FileItem Exe in ExecutablesThatNeedDsyms)
 			{
-				OutputFiles.Add(GenerateDebugInfo(Exe, BinaryLinkEnvironment, ActionGraph));
+				OutputFiles.Add(GenerateDebugInfo(Exe, BinaryLinkEnvironment, Actions));
 			}
 			ExecutablesThatNeedDsyms.Clear();
 
diff --git a/Engine/Source/Programs/UnrealBuildTool/Platform/Windows/PVSToolChain.cs b/Engine/Source/Programs/UnrealBuildTool/Platform/Windows/PVSToolChain.cs
--- a/Engine/Source/Programs/UnrealBuildTool/Platform/Windows/PVSToolChain.cs
+++ b/Engine/Source/Programs/UnrealBuildTool/Platform/Windows/PVSToolChain.cs
@@ -191,7 +191,7 @@ namespace UnrealBuildTool
 			return Path.GetFullPath(ActionThread.ExpandEnvironmentVariables(IncludePath));
 		}
 
-		public override CPPOutput CompileCPPFiles(CppCompileEnvironment CompileEnvironment, List<FileItem> InputFiles, DirectoryReference OutputDir, string ModuleName, ActionGraph ActionGraph)
+		public override CPPOutput CompileCPPFiles(CppCompileEnvironment CompileEnvironment, List<FileItem> InputFiles, DirectoryReference OutputDir, string ModuleName, List<Action> Actions)
 		{
 			// Get the MSVC arguments required to compile all files in this batch
 			List<string> SharedArguments = new List<string>();
@@ -244,7 +244,7 @@ namespace UnrealBuildTool
 				// Preprocess the source file
 				FileItem PreprocessedFileItem = FileItem.GetItemByFileReference(PreprocessedFileLocation);
 
-				Action PreprocessAction = ActionGraph.Add(ActionType.Compile);
+				Action PreprocessAction = new Action(ActionType.Compile);
 				PreprocessAction.CommandPath = EnvVars.CompilerPath.FullName;
 				PreprocessAction.WorkingDirectory = UnrealBuildTool.EngineSourceDirectory.FullName;
 				PreprocessAction.CommandArguments = " @\"" + ResponseFileItem.AbsolutePath + "\"";
@@ -253,6 +253,7 @@ namespace UnrealBuildTool
 				PreprocessAction.PrerequisiteItems.Add(ResponseFileItem);
 				PreprocessAction.ProducedItems.Add(PreprocessedFileItem);
 				PreprocessAction.bShouldOutputStatusDescription = false;
+				Actions.Add(PreprocessAction);
 
 				// Write the PVS studio config file
 				StringBuilder ConfigFileContents = new StringBuilder();
@@ -301,7 +302,7 @@ namespace UnrealBuildTool
 				FileReference OutputFileLocation = FileReference.Combine(OutputDir, BaseFileName + ".pvslog");
 				FileItem OutputFileItem = FileItem.GetItemByFileReference(OutputFileLocation);
 
-				Action AnalyzeAction = ActionGraph.Add(ActionType.Compile);
+				Action AnalyzeAction = new Action(ActionType.Compile);
 				AnalyzeAction.CommandDescription = "Analyzing";
 				AnalyzeAction.StatusDescription = BaseFileName;
 				AnalyzeAction.WorkingDirectory = UnrealBuildTool.EngineSourceDirectory.FullName;
@@ -316,18 +317,19 @@ namespace UnrealBuildTool
 				AnalyzeAction.PrerequisiteItems.Add(PreprocessedFileItem);
 				AnalyzeAction.ProducedItems.Add(OutputFileItem);
 				AnalyzeAction.DeleteItems.Add(OutputFileItem); // PVS Studio will append by default, so need to delete produced items
+				Actions.Add(AnalyzeAction);
 
 				Result.ObjectFiles.AddRange(AnalyzeAction.ProducedItems);
 			}
 			return Result;
 		}
 
-		public override FileItem LinkFiles(LinkEnvironment LinkEnvironment, bool bBuildImportLibraryOnly, ActionGraph ActionGraph)
+		public override FileItem LinkFiles(LinkEnvironment LinkEnvironment, bool bBuildImportLibraryOnly, List<Action> Actions)
 		{
 			throw new BuildException("Unable to link with PVS toolchain.");
 		}
 
-		public override void FinalizeOutput(ReadOnlyTargetRules Target, List<FileItem> OutputItems, ActionGraph ActionGraph)
+		public override void FinalizeOutput(ReadOnlyTargetRules Target, List<FileItem> OutputItems, List<Action> Actions)
 		{
 			FileReference OutputFile;
 			if (Target.ProjectFile == null)
@@ -343,7 +345,7 @@ namespace UnrealBuildTool
 
 			FileItem InputFileListItem = FileItem.CreateIntermediateTextFile(OutputFile.ChangeExtension(".input"), InputFiles.Select(x => x.FullName));
 
-			Action AnalyzeAction = ActionGraph.Add(ActionType.Compile);
+			Action AnalyzeAction = new Action(ActionType.Compile);
 			AnalyzeAction.CommandPath = UnrealBuildTool.GetUBTPath();
 			AnalyzeAction.CommandArguments = String.Format("-Mode=PVSGather -Input=\"{0}\" -Output=\"{1}\"", InputFileListItem.Location, OutputFile);
 			AnalyzeAction.WorkingDirectory = UnrealBuildTool.EngineSourceDirectory.FullName;
@@ -351,6 +353,7 @@ namespace UnrealBuildTool
 			AnalyzeAction.PrerequisiteItems.AddRange(OutputItems);
 			AnalyzeAction.ProducedItems.Add(FileItem.GetItemByFileReference(OutputFile));
 			AnalyzeAction.DeleteItems.AddRange(AnalyzeAction.ProducedItems);
+			Actions.Add(AnalyzeAction);
 
 			OutputItems.AddRange(AnalyzeAction.ProducedItems);
 		}
diff --git a/Engine/Source/Programs/UnrealBuildTool/Platform/Windows/VCToolChain.cs b/Engine/Source/Programs/UnrealBuildTool/Platform/Windows/VCToolChain.cs
--- a/Engine/Source/Programs/UnrealBuildTool/Platform/Windows/VCToolChain.cs
+++ b/Engine/Source/Programs/UnrealBuildTool/Platform/Windows/VCToolChain.cs
@@ -841,7 +841,7 @@ namespace UnrealBuildTool
 			}
 		}
 
-		public override CPPOutput CompileCPPFiles(CppCompileEnvironment CompileEnvironment, List<FileItem> InputFiles, DirectoryReference OutputDir, string ModuleName, ActionGraph ActionGraph)
+		public override CPPOutput CompileCPPFiles(CppCompileEnvironment CompileEnvironment, List<FileItem> InputFiles, DirectoryReference OutputDir, string ModuleName, List<Action> Actions)
 		{
 			List<string> SharedArguments = new List<string>();
 			AppendCLArguments_Global(CompileEnvironment, SharedArguments);
@@ -886,7 +886,7 @@ namespace UnrealBuildTool
 			CPPOutput Result = new CPPOutput();
 			foreach (FileItem SourceFile in InputFiles)
 			{
-				Action CompileAction = ActionGraph.Add(ActionType.Compile);
+				Action CompileAction = new Action(ActionType.Compile);
 				CompileAction.CommandDescription = "Compile";
 
 				// ensure compiler timings are captured when we execute the action.
@@ -1105,17 +1105,19 @@ namespace UnrealBuildTool
 				{
 					CompileAction.bCanExecuteRemotelyWithSNDBS = false;
 				}
+
+				Actions.Add(CompileAction);
 			}
 			return Result;
 		}
 
-		public override CPPOutput CompileRCFiles(CppCompileEnvironment CompileEnvironment, List<FileItem> InputFiles, DirectoryReference OutputDir, ActionGraph ActionGraph)
+		public override CPPOutput CompileRCFiles(CppCompileEnvironment CompileEnvironment, List<FileItem> InputFiles, DirectoryReference OutputDir, List<Action> Actions)
 		{
 			CPPOutput Result = new CPPOutput();
 
 			foreach (FileItem RCFile in InputFiles)
 			{
-				Action CompileAction = ActionGraph.Add(ActionType.Compile);
+				Action CompileAction = new Action(ActionType.Compile);
 				CompileAction.CommandDescription = "Resource";
 				CompileAction.WorkingDirectory = UnrealBuildTool.EngineSourceDirectory.FullName;
 				CompileAction.CommandPath = EnvVars.ResourceCompilerPath.FullName;
@@ -1183,12 +1185,14 @@ namespace UnrealBuildTool
 
 				// Add the C++ source file and its included files to the prerequisite item list.
 				CompileAction.PrerequisiteItems.Add(RCFile);
+
+				Actions.Add(CompileAction);
 			}
 
 			return Result;
 		}
 
-		public override FileItem LinkFiles(LinkEnvironment LinkEnvironment, bool bBuildImportLibraryOnly, ActionGraph ActionGraph)
+		public override FileItem LinkFiles(LinkEnvironment LinkEnvironment, bool bBuildImportLibraryOnly, List<Action> Actions)
 		{
 			if (LinkEnvironment.bIsBuildingDotNetAssembly)
 			{
@@ -1406,7 +1410,7 @@ namespace UnrealBuildTool
 			}
 
 			// Create an action that invokes the linker.
-			Action LinkAction = ActionGraph.Add(ActionType.Link);
+			Action LinkAction = new Action(ActionType.Link);
 			LinkAction.CommandDescription = "Link";
 			LinkAction.WorkingDirectory = UnrealBuildTool.EngineSourceDirectory.FullName;
 			LinkAction.CommandPath = bIsBuildingLibraryOrImportLibrary ? EnvVars.LibraryManagerPath.FullName : EnvVars.LinkerPath.FullName;
@@ -1439,6 +1443,7 @@ namespace UnrealBuildTool
 
 			// Allow remote linking.  Especially in modular builds with many small DLL files, this is almost always very efficient
 			LinkAction.bCanExecuteRemotely = true;
+			Actions.Add(LinkAction);
 
 			Log.TraceVerbose("     Linking: " + LinkAction.StatusDescription);
 			Log.TraceVerbose("     Command: " + LinkAction.CommandArguments);
diff --git a/Engine/Source/Programs/UnrealBuildTool/System/Action.cs b/Engine/Source/Programs/UnrealBuildTool/System/Action.cs
--- a/Engine/Source/Programs/UnrealBuildTool/System/Action.cs
+++ b/Engine/Source/Programs/UnrealBuildTool/System/Action.cs
@@ -4,6 +4,7 @@ using System;
 using System.Collections.Generic;
 using System.IO;
 using System.Linq;
+using System.Reflection;
 using System.Text;
 using System.Threading.Tasks;
 using Tools.DotNETCommon;
@@ -212,6 +213,26 @@ namespace UnrealBuildTool
 			Writer.WriteFileItem(DependencyListFile);
 		}
 
+		/// <summary>
+		/// Creates an action which calls UBT recursively
+		/// </summary>
+		/// <param name="Type">Type of the action</param>
+		/// <param name="Arguments">Arguments for the action</param>
+		/// <returns>New action instance</returns>
+		public static Action CreateRecursiveAction<T>(ActionType Type, string Arguments) where T : ToolMode
+		{
+			ToolModeAttribute Attribute = typeof(T).GetCustomAttribute<ToolModeAttribute>();
+			if(Attribute == null)
+			{
+				throw new BuildException("Missing ToolModeAttribute on {0}", typeof(T).Name);
+			}
+
+			Action NewAction = new Action(Type);
+			NewAction.CommandPath = UnrealBuildTool.GetUBTPath();
+			NewAction.CommandArguments = String.Format("-Mode={0} {1}", Attribute.Name, Arguments);
+			return NewAction;
+		}
+
 		/// <summary>
 		/// Finds conflicts betwee two actions, and prints them to the log
 		/// </summary>
diff --git a/Engine/Source/Programs/UnrealBuildTool/System/ActionGraph.cs b/Engine/Source/Programs/UnrealBuildTool/System/ActionGraph.cs
--- a/Engine/Source/Programs/UnrealBuildTool/System/ActionGraph.cs
+++ b/Engine/Source/Programs/UnrealBuildTool/System/ActionGraph.cs
@@ -12,32 +12,8 @@ using System.Threading.Tasks;
 
 namespace UnrealBuildTool
 {
-	class ActionGraph
+	static class ActionGraph
 	{
-		/// <summary>
-		/// List of all the actions
-		/// </summary>
-		public List<Action> Actions = new List<Action>();
-
-		public ActionGraph()
-		{
-		}
-
-		public Action Add(ActionType Type)
-		{
-			Action NewAction = new Action(Type);
-			Actions.Add(NewAction);
-			return NewAction;
-		}
-
-		public Action AddRecursiveCall(ActionType Type, string Arguments)
-		{
-			Action NewAction = Add(Type);
-			NewAction.CommandPath = UnrealBuildTool.GetUBTPath();
-			NewAction.CommandArguments = Arguments;
-			return NewAction;
-		}
-
 		/// <summary>
 		/// Links the actions together and sets up their dependencies
 		/// </summary>
diff --git a/Engine/Source/Programs/UnrealBuildTool/ToolChain/UEToolChain.cs b/Engine/Source/Programs/UnrealBuildTool/ToolChain/UEToolChain.cs
--- a/Engine/Source/Programs/UnrealBuildTool/ToolChain/UEToolChain.cs
+++ b/Engine/Source/Programs/UnrealBuildTool/ToolChain/UEToolChain.cs
@@ -26,18 +26,18 @@ namespace UnrealBuildTool
 		{
 		}
 
-		public abstract CPPOutput CompileCPPFiles(CppCompileEnvironment CompileEnvironment, List<FileItem> InputFiles, DirectoryReference OutputDir, string ModuleName, ActionGraph ActionGraph);
+		public abstract CPPOutput CompileCPPFiles(CppCompileEnvironment CompileEnvironment, List<FileItem> InputFiles, DirectoryReference OutputDir, string ModuleName, List<Action> Actions);
 
-		public virtual CPPOutput CompileRCFiles(CppCompileEnvironment Environment, List<FileItem> InputFiles, DirectoryReference OutputDir, ActionGraph ActionGraph)
+		public virtual CPPOutput CompileRCFiles(CppCompileEnvironment Environment, List<FileItem> InputFiles, DirectoryReference OutputDir, List<Action> Actions)
 		{
 			CPPOutput Result = new CPPOutput();
 			return Result;
 		}
 
-		public abstract FileItem LinkFiles(LinkEnvironment LinkEnvironment, bool bBuildImportLibraryOnly, ActionGraph ActionGraph);
-		public virtual FileItem[] LinkAllFiles(LinkEnvironment LinkEnvironment, bool bBuildImportLibraryOnly, ActionGraph ActionGraph)
+		public abstract FileItem LinkFiles(LinkEnvironment LinkEnvironment, bool bBuildImportLibraryOnly, List<Action> Actions);
+		public virtual FileItem[] LinkAllFiles(LinkEnvironment LinkEnvironment, bool bBuildImportLibraryOnly, List<Action> Actions)
 		{
-			return new FileItem[] { LinkFiles(LinkEnvironment, bBuildImportLibraryOnly, ActionGraph) };
+			return new FileItem[] { LinkFiles(LinkEnvironment, bBuildImportLibraryOnly, Actions) };
 		}
 
 
@@ -53,7 +53,7 @@ namespace UnrealBuildTool
 			return FileReference.Combine(LinkEnvironment.IntermediateDirectory, OutputFile.Location.GetFileName() + ".response");
 		}
 
-		public virtual ICollection<FileItem> PostBuild(FileItem Executable, LinkEnvironment ExecutableLinkEnvironment, ActionGraph ActionGraph)
+		public virtual ICollection<FileItem> PostBuild(FileItem Executable, LinkEnvironment ExecutableLinkEnvironment, List<Action> Actions)
 		{
 			return new List<FileItem>();
 		}
@@ -66,7 +66,7 @@ namespace UnrealBuildTool
 		{
 		}
 
-		public virtual void FinalizeOutput(ReadOnlyTargetRules Target, List<FileItem> OutputItems, ActionGraph ActionGraph)
+		public virtual void FinalizeOutput(ReadOnlyTargetRules Target, List<FileItem> OutputItems, List<Action> Actions)
 		{
 		}
 
