UBT: Add more detailed timing info.

#rnx

diff --git a/Engine/Source/Programs/UnrealBuildTool/Configuration/UEBuildTarget.cs b/Engine/Source/Programs/UnrealBuildTool/Configuration/UEBuildTarget.cs
--- a/Engine/Source/Programs/UnrealBuildTool/Configuration/UEBuildTarget.cs
+++ b/Engine/Source/Programs/UnrealBuildTool/Configuration/UEBuildTarget.cs
@@ -1313,14 +1313,17 @@ namespace UnrealBuildTool
 			{
 				ExeDir = ExeDir.ParentDirectory.ParentDirectory.ParentDirectory;
 			}
-			foreach (UEBuildBinary Binary in Binaries)
+			using(Timeline.ScopeEvent("UEBuildBinary.Build()"))
 			{
-				List<FileItem> BinaryOutputItems = Binary.Build(Rules, TargetToolChain, GlobalCompileEnvironment, GlobalLinkEnvironment, SharedPCHs, WorkingSet, Prerequisites, ExeDir, Actions);
-				if(!bCompileMonolithic)
+				foreach (UEBuildBinary Binary in Binaries)
 				{
-					ModuleNameToOutputItems[Binary.PrimaryModule.Name] = BinaryOutputItems.ToArray();
+					List<FileItem> BinaryOutputItems = Binary.Build(Rules, TargetToolChain, GlobalCompileEnvironment, GlobalLinkEnvironment, SharedPCHs, WorkingSet, Prerequisites, ExeDir, Actions);
+					if(!bCompileMonolithic)
+					{
+						ModuleNameToOutputItems[Binary.PrimaryModule.Name] = BinaryOutputItems.ToArray();
+					}
+					OutputItems.AddRange(BinaryOutputItems);
 				}
-				OutputItems.AddRange(BinaryOutputItems);
 			}
 
 			// Prepare all the runtime dependencies, copying them from their source folders if necessary
diff --git a/Engine/Source/Programs/UnrealBuildTool/Modes/BuildMode.cs b/Engine/Source/Programs/UnrealBuildTool/Modes/BuildMode.cs
--- a/Engine/Source/Programs/UnrealBuildTool/Modes/BuildMode.cs
+++ b/Engine/Source/Programs/UnrealBuildTool/Modes/BuildMode.cs
@@ -85,8 +85,6 @@ namespace UnrealBuildTool
 			Arguments.TryGetValue("-RemoteIni=", out RemoteIniPath);
 			UnrealBuildTool.SetRemoteIniPath(RemoteIniPath);
 
-			Timeline.AddEvent("Basic UBT initialization");
-
 			// now that we know the available platforms, we can delete other platforms' junk. if we're only building specific modules from the editor, don't touch anything else (it may be in use).
 			if (!BuildConfiguration.bIgnoreJunk)
 			{
@@ -107,8 +105,6 @@ namespace UnrealBuildTool
 				// action graph implies using the dependency resolve cache
 				bool GeneratingActionGraph = Arguments.HasOption("-Graph");
 
-				Timeline.AddEvent("RunUBT() initialization complete");
-
 				bool bSkipRulesCompile = Arguments.Any(x => x.Equals("-skiprulescompile", StringComparison.InvariantCultureIgnoreCase));
 
 				List<TargetDescriptor> TargetDescriptors = new List<TargetDescriptor>();
@@ -198,7 +194,11 @@ namespace UnrealBuildTool
 						if(Makefile == null)
 						{
 							// Create the target
-							UEBuildTarget Target = UEBuildTarget.CreateTarget(TargetDesc, bSkipRulesCompile, BuildConfiguration.SingleFileToCompile != null, BuildConfiguration.bUsePrecompiled);
+							UEBuildTarget Target;
+							using(Timeline.ScopeEvent("UEBuildTarget.CreateTarget"))
+							{
+								Target = UEBuildTarget.CreateTarget(TargetDesc, bSkipRulesCompile, BuildConfiguration.SingleFileToCompile != null, BuildConfiguration.bUsePrecompiled);
+							}
 
 							// Create the header cache for this target
 							FileReference DependencyCacheFile = DependencyCache.GetDependencyCachePathForTarget(Target.ProjectFile, Target.Platform, Target.TargetName);
@@ -220,7 +220,11 @@ namespace UnrealBuildTool
 							List<Action> Actions = new List<Action>();
 							BuildPrerequisites Prerequisites = new BuildPrerequisites();
 
-							ECompilationResult BuildResult = Target.Build(BuildConfiguration, Headers, OutputItems, ModuleNameToOutputItems, UObjectModules, WorkingSet, Actions, Prerequisites, bIsAssemblingBuild);
+							ECompilationResult BuildResult;
+							using(Timeline.ScopeEvent("UEBuildTarget.Build()"))
+							{
+								BuildResult = Target.Build(BuildConfiguration, Headers, OutputItems, ModuleNameToOutputItems, UObjectModules, WorkingSet, Actions, Prerequisites, bIsAssemblingBuild);
+							}
 							if (BuildResult != ECompilationResult.Succeeded)
 							{
 								return (int)BuildResult;
diff --git a/Engine/Source/Programs/UnrealBuildTool/System/ExternalExecution.cs b/Engine/Source/Programs/UnrealBuildTool/System/ExternalExecution.cs
--- a/Engine/Source/Programs/UnrealBuildTool/System/ExternalExecution.cs
+++ b/Engine/Source/Programs/UnrealBuildTool/System/ExternalExecution.cs
@@ -1208,7 +1208,12 @@ namespace UnrealBuildTool
 							UBTArguments.Append(" -logsuffix=UHT");
 						}
 
-						if (RunExternalDotNETExecutable(UnrealBuildTool.GetUBTPath(), UBTArguments.ToString()) != 0)
+						int ExitCode;
+						using(Timeline.ScopeEvent("Buildng UnrealHeaderTool"))
+						{
+							ExitCode = RunExternalDotNETExecutable(UnrealBuildTool.GetUBTPath(), UBTArguments.ToString());
+						}
+						if(ExitCode != 0)
 						{
 							return ECompilationResult.OtherCompilationError;
 						}
