UBT: Add a cache for source file metadata (containing text from the first include directive, and a flag for whether the file contains reflection markup). Replaces the previous include dependency cache, since we now only require the first include directive. Saves ~2s whenever UHTModuleInfo array has to be regenerated.

diff --git a/Engine/Source/Programs/UnrealBuildTool/Configuration/UEBuildModuleCPP.cs b/Engine/Source/Programs/UnrealBuildTool/Configuration/UEBuildModuleCPP.cs
--- a/Engine/Source/Programs/UnrealBuildTool/Configuration/UEBuildModuleCPP.cs
+++ b/Engine/Source/Programs/UnrealBuildTool/Configuration/UEBuildModuleCPP.cs
@@ -807,7 +807,7 @@ namespace UnrealBuildTool
 				// Create a new C++ environment to compile the PCH
 				CppCompileEnvironment PchEnvironment = new CppCompileEnvironment(CompileEnvironment);
 				PchEnvironment.Definitions.Clear();
-				PchEnvironment.IncludePaths.UserIncludePaths.Add(File.Location.Directory); // Need to be able to include headers in the same directory as the source file
+				PchEnvironment.UserIncludePaths.Add(File.Location.Directory); // Need to be able to include headers in the same directory as the source file
 				PchEnvironment.PrecompiledHeaderAction = PrecompiledHeaderAction.Create;
 				PchEnvironment.PrecompiledHeaderIncludeFilename = DedicatedPchFile.Location;
 
@@ -943,19 +943,19 @@ namespace UnrealBuildTool
 					InvalidIncludeDirectiveMessages = new List<string>();
 					if (Rules != null && Rules.bEnforceIWYU && Target.bEnforceIWYU)
 					{
-						foreach (FileItem CPPFile in CppFiles)
+						foreach (FileItem CppFile in CppFiles)
 						{
-							List<DependencyInclude> DirectIncludeFilenames = ModuleCompileEnvironment.Headers.GetDirectIncludeDependencies(CPPFile, bOnlyCachedDependencies: false);
-							if (DirectIncludeFilenames.Count > 0)
+							string FirstInclude = ModuleCompileEnvironment.MetadataCache.GetFirstInclude(CppFile);
+							if(FirstInclude != null)
 							{
-								string IncludeName = Path.GetFileNameWithoutExtension(DirectIncludeFilenames[0].IncludeName);
-								string ExpectedName = CPPFile.Location.GetFileNameWithoutExtension();
-								if (String.Compare(IncludeName, ExpectedName, StringComparison.InvariantCultureIgnoreCase) != 0)
+								string IncludeName = Path.GetFileNameWithoutExtension(FirstInclude);
+								string ExpectedName = CppFile.Location.GetFileNameWithoutExtension();
+								if (String.Compare(IncludeName, ExpectedName, StringComparison.OrdinalIgnoreCase) != 0)
 								{
 									FileReference HeaderFile;
 									if (NameToHeaderFile.TryGetValue(ExpectedName, out HeaderFile) && !IgnoreMismatchedHeader(ExpectedName))
 									{
-										InvalidIncludeDirectiveMessages.Add(String.Format("{0}(1): error: Expected {1} to be first header included.", CPPFile.Location, HeaderFile.GetFileName()));
+										InvalidIncludeDirectiveMessages.Add(String.Format("{0}(1): error: Expected {1} to be first header included.", CppFile.Location, HeaderFile.GetFileName()));
 									}
 								}
 							}
@@ -1075,7 +1075,7 @@ namespace UnrealBuildTool
 			}
 
 			// Setup the compile environment for the module.
-			SetupPrivateCompileEnvironment(Result.IncludePaths.UserIncludePaths, Result.IncludePaths.SystemIncludePaths, Result.Definitions, Result.AdditionalFrameworks, (Rules != null)? Rules.bLegacyPublicIncludePaths.Value : true);
+			SetupPrivateCompileEnvironment(Result.UserIncludePaths, Result.SystemIncludePaths, Result.Definitions, Result.AdditionalFrameworks, (Rules != null)? Rules.bLegacyPublicIncludePaths.Value : true);
 
 			return Result;
 		}
@@ -1117,7 +1117,7 @@ namespace UnrealBuildTool
 			// Now set up the compile environment for the modules in the original order that we encountered them
 			foreach (UEBuildModule Module in ModuleToIncludePathsOnlyFlag.Keys)
 			{
-				Module.AddModuleToCompileEnvironment(null, CompileEnvironment.IncludePaths.UserIncludePaths, CompileEnvironment.IncludePaths.SystemIncludePaths, CompileEnvironment.Definitions, CompileEnvironment.AdditionalFrameworks, (Rules != null)? Rules.bLegacyPublicIncludePaths.Value : true);
+				Module.AddModuleToCompileEnvironment(null, CompileEnvironment.UserIncludePaths, CompileEnvironment.SystemIncludePaths, CompileEnvironment.Definitions, CompileEnvironment.AdditionalFrameworks, (Rules != null)? Rules.bLegacyPublicIncludePaths.Value : true);
 			}
 			return CompileEnvironment;
 		}
diff --git a/Engine/Source/Programs/UnrealBuildTool/Configuration/UEBuildTarget.cs b/Engine/Source/Programs/UnrealBuildTool/Configuration/UEBuildTarget.cs
--- a/Engine/Source/Programs/UnrealBuildTool/Configuration/UEBuildTarget.cs
+++ b/Engine/Source/Programs/UnrealBuildTool/Configuration/UEBuildTarget.cs
@@ -446,6 +446,11 @@ namespace UnrealBuildTool
 		/// </summary>
 		public RulesAssembly RulesAssembly;
 
+		/// <summary>
+		/// Cache of source file metadata for this target
+		/// </summary>
+		public SourceFileMetadataCache MetadataCache;
+
 		/// <summary>
 		/// The project file for this target
 		/// </summary>
@@ -597,6 +602,7 @@ namespace UnrealBuildTool
 		/// <param name="InRulesAssembly">The chain of rules assemblies that this target was created with</param>
 		public UEBuildTarget(TargetDescriptor InDesc, ReadOnlyTargetRules InRules, RulesAssembly InRulesAssembly)
 		{
+			MetadataCache = SourceFileMetadataCache.CreateHierarchy(InDesc.ProjectFile);
 			ProjectFile = InDesc.ProjectFile;
 			AppName = InDesc.Name;
 			TargetName = InDesc.Name;
@@ -1072,7 +1078,9 @@ namespace UnrealBuildTool
 			CppPlatform CppPlatform = UEBuildPlatform.GetBuildPlatform(Platform).DefaultCppPlatform;
 			CppConfiguration CppConfiguration = GetCppConfiguration(Configuration);
 
-			CppCompileEnvironment GlobalCompileEnvironment = new CppCompileEnvironment(CppPlatform, CppConfiguration, Architecture, null);
+			SourceFileMetadataCache MetadataCache = SourceFileMetadataCache.CreateHierarchy(ProjectFile);
+
+			CppCompileEnvironment GlobalCompileEnvironment = new CppCompileEnvironment(CppPlatform, CppConfiguration, Architecture, MetadataCache);
 			LinkEnvironment GlobalLinkEnvironment = new LinkEnvironment(GlobalCompileEnvironment.Platform, GlobalCompileEnvironment.Configuration, GlobalCompileEnvironment.Architecture);
 
 			UEToolChain TargetToolChain = CreateToolchain(CppPlatform);
@@ -1084,12 +1092,14 @@ namespace UnrealBuildTool
 		/// <summary>
 		/// Builds the target, appending list of output files and returns building result.
 		/// </summary>
-		public ECompilationResult Build(BuildConfiguration BuildConfiguration, CPPHeaders Headers, List<FileItem> OutputItems, Dictionary<string, FileItem[]> ModuleNameToOutputItems, List<UHTModuleInfo> UObjectModules, ISourceFileWorkingSet WorkingSet, List<Action> Actions, BuildPrerequisites Prerequisites, bool bIsAssemblingBuild)
+		public ECompilationResult Build(BuildConfiguration BuildConfiguration, List<FileItem> OutputItems, Dictionary<string, FileItem[]> ModuleNameToOutputItems, List<UHTModuleInfo> UObjectModules, ISourceFileWorkingSet WorkingSet, List<Action> Actions, BuildPrerequisites Prerequisites, bool bIsAssemblingBuild)
 		{
 			CppPlatform CppPlatform = UEBuildPlatform.GetBuildPlatform(Platform).DefaultCppPlatform;
 			CppConfiguration CppConfiguration = GetCppConfiguration(Configuration);
 
-			CppCompileEnvironment GlobalCompileEnvironment = new CppCompileEnvironment(CppPlatform, CppConfiguration, Architecture, Headers);
+			SourceFileMetadataCache MetadataCache = SourceFileMetadataCache.CreateHierarchy(ProjectFile);
+
+			CppCompileEnvironment GlobalCompileEnvironment = new CppCompileEnvironment(CppPlatform, CppConfiguration, Architecture, MetadataCache);
 			LinkEnvironment GlobalLinkEnvironment = new LinkEnvironment(GlobalCompileEnvironment.Platform, GlobalCompileEnvironment.Configuration, GlobalCompileEnvironment.Architecture);
 
 			UEToolChain TargetToolChain = CreateToolchain(CppPlatform);
@@ -1263,7 +1273,7 @@ namespace UnrealBuildTool
 			HashSet<UEBuildModuleCPP> ModulesToGenerateHeadersFor = GatherDependencyModules(OriginalBinaries.ToList());
 			using(Timeline.ScopeEvent("ExternalExecution.SetupUObjectModules()"))
 			{
-				ExternalExecution.SetupUObjectModules(ModulesToGenerateHeadersFor, Rules.Platform, ProjectDescriptor, UObjectModules, Prerequisites.UObjectModuleHeaders, Rules.GeneratedCodeVersion, bIsAssemblingBuild);
+				ExternalExecution.SetupUObjectModules(ModulesToGenerateHeadersFor, Rules.Platform, ProjectDescriptor, UObjectModules, Prerequisites.UObjectModuleHeaders, Rules.GeneratedCodeVersion, bIsAssemblingBuild, MetadataCache);
 			}
 
 			// NOTE: Even in Gather mode, we need to run UHT to make sure the files exist for the static action graph to be setup correctly.  This is because UHT generates .cpp
@@ -2835,7 +2845,7 @@ namespace UnrealBuildTool
 			GlobalCompileEnvironment.bPGOOptimize = Rules.bPGOOptimize;
 			GlobalCompileEnvironment.bPGOProfile = Rules.bPGOProfile;
 			GlobalCompileEnvironment.bAllowRemotelyCompiledPCHs = Rules.bAllowRemotelyCompiledPCHs;
-			GlobalCompileEnvironment.IncludePaths.bCheckSystemHeadersForModification = Rules.bCheckSystemHeadersForModification;
+			GlobalCompileEnvironment.bCheckSystemHeadersForModification = Rules.bCheckSystemHeadersForModification;
 			GlobalCompileEnvironment.bPrintTimingInfo = Rules.bPrintToolChainTimingInfo;
 			GlobalCompileEnvironment.bUseRTTI = Rules.bForceEnableRTTI;
 			GlobalCompileEnvironment.bUseInlining = Rules.bUseInlining;
@@ -2878,7 +2888,7 @@ namespace UnrealBuildTool
 			GlobalCompileEnvironment.Definitions.Add("USE_VORBIS_FOR_STREAMING=1");
 
 			// Add the 'Engine/Source' path as a global include path for all modules
-			GlobalCompileEnvironment.IncludePaths.UserIncludePaths.Add(UnrealBuildTool.EngineSourceDirectory);
+			GlobalCompileEnvironment.UserIncludePaths.Add(UnrealBuildTool.EngineSourceDirectory);
 
 			[REDACTED_EMAIL]: Do any platform specific tool chain initialization here if required
 
@@ -3193,10 +3203,10 @@ namespace UnrealBuildTool
 					FileReference CppFile = DirectoryReference.EnumerateFiles(RulesObject.Directory, "*.cpp", SearchOption.AllDirectories).FirstOrDefault();
 					if(CppFile != null)
 					{
-						List<DependencyInclude> Includes = CPPHeaders.GetUncachedDirectIncludeDependencies(CppFile, ProjectFile);
-						if(Includes.Count > 0)
+						string IncludeFile = MetadataCache.GetFirstInclude(FileItem.GetItemByFileReference(CppFile));
+						if(IncludeFile != null)
 						{
-							FileReference PchIncludeFile = DirectoryReference.EnumerateFiles(RulesObject.Directory, Path.GetFileName(Includes[0].IncludeName), SearchOption.AllDirectories).FirstOrDefault();
+							FileReference PchIncludeFile = DirectoryReference.EnumerateFiles(RulesObject.Directory, Path.GetFileName(IncludeFile), SearchOption.AllDirectories).FirstOrDefault();
 							if(PchIncludeFile != null)
 							{
 								RulesObject.PrivatePCHHeaderFile = PchIncludeFile.MakeRelativeTo(RulesObject.Directory).Replace(Path.DirectorySeparatorChar, '/');
diff --git a/Engine/Source/Programs/UnrealBuildTool/Modes/BuildMode.cs b/Engine/Source/Programs/UnrealBuildTool/Modes/BuildMode.cs
--- a/Engine/Source/Programs/UnrealBuildTool/Modes/BuildMode.cs
+++ b/Engine/Source/Programs/UnrealBuildTool/Modes/BuildMode.cs
@@ -96,8 +96,6 @@ namespace UnrealBuildTool
 			// Reset global configurations
 			string ExecutorName = "Unknown";
 
-			List<CPPHeaders> HeadersList = new List<CPPHeaders>();
-
 			Stopwatch ExecutorTimer = new Stopwatch();
 
 			try
@@ -200,17 +198,6 @@ namespace UnrealBuildTool
 								Target = UEBuildTarget.CreateTarget(TargetDesc, bSkipRulesCompile, BuildConfiguration.SingleFileToCompile != null, BuildConfiguration.bUsePrecompiled);
 							}
 
-							// Create the header cache for this target
-							FileReference DependencyCacheFile = DependencyCache.GetDependencyCachePathForTarget(Target.ProjectFile, Target.Platform, Target.TargetName);
-							CPPHeaders Headers = new CPPHeaders(Target.ProjectFile, DependencyCacheFile);
-							HeadersList.Add(Headers);
-
-							// Read the dependency cache
-							using(Timeline.ScopeEvent("DependencyCache.Create()"))
-							{
-								Headers.IncludeDependencyCache = DependencyCache.Create(DependencyCache.GetDependencyCachePathForTarget(Target.ProjectFile, Target.Platform, Target.TargetName));
-							}
-
 							// Build the target
 							const bool bIsAssemblingBuild = true;
 
@@ -223,7 +210,7 @@ namespace UnrealBuildTool
 							ECompilationResult BuildResult;
 							using(Timeline.ScopeEvent("UEBuildTarget.Build()"))
 							{
-								BuildResult = Target.Build(BuildConfiguration, Headers, OutputItems, ModuleNameToOutputItems, UObjectModules, WorkingSet, Actions, Prerequisites, bIsAssemblingBuild);
+								BuildResult = Target.Build(BuildConfiguration, OutputItems, ModuleNameToOutputItems, UObjectModules, WorkingSet, Actions, Prerequisites, bIsAssemblingBuild);
 							}
 							if (BuildResult != ECompilationResult.Succeeded)
 							{
@@ -408,34 +395,15 @@ namespace UnrealBuildTool
 			}
 			finally
 			{
-				// Save the dependency caches
+				// Save all the caches
+				SourceFileMetadataCache.SaveAll();
 				CppDependencyCache.SaveAll();
 
-				// Save the include dependency cache.
-				foreach (CPPHeaders Headers in HeadersList)
-				{
-					// NOTE: It's very important that we save the include cache, even if a build exception was thrown (compile error, etc), because we need to make sure that
-					//    any C++ include dependencies that we computed for out of date source files are saved.  Remember, the build may fail *after* some build products
-					//    are successfully built.  If we didn't save our dependency cache after build failures, source files for those build products that were successfully
-					//    built before the failure would not be considered out of date on the next run, so this is our only chance to cache C++ includes for those files!
-
-					if (Headers.IncludeDependencyCache != null)
-					{
-						using(Timeline.ScopeEvent("DependencyCache.Save()"))
-						{
-							Headers.IncludeDependencyCache.Save();
-						}
-					}
-				}
-
 				// Figure out how long we took to execute.
 				if (ExecutorName != "Unknown")
 				{
 					Log.TraceInformation("Total build time: {0:0.00} seconds ({1} executor: {2:0.00} seconds)", BuildTimer.Elapsed.TotalSeconds, ExecutorName, ExecutorTimer.Elapsed.TotalSeconds);
 				}
-
-				// Print some performance info
-				Log.TraceLog("DirectIncludes cache miss time: {0}s ({1} misses)", CPPHeaders.DirectIncludeCacheMissesTotalTime, CPPHeaders.TotalDirectIncludeCacheMisses);
 			}
 		}
 
diff --git a/Engine/Source/Programs/UnrealBuildTool/Modes/CleanMode.cs b/Engine/Source/Programs/UnrealBuildTool/Modes/CleanMode.cs
--- a/Engine/Source/Programs/UnrealBuildTool/Modes/CleanMode.cs
+++ b/Engine/Source/Programs/UnrealBuildTool/Modes/CleanMode.cs
@@ -127,8 +127,8 @@ namespace UnrealBuildTool
 				}
 
 				// Add all the makefiles and caches to be deleted
-				FilesToDelete.Add(DependencyCache.GetDependencyCachePathForTarget(Target.ProjectFile, Target.Platform, Target.Name));
 				FilesToDelete.Add(TargetMakefile.GetLocation(Target.ProjectFile, Target.Name, Target.Platform, Target.Configuration));
+				FilesToDelete.UnionWith(SourceFileMetadataCache.GetFilesToClean(Target.ProjectFile));
 				FilesToDelete.UnionWith(ActionHistory.GetFilesToClean(Target.ProjectFile, Target.Name, Target.Platform, Target.Type));
 
 				// Add all the intermediate folders to be deleted
diff --git a/Engine/Source/Programs/UnrealBuildTool/Platform/Android/AndroidToolChain.cs b/Engine/Source/Programs/UnrealBuildTool/Platform/Android/AndroidToolChain.cs
--- a/Engine/Source/Programs/UnrealBuildTool/Platform/Android/AndroidToolChain.cs
+++ b/Engine/Source/Programs/UnrealBuildTool/Platform/Android/AndroidToolChain.cs
@@ -1304,14 +1304,14 @@ namespace UnrealBuildTool
 					}
 
 					// Add include paths to the argument list (filtered by architecture)
-					foreach (DirectoryReference IncludePath in CompileEnvironment.IncludePaths.SystemIncludePaths)
+					foreach (DirectoryReference IncludePath in CompileEnvironment.SystemIncludePaths)
 					{
 						if (IsDirectoryForArch(IncludePath.FullName, Arch))
 						{
 							Arguments += string.Format(" -I\"{0}\"", IncludePath);
 						}
 					}
-					foreach (DirectoryReference IncludePath in CompileEnvironment.IncludePaths.UserIncludePaths)
+					foreach (DirectoryReference IncludePath in CompileEnvironment.UserIncludePaths)
 					{
 						if (IsDirectoryForArch(IncludePath.FullName, Arch))
 						{
diff --git a/Engine/Source/Programs/UnrealBuildTool/Platform/Android/UEBuildAndroid.cs b/Engine/Source/Programs/UnrealBuildTool/Platform/Android/UEBuildAndroid.cs
--- a/Engine/Source/Programs/UnrealBuildTool/Platform/Android/UEBuildAndroid.cs
+++ b/Engine/Source/Programs/UnrealBuildTool/Platform/Android/UEBuildAndroid.cs
@@ -354,21 +354,21 @@ namespace UnrealBuildTool
 			CompileEnvironment.Definitions.Add("USE_NULL_RHI=0");
 
 			DirectoryReference NdkDir = new DirectoryReference(NDKPath);
-			CompileEnvironment.IncludePaths.SystemIncludePaths.Add(DirectoryReference.Combine(NdkDir, "sources/cxx-stl/gnu-libstdc++/" + GccVersion + "/include"));
+			CompileEnvironment.SystemIncludePaths.Add(DirectoryReference.Combine(NdkDir, "sources/cxx-stl/gnu-libstdc++/" + GccVersion + "/include"));
 
 			// the toolchain will actually filter these out
-			CompileEnvironment.IncludePaths.SystemIncludePaths.Add(DirectoryReference.Combine(NdkDir, "sources/cxx-stl/gnu-libstdc++/" + GccVersion + "/libs/armeabi-v7a/include"));
-			CompileEnvironment.IncludePaths.SystemIncludePaths.Add(DirectoryReference.Combine(NdkDir, "sources/cxx-stl/gnu-libstdc++/" + GccVersion + "/libs/arm64-v8a/include"));
-			CompileEnvironment.IncludePaths.SystemIncludePaths.Add(DirectoryReference.Combine(NdkDir, "sources/cxx-stl/gnu-libstdc++/" + GccVersion + "/libs/x86/include"));
-			CompileEnvironment.IncludePaths.SystemIncludePaths.Add(DirectoryReference.Combine(NdkDir, "sources/cxx-stl/gnu-libstdc++/" + GccVersion + "/libs/x86_64/include"));
+			CompileEnvironment.SystemIncludePaths.Add(DirectoryReference.Combine(NdkDir, "sources/cxx-stl/gnu-libstdc++/" + GccVersion + "/libs/armeabi-v7a/include"));
+			CompileEnvironment.SystemIncludePaths.Add(DirectoryReference.Combine(NdkDir, "sources/cxx-stl/gnu-libstdc++/" + GccVersion + "/libs/arm64-v8a/include"));
+			CompileEnvironment.SystemIncludePaths.Add(DirectoryReference.Combine(NdkDir, "sources/cxx-stl/gnu-libstdc++/" + GccVersion + "/libs/x86/include"));
+			CompileEnvironment.SystemIncludePaths.Add(DirectoryReference.Combine(NdkDir, "sources/cxx-stl/gnu-libstdc++/" + GccVersion + "/libs/x86_64/include"));
 
 			LinkEnvironment.LibraryPaths.Add(DirectoryReference.Combine(NdkDir, "sources/cxx-stl/gnu-libstdc++/" + GccVersion + "/libs/armeabi-v7a"));
 			LinkEnvironment.LibraryPaths.Add(DirectoryReference.Combine(NdkDir, "sources/cxx-stl/gnu-libstdc++/" + GccVersion + "/libs/arm64-v8a"));
 			LinkEnvironment.LibraryPaths.Add(DirectoryReference.Combine(NdkDir, "sources/cxx-stl/gnu-libstdc++/" + GccVersion + "/libs/x86"));
 			LinkEnvironment.LibraryPaths.Add(DirectoryReference.Combine(NdkDir, "sources/cxx-stl/gnu-libstdc++/" + GccVersion + "/libs/x86_64"));
 
-			CompileEnvironment.IncludePaths.SystemIncludePaths.Add(DirectoryReference.Combine(NdkDir, "sources/android/native_app_glue"));
-			CompileEnvironment.IncludePaths.SystemIncludePaths.Add(DirectoryReference.Combine(NdkDir, "sources/android/cpufeatures"));
+			CompileEnvironment.SystemIncludePaths.Add(DirectoryReference.Combine(NdkDir, "sources/android/native_app_glue"));
+			CompileEnvironment.SystemIncludePaths.Add(DirectoryReference.Combine(NdkDir, "sources/android/cpufeatures"));
 
 
 			// if the project has an Oodle compression Dll, enable the decompressor on Android
diff --git a/Engine/Source/Programs/UnrealBuildTool/Platform/HTML5/HTML5ToolChain.cs b/Engine/Source/Programs/UnrealBuildTool/Platform/HTML5/HTML5ToolChain.cs
--- a/Engine/Source/Programs/UnrealBuildTool/Platform/HTML5/HTML5ToolChain.cs
+++ b/Engine/Source/Programs/UnrealBuildTool/Platform/HTML5/HTML5ToolChain.cs
@@ -346,11 +346,11 @@ namespace UnrealBuildTool
 			CPPOutput Result = new CPPOutput();
 
 			// Add include paths to the argument list.
-			foreach (DirectoryReference IncludePath in CompileEnvironment.IncludePaths.UserIncludePaths)
+			foreach (DirectoryReference IncludePath in CompileEnvironment.UserIncludePaths)
 			{
 				AddIncludePath(ref Arguments, IncludePath);
 			}
-			foreach (DirectoryReference IncludePath in CompileEnvironment.IncludePaths.SystemIncludePaths)
+			foreach (DirectoryReference IncludePath in CompileEnvironment.SystemIncludePaths)
 			{
 				AddIncludePath(ref Arguments, IncludePath);
 			}
diff --git a/Engine/Source/Programs/UnrealBuildTool/Platform/IOS/IOSToolChain.cs b/Engine/Source/Programs/UnrealBuildTool/Platform/IOS/IOSToolChain.cs
--- a/Engine/Source/Programs/UnrealBuildTool/Platform/IOS/IOSToolChain.cs
+++ b/Engine/Source/Programs/UnrealBuildTool/Platform/IOS/IOSToolChain.cs
@@ -564,8 +564,8 @@ namespace UnrealBuildTool
 			}
 
 			// Add include paths to the argument list.
-			HashSet<DirectoryReference> AllIncludes = new HashSet<DirectoryReference>(CompileEnvironment.IncludePaths.UserIncludePaths);
-			AllIncludes.UnionWith(CompileEnvironment.IncludePaths.SystemIncludePaths);
+			HashSet<DirectoryReference> AllIncludes = new HashSet<DirectoryReference>(CompileEnvironment.UserIncludePaths);
+			AllIncludes.UnionWith(CompileEnvironment.SystemIncludePaths);
 			foreach (DirectoryReference IncludePath in AllIncludes)
 			{
 				Arguments += string.Format(" -I\"{0}\"", IncludePath.FullName);
diff --git a/Engine/Source/Programs/UnrealBuildTool/Platform/Linux/LinuxToolChain.cs b/Engine/Source/Programs/UnrealBuildTool/Platform/Linux/LinuxToolChain.cs
--- a/Engine/Source/Programs/UnrealBuildTool/Platform/Linux/LinuxToolChain.cs
+++ b/Engine/Source/Programs/UnrealBuildTool/Platform/Linux/LinuxToolChain.cs
@@ -1142,11 +1142,11 @@ namespace UnrealBuildTool
 			}
 
 			// Add include paths to the argument list.
-			foreach (DirectoryReference IncludePath in CompileEnvironment.IncludePaths.UserIncludePaths)
+			foreach (DirectoryReference IncludePath in CompileEnvironment.UserIncludePaths)
 			{
 				Arguments += string.Format(" -I\"{0}\"", IncludePath.FullName.Replace('\\', '/'));
 			}
-			foreach (DirectoryReference IncludePath in CompileEnvironment.IncludePaths.SystemIncludePaths)
+			foreach (DirectoryReference IncludePath in CompileEnvironment.SystemIncludePaths)
 			{
 				Arguments += string.Format(" -I\"{0}\"", IncludePath.FullName.Replace('\\', '/'));
 			}
diff --git a/Engine/Source/Programs/UnrealBuildTool/Platform/Linux/UEBuildLinux.cs b/Engine/Source/Programs/UnrealBuildTool/Platform/Linux/UEBuildLinux.cs
--- a/Engine/Source/Programs/UnrealBuildTool/Platform/Linux/UEBuildLinux.cs
+++ b/Engine/Source/Programs/UnrealBuildTool/Platform/Linux/UEBuildLinux.cs
@@ -454,7 +454,7 @@ namespace UnrealBuildTool
 			string BaseLinuxPath = SDK.GetBaseLinuxPathForArchitecture(Target.Architecture);
 			if (BuildHostPlatform.Current.Platform == UnrealTargetPlatform.Linux && String.IsNullOrEmpty(BaseLinuxPath))
 			{
-				CompileEnvironment.IncludePaths.SystemIncludePaths.Add(new DirectoryReference("/usr/include"));
+				CompileEnvironment.SystemIncludePaths.Add(new DirectoryReference("/usr/include"));
 			}
 
 			if (CompileEnvironment.bAllowLTCG != LinkEnvironment.bAllowLTCG)
diff --git a/Engine/Source/Programs/UnrealBuildTool/Platform/Lumin/UEBuildLumin.cs b/Engine/Source/Programs/UnrealBuildTool/Platform/Lumin/UEBuildLumin.cs
--- a/Engine/Source/Programs/UnrealBuildTool/Platform/Lumin/UEBuildLumin.cs
+++ b/Engine/Source/Programs/UnrealBuildTool/Platform/Lumin/UEBuildLumin.cs
@@ -166,9 +166,9 @@ namespace UnrealBuildTool
 			CompileEnvironment.Definitions.Add("WITH_OGGVORBIS=1");
 
 			DirectoryReference MLSDKDir = new DirectoryReference("$(MLSDK)", DirectoryReference.Sanitize.None);
-			CompileEnvironment.IncludePaths.SystemIncludePaths.Add(DirectoryReference.Combine(MLSDKDir, "lumin/stl/gnu-libstdc++/include"));
-			CompileEnvironment.IncludePaths.SystemIncludePaths.Add(DirectoryReference.Combine(MLSDKDir, "lumin/stl/gnu-libstdc++/include/aarch64-linux-android"));
-			CompileEnvironment.IncludePaths.SystemIncludePaths.Add(DirectoryReference.Combine(MLSDKDir, "include"));
+			CompileEnvironment.SystemIncludePaths.Add(DirectoryReference.Combine(MLSDKDir, "lumin/stl/gnu-libstdc++/include"));
+			CompileEnvironment.SystemIncludePaths.Add(DirectoryReference.Combine(MLSDKDir, "lumin/stl/gnu-libstdc++/include/aarch64-linux-android"));
+			CompileEnvironment.SystemIncludePaths.Add(DirectoryReference.Combine(MLSDKDir, "include"));
 
 			LinkEnvironment.LibraryPaths.Add(DirectoryReference.Combine(MLSDKDir, "lib/lumin"));
 			LinkEnvironment.LibraryPaths.Add(DirectoryReference.Combine(MLSDKDir, "lumin/stl/gnu-libstdc++/lib"));
diff --git a/Engine/Source/Programs/UnrealBuildTool/Platform/Mac/MacToolChain.cs b/Engine/Source/Programs/UnrealBuildTool/Platform/Mac/MacToolChain.cs
--- a/Engine/Source/Programs/UnrealBuildTool/Platform/Mac/MacToolChain.cs
+++ b/Engine/Source/Programs/UnrealBuildTool/Platform/Mac/MacToolChain.cs
@@ -406,8 +406,8 @@ namespace UnrealBuildTool
 			}
 
 			// Add include paths to the argument list.
-			HashSet<DirectoryReference> AllIncludes = new HashSet<DirectoryReference>(CompileEnvironment.IncludePaths.UserIncludePaths);
-			AllIncludes.UnionWith(CompileEnvironment.IncludePaths.SystemIncludePaths);
+			HashSet<DirectoryReference> AllIncludes = new HashSet<DirectoryReference>(CompileEnvironment.UserIncludePaths);
+			AllIncludes.UnionWith(CompileEnvironment.SystemIncludePaths);
 			foreach (DirectoryReference IncludePath in AllIncludes)
 			{
 				Arguments.Append(" -I\"");
diff --git a/Engine/Source/Programs/UnrealBuildTool/Platform/Windows/PVSToolChain.cs b/Engine/Source/Programs/UnrealBuildTool/Platform/Windows/PVSToolChain.cs
--- a/Engine/Source/Programs/UnrealBuildTool/Platform/Windows/PVSToolChain.cs
+++ b/Engine/Source/Programs/UnrealBuildTool/Platform/Windows/PVSToolChain.cs
@@ -204,11 +204,11 @@ namespace UnrealBuildTool
 			{
 				SharedArguments.Add("/D _SILENCE_STDEXT_HASH_DEPRECATION_WARNINGS=1");
 			}
-			foreach (DirectoryReference IncludePath in CompileEnvironment.IncludePaths.UserIncludePaths)
+			foreach (DirectoryReference IncludePath in CompileEnvironment.UserIncludePaths)
 			{
 				SharedArguments.Add(String.Format("/I \"{0}\"", IncludePath));
 			}
-			foreach (DirectoryReference IncludePath in CompileEnvironment.IncludePaths.SystemIncludePaths)
+			foreach (DirectoryReference IncludePath in CompileEnvironment.SystemIncludePaths)
 			{
 				SharedArguments.Add(String.Format("/I \"{0}\"", IncludePath));
 			}
diff --git a/Engine/Source/Programs/UnrealBuildTool/Platform/Windows/VCToolChain.cs b/Engine/Source/Programs/UnrealBuildTool/Platform/Windows/VCToolChain.cs
--- a/Engine/Source/Programs/UnrealBuildTool/Platform/Windows/VCToolChain.cs
+++ b/Engine/Source/Programs/UnrealBuildTool/Platform/Windows/VCToolChain.cs
@@ -847,12 +847,12 @@ namespace UnrealBuildTool
 			AppendCLArguments_Global(CompileEnvironment, SharedArguments);
 
 			// Add include paths to the argument list.
-			foreach (DirectoryReference IncludePath in CompileEnvironment.IncludePaths.UserIncludePaths)
+			foreach (DirectoryReference IncludePath in CompileEnvironment.UserIncludePaths)
 			{
 				AddIncludePath(SharedArguments, IncludePath);
 			}
 
-			foreach (DirectoryReference IncludePath in CompileEnvironment.IncludePaths.SystemIncludePaths)
+			foreach (DirectoryReference IncludePath in CompileEnvironment.SystemIncludePaths)
 			{
 				AddSystemIncludePath(SharedArguments, IncludePath);
 			}
@@ -1143,13 +1143,13 @@ namespace UnrealBuildTool
 				Arguments.Add("/l 0x409");
 
 				// Include paths. Don't use AddIncludePath() here, since it uses the full path and exceeds the max command line length.
-				foreach (DirectoryReference IncludePath in CompileEnvironment.IncludePaths.UserIncludePaths)
+				foreach (DirectoryReference IncludePath in CompileEnvironment.UserIncludePaths)
 				{
 					Arguments.Add(String.Format("/I \"{0}\"", IncludePath));
 				}
 
 				// System include paths.
-				foreach (DirectoryReference SystemIncludePath in CompileEnvironment.IncludePaths.SystemIncludePaths)
+				foreach (DirectoryReference SystemIncludePath in CompileEnvironment.SystemIncludePaths)
 				{
 					Arguments.Add(String.Format("/I \"{0}\"", SystemIncludePath));
 				}
diff --git a/Engine/Source/Programs/UnrealBuildTool/ProjectFiles/ProjectFileGenerator.cs b/Engine/Source/Programs/UnrealBuildTool/ProjectFiles/ProjectFileGenerator.cs
--- a/Engine/Source/Programs/UnrealBuildTool/ProjectFiles/ProjectFileGenerator.cs
+++ b/Engine/Source/Programs/UnrealBuildTool/ProjectFiles/ProjectFileGenerator.cs
@@ -1722,8 +1722,8 @@ namespace UnrealBuildTool
 									{
 										CppCompileEnvironment ModuleCompileEnvironment = Module.CreateModuleCompileEnvironment(Target.Rules, BinaryCompileEnvironment);
 										ProjectFileForIDE.AddIntelliSensePreprocessorDefinitions(ModuleCompileEnvironment.Definitions);
-										ProjectFileForIDE.AddIntelliSenseIncludePaths(ModuleCompileEnvironment.IncludePaths.SystemIncludePaths, true);
-										ProjectFileForIDE.AddIntelliSenseIncludePaths(ModuleCompileEnvironment.IncludePaths.UserIncludePaths, false);
+										ProjectFileForIDE.AddIntelliSenseIncludePaths(ModuleCompileEnvironment.SystemIncludePaths, true);
+										ProjectFileForIDE.AddIntelliSenseIncludePaths(ModuleCompileEnvironment.UserIncludePaths, false);
 									}
 								}
 							}
diff --git a/Engine/Source/Programs/UnrealBuildTool/System/CPPHeaders.cs b/Engine/Source/Programs/UnrealBuildTool/System/CPPHeaders.cs
deleted file mode 100644
diff --git a/Engine/Source/Programs/UnrealBuildTool/System/CppCompileEnvironment.cs b/Engine/Source/Programs/UnrealBuildTool/System/CppCompileEnvironment.cs
--- a/Engine/Source/Programs/UnrealBuildTool/System/CppCompileEnvironment.cs
+++ b/Engine/Source/Programs/UnrealBuildTool/System/CppCompileEnvironment.cs
@@ -90,6 +90,11 @@ namespace UnrealBuildTool
 		/// </summary>
 		public readonly string Architecture;
 
+		/// <summary>
+		/// Cache of source file metadata
+		/// </summary>
+		public readonly SourceFileMetadataCache MetadataCache;
+
 		/// <summary>
 		/// The name of the header file which is precompiled.
 		/// </summary>
@@ -277,9 +282,20 @@ namespace UnrealBuildTool
 		public bool bAllowRemotelyCompiledPCHs = false;
 
 		/// <summary>
-		/// The include paths to look for included files in.
+		/// Ordered list of include paths for the module
+		/// </summary>
+		public HashSet<DirectoryReference> UserIncludePaths;
+
+		/// <summary>
+		/// The include paths where changes to contained files won't cause dependent C++ source files to
+		/// be recompiled, unless BuildConfiguration.bCheckSystemHeadersForModification==true.
 		/// </summary>
-		public readonly CppIncludePaths IncludePaths = new CppIncludePaths();
+		public HashSet<DirectoryReference> SystemIncludePaths;
+
+		/// <summary>
+		/// Whether headers in system paths should be checked for modification when determining outdated actions.
+		/// </summary>
+		public bool bCheckSystemHeadersForModification;
 
 		/// <summary>
 		/// List of header files to force include
@@ -306,11 +322,6 @@ namespace UnrealBuildTool
 		/// </summary>
 		public FileItem PrecompiledHeaderFile = null;
 
-		/// <summary>
-		/// Header file cache for this target
-		/// </summary>
-		public CPPHeaders Headers;
-
 		/// <summary>
 		/// Whether or not UHT is being built
 		/// </summary>
@@ -324,12 +335,14 @@ namespace UnrealBuildTool
 		/// <summary>
 		/// Default constructor.
 		/// </summary>
-        public CppCompileEnvironment(CppPlatform Platform, CppConfiguration Configuration, string Architecture, CPPHeaders Headers)
+        public CppCompileEnvironment(CppPlatform Platform, CppConfiguration Configuration, string Architecture, SourceFileMetadataCache MetadataCache)
 		{
 			this.Platform = Platform;
 			this.Configuration = Configuration;
 			this.Architecture = Architecture;
-			this.Headers = Headers;
+			this.MetadataCache = MetadataCache;
+			this.UserIncludePaths = new HashSet<DirectoryReference>();
+			this.SystemIncludePaths = new HashSet<DirectoryReference>();
 		}
 
 		/// <summary>
@@ -341,6 +354,7 @@ namespace UnrealBuildTool
 			Platform = Other.Platform;
 			Configuration = Other.Configuration;
 			Architecture = Other.Architecture;
+			MetadataCache = Other.MetadataCache;
 			PrecompiledHeaderIncludeFilename = Other.PrecompiledHeaderIncludeFilename;
 			PrecompiledHeaderAction = Other.PrecompiledHeaderAction;
 			bUseSharedBuildEnvironment = Other.bUseSharedBuildEnvironment;
@@ -377,13 +391,14 @@ namespace UnrealBuildTool
 			bPrintTimingInfo = Other.bPrintTimingInfo;
 			bGenerateDependenciesFile = Other.bGenerateDependenciesFile;
 			bAllowRemotelyCompiledPCHs = Other.bAllowRemotelyCompiledPCHs;
-			IncludePaths = new CppIncludePaths(Other.IncludePaths);
+			UserIncludePaths = new HashSet<DirectoryReference>(Other.UserIncludePaths);
+			SystemIncludePaths = new HashSet<DirectoryReference>(Other.SystemIncludePaths);
+			bCheckSystemHeadersForModification = Other.bCheckSystemHeadersForModification;
 			ForceIncludeFiles.AddRange(Other.ForceIncludeFiles);
 			Definitions.AddRange(Other.Definitions);
 			AdditionalArguments = Other.AdditionalArguments;
 			AdditionalFrameworks.AddRange(Other.AdditionalFrameworks);
 			PrecompiledHeaderFile = Other.PrecompiledHeaderFile;
-			Headers = Other.Headers;
 			bHackHeaderGenerator = Other.bHackHeaderGenerator;
 			bHideSymbolsByDefault = Other.bHideSymbolsByDefault;
 		}
diff --git a/Engine/Source/Programs/UnrealBuildTool/System/DependencyCache.cs b/Engine/Source/Programs/UnrealBuildTool/System/DependencyCache.cs
deleted file mode 100644
diff --git a/Engine/Source/Programs/UnrealBuildTool/System/ExternalExecution.cs b/Engine/Source/Programs/UnrealBuildTool/System/ExternalExecution.cs
--- a/Engine/Source/Programs/UnrealBuildTool/System/ExternalExecution.cs
+++ b/Engine/Source/Programs/UnrealBuildTool/System/ExternalExecution.cs
@@ -355,7 +355,7 @@ namespace UnrealBuildTool
 		/// Generates a UHTModuleInfo for a particular named module under a directory.
 		/// </summary>
 		/// <returns></returns>
-		public static UHTModuleInfo CreateUHTModuleInfo(IEnumerable<FileReference> HeaderFilenames, string ModuleName, FileReference ModuleRulesFile, DirectoryReference ModuleDirectory, UHTModuleType ModuleType, EGeneratedCodeVersion GeneratedCodeVersion, bool bIsReadOnly)
+		public static UHTModuleInfo CreateUHTModuleInfo(IEnumerable<FileReference> HeaderFilenames, string ModuleName, FileReference ModuleRulesFile, DirectoryReference ModuleDirectory, UHTModuleType ModuleType, EGeneratedCodeVersion GeneratedCodeVersion, bool bIsReadOnly, SourceFileMetadataCache MetadataCache)
 		{
 			DirectoryReference ClassesFolder = DirectoryReference.Combine(ModuleDirectory, "Classes");
 			DirectoryReference PublicFolder = DirectoryReference.Combine(ModuleDirectory, "Public");
@@ -370,7 +370,7 @@ namespace UnrealBuildTool
 				// UObjects or not, we can skip doing a test here.
 				FileItem UObjectHeaderFileItem = FileItem.GetItemByFileReference(Header);
 
-				if (CPPHeaders.DoesFileContainUObjects(UObjectHeaderFileItem.AbsolutePath))
+				if (MetadataCache.ContainsReflectionMarkup(UObjectHeaderFileItem))
 				{
 					if (UObjectHeaderFileItem.Location.IsUnderDirectory(ClassesFolder))
 					{
@@ -609,7 +609,7 @@ namespace UnrealBuildTool
 			}
 		}
 
-		public static void SetupUObjectModules(IEnumerable<UEBuildModuleCPP> ModulesToGenerateHeadersFor, UnrealTargetPlatform Platform, ProjectDescriptor ProjectDescriptor, List<UHTModuleInfo> UObjectModules, List<UHTModuleHeaderInfo> UObjectModuleHeaders, EGeneratedCodeVersion GeneratedCodeVersion, bool bIsAssemblingBuild)
+		public static void SetupUObjectModules(IEnumerable<UEBuildModuleCPP> ModulesToGenerateHeadersFor, UnrealTargetPlatform Platform, ProjectDescriptor ProjectDescriptor, List<UHTModuleInfo> UObjectModules, List<UHTModuleHeaderInfo> UObjectModuleHeaders, EGeneratedCodeVersion GeneratedCodeVersion, bool bIsAssemblingBuild, SourceFileMetadataCache MetadataCache)
 		{
 			// Find the type of each module
 			Dictionary<UEBuildModuleCPP, UHTModuleType> ModuleToType = new Dictionary<UEBuildModuleCPP, UHTModuleType>();
@@ -628,7 +628,7 @@ namespace UnrealBuildTool
 				List<FileReference> HeaderFiles = new List<FileReference>();
 				FindHeaders(new DirectoryInfo(Module.ModuleDirectory.FullName), ExcludedFolders, HeaderFiles);
 
-				UHTModuleInfo Info = ExternalExecution.CreateUHTModuleInfo(HeaderFiles, Module.Name, Module.RulesFile, Module.ModuleDirectory, ModuleToType[Module], GeneratedCodeVersion, Module.Rules.bUsePrecompiled);
+				UHTModuleInfo Info = ExternalExecution.CreateUHTModuleInfo(HeaderFiles, Module.Name, Module.RulesFile, Module.ModuleDirectory, ModuleToType[Module], GeneratedCodeVersion, Module.Rules.bUsePrecompiled, MetadataCache);
 				if (Info.PublicUObjectClassesHeaders.Count > 0 || Info.PrivateUObjectHeaders.Count > 0 || Info.PublicUObjectHeaders.Count > 0)
 				{
 					// Set a flag indicating that we need to add the generated headers directory
diff --git a/Engine/Source/Programs/UnrealBuildTool/System/SourceFileMetadataCache.cs b/Engine/Source/Programs/UnrealBuildTool/System/SourceFileMetadataCache.cs
new file mode 100644
--- /dev/null
+++ b/Engine/Source/Programs/UnrealBuildTool/System/SourceFileMetadataCache.cs
@@ -0,0 +1,367 @@
+using System;
+using System.Collections.Concurrent;
+using System.Collections.Generic;
+using System.Diagnostics;
+using System.IO;
+using System.Linq;
+using System.Text;
+using System.Text.RegularExpressions;
+using System.Threading.Tasks;
+using Tools.DotNETCommon;
+
+namespace UnrealBuildTool
+{
+	/// <summary>
+	/// Caches information about C++ source files; whether they contain reflection markup, what the first included header is, and so on.
+	/// </summary>
+	class SourceFileMetadataCache
+	{
+		/// <summary>
+		/// Information about the first file included from a source file
+		/// </summary>
+		class IncludeInfo
+		{
+			/// <summary>
+			/// Last write time of the file when the data was cached
+			/// </summary>
+			public long LastWriteTimeUtc;
+
+			/// <summary>
+			/// Contents of the include directive
+			/// </summary>
+			public string IncludeText;
+		}
+
+		/// <summary>
+		/// Information about whether a file contains reflection markup
+		/// </summary>
+		class ReflectionInfo
+		{
+			/// <summary>
+			/// Last write time of the file when the data was cached
+			/// </summary>
+			public long LastWriteTimeUtc;
+
+			/// <summary>
+			/// Whether or not the file contains reflection markup
+			/// </summary>
+			public bool bContainsMarkup;
+		}
+
+		/// <summary>
+		/// The current file version
+		/// </summary>
+		public const int CurrentVersion = 2;
+
+		/// <summary>
+		/// Location of this dependency cache
+		/// </summary>
+		FileReference Location;
+
+		/// <summary>
+		/// Directory for files to cache dependencies for.
+		/// </summary>
+		DirectoryReference BaseDirectory;
+
+		/// <summary>
+		/// The parent cache.
+		/// </summary>
+		SourceFileMetadataCache Parent;
+
+		/// <summary>
+		/// Map from file item to source file info
+		/// </summary>
+		ConcurrentDictionary<FileItem, IncludeInfo> FileToIncludeInfo = new ConcurrentDictionary<FileItem, IncludeInfo>();
+
+		/// <summary>
+		/// Map from file item to header file info
+		/// </summary>
+		ConcurrentDictionary<FileItem, ReflectionInfo> FileToReflectionInfo = new ConcurrentDictionary<FileItem, ReflectionInfo>();
+
+		/// <summary>
+		/// Whether the cache has been modified and needs to be saved
+		/// </summary>
+		bool bModified;
+
+		/// <summary>
+		/// Regex that matches C++ code with UObject declarations which we will need to generated code for.
+		/// </summary>
+		static readonly Regex ReflectionMarkupRegex = new Regex("^\\s*U(CLASS|STRUCT|ENUM|INTERFACE|DELEGATE)\\b", RegexOptions.Compiled | RegexOptions.Multiline);
+
+		/// <summary>
+		/// Regex that matches #include statements.
+		/// </summary>
+		static readonly Regex IncludeRegex = new Regex("^[ \t]*#[ \t]*include[ \t]*[<\"](?<HeaderFile>[^\">]*)[\">]", RegexOptions.Compiled | RegexOptions.Singleline | RegexOptions.ExplicitCapture);
+
+		/// <summary>
+		/// Regex that matches #import directives in mm files
+		/// </summary>
+		static readonly Regex ImportRegex = new Regex("^[ \t]*#[ \t]*import[ \t]*[<\"](?<HeaderFile>[^\">]*)[\">]", RegexOptions.Compiled | RegexOptions.Singleline | RegexOptions.ExplicitCapture);
+
+		/// <summary>
+		/// Static cache of all constructed dependency caches
+		/// </summary>
+		static Dictionary<FileReference, SourceFileMetadataCache> Caches = new Dictionary<FileReference, SourceFileMetadataCache>();
+
+		/// <summary>
+		/// Constructs a dependency cache. This method is private; call CppDependencyCache.Create() to create a cache hierarchy for a given project.
+		/// </summary>
+		/// <param name="Location">File to store the cache</param>
+		/// <param name="BaseDir">Base directory for files that this cache should store data for</param>
+		/// <param name="Parent">The parent cache to use</param>
+		private SourceFileMetadataCache(FileReference Location, DirectoryReference BaseDir, SourceFileMetadataCache Parent)
+		{
+			this.Location = Location;
+			this.BaseDirectory = BaseDir;
+			this.Parent = Parent;
+
+			if(FileReference.Exists(Location))
+			{
+				Read();
+			}
+		}
+
+		/// <summary>
+		/// Gets the first included file from a source file
+		/// </summary>
+		/// <param name="SourceFile">The source file to parse</param>
+		/// <returns>Text from the first include directive. Null if the file did not contain any include directives.</returns>
+		public string GetFirstInclude(FileItem SourceFile)
+		{
+			if(Parent != null && !SourceFile.Location.IsUnderDirectory(BaseDirectory))
+			{
+				return Parent.GetFirstInclude(SourceFile);
+			}
+			else
+			{
+				IncludeInfo IncludeInfo;
+				if(!FileToIncludeInfo.TryGetValue(SourceFile, out IncludeInfo) || SourceFile.LastWriteTimeUtc.Ticks > IncludeInfo.LastWriteTimeUtc)
+				{
+					IncludeInfo = new IncludeInfo();
+					IncludeInfo.LastWriteTimeUtc = SourceFile.LastWriteTimeUtc.Ticks;
+					IncludeInfo.IncludeText = ParseFirstInclude(SourceFile.Location);
+					FileToIncludeInfo[SourceFile] = IncludeInfo;
+					bModified = true;
+				}
+				return IncludeInfo.IncludeText;
+			}
+		}
+
+		/// <summary>
+		/// Determines whether the given file contains reflection markup
+		/// </summary>
+		/// <param name="SourceFile">The source file to parse</param>
+		/// <returns>True if the file contains reflection markup</returns>
+		public bool ContainsReflectionMarkup(FileItem SourceFile)
+		{
+			if(Parent != null && !SourceFile.Location.IsUnderDirectory(BaseDirectory))
+			{
+				return Parent.ContainsReflectionMarkup(SourceFile);
+			}
+			else
+			{
+				ReflectionInfo ReflectionInfo;
+				if(!FileToReflectionInfo.TryGetValue(SourceFile, out ReflectionInfo) || SourceFile.LastWriteTimeUtc.Ticks > ReflectionInfo.LastWriteTimeUtc)
+				{
+					ReflectionInfo = new ReflectionInfo();
+					ReflectionInfo.LastWriteTimeUtc = SourceFile.LastWriteTimeUtc.Ticks;
+					ReflectionInfo.bContainsMarkup = ReflectionMarkupRegex.IsMatch(FileReference.ReadAllText(SourceFile.Location));
+					FileToReflectionInfo[SourceFile] = ReflectionInfo;
+					bModified = true;
+				}
+				return ReflectionInfo.bContainsMarkup;
+			}
+		}
+
+		/// <summary>
+		/// Parse the first include directive from a source file
+		/// </summary>
+		/// <param name="SourceFile">The source file to parse</param>
+		/// <returns>The first include directive</returns>
+		static string ParseFirstInclude(FileReference SourceFile)
+		{
+			bool bMatchImport = SourceFile.HasExtension(".m") || SourceFile.HasExtension(".mm");
+			using(StreamReader Reader = new StreamReader(SourceFile.FullName, true))
+			{
+				for(;;)
+				{
+					string Line = Reader.ReadLine();
+					if(Line == null)
+					{
+						return null;
+					}
+
+					Match IncludeMatch = IncludeRegex.Match(Line);
+					if(IncludeMatch.Success)
+					{
+						return IncludeMatch.Groups[1].Value;
+					}
+
+					if(bMatchImport)
+					{
+						Match ImportMatch = ImportRegex.Match(Line);
+						if(ImportMatch.Success)
+						{
+							return IncludeMatch.Groups[1].Value;
+						}
+					}
+				}
+			}
+		}
+
+		/// <summary>
+		/// Creates a cache hierarchy for a particular target
+		/// </summary>
+		/// <param name="ProjectFile">Project file for the target being built</param>
+		/// <returns>Dependency cache hierarchy for the given project</returns>
+		public static SourceFileMetadataCache CreateHierarchy(FileReference ProjectFile)
+		{
+			SourceFileMetadataCache Cache = null;
+
+			if(ProjectFile == null || !UnrealBuildTool.IsEngineInstalled())
+			{
+				FileReference EngineCacheLocation = FileReference.Combine(UnrealBuildTool.EngineDirectory, "Intermediate", "Build", "SourceFileCache.bin");
+				Cache = FindOrAddCache(EngineCacheLocation, UnrealBuildTool.EngineDirectory, Cache);
+			}
+
+			if(ProjectFile != null)
+			{
+				FileReference ProjectCacheLocation = FileReference.Combine(ProjectFile.Directory, "Intermediate", "Build", "SourceFileCache.bin");
+				Cache = FindOrAddCache(ProjectCacheLocation, ProjectFile.Directory, Cache);
+			}
+
+			return Cache;
+		}
+
+		/// <summary>
+		/// Enumerates all the locations of metadata caches for the given target
+		/// </summary>
+		/// <param name="ProjectFile">Project file for the target being built</param>
+		/// <returns>Dependency cache hierarchy for the given project</returns>
+		public static IEnumerable<FileReference> GetFilesToClean(FileReference ProjectFile)
+		{
+			if(ProjectFile == null || !UnrealBuildTool.IsEngineInstalled())
+			{
+				yield return FileReference.Combine(UnrealBuildTool.EngineDirectory, "Intermediate", "Build", "SourceFileCache.bin");
+			}
+			if(ProjectFile != null)
+			{
+				yield return FileReference.Combine(ProjectFile.Directory, "Intermediate", "Build", "SourceFileCache.bin");
+			}
+		}
+
+		/// <summary>
+		/// Reads a cache from the given location, or creates it with the given settings
+		/// </summary>
+		/// <param name="Location">File to store the cache</param>
+		/// <param name="BaseDirectory">Base directory for files that this cache should store data for</param>
+		/// <param name="Parent">The parent cache to use</param>
+		/// <returns>Reference to a dependency cache with the given settings</returns>
+		static SourceFileMetadataCache FindOrAddCache(FileReference Location, DirectoryReference BaseDirectory, SourceFileMetadataCache Parent)
+		{
+			lock(Caches)
+			{
+				SourceFileMetadataCache Cache;
+				if(Caches.TryGetValue(Location, out Cache))
+				{
+					Debug.Assert(Cache.BaseDirectory == BaseDirectory);
+					Debug.Assert(Cache.Parent == Parent);
+				}
+				else
+				{
+					Cache = new SourceFileMetadataCache(Location, BaseDirectory, Parent);
+					Caches.Add(Location, Cache);
+				}
+				return Cache;
+			}
+		}
+
+		/// <summary>
+		/// Save all the caches that have been modified
+		/// </summary>
+		public static void SaveAll()
+		{
+			Parallel.ForEach(Caches.Values, Cache => { if(Cache.bModified){ Cache.Write(); } });
+		}
+
+		/// <summary>
+		/// Reads data for this dependency cache from disk
+		/// </summary>
+		private void Read()
+		{
+			try
+			{
+				using(BinaryArchiveReader Reader = new BinaryArchiveReader(Location))
+				{
+					int Version = Reader.ReadInt();
+					if(Version != CurrentVersion)
+					{
+						Log.TraceLog("Unable to read dependency cache from {0}; version {1} vs current {2}", Location, Version, CurrentVersion);
+						return;
+					}
+
+					int FileToFirstIncludeCount = Reader.ReadInt();
+					for(int Idx = 0; Idx < FileToFirstIncludeCount; Idx++)
+					{
+						FileItem File = Reader.ReadFileItem();
+						
+						IncludeInfo IncludeInfo = new IncludeInfo();
+						IncludeInfo.LastWriteTimeUtc = Reader.ReadLong();
+						IncludeInfo.IncludeText = Reader.ReadString();
+
+						FileToIncludeInfo[File] = IncludeInfo;
+					}
+
+					int FileToMarkupFlagCount = Reader.ReadInt();
+					for(int Idx = 0; Idx < FileToMarkupFlagCount; Idx++)
+					{
+						FileItem File = Reader.ReadFileItem();
+
+						ReflectionInfo ReflectionInfo = new ReflectionInfo();
+						ReflectionInfo.LastWriteTimeUtc = Reader.ReadLong();
+						ReflectionInfo.bContainsMarkup = Reader.ReadBool();
+
+						FileToReflectionInfo[File] = ReflectionInfo;
+					}
+				}
+			}
+			catch(Exception Ex)
+			{
+				Log.TraceWarning("Unable to read {0}. See log for additional information.", Location);
+				Log.TraceLog("{0}", ExceptionUtils.FormatExceptionDetails(Ex));
+			}
+		}
+
+		/// <summary>
+		/// Writes data for this dependency cache to disk
+		/// </summary>
+		private void Write()
+		{
+			DirectoryReference.CreateDirectory(Location.Directory);
+			using(FileStream Stream = File.Open(Location.FullName, FileMode.Create, FileAccess.Write, FileShare.Read))
+			{
+				using(BinaryArchiveWriter Writer = new BinaryArchiveWriter(Stream))
+				{
+					Writer.WriteInt(CurrentVersion);
+
+					Writer.WriteInt(FileToIncludeInfo.Count);
+					foreach(KeyValuePair<FileItem, IncludeInfo> Pair in FileToIncludeInfo)
+					{
+						Writer.WriteFileItem(Pair.Key);
+						Writer.WriteLong(Pair.Value.LastWriteTimeUtc);
+						Writer.WriteString(Pair.Value.IncludeText);
+					}
+
+					Writer.WriteInt(FileToReflectionInfo.Count);
+					foreach(KeyValuePair<FileItem, ReflectionInfo> Pair in FileToReflectionInfo)
+					{
+						Writer.WriteFileItem(Pair.Key);
+						Writer.WriteLong(Pair.Value.LastWriteTimeUtc);
+						Writer.WriteBool(Pair.Value.bContainsMarkup);
+					}
+				}
+			}
+			bModified = false;
+		}
+	}
+}
diff --git a/Engine/Source/Programs/UnrealBuildTool/System/TargetMakefile.cs b/Engine/Source/Programs/UnrealBuildTool/System/TargetMakefile.cs
--- a/Engine/Source/Programs/UnrealBuildTool/System/TargetMakefile.cs
+++ b/Engine/Source/Programs/UnrealBuildTool/System/TargetMakefile.cs
@@ -188,8 +188,8 @@ namespace UnrealBuildTool
 		/// <param name="MakefilePath">Path to the makefile to load</param>
 		/// <param name="ProjectFile">Path to the project file</param>
 		/// <param name="Platform">Platform for this makefile</param>
-		/// <param name="ReasonNotLoaded">If the function returns null, this string will contain the reason why</param>
 		/// <param name="WorkingSet">Interface to query which source files are in the working set</param>
+		/// <param name="ReasonNotLoaded">If the function returns null, this string will contain the reason why</param>
 		/// <returns>The loaded makefile, or null if it failed for some reason.  On failure, the 'ReasonNotLoaded' variable will contain information about why</returns>
 		public static TargetMakefile Load(FileReference MakefilePath, FileReference ProjectFile, UnrealTargetPlatform Platform, ISourceFileWorkingSet WorkingSet, out string ReasonNotLoaded)
 		{
@@ -387,9 +387,10 @@ namespace UnrealBuildTool
 				// Makefile is invalid if:
 				// * There are any newer files which contain no UHT data, but were previously in the makefile
 				// * There are any newer files contain data which needs processing by UHT, but weren't not previously in the makefile
+				SourceFileMetadataCache MetadataCache = SourceFileMetadataCache.CreateHierarchy(ProjectFile);
 				foreach (FileReference Filename in HFilesNewerThanMakefile)
 				{
-					bool bContainsUHTData = CPPHeaders.DoesFileContainUObjects(Filename.FullName);
+					bool bContainsUHTData = MetadataCache.ContainsReflectionMarkup(FileItem.GetItemByFileReference(Filename));
 					bool bWasProcessed = AllUHTHeaders.Contains(Filename);
 					if (bContainsUHTData != bWasProcessed)
 					{
diff --git a/Engine/Source/Programs/UnrealBuildTool/UnrealBuildTool.csproj b/Engine/Source/Programs/UnrealBuildTool/UnrealBuildTool.csproj
--- a/Engine/Source/Programs/UnrealBuildTool/UnrealBuildTool.csproj
+++ b/Engine/Source/Programs/UnrealBuildTool/UnrealBuildTool.csproj
@@ -190,6 +190,7 @@
     <Compile Include="System\CppCompileEnvironment.cs" />
     <Compile Include="System\CppDependencyCache.cs" />
     <Compile Include="System\CppIncludeParser.cs" />
+    <Compile Include="System\SourceFileMetadataCache.cs" />
     <Compile Include="System\DataDrivenPlatformInfo.cs" />
     <Compile Include="System\EncryptionAndSigning.cs" />
     <Compile Include="System\HotReload.cs" />
@@ -342,8 +343,6 @@
     <Compile Include="System\PluginDescriptor.cs" />
     <Compile Include="System\ProjectDescriptor.cs" />
     <Compile Include="System\DirectoryLookupCache.cs" />
-    <Compile Include="System\CPPHeaders.cs" />
-    <Compile Include="System\DependencyCache.cs" />
     <Compile Include="System\DynamicCompilation.cs" />
     <Compile Include="System\ExternalExecution.cs" />
     <Compile Include="System\FileItem.cs" />
