PR #5356: Modified fix to accommodate default build environment being set to shared by the time Intellisense is being gathered.

#rnx

diff --git a/Engine/Source/Programs/UnrealBuildTool/ProjectFiles/ProjectFileGenerator.cs b/Engine/Source/Programs/UnrealBuildTool/ProjectFiles/ProjectFileGenerator.cs
--- a/Engine/Source/Programs/UnrealBuildTool/ProjectFiles/ProjectFileGenerator.cs
+++ b/Engine/Source/Programs/UnrealBuildTool/ProjectFiles/ProjectFileGenerator.cs
@@ -1693,7 +1693,7 @@ namespace UnrealBuildTool
 						{
 							NewArguments.Add("-precompile");
 						}
-						if (CurTarget.TargetRules.Type == TargetType.Editor && CurTarget.TargetRules.BuildEnvironment == TargetBuildEnvironment.Default)
+						if (CurTarget.TargetRules.Type == TargetType.Editor && CurTarget.TargetRules.BuildEnvironment != TargetBuildEnvironment.Unique)
 						{
 							NewArguments.Add("-SharedBuildEnvironment"); // Force the shared build environment, even though we have -Monolithic
 						}
