Add missing copyright notice.

#rnx

diff --git a/Engine/Source/Programs/UnrealBuildTool/System/SourceFileMetadataCache.cs b/Engine/Source/Programs/UnrealBuildTool/System/SourceFileMetadataCache.cs
--- a/Engine/Source/Programs/UnrealBuildTool/System/SourceFileMetadataCache.cs
+++ b/Engine/Source/Programs/UnrealBuildTool/System/SourceFileMetadataCache.cs
@@ -1,3 +1,5 @@
+// Copyright 1998-2019 Epic Games, Inc. All Rights Reserved.
+
 using System;
 using System.Collections.Concurrent;
 using System.Collections.Generic;
