UBT: Add a DirectoryItem class, similar to FileItem, which caches directory contents and properties. Allows constructing FileItem objects from enumerated FileInfo objects, avoiding additional I/O hit to determine attributes. DirectoryLookupCache is converted to use DirectoryItem objects as part of this change.

#rnx

diff --git a/Engine/Source/Programs/UnrealBuildTool/Configuration/UEBuildModuleCPP.cs b/Engine/Source/Programs/UnrealBuildTool/Configuration/UEBuildModuleCPP.cs
--- a/Engine/Source/Programs/UnrealBuildTool/Configuration/UEBuildModuleCPP.cs
+++ b/Engine/Source/Programs/UnrealBuildTool/Configuration/UEBuildModuleCPP.cs
@@ -245,14 +245,8 @@ namespace UnrealBuildTool
 			{
 				PublicIncludePaths.Add(PublicDirectory);
 
-				string[] ExcludedFolderNames = UEBuildPlatform.GetBuildPlatform(Rules.Target.Platform).GetExcludedFolderNames();
-				foreach (DirectoryReference PublicSubDirectory in DirectoryLookupCache.EnumerateDirectoriesRecursively(PublicDirectory))
-				{
-					if(!PublicSubDirectory.ContainsAnyNames(ExcludedFolderNames, PublicDirectory))
-					{
-						LegacyPublicIncludePaths.Add(PublicSubDirectory);
-					}
-				}
+				HashSet<string> ExcludeNames = new HashSet<string>(UEBuildPlatform.GetBuildPlatform(Rules.Target.Platform).GetExcludedFolderNames(), StringComparer.OrdinalIgnoreCase);
+				EnumerateLegacyIncludePaths(DirectoryItem.GetItemByDirectoryReference(PublicDirectory), ExcludeNames, LegacyPublicIncludePaths);
 			}
 
 			// Add the base private directory for this module
@@ -263,6 +257,24 @@ namespace UnrealBuildTool
 			}
 		}
 
+		/// <summary>
+		/// Enumerates legacy include paths under a given base directory
+		/// </summary>
+		/// <param name="BaseDirectory">The directory to start from. This directory is not added to the output list.</param>
+		/// <param name="ExcludeNames">Set of folder names to exclude from the search.</param>
+		/// <param name="LegacyPublicIncludePaths">List populated with the discovered directories</param>
+		static void EnumerateLegacyIncludePaths(DirectoryItem BaseDirectory, HashSet<string> ExcludeNames, HashSet<DirectoryReference> LegacyPublicIncludePaths)
+		{
+			foreach(DirectoryItem SubDirectory in BaseDirectory.EnumerateDirectories())
+			{
+				if(!ExcludeNames.Contains(SubDirectory.Name))
+				{
+					LegacyPublicIncludePaths.Add(SubDirectory.Location);
+					EnumerateLegacyIncludePaths(SubDirectory, ExcludeNames, LegacyPublicIncludePaths);
+				}
+			}
+		}
+
 		/// <summary>
 		/// Path to the precompiled manifest location
 		/// </summary>
diff --git a/Engine/Source/Programs/UnrealBuildTool/System/ActionGraph.cs b/Engine/Source/Programs/UnrealBuildTool/System/ActionGraph.cs
--- a/Engine/Source/Programs/UnrealBuildTool/System/ActionGraph.cs
+++ b/Engine/Source/Programs/UnrealBuildTool/System/ActionGraph.cs
@@ -447,7 +447,7 @@ namespace UnrealBuildTool
 							bool bPrerequisiteItemIsNewerThanLastExecution = TimeDifference.TotalSeconds > 1;
 							if (bPrerequisiteItemIsNewerThanLastExecution)
 							{
-								Log.TraceLog("{0}: Prerequisite {1} is newer than the last execution of the action: {2} vs {3}", RootAction.StatusDescription, Path.GetFileName(PrerequisiteItem.AbsolutePath), PrerequisiteItem.LastWriteTimeUtc.LocalDateTime, LastExecutionTimeUtc.LocalDateTime);
+								Log.TraceLog("{0}: Prerequisite {1} is newer than the last execution of the action: {2} vs {3}", RootAction.StatusDescription, Path.GetFileName(PrerequisiteItem.AbsolutePath), PrerequisiteItem.LastWriteTimeUtc.ToLocalTime(), LastExecutionTimeUtc.LocalDateTime);
 								bIsOutdated = true;
 								break;
 							}
@@ -474,7 +474,7 @@ namespace UnrealBuildTool
 									"{0}: Dependency {1} is newer than the last execution of the action: {2} vs {3}",
 									RootAction.StatusDescription,
 									Path.GetFileName(DependencyFile.AbsolutePath),
-									DependencyFile.LastWriteTimeUtc.LocalDateTime,
+									DependencyFile.LastWriteTimeUtc.ToLocalTime(),
 									LastExecutionTimeUtc.LocalDateTime
 									);
 								bIsOutdated = true;
diff --git a/Engine/Source/Programs/UnrealBuildTool/System/DirectoryItem.cs b/Engine/Source/Programs/UnrealBuildTool/System/DirectoryItem.cs
new file mode 100644
--- /dev/null
+++ b/Engine/Source/Programs/UnrealBuildTool/System/DirectoryItem.cs
@@ -0,0 +1,223 @@
+using System;
+using System.Collections.Concurrent;
+using System.Collections.Generic;
+using System.IO;
+using System.Linq;
+using System.Text;
+using System.Threading.Tasks;
+using Tools.DotNETCommon;
+
+namespace UnrealBuildTool
+{
+	/// <summary>
+	/// Stores the state of a directory. May or may not exist.
+	/// </summary>
+	class DirectoryItem
+	{
+		/// <summary>
+		/// Full path to the directory on disk
+		/// </summary>
+		public readonly DirectoryReference Location;
+
+		/// <summary>
+		/// Cached value for whether the directory exists
+		/// </summary>
+		DirectoryInfo Info;
+
+		/// <summary>
+		/// Cached map of name to subdirectory item
+		/// </summary>
+		List<DirectoryItem> Directories;
+
+		/// <summary>
+		/// Cached map of name to file
+		/// </summary>
+		List<FileItem> Files;
+
+		/// <summary>
+		/// Global map of location to item
+		/// </summary>
+		static ConcurrentDictionary<DirectoryReference, DirectoryItem> LocationToItem = new ConcurrentDictionary<DirectoryReference, DirectoryItem>();
+
+		/// <summary>
+		/// Constructor
+		/// </summary>
+		/// <param name="Location">Path to this directory</param>
+		/// <param name="Info">Information about this directory</param>
+		private DirectoryItem(DirectoryReference Location, DirectoryInfo Info)
+		{
+			this.Location = Location;
+			this.Info = Info;
+		}
+
+		/// <summary>
+		/// The name of this directory
+		/// </summary>
+		public string Name
+		{
+			get { return Info.Name; }
+		}
+
+		/// <summary>
+		/// Whether the directory exists or not
+		/// </summary>
+		public bool Exists
+		{
+			get { return Info.Exists; }
+		}
+
+		/// <summary>
+		/// The last write time of the file.
+		/// </summary>
+		public DateTime LastWriteTimeUtc
+		{
+			get { return Info.LastWriteTimeUtc; }
+		}
+
+		/// <summary>
+		/// Gets a new directory item by combining the existing directory item with the given path fragments
+		/// </summary>
+		/// <param name="BaseDirectory">Base directory to append path fragments to</param>
+		/// <param name="Fragments">The path fragments to append</param>
+		/// <returns>Directory item corresponding to the combined path</returns>
+		public static DirectoryItem Combine(DirectoryItem BaseDirectory, params string[] Fragments)
+		{
+			return DirectoryItem.GetItemByDirectoryReference(DirectoryReference.Combine(BaseDirectory.Location, Fragments));
+		}
+
+		/// <summary>
+		/// Finds or creates a directory item from its location
+		/// </summary>
+		/// <param name="Location">Path to the directory</param>
+		/// <returns>The directory item for this location</returns>
+		public static DirectoryItem GetItemByDirectoryReference(DirectoryReference Location)
+		{
+			DirectoryItem Result;
+			if(!LocationToItem.TryGetValue(Location, out Result))
+			{
+				DirectoryItem NewItem = new DirectoryItem(Location, new DirectoryInfo(Location.FullName));
+				if(LocationToItem.TryAdd(Location, NewItem))
+				{
+					Result = NewItem;
+				}
+				else
+				{
+					Result = LocationToItem[Location];
+				}
+			}
+			return Result;
+		}
+
+		/// <summary>
+		/// Finds or creates a directory item from a DirectoryInfo object
+		/// </summary>
+		/// <param name="Info">Path to the directory</param>
+		/// <returns>The directory item for this location</returns>
+		public static DirectoryItem GetItemByDirectoryInfo(DirectoryInfo Info)
+		{
+			DirectoryReference Location = new DirectoryReference(Info);
+
+			DirectoryItem Result;
+			if(!LocationToItem.TryGetValue(Location, out Result))
+			{
+				DirectoryItem NewItem = new DirectoryItem(Location, Info);
+				if(LocationToItem.TryAdd(Location, NewItem))
+				{
+					Result = NewItem;
+				}
+				else
+				{
+					Result = LocationToItem[Location];
+				}
+			}
+			return Result;
+		}
+
+		/// <summary>
+		/// Reset the contents of the directory and allow them to be fetched again
+		/// </summary>
+		public void Refresh()
+		{
+			Directories = null;
+			Files = null;
+		}
+
+		/// <summary>
+		/// Caches the subdirectories of this directories
+		/// </summary>
+		public void CacheDirectories()
+		{
+			if(Directories == null)
+			{
+				List<DirectoryItem> NewDirectories = new List<DirectoryItem>();
+				if(Info.Exists)
+				{
+					foreach(DirectoryInfo SubDirectoryInfo in Info.EnumerateDirectories())
+					{
+						if(SubDirectoryInfo.Name.Length == 1 && SubDirectoryInfo.Name[0] == '.')
+						{
+							continue;
+						}
+						else if(SubDirectoryInfo.Name.Length == 2 && SubDirectoryInfo.Name[0] == '.' && SubDirectoryInfo.Name[1] == '.')
+						{
+							continue;
+						}
+						else
+						{
+							NewDirectories.Add(DirectoryItem.GetItemByDirectoryInfo(SubDirectoryInfo));
+						}
+					}
+				}
+				Directories = NewDirectories;
+			}
+		}
+
+		/// <summary>
+		/// Enumerates all the subdirectories
+		/// </summary>
+		/// <returns>Sequence of subdirectory items</returns>
+		public IEnumerable<DirectoryItem> EnumerateDirectories()
+		{
+			CacheDirectories();
+			return Directories;
+		}
+
+		/// <summary>
+		/// Caches the files in this directory
+		/// </summary>
+		public void CacheFiles()
+		{
+			if(Files == null)
+			{
+				List<FileItem> NewFiles = new List<FileItem>();
+				if(Info.Exists)
+				{
+					foreach(FileInfo FileInfo in Info.EnumerateFiles())
+					{
+						NewFiles.Add(FileItem.GetItemByFileInfo(FileInfo));
+					}
+				}
+				Files = NewFiles;
+			}
+		}
+
+		/// <summary>
+		/// Enumerates all the files
+		/// </summary>
+		/// <returns>Sequence of FileItems</returns>
+		public IEnumerable<FileItem> EnumerateFiles()
+		{
+			CacheFiles();
+			return Files;
+		}
+
+		/// <summary>
+		/// Formats this object as a string for debugging
+		/// </summary>
+		/// <returns>Location of the directory</returns>
+		public override string ToString()
+		{
+			return Location.FullName;
+		}
+	}
+}
diff --git a/Engine/Source/Programs/UnrealBuildTool/System/DirectoryLookupCache.cs b/Engine/Source/Programs/UnrealBuildTool/System/DirectoryLookupCache.cs
--- a/Engine/Source/Programs/UnrealBuildTool/System/DirectoryLookupCache.cs
+++ b/Engine/Source/Programs/UnrealBuildTool/System/DirectoryLookupCache.cs
@@ -1,123 +1,38 @@
-﻿// Copyright 1998-2019 Epic Games, Inc. All Rights Reserved.
+// Copyright 1998-2019 Epic Games, Inc. All Rights Reserved.
 
 using System;
 using System.IO;
 using System.Collections.Generic;
 using Tools.DotNETCommon;
+using System.Linq;
 
 namespace UnrealBuildTool
 {
-	class DirectoryCache
-	{
-		public readonly DirectoryReference Location;
-		public readonly bool Exists;
-		HashSet<FileReference> Files;
-		HashSet<DirectoryReference> Directories;
-
-		public DirectoryCache(DirectoryReference InDirectory)
-		{
-			Location = InDirectory;
-			Exists = DirectoryReference.Exists(Location);
-
-			if (!Exists)
-			{
-				Files = new HashSet<FileReference>();
-				Directories = new HashSet<DirectoryReference>();
-			}
-		}
-
-		private void CacheFiles()
-		{
-			if (Files == null)
-			{
-				Files = new HashSet<FileReference>(DirectoryReference.EnumerateFiles(Location));
-			}
-		}
-
-		public bool HasFile(FileReference CheckFile)
-		{
-			CacheFiles();
-			return Files.Contains(CheckFile);
-		}
-
-		public IEnumerable<FileReference> EnumerateFiles()
-		{
-			CacheFiles();
-			return Files;
-		}
-
-		private void CacheDirectories()
-		{
-			if (Directories == null)
-			{
-				Directories = new HashSet<DirectoryReference>(DirectoryReference.EnumerateDirectories(Location));
-			}
-		}
-
-		public IEnumerable<DirectoryReference> EnumerateDirectories()
-		{
-			CacheDirectories();
-			return Directories;
-		}
-	}
-
 	static class DirectoryLookupCache
 	{
-		static DirectoryCache FindOrCreateDirectoryCache(DirectoryReference Directory)
-		{
-			DirectoryCache FoundDirectoryCache;
-			if (!Directories.TryGetValue(Directory, out FoundDirectoryCache))
-			{
-				FoundDirectoryCache = new DirectoryCache(Directory);
-				Directories.Add(Directory, FoundDirectoryCache);
-			}
-			return FoundDirectoryCache;
-		}
-
-
 		static public bool FileExists(FileReference File)
 		{
-			DirectoryCache FoundDirectoryCache = FindOrCreateDirectoryCache(File.Directory);
-			return FoundDirectoryCache.HasFile(File);
+			return FileItem.GetItemByFileReference(File).Exists;
 		}
 
 		static public bool DirectoryExists(DirectoryReference Directory)
 		{
-			DirectoryCache FoundDirectoryCache = FindOrCreateDirectoryCache(Directory);
-			return FoundDirectoryCache.Exists;
+			return DirectoryItem.GetItemByDirectoryReference(Directory).Exists;
 		}
 
 		static public IEnumerable<FileReference> EnumerateFiles(DirectoryReference Directory)
 		{
-			DirectoryCache FoundDirectoryCache = FindOrCreateDirectoryCache(Directory);
-			return FoundDirectoryCache.EnumerateFiles();
+			return DirectoryItem.GetItemByDirectoryReference(Directory).EnumerateFiles().Select(x => x.Location);
 		}
 
 		static public IEnumerable<DirectoryReference> EnumerateDirectories(DirectoryReference Directory)
 		{
-			DirectoryCache FoundDirectoryCache = FindOrCreateDirectoryCache(Directory);
-			return FoundDirectoryCache.EnumerateDirectories();
-		}
-
-		static public IEnumerable<DirectoryReference> EnumerateDirectoriesRecursively(DirectoryReference Directory)
-		{
-			DirectoryCache FoundDirectoryCache = FindOrCreateDirectoryCache(Directory);
-			foreach (DirectoryReference SubDirectory in FoundDirectoryCache.EnumerateDirectories())
-			{
-				yield return SubDirectory;
-
-				foreach (DirectoryReference ChildSubDirectory in EnumerateDirectoriesRecursively(SubDirectory))
-				{
-					yield return ChildSubDirectory;
-				}
-			}
+			return DirectoryItem.GetItemByDirectoryReference(Directory).EnumerateDirectories().Select(x => x.Location);
 		}
 
         static public void InvalidateCachedDirectory(DirectoryReference Directory)
         {
-            Directories.Remove(Directory);
+			DirectoryItem.GetItemByDirectoryReference(Directory).Refresh();
         }
-
-		static Dictionary<DirectoryReference, DirectoryCache> Directories = new Dictionary<DirectoryReference, DirectoryCache>();
 	}
 }
diff --git a/Engine/Source/Programs/UnrealBuildTool/System/ExternalExecution.cs b/Engine/Source/Programs/UnrealBuildTool/System/ExternalExecution.cs
--- a/Engine/Source/Programs/UnrealBuildTool/System/ExternalExecution.cs
+++ b/Engine/Source/Programs/UnrealBuildTool/System/ExternalExecution.cs
@@ -355,7 +355,7 @@ namespace UnrealBuildTool
 		/// Generates a UHTModuleInfo for a particular named module under a directory.
 		/// </summary>
 		/// <returns></returns>
-		public static UHTModuleInfo CreateUHTModuleInfo(IEnumerable<FileReference> HeaderFilenames, string ModuleName, FileReference ModuleRulesFile, DirectoryReference ModuleDirectory, UHTModuleType ModuleType, EGeneratedCodeVersion GeneratedCodeVersion, bool bIsReadOnly, SourceFileMetadataCache MetadataCache)
+		public static UHTModuleInfo CreateUHTModuleInfo(List<FileItem> HeaderFiles, string ModuleName, FileReference ModuleRulesFile, DirectoryReference ModuleDirectory, UHTModuleType ModuleType, EGeneratedCodeVersion GeneratedCodeVersion, bool bIsReadOnly, SourceFileMetadataCache MetadataCache)
 		{
 			DirectoryReference ClassesFolder = DirectoryReference.Combine(ModuleDirectory, "Classes");
 			DirectoryReference PublicFolder = DirectoryReference.Combine(ModuleDirectory, "Public");
@@ -364,25 +364,23 @@ namespace UnrealBuildTool
 			List<FileItem> PublicUObjectHeaders = new List<FileItem>();
 			List<FileItem> PrivateUObjectHeaders = new List<FileItem>();
 
-			foreach (FileReference Header in HeaderFilenames)
+			foreach (FileItem HeaderFile in HeaderFiles)
 			{
 				// Check to see if we know anything about this file.  If we have up-to-date cached information about whether it has
 				// UObjects or not, we can skip doing a test here.
-				FileItem UObjectHeaderFileItem = FileItem.GetItemByFileReference(Header);
-
-				if (MetadataCache.ContainsReflectionMarkup(UObjectHeaderFileItem))
+				if (MetadataCache.ContainsReflectionMarkup(HeaderFile))
 				{
-					if (UObjectHeaderFileItem.Location.IsUnderDirectory(ClassesFolder))
+					if (HeaderFile.Location.IsUnderDirectory(ClassesFolder))
 					{
-						PublicClassesUObjectHeaders.Add(UObjectHeaderFileItem);
+						PublicClassesUObjectHeaders.Add(HeaderFile);
 					}
-					else if (UObjectHeaderFileItem.Location.IsUnderDirectory(PublicFolder))
+					else if (HeaderFile.Location.IsUnderDirectory(PublicFolder))
 					{
-						PublicUObjectHeaders.Add(UObjectHeaderFileItem);
+						PublicUObjectHeaders.Add(HeaderFile);
 					}
 					else
 					{
-						PrivateUObjectHeaders.Add(UObjectHeaderFileItem);
+						PrivateUObjectHeaders.Add(HeaderFile);
 					}
 				}
 			}
@@ -591,20 +589,23 @@ namespace UnrealBuildTool
 		/// <summary>
 		/// Find all the headers under the given base directory, excluding any other platform folders.
 		/// </summary>
-		/// <param name="BaseDir">Base directory to search</param>
+		/// <param name="BaseDirectory">Base directory to search</param>
 		/// <param name="ExcludeFolders">Array of folders to exclude</param>
 		/// <param name="Headers">Receives the list of headers that was found</param>
-		static void FindHeaders(DirectoryInfo BaseDir, string[] ExcludeFolders, List<FileReference> Headers)
+		static void FindHeaders(DirectoryItem BaseDirectory, HashSet<string> ExcludeFolders, List<FileItem> Headers)
 		{
-			if (!ExcludeFolders.Any(x => x.Equals(BaseDir.Name, StringComparison.InvariantCultureIgnoreCase)))
+			foreach(DirectoryItem SubDirectory in BaseDirectory.EnumerateDirectories())
 			{
-				foreach (DirectoryInfo SubDir in BaseDir.EnumerateDirectories())
+				if(!ExcludeFolders.Contains(SubDirectory.Name))
 				{
-					FindHeaders(SubDir, ExcludeFolders, Headers);
+					FindHeaders(SubDirectory, ExcludeFolders, Headers);
 				}
-				foreach (FileInfo File in BaseDir.EnumerateFiles("*.h"))
+			}
+			foreach(FileItem File in BaseDirectory.EnumerateFiles())
+			{
+				if(File.HasExtension(".h"))
 				{
-					Headers.Add(new FileReference(File));
+					Headers.Add(File);
 				}
 			}
 		}
@@ -622,11 +623,13 @@ namespace UnrealBuildTool
 			List<UEBuildModuleCPP> ModulesSortedByType = ModulesToGenerateHeadersFor.OrderBy(c => ModuleToType[c]).ToList();
 			StableTopologicalSort(ModulesSortedByType);
 
-			string[] ExcludedFolders = UEBuildPlatform.GetBuildPlatform(Platform, true).GetExcludedFolderNames();
+			HashSet<string> ExcludedFolders = new HashSet<string>(UEBuildPlatform.GetBuildPlatform(Platform, true).GetExcludedFolderNames(), StringComparer.OrdinalIgnoreCase);
 			foreach (UEBuildModuleCPP Module in ModulesSortedByType)
 			{
-				List<FileReference> HeaderFiles = new List<FileReference>();
-				FindHeaders(new DirectoryInfo(Module.ModuleDirectory.FullName), ExcludedFolders, HeaderFiles);
+				DirectoryItem ModuleDirectory = DirectoryItem.GetItemByDirectoryReference(Module.ModuleDirectory);
+
+				List<FileItem> HeaderFiles = new List<FileItem>();
+				FindHeaders(ModuleDirectory, ExcludedFolders, HeaderFiles);
 
 				UHTModuleInfo Info = ExternalExecution.CreateUHTModuleInfo(HeaderFiles, Module.Name, Module.RulesFile, Module.ModuleDirectory, ModuleToType[Module], GeneratedCodeVersion, Module.Rules.bUsePrecompiled, MetadataCache);
 				if (Info.PublicUObjectClassesHeaders.Count > 0 || Info.PrivateUObjectHeaders.Count > 0 || Info.PublicUObjectHeaders.Count > 0)
@@ -935,7 +938,7 @@ namespace UnrealBuildTool
 
 				foreach (FileItem HeaderFile in AllUObjectHeaders)
 				{
-					DateTime HeaderFileTimestamp = HeaderFile.Info.LastWriteTime;
+					DateTime HeaderFileTimestamp = HeaderFile.LastWriteTimeUtc;
 
 					// Has the source header changed since we last generated headers successfully?
 					if (HeaderFileTimestamp > SavedTimestamp)
diff --git a/Engine/Source/Programs/UnrealBuildTool/System/FileItem.cs b/Engine/Source/Programs/UnrealBuildTool/System/FileItem.cs
--- a/Engine/Source/Programs/UnrealBuildTool/System/FileItem.cs
+++ b/Engine/Source/Programs/UnrealBuildTool/System/FileItem.cs
@@ -26,7 +26,7 @@ namespace UnrealBuildTool
 		/// <summary>
 		/// The information about the file.
 		/// </summary>
-		public FileInfo Info;
+		FileInfo Info;
 
 		/// <summary>
 		/// A case-insensitive dictionary that's used to map each unique file name to a single FileItem object.
@@ -44,6 +44,14 @@ namespace UnrealBuildTool
 			this.Info = Info;
 		}
 
+		/// <summary>
+		/// Name of this file
+		/// </summary>
+		public string Name
+		{
+			get { return Info.Name; }
+		}
+
 		/// <summary>
 		/// Accessor for the absolute path to the file
 		/// </summary>
@@ -71,11 +79,21 @@ namespace UnrealBuildTool
 		/// <summary>
 		/// The last write time of the file.
 		/// </summary>
-		public DateTimeOffset LastWriteTimeUtc
+		public DateTime LastWriteTimeUtc
 		{
 			get { return Info.LastWriteTimeUtc; }
 		}
 
+		/// <summary>
+		/// Determines if the file has the given extension
+		/// </summary>
+		/// <param name="Extension">The extension to check for</param>
+		/// <returns>True if the file has the given extension, false otherwise</returns>
+		public bool HasExtension(string Extension)
+		{
+			return Location.HasExtension(Extension);
+		}
+
 		/// <summary>
 		/// Gets a FileItem corresponding to the given path
 		/// </summary>
@@ -86,6 +104,31 @@ namespace UnrealBuildTool
 			return GetItemByFileReference(new FileReference(FilePath));
 		}
 
+		/// <summary>
+		/// Gets a FileItem for a given path
+		/// </summary>
+		/// <param name="Info">Information about the file</param>
+		/// <returns>The FileItem that represents the given a full file path.</returns>
+		public static FileItem GetItemByFileInfo(FileInfo Info)
+		{
+			FileReference Location = new FileReference(Info);
+
+			FileItem Result;
+			if (!UniqueSourceFileMap.TryGetValue(Location, out Result))
+			{
+				FileItem NewFileItem = new FileItem(Location, Info);
+				if(UniqueSourceFileMap.TryAdd(Location, NewFileItem))
+				{
+					Result = NewFileItem;
+				}
+				else
+				{
+					Result = UniqueSourceFileMap[Location];
+				}
+			}
+			return Result;
+		}
+
 		/// <summary>
 		/// Gets a FileItem for a given path
 		/// </summary>
diff --git a/Engine/Source/Programs/UnrealBuildTool/System/Plugins.cs b/Engine/Source/Programs/UnrealBuildTool/System/Plugins.cs
--- a/Engine/Source/Programs/UnrealBuildTool/System/Plugins.cs
+++ b/Engine/Source/Programs/UnrealBuildTool/System/Plugins.cs
@@ -326,10 +326,13 @@ namespace UnrealBuildTool
 			if (!PluginFileCache.TryGetValue(ParentDirectory, out FileNames))
 			{
 				FileNames = new List<FileReference>();
-				if (DirectoryReference.Exists(ParentDirectory))
+
+				DirectoryItem ParentDirectoryItem = DirectoryItem.GetItemByDirectoryReference(ParentDirectory);
+				if (ParentDirectoryItem.Exists)
 				{
-					EnumeratePluginsInternal(ParentDirectory, FileNames);
+					EnumeratePluginsInternal(ParentDirectoryItem, FileNames);
 				}
+
 				PluginFileCache.Add(ParentDirectory, FileNames);
 			}
 			return FileNames;
@@ -340,16 +343,21 @@ namespace UnrealBuildTool
 		/// </summary>
 		/// <param name="ParentDirectory">Parent directory to look in. Plugins will be found in any *subfolders* of this directory.</param>
 		/// <param name="FileNames">List of filenames. Will have all the discovered .uplugin files appended to it.</param>
-		static void EnumeratePluginsInternal(DirectoryReference ParentDirectory, List<FileReference> FileNames)
+		static void EnumeratePluginsInternal(DirectoryItem ParentDirectory, List<FileReference> FileNames)
 		{
-			foreach (DirectoryReference ChildDirectory in DirectoryReference.EnumerateDirectories(ParentDirectory))
+			foreach (DirectoryItem ChildDirectory in ParentDirectory.EnumerateDirectories())
 			{
-				int InitialFileNamesCount = FileNames.Count;
-				foreach (FileReference PluginFile in DirectoryReference.EnumerateFiles(ChildDirectory, "*.uplugin"))
+				bool bSearchSubDirectories = true;
+				foreach (FileItem PluginFile in ChildDirectory.EnumerateFiles())
 				{
-					FileNames.Add(PluginFile);
+					if(PluginFile.HasExtension(".uplugin"))
+					{
+						FileNames.Add(PluginFile.Location);
+						bSearchSubDirectories = false;
+					}
 				}
-				if (FileNames.Count == InitialFileNamesCount)
+
+				if (bSearchSubDirectories)
 				{
 					EnumeratePluginsInternal(ChildDirectory, FileNames);
 				}
diff --git a/Engine/Source/Programs/UnrealBuildTool/System/Unity.cs b/Engine/Source/Programs/UnrealBuildTool/System/Unity.cs
--- a/Engine/Source/Programs/UnrealBuildTool/System/Unity.cs
+++ b/Engine/Source/Programs/UnrealBuildTool/System/Unity.cs
@@ -45,7 +45,7 @@ namespace UnrealBuildTool
 			{
 				Files.Add(File);
 
-				long FileLength = File.Info.Length;
+				long FileLength = File.Length;
 				TotalLength += FileLength;
 				VirtualLength += FileLength;
 			}
@@ -105,7 +105,7 @@ namespace UnrealBuildTool
 			/// <param name="File">The file to add virtually.  Only the size of the file is tracked.</param>
 			public void AddVirtualFile(FileItem File)
 			{
-				CurrentCollection.AddVirtualFile(File.Info.Length);
+				CurrentCollection.AddVirtualFile(File.Length);
 				if (SplitLength != -1 && CurrentCollection.VirtualLength > SplitLength)
 				{
 					EndCurrentUnityFile();
@@ -171,7 +171,7 @@ namespace UnrealBuildTool
 			UEBuildPlatform BuildPlatform = UEBuildPlatform.GetBuildPlatformForCPPTargetPlatform(CompileEnvironment.Platform);
 
 			// Figure out size of all input files combined. We use this to determine whether to use larger unity threshold or not.
-			long TotalBytesInCPPFiles = CPPFiles.Sum(F => F.Info.Length);
+			long TotalBytesInCPPFiles = CPPFiles.Sum(F => F.Length);
 
 			// We have an increased threshold for unity file size if, and only if, all files fit into the same unity file. This
 			// is beneficial when dealing with PCH files. The default PCH creation limit is X unity files so if we generate < X 
diff --git a/Engine/Source/Programs/UnrealBuildTool/UnrealBuildTool.csproj b/Engine/Source/Programs/UnrealBuildTool/UnrealBuildTool.csproj
--- a/Engine/Source/Programs/UnrealBuildTool/UnrealBuildTool.csproj
+++ b/Engine/Source/Programs/UnrealBuildTool/UnrealBuildTool.csproj
@@ -192,6 +192,7 @@
     <Compile Include="System\CppIncludeParser.cs" />
     <Compile Include="System\SourceFileMetadataCache.cs" />
     <Compile Include="System\DataDrivenPlatformInfo.cs" />
+    <Compile Include="System\DirectoryItem.cs" />
     <Compile Include="System\EncryptionAndSigning.cs" />
     <Compile Include="System\HotReload.cs" />
     <Compile Include="System\NativeProjects.cs" />
