UBT: Improve performance of source file discovery within UBT (removing duplicate file enumeration, add caching, filter files at the point of discovery).

#rnx

diff --git a/Engine/Source/Programs/UnrealBuildTool/Configuration/UEBuildModuleCPP.cs b/Engine/Source/Programs/UnrealBuildTool/Configuration/UEBuildModuleCPP.cs
--- a/Engine/Source/Programs/UnrealBuildTool/Configuration/UEBuildModuleCPP.cs
+++ b/Engine/Source/Programs/UnrealBuildTool/Configuration/UEBuildModuleCPP.cs
@@ -16,48 +16,21 @@ namespace UnrealBuildTool
 	/// </summary>
 	class UEBuildModuleCPP : UEBuildModule
 	{
-		public class SourceFilesClass
+		/// <summary>
+		/// Stores a list of all source files, of different types
+		/// </summary>
+		class SourceFileCollection
 		{
+			public readonly List<FileItem> HFiles = new List<FileItem>();
 			public readonly List<FileItem> CPPFiles = new List<FileItem>();
 			public readonly List<FileItem> CFiles = new List<FileItem>();
 			public readonly List<FileItem> CCFiles = new List<FileItem>();
 			public readonly List<FileItem> MMFiles = new List<FileItem>();
 			public readonly List<FileItem> RCFiles = new List<FileItem>();
 
-			public int Count
-			{
-				get
-				{
-					return CPPFiles.Count +
-						   CFiles.Count +
-						   CCFiles.Count +
-						   MMFiles.Count +
-						   RCFiles.Count;
-				}
-			}
-
-			/// <summary>
-			/// Copy from list to list helper.
-			/// </summary>
-			/// <param name="From">Source list.</param>
-			/// <param name="To">Destination list.</param>
-			private static void CopyFromListToList(List<FileItem> From, List<FileItem> To)
+			public IEnumerable<FileItem> AllFiles
 			{
-				To.Clear();
-				To.AddRange(From);
-			}
-
-			/// <summary>
-			/// Copies file lists from other SourceFilesClass to this.
-			/// </summary>
-			/// <param name="Other">Source object.</param>
-			public void CopyFrom(SourceFilesClass Other)
-			{
-				CopyFromListToList(Other.CPPFiles, CPPFiles);
-				CopyFromListToList(Other.CFiles, CFiles);
-				CopyFromListToList(Other.CCFiles, CCFiles);
-				CopyFromListToList(Other.MMFiles, MMFiles);
-				CopyFromListToList(Other.RCFiles, RCFiles);
+				get { return CPPFiles.Concat(CFiles).Concat(CCFiles).Concat(MMFiles).Concat(RCFiles); }
 			}
 		}
 
@@ -103,37 +76,6 @@ namespace UnrealBuildTool
 			Directories.UnionWith(SourceDirectories);
 		}
 
-		/// <summary>
-		/// Categorizes source files into per-extension buckets
-		/// </summary>
-		private static void CategorizeSourceFiles(IEnumerable<FileItem> InSourceFiles, SourceFilesClass OutSourceFiles)
-		{
-			foreach (FileItem SourceFile in InSourceFiles)
-			{
-				string Extension = Path.GetExtension(SourceFile.AbsolutePath).ToUpperInvariant();
-				if (Extension == ".CPP")
-				{
-					OutSourceFiles.CPPFiles.Add(SourceFile);
-				}
-				else if (Extension == ".C")
-				{
-					OutSourceFiles.CFiles.Add(SourceFile);
-				}
-				else if (Extension == ".CC")
-				{
-					OutSourceFiles.CCFiles.Add(SourceFile);
-				}
-				else if (Extension == ".MM" || Extension == ".M")
-				{
-					OutSourceFiles.MMFiles.Add(SourceFile);
-				}
-				else if (Extension == ".RC")
-				{
-					OutSourceFiles.RCFiles.Add(SourceFile);
-				}
-			}
-		}
-
 		/// <summary>
 		/// List of whitelisted circular dependencies. Please do NOT add new modules here; refactor to allow the modules to be decoupled instead.
 		/// </summary>
@@ -245,7 +187,7 @@ namespace UnrealBuildTool
 			{
 				PublicIncludePaths.Add(PublicDirectory);
 
-				HashSet<string> ExcludeNames = new HashSet<string>(UEBuildPlatform.GetBuildPlatform(Rules.Target.Platform).GetExcludedFolderNames(), StringComparer.OrdinalIgnoreCase);
+				HashSet<string> ExcludeNames = UEBuildPlatform.GetBuildPlatform(Rules.Target.Platform).GetExcludedFolderNames();
 				EnumerateLegacyIncludePaths(DirectoryItem.GetItemByDirectoryReference(PublicDirectory), ExcludeNames, LegacyPublicIncludePaths);
 			}
 
@@ -330,19 +272,12 @@ namespace UnrealBuildTool
 				return LinkInputFiles;
 			}
 
-			List<FileReference> SourceFilePaths = SourceFileSearch.FindModuleSourceFiles(ModuleRulesFile: RulesFile, SearchedDirectories: Prerequisites.SourceDirectories);
-			List<FileItem> SourceFiles = GetCPlusPlusFilesToBuild(SourceFilePaths, ModuleDirectory, Target.Platform);
+			SourceFileCollection SourceFilesToBuild = FindSourceFiles(Target.Platform, Prerequisites.SourceDirectories);
 
-			SourceFilesClass SourceFilesFound = new SourceFilesClass();
-			CategorizeSourceFiles(SourceFiles, SourceFilesFound);
-
-			SourceFilesClass SourceFilesToBuild = new SourceFilesClass();
-			SourceFilesToBuild.CopyFrom(SourceFilesFound);
-
-			SourceDirectories = new HashSet<DirectoryReference>(SourceFiles.Select(x => x.Location.Directory));
+			SourceDirectories = new HashSet<DirectoryReference>(SourceFilesToBuild.AllFiles.Select(x => x.Location.Directory));
 
 			// Process all of the header file dependencies for this module
-			CheckFirstIncludeMatchesEachCppFile(Target, ModuleCompileEnvironment, SourceFilesToBuild.CPPFiles);
+			CheckFirstIncludeMatchesEachCppFile(Target, ModuleCompileEnvironment, SourceFilesToBuild.HFiles, SourceFilesToBuild.CPPFiles);
 
 			// Should we force a precompiled header to be generated for this module?  Usually, we only bother with a
 			// precompiled header if there are at least several source files in the module (after combining them for unity
@@ -931,24 +866,19 @@ namespace UnrealBuildTool
 		/// </summary>
 		/// <param name="Target">The target being compiled</param>
 		/// <param name="ModuleCompileEnvironment">Compile environment for the module</param>
+		/// <param name="HeaderFiles">All header files for this module</param>
 		/// <param name="CppFiles">List of C++ source files</param>
-		private void CheckFirstIncludeMatchesEachCppFile(ReadOnlyTargetRules Target, CppCompileEnvironment ModuleCompileEnvironment, List<FileItem> CppFiles)
+		private void CheckFirstIncludeMatchesEachCppFile(ReadOnlyTargetRules Target, CppCompileEnvironment ModuleCompileEnvironment, List<FileItem> HeaderFiles, List<FileItem> CppFiles)
 		{
 			if(Rules.PCHUsage == ModuleRules.PCHUsageMode.UseExplicitOrSharedPCHs)
 			{
 				if(InvalidIncludeDirectiveMessages == null)
 				{
-					// Find all the source files in this module
-					List<FileReference> ModuleFiles = SourceFileSearch.FindModuleSourceFiles(RulesFile);
-
 					// Find headers used by the source file.
 					Dictionary<string, FileReference> NameToHeaderFile = new Dictionary<string, FileReference>();
-					foreach(FileReference ModuleFile in ModuleFiles)
+					foreach(FileItem HeaderFile in HeaderFiles)
 					{
-						if(ModuleFile.HasExtension(".h"))
-						{
-							NameToHeaderFile[ModuleFile.GetFileNameWithoutExtension()] = ModuleFile;
-						}
+						NameToHeaderFile[HeaderFile.Location.GetFileNameWithoutExtension()] = HeaderFile.Location;
 					}
 
 					// Find the directly included files for each source file, and make sure it includes the matching header if possible
@@ -1167,100 +1097,68 @@ namespace UnrealBuildTool
 		}
 
 		/// <summary>
-		/// Given a list of source files for a module, filters them into a list of files that should actually be included in a build
+		/// Finds all the source files that should be built for this module
 		/// </summary>
-		/// <param name="SourceFiles">Original list of files, which may contain non-source</param>
-		/// <param name="SourceFilesBaseDirectory">Directory that the source files are in</param>
-		/// <param name="TargetPlatform">The platform we're going to compile for</param>
-		/// <returns>The list of source files to actually compile</returns>
-		static List<FileItem> GetCPlusPlusFilesToBuild(List<FileReference> SourceFiles, DirectoryReference SourceFilesBaseDirectory, UnrealTargetPlatform TargetPlatform)
+		/// <param name="Platform">The platform the module is being built for</param>
+		/// <param name="SearchedDirectories">Receives a set of directories that have been searched</param>
+		/// <returns>Set of source files that should be built</returns>
+		SourceFileCollection FindSourceFiles(UnrealTargetPlatform Platform, HashSet<DirectoryReference> SearchedDirectories)
 		{
-			// Make a list of all platform name strings that we're *not* currently compiling, to speed
-			// up file path comparisons later on
-			List<UnrealTargetPlatform> SupportedPlatforms = new List<UnrealTargetPlatform>();
-			SupportedPlatforms.Add(TargetPlatform);
-			List<string> OtherPlatformNameStrings = Utils.MakeListOfUnsupportedPlatforms(SupportedPlatforms);
+			HashSet<string> ExcludedNames = UEBuildPlatform.GetBuildPlatform(Platform).GetExcludedFolderNames();
 
+			SourceFileCollection SourceFilesToBuild = new SourceFileCollection();
+			DirectoryItem ModuleDirectoryItem = DirectoryItem.GetItemByDirectoryReference(ModuleDirectory);
+			FindSourceFilesRecursive(ModuleDirectoryItem, ExcludedNames, SearchedDirectories, SourceFilesToBuild);
 
-			// @todo projectfiles: Consider saving out cached list of source files for modules so we don't need to harvest these each time
+			return SourceFilesToBuild;
+		}
 
-			List<FileItem> FilteredFileItems = new List<FileItem>();
-			FilteredFileItems.Capacity = SourceFiles.Count;
+		/// <summary>
+		/// Finds all the source files that should be built for this module
+		/// </summary>
+		/// <param name="BaseDirectory">Directory to search from</param>
+		/// <param name="ExcludedNames">Set of excluded directory names (eg. other platforms)</param>
+		/// <param name="SearchedDirectories">Receives a set of directories that have been searched</param>
+		/// <param name="SourceFiles">Collection of source files, categorized by type</param>
+		static void FindSourceFilesRecursive(DirectoryItem BaseDirectory, HashSet<string> ExcludedNames, HashSet<DirectoryReference> SearchedDirectories, SourceFileCollection SourceFiles)
+		{
+			SearchedDirectories.Add(BaseDirectory.Location);
 
-			// @todo projectfiles: hard-coded source file set.  Should be made extensible by platform tool chains.
-			string[] CompilableSourceFileTypes = new string[]
-				{
-					".cpp",
-					".c",
-					".cc",
-					".mm",
-					".m",
-					".rc",
-					".manifest"
-				};
-
-			// When generating project files, we have no file to extract source from, so we'll locate the code files manually
-			foreach (FileReference SourceFilePath in SourceFiles)
+			foreach(DirectoryItem SubDirectory in BaseDirectory.EnumerateDirectories())
 			{
-				// We're only able to compile certain types of files
-				bool IsCompilableSourceFile = false;
-				foreach (string CurExtension in CompilableSourceFileTypes)
+				if(!ExcludedNames.Contains(SubDirectory.Name))
 				{
-					if (SourceFilePath.HasExtension(CurExtension))
-					{
-						IsCompilableSourceFile = true;
-						break;
-					}
+					FindSourceFilesRecursive(SubDirectory, ExcludedNames, SearchedDirectories, SourceFiles);
 				}
+			}
 
-				if (IsCompilableSourceFile)
-		{
-					if (SourceFilePath.IsUnderDirectory(SourceFilesBaseDirectory))
+			foreach(FileItem SourceFile in BaseDirectory.EnumerateFiles())
 			{
-						// Store the path as relative to the project file
-						string RelativeFilePath = SourceFilePath.MakeRelativeTo(SourceFilesBaseDirectory);
-
-						// All compiled files should always be in a sub-directory under the project file directory.  We enforce this here.
-						if (Path.IsPathRooted(RelativeFilePath) || RelativeFilePath.StartsWith(".."))
+				if (SourceFile.HasExtension(".h"))
 				{
-							throw new BuildException("Error: Found source file {0} in project whose path was not relative to the base directory of the source files", RelativeFilePath);
+					SourceFiles.HFiles.Add(SourceFile);
 				}
-
-						// Check for source files that don't belong to the platform we're currently compiling.  We'll filter
-						// those source files out
-						bool IncludeThisFile = true;
-						foreach (string CurPlatformName in OtherPlatformNameStrings)
+				if (SourceFile.HasExtension(".cpp"))
 				{
-							if (RelativeFilePath.IndexOf(Path.DirectorySeparatorChar + CurPlatformName + Path.DirectorySeparatorChar, StringComparison.InvariantCultureIgnoreCase) != -1
-								|| RelativeFilePath.StartsWith(CurPlatformName + Path.DirectorySeparatorChar))
-							{
-								IncludeThisFile = false;
-								break;
-							}
+					SourceFiles.CPPFiles.Add(SourceFile);
 				}
-
-						if (IncludeThisFile)
+				else if (SourceFile.HasExtension(".c"))
 				{
-							FilteredFileItems.Add(FileItem.GetItemByFileReference(SourceFilePath));
-						}
-					}
+					SourceFiles.CFiles.Add(SourceFile);
 				}
+				else if (SourceFile.HasExtension(".cc"))
+				{
+					SourceFiles.CCFiles.Add(SourceFile);
 				}
-
-			// @todo projectfiles: Consider enabling this error but changing it to a warning instead.  It can fire for
-			//    projects that are being digested for IntelliSense (because the module was set as a cross-
-			//	  platform dependency), but all of their source files were filtered out due to platform masking
-			//    in the project generator
-			bool AllowEmptyProjects = true;
-			if (!AllowEmptyProjects)
+				else if (SourceFile.HasExtension(".m") || SourceFile.HasExtension(".mm"))
 				{
-				if (FilteredFileItems.Count == 0)
+					SourceFiles.MMFiles.Add(SourceFile);
+				}
+				else if (SourceFile.HasExtension(".rc"))
 				{
-					throw new BuildException("Could not find any valid source files for base directory {0}.  Project has {1} files in it", SourceFilesBaseDirectory, SourceFiles.Count);
+					SourceFiles.RCFiles.Add(SourceFile);
 				}
 			}
-
-			return FilteredFileItems;
 		}
 	}
 }
diff --git a/Engine/Source/Programs/UnrealBuildTool/Configuration/UEBuildPlatform.cs b/Engine/Source/Programs/UnrealBuildTool/Configuration/UEBuildPlatform.cs
--- a/Engine/Source/Programs/UnrealBuildTool/Configuration/UEBuildPlatform.cs
+++ b/Engine/Source/Programs/UnrealBuildTool/Configuration/UEBuildPlatform.cs
@@ -36,12 +36,12 @@ namespace UnrealBuildTool
 		/// <summary>
 		/// Cached copy of the list of folders to include for this platform
 		/// </summary>
-		private string[] CachedIncludedFolderNames;
+		private HashSet<string> CachedIncludedFolderNames;
 
 		/// <summary>
 		/// Cached copy of the list of folders to exclude for this platform
 		/// </summary>
-		private string[] CachedExcludedFolderNames;
+		private HashSet<string> CachedExcludedFolderNames;
 
 		/// <summary>
 		/// Constructor.
@@ -119,11 +119,11 @@ namespace UnrealBuildTool
 		/// <summary>
 		/// Finds a list of folder names to include when building for this platform
 		/// </summary>
-		virtual public string[] GetIncludedFolderNames()
+		virtual public HashSet<string> GetIncludedFolderNames()
 		{
 			if(CachedIncludedFolderNames == null)
 			{
-				List<string> Names = new List<string>();
+				HashSet<string> Names = new HashSet<string>(DirectoryReference.Comparer);
 
 				Names.Add(Platform.ToString());
 				foreach(UnrealPlatformGroup Group in UEBuildPlatform.GetPlatformGroups(Platform))
@@ -131,7 +131,7 @@ namespace UnrealBuildTool
 					Names.Add(Group.ToString());
 				}
 
-				CachedIncludedFolderNames = Names.ToArray();
+				CachedIncludedFolderNames = Names;
 			}
 			return CachedIncludedFolderNames;
 		}
@@ -139,11 +139,11 @@ namespace UnrealBuildTool
 		/// <summary>
 		/// Finds a list of folder names to exclude when building for this platform
 		/// </summary>
-		public string[] GetExcludedFolderNames()
+		public HashSet<string> GetExcludedFolderNames()
 		{
 			if(CachedExcludedFolderNames == null)
 			{
-				CachedExcludedFolderNames = GetPlatformFolderNames().Except(GetIncludedFolderNames()).ToArray();
+				CachedExcludedFolderNames = new HashSet<string>(GetPlatformFolderNames().Except(GetIncludedFolderNames()), DirectoryReference.Comparer);
 			}
 			return CachedExcludedFolderNames;
 		}
diff --git a/Engine/Source/Programs/UnrealBuildTool/System/ExternalExecution.cs b/Engine/Source/Programs/UnrealBuildTool/System/ExternalExecution.cs
--- a/Engine/Source/Programs/UnrealBuildTool/System/ExternalExecution.cs
+++ b/Engine/Source/Programs/UnrealBuildTool/System/ExternalExecution.cs
@@ -623,7 +623,7 @@ namespace UnrealBuildTool
 			List<UEBuildModuleCPP> ModulesSortedByType = ModulesToGenerateHeadersFor.OrderBy(c => ModuleToType[c]).ToList();
 			StableTopologicalSort(ModulesSortedByType);
 
-			HashSet<string> ExcludedFolders = new HashSet<string>(UEBuildPlatform.GetBuildPlatform(Platform, true).GetExcludedFolderNames(), StringComparer.OrdinalIgnoreCase);
+			HashSet<string> ExcludedFolders = UEBuildPlatform.GetBuildPlatform(Platform, true).GetExcludedFolderNames();
 			foreach (UEBuildModuleCPP Module in ModulesSortedByType)
 			{
 				DirectoryItem ModuleDirectory = DirectoryItem.GetItemByDirectoryReference(Module.ModuleDirectory);
diff --git a/Engine/Source/Programs/UnrealBuildTool/System/PlatformExports.cs b/Engine/Source/Programs/UnrealBuildTool/System/PlatformExports.cs
--- a/Engine/Source/Programs/UnrealBuildTool/System/PlatformExports.cs
+++ b/Engine/Source/Programs/UnrealBuildTool/System/PlatformExports.cs
@@ -85,7 +85,7 @@ namespace UnrealBuildTool
 		public static string[] GetIncludedFolderNames(UnrealTargetPlatform Platform)
 		{
 			UEBuildPlatform BuildPlatform = UEBuildPlatform.GetBuildPlatform(Platform, false);
-			return BuildPlatform.GetIncludedFolderNames();
+			return BuildPlatform.GetIncludedFolderNames().ToArray();
 		}
 
 		/// <summary>
@@ -96,7 +96,7 @@ namespace UnrealBuildTool
 		public static string[] GetExcludedFolderNames(UnrealTargetPlatform Platform)
 		{
 			UEBuildPlatform BuildPlatform = UEBuildPlatform.GetBuildPlatform(Platform, false);
-			return BuildPlatform.GetExcludedFolderNames();
+			return BuildPlatform.GetExcludedFolderNames().ToArray();
 		}
 
 		/// <summary>
