UBT: Add a read-only hashset wrapper, for returning platform names from the platform layer.

#rnx

diff --git a/Engine/Source/Programs/DotNETCommon/DotNETUtilities/DotNETUtilities.csproj b/Engine/Source/Programs/DotNETCommon/DotNETUtilities/DotNETUtilities.csproj
--- a/Engine/Source/Programs/DotNETCommon/DotNETUtilities/DotNETUtilities.csproj
+++ b/Engine/Source/Programs/DotNETCommon/DotNETUtilities/DotNETUtilities.csproj
@@ -151,6 +151,7 @@
     <Compile Include="Perforce\Records\UserRecord.cs" />
     <Compile Include="Perforce\Records\WhereRecord.cs" />
     <Compile Include="Properties\AssemblyInfo.cs" />
+    <Compile Include="ReadOnlyHashSet.cs" />
     <Compile Include="RestrictedFolders.cs" />
     <Compile Include="ResXResource\UEResXReader.cs" />
     <Compile Include="ResXResource\UEResXWriter.cs" />
diff --git a/Engine/Source/Programs/DotNETCommon/DotNETUtilities/ReadOnlyHashSet.cs b/Engine/Source/Programs/DotNETCommon/DotNETUtilities/ReadOnlyHashSet.cs
new file mode 100644
--- /dev/null
+++ b/Engine/Source/Programs/DotNETCommon/DotNETUtilities/ReadOnlyHashSet.cs
@@ -0,0 +1,95 @@
+// Copyright 1998-2019 Epic Games, Inc. All Rights Reserved.
+
+using System;
+using System.Collections;
+using System.Collections.Generic;
+using System.Linq;
+using System.Text;
+using System.Threading.Tasks;
+
+namespace Tools.DotNETCommon
+{
+	/// <summary>
+	/// Wrapper around the HashSet container that only allows read operations
+	/// </summary>
+	/// <typeparam name="T">Type of element for the hashset</typeparam>
+	public class ReadOnlyHashSet<T> : IReadOnlyCollection<T>
+	{
+		/// <summary>
+		/// The mutable hashset
+		/// </summary>
+		HashSet<T> Inner;
+
+		/// <summary>
+		/// Constructor
+		/// </summary>
+		/// <param name="Inner">The mutable hashset</param>
+		public ReadOnlyHashSet(HashSet<T> Inner)
+		{
+			this.Inner = Inner;
+		}
+
+		/// <summary>
+		/// Constructor
+		/// </summary>
+		/// <param name="Collection">Elements for the hash set</param>
+		public ReadOnlyHashSet(IEnumerable<T> Elements)
+		{
+			this.Inner = new HashSet<T>(Elements);
+		}
+
+		/// <summary>
+		/// Constructor
+		/// </summary>
+		/// <param name="Collection">Elements for the hash set</param>
+		/// <param name="Comparer">Comparer for elements in the set</param>
+		public ReadOnlyHashSet(IEnumerable<T> Elements, IEqualityComparer<T> Comparer)
+		{
+			this.Inner = new HashSet<T>(Elements, Comparer);
+		}
+
+		/// <summary>
+		/// Number of elements in the set
+		/// </summary>
+		public int Count
+		{
+			get { return Inner.Count; }
+		}
+
+		/// <summary>
+		/// The comparer for elements in the set
+		/// </summary>
+		public IEqualityComparer<T> Comparer
+		{
+			get { return Inner.Comparer; }
+		}
+
+		/// <summary>
+		/// Tests whether a given item is in the set
+		/// </summary>
+		/// <param name="Item">Item to check for</param>
+		/// <returns>True if the item is in the set</returns>
+		public bool Contains(T Item)
+		{
+			return Inner.Contains(Item);
+		}
+
+		/// <summary>
+		/// Gets an enumerator for set elements
+		/// </summary>
+		/// <returns>Enumerator instance</returns>
+		public IEnumerator<T> GetEnumerator()
+		{
+			return Inner.GetEnumerator();
+		}
+
+		/// <summary>
+		/// Gets an enumerator for set elements
+		/// </summary>
+		/// <returns>Enumerator instance</returns>
+		IEnumerator IEnumerable.GetEnumerator()
+		{
+			return ((IEnumerable)Inner).GetEnumerator();
+		}
+	}
+}
diff --git a/Engine/Source/Programs/UnrealBuildTool/Configuration/UEBuildModuleCPP.cs b/Engine/Source/Programs/UnrealBuildTool/Configuration/UEBuildModuleCPP.cs
--- a/Engine/Source/Programs/UnrealBuildTool/Configuration/UEBuildModuleCPP.cs
+++ b/Engine/Source/Programs/UnrealBuildTool/Configuration/UEBuildModuleCPP.cs
@@ -187,7 +187,7 @@ namespace UnrealBuildTool
 			{
 				PublicIncludePaths.Add(PublicDirectory);
 
-				HashSet<string> ExcludeNames = UEBuildPlatform.GetBuildPlatform(Rules.Target.Platform).GetExcludedFolderNames();
+				ReadOnlyHashSet<string> ExcludeNames = UEBuildPlatform.GetBuildPlatform(Rules.Target.Platform).GetExcludedFolderNames();
 				EnumerateLegacyIncludePaths(DirectoryItem.GetItemByDirectoryReference(PublicDirectory), ExcludeNames, LegacyPublicIncludePaths);
 			}
 
@@ -205,7 +205,7 @@ namespace UnrealBuildTool
 		/// <param name="BaseDirectory">The directory to start from. This directory is not added to the output list.</param>
 		/// <param name="ExcludeNames">Set of folder names to exclude from the search.</param>
 		/// <param name="LegacyPublicIncludePaths">List populated with the discovered directories</param>
-		static void EnumerateLegacyIncludePaths(DirectoryItem BaseDirectory, HashSet<string> ExcludeNames, HashSet<DirectoryReference> LegacyPublicIncludePaths)
+		static void EnumerateLegacyIncludePaths(DirectoryItem BaseDirectory, ReadOnlyHashSet<string> ExcludeNames, HashSet<DirectoryReference> LegacyPublicIncludePaths)
 		{
 			foreach(DirectoryItem SubDirectory in BaseDirectory.EnumerateDirectories())
 			{
@@ -1104,7 +1104,7 @@ namespace UnrealBuildTool
 		/// <returns>Set of source files that should be built</returns>
 		SourceFileCollection FindSourceFiles(UnrealTargetPlatform Platform, HashSet<DirectoryReference> SearchedDirectories)
 		{
-			HashSet<string> ExcludedNames = UEBuildPlatform.GetBuildPlatform(Platform).GetExcludedFolderNames();
+			ReadOnlyHashSet<string> ExcludedNames = UEBuildPlatform.GetBuildPlatform(Platform).GetExcludedFolderNames();
 
 			SourceFileCollection SourceFilesToBuild = new SourceFileCollection();
 			DirectoryItem ModuleDirectoryItem = DirectoryItem.GetItemByDirectoryReference(ModuleDirectory);
@@ -1120,7 +1120,7 @@ namespace UnrealBuildTool
 		/// <param name="ExcludedNames">Set of excluded directory names (eg. other platforms)</param>
 		/// <param name="SearchedDirectories">Receives a set of directories that have been searched</param>
 		/// <param name="SourceFiles">Collection of source files, categorized by type</param>
-		static void FindSourceFilesRecursive(DirectoryItem BaseDirectory, HashSet<string> ExcludedNames, HashSet<DirectoryReference> SearchedDirectories, SourceFileCollection SourceFiles)
+		static void FindSourceFilesRecursive(DirectoryItem BaseDirectory, ReadOnlyHashSet<string> ExcludedNames, HashSet<DirectoryReference> SearchedDirectories, SourceFileCollection SourceFiles)
 		{
 			SearchedDirectories.Add(BaseDirectory.Location);
 
diff --git a/Engine/Source/Programs/UnrealBuildTool/Configuration/UEBuildPlatform.cs b/Engine/Source/Programs/UnrealBuildTool/Configuration/UEBuildPlatform.cs
--- a/Engine/Source/Programs/UnrealBuildTool/Configuration/UEBuildPlatform.cs
+++ b/Engine/Source/Programs/UnrealBuildTool/Configuration/UEBuildPlatform.cs
@@ -36,12 +36,12 @@ namespace UnrealBuildTool
 		/// <summary>
 		/// Cached copy of the list of folders to include for this platform
 		/// </summary>
-		private HashSet<string> CachedIncludedFolderNames;
+		private ReadOnlyHashSet<string> CachedIncludedFolderNames;
 
 		/// <summary>
 		/// Cached copy of the list of folders to exclude for this platform
 		/// </summary>
-		private HashSet<string> CachedExcludedFolderNames;
+		private ReadOnlyHashSet<string> CachedExcludedFolderNames;
 
 		/// <summary>
 		/// Constructor.
@@ -119,7 +119,7 @@ namespace UnrealBuildTool
 		/// <summary>
 		/// Finds a list of folder names to include when building for this platform
 		/// </summary>
-		virtual public HashSet<string> GetIncludedFolderNames()
+		virtual public ReadOnlyHashSet<string> GetIncludedFolderNames()
 		{
 			if(CachedIncludedFolderNames == null)
 			{
@@ -131,7 +131,7 @@ namespace UnrealBuildTool
 					Names.Add(Group.ToString());
 				}
 
-				CachedIncludedFolderNames = Names;
+				CachedIncludedFolderNames = new ReadOnlyHashSet<string>(Names, DirectoryReference.Comparer);
 			}
 			return CachedIncludedFolderNames;
 		}
@@ -139,11 +139,11 @@ namespace UnrealBuildTool
 		/// <summary>
 		/// Finds a list of folder names to exclude when building for this platform
 		/// </summary>
-		public HashSet<string> GetExcludedFolderNames()
+		public ReadOnlyHashSet<string> GetExcludedFolderNames()
 		{
 			if(CachedExcludedFolderNames == null)
 			{
-				CachedExcludedFolderNames = new HashSet<string>(GetPlatformFolderNames().Except(GetIncludedFolderNames()), DirectoryReference.Comparer);
+				CachedExcludedFolderNames = new ReadOnlyHashSet<string>(GetPlatformFolderNames().Except(GetIncludedFolderNames()), DirectoryReference.Comparer);
 			}
 			return CachedExcludedFolderNames;
 		}
diff --git a/Engine/Source/Programs/UnrealBuildTool/System/DirectoryItem.cs b/Engine/Source/Programs/UnrealBuildTool/System/DirectoryItem.cs
--- a/Engine/Source/Programs/UnrealBuildTool/System/DirectoryItem.cs
+++ b/Engine/Source/Programs/UnrealBuildTool/System/DirectoryItem.cs
@@ -1,3 +1,5 @@
+// Copyright 1998-2019 Epic Games, Inc. All Rights Reserved.
+
 using System;
 using System.Collections.Concurrent;
 using System.Collections.Generic;
diff --git a/Engine/Source/Programs/UnrealBuildTool/System/ExternalExecution.cs b/Engine/Source/Programs/UnrealBuildTool/System/ExternalExecution.cs
--- a/Engine/Source/Programs/UnrealBuildTool/System/ExternalExecution.cs
+++ b/Engine/Source/Programs/UnrealBuildTool/System/ExternalExecution.cs
@@ -592,7 +592,7 @@ namespace UnrealBuildTool
 		/// <param name="BaseDirectory">Base directory to search</param>
 		/// <param name="ExcludeFolders">Array of folders to exclude</param>
 		/// <param name="Headers">Receives the list of headers that was found</param>
-		static void FindHeaders(DirectoryItem BaseDirectory, HashSet<string> ExcludeFolders, List<FileItem> Headers)
+		static void FindHeaders(DirectoryItem BaseDirectory, ReadOnlyHashSet<string> ExcludeFolders, List<FileItem> Headers)
 		{
 			foreach(DirectoryItem SubDirectory in BaseDirectory.EnumerateDirectories())
 			{
@@ -623,7 +623,7 @@ namespace UnrealBuildTool
 			List<UEBuildModuleCPP> ModulesSortedByType = ModulesToGenerateHeadersFor.OrderBy(c => ModuleToType[c]).ToList();
 			StableTopologicalSort(ModulesSortedByType);
 
-			HashSet<string> ExcludedFolders = UEBuildPlatform.GetBuildPlatform(Platform, true).GetExcludedFolderNames();
+			ReadOnlyHashSet<string> ExcludedFolders = UEBuildPlatform.GetBuildPlatform(Platform, true).GetExcludedFolderNames();
 			foreach (UEBuildModuleCPP Module in ModulesSortedByType)
 			{
 				DirectoryItem ModuleDirectory = DirectoryItem.GetItemByDirectoryReference(Module.ModuleDirectory);
