UBT: Optimize file operations to use cached file/directory info, and to discover plugins/rules files in parallel. Also add more stats for different parts of the build.

#rnx

diff --git a/Engine/Source/Programs/UnrealBuildTool/Configuration/UEBuildTarget.cs b/Engine/Source/Programs/UnrealBuildTool/Configuration/UEBuildTarget.cs
--- a/Engine/Source/Programs/UnrealBuildTool/Configuration/UEBuildTarget.cs
+++ b/Engine/Source/Programs/UnrealBuildTool/Configuration/UEBuildTarget.cs
@@ -261,28 +261,33 @@ namespace UnrealBuildTool
 		/// <summary>
 		/// Creates a target object for the specified target name.
 		/// </summary>
-		/// <param name="Desc">Information about the target</param>
+		/// <param name="Descriptor">Information about the target</param>
 		/// <param name="bSkipRulesCompile">Whether to skip compiling any rules assemblies</param>
 		/// <param name="bCompilingSingleFile">Whether we're compiling a single file</param>
 		/// <param name="bUsePrecompiled">Whether to use a precompiled engine/enterprise build</param>
 		/// <returns>The build target object for the specified build rules source file</returns>
-		public static UEBuildTarget CreateTarget(TargetDescriptor Desc, bool bSkipRulesCompile, bool bCompilingSingleFile, bool bUsePrecompiled)
+		public static UEBuildTarget Create(TargetDescriptor Descriptor, bool bSkipRulesCompile, bool bCompilingSingleFile, bool bUsePrecompiled)
 		{
-			RulesAssembly RulesAssembly = RulesCompiler.CreateTargetRulesAssembly(Desc.ProjectFile, Desc.Name, bSkipRulesCompile, bUsePrecompiled, Desc.ForeignPlugin);
-
-			TargetRules RulesObject = RulesAssembly.CreateTargetRules(Desc.Name, Desc.Platform, Desc.Configuration, Desc.Architecture, Desc.ProjectFile, Desc.AdditionalArguments);
-			if ((ProjectFileGenerator.bGenerateProjectFiles == false) && !RulesObject.GetSupportedPlatforms().Contains(Desc.Platform))
+			RulesAssembly RulesAssembly;
+			using(Timeline.ScopeEvent("RulesCompiler.CreateTargetRulesAssembly()"))
 			{
-				throw new BuildException("{0} does not support the {1} platform.", Desc.Name, Desc.Platform.ToString());
+				RulesAssembly = RulesCompiler.CreateTargetRulesAssembly(Descriptor.ProjectFile, Descriptor.Name, bSkipRulesCompile, bUsePrecompiled, Descriptor.ForeignPlugin);
+			}
+	
+			TargetRules RulesObject;
+			using(Timeline.ScopeEvent("RulesAssembly.CreateTargetRules()"))
+			{
+				RulesObject = RulesAssembly.CreateTargetRules(Descriptor.Name, Descriptor.Platform, Descriptor.Configuration, Descriptor.Architecture, Descriptor.ProjectFile, Descriptor.AdditionalArguments);
+			}
+			if ((ProjectFileGenerator.bGenerateProjectFiles == false) && !RulesObject.GetSupportedPlatforms().Contains(Descriptor.Platform))
+			{
+				throw new BuildException("{0} does not support the {1} platform.", Descriptor.Name, Descriptor.Platform.ToString());
 			}
-
-			// Now that we found the actual Editor target, make sure we're no longer using the old TargetName (which is the Game target)
-			Desc.Name = RulesObject.Name;
 
 			// If we're using the shared build environment, make sure all the settings are valid
 			if(RulesObject.BuildEnvironment == TargetBuildEnvironment.Shared)
 			{
-				ValidateSharedEnvironment(RulesAssembly, Desc.Name, RulesObject);
+				ValidateSharedEnvironment(RulesAssembly, Descriptor.Name, RulesObject);
 			}
 
 			// If we're precompiling, generate a list of all the files that we depend on
@@ -321,7 +326,7 @@ namespace UnrealBuildTool
 			}
 
 			// If we're compiling a plugin, and this target is monolithic, just create the object files
-			if(Desc.ForeignPlugin != null && RulesObject.LinkType == TargetLinkType.Monolithic)
+			if(Descriptor.ForeignPlugin != null && RulesObject.LinkType == TargetLinkType.Monolithic)
 			{
 				// Don't actually want an executable
 				RulesObject.bDisableLinking = true;
@@ -338,8 +343,16 @@ namespace UnrealBuildTool
 			}
 
 			// Generate a build target from this rules module
-			UEBuildTarget BuildTarget = new UEBuildTarget(Desc, new ReadOnlyTargetRules(RulesObject), RulesAssembly);
-			return BuildTarget;
+			UEBuildTarget Target;
+			using(Timeline.ScopeEvent("UEBuildTarget constructor"))
+			{
+				Target = new UEBuildTarget(Descriptor, new ReadOnlyTargetRules(RulesObject), RulesAssembly);
+			}
+			using(Timeline.ScopeEvent("UEBuildTarget.PreBuildSetup()"))
+			{
+				Target.PreBuildSetup();
+			}
+			return Target;
 		}
 
 		/// <summary>
@@ -597,22 +610,22 @@ namespace UnrealBuildTool
 		/// <summary>
 		/// Constructor.
 		/// </summary>
-		/// <param name="InDesc">Target descriptor</param>
+		/// <param name="InDescriptor">Target descriptor</param>
 		/// <param name="InRules">The target rules, as created by RulesCompiler.</param>
 		/// <param name="InRulesAssembly">The chain of rules assemblies that this target was created with</param>
-		public UEBuildTarget(TargetDescriptor InDesc, ReadOnlyTargetRules InRules, RulesAssembly InRulesAssembly)
+		private UEBuildTarget(TargetDescriptor InDescriptor, ReadOnlyTargetRules InRules, RulesAssembly InRulesAssembly)
 		{
-			MetadataCache = SourceFileMetadataCache.CreateHierarchy(InDesc.ProjectFile);
-			ProjectFile = InDesc.ProjectFile;
-			AppName = InDesc.Name;
-			TargetName = InDesc.Name;
-			Platform = InDesc.Platform;
-			Configuration = InDesc.Configuration;
-			Architecture = InDesc.Architecture;
+			MetadataCache = SourceFileMetadataCache.CreateHierarchy(InDescriptor.ProjectFile);
+			ProjectFile = InDescriptor.ProjectFile;
+			AppName = InDescriptor.Name;
+			TargetName = InDescriptor.Name;
+			Platform = InDescriptor.Platform;
+			Configuration = InDescriptor.Configuration;
+			Architecture = InDescriptor.Architecture;
 			Rules = InRules;
 			RulesAssembly = InRulesAssembly;
 			TargetType = Rules.Type;
-			ForeignPlugin = InDesc.ForeignPlugin;
+			ForeignPlugin = InDescriptor.ForeignPlugin;
 			bDeployAfterCompile = InRules.bDeployAfterCompile && !InRules.bDisableLinking;
 
 			// now that we have the platform, we can set the intermediate path to include the platform/architecture name
@@ -682,9 +695,6 @@ namespace UnrealBuildTool
 			{
 				ProjectDescriptor = ProjectDescriptor.FromFile(ProjectFile);
 			}
-
-			// Create all the binaries and modules
-			PreBuildSetup();
 		}
 
 		/// <summary>
diff --git a/Engine/Source/Programs/UnrealBuildTool/Modes/BuildMode.cs b/Engine/Source/Programs/UnrealBuildTool/Modes/BuildMode.cs
--- a/Engine/Source/Programs/UnrealBuildTool/Modes/BuildMode.cs
+++ b/Engine/Source/Programs/UnrealBuildTool/Modes/BuildMode.cs
@@ -177,7 +177,7 @@ namespace UnrealBuildTool
 						TargetMakefile Makefile = null;
 						if(BuildConfiguration.bUseUBTMakefiles)
 						{
-							using(Timeline.ScopeEvent("UBTMakefile.LoadUBTMakefile"))
+							using(Timeline.ScopeEvent("TargetMakefile.Load()"))
 							{
 								string ReasonNotLoaded;
 								Makefile = TargetMakefile.Load(MakefileLocation, TargetDesc.ProjectFile, TargetDesc.Platform, WorkingSet, out ReasonNotLoaded);
@@ -193,9 +193,9 @@ namespace UnrealBuildTool
 						{
 							// Create the target
 							UEBuildTarget Target;
-							using(Timeline.ScopeEvent("UEBuildTarget.CreateTarget"))
+							using(Timeline.ScopeEvent("UEBuildTarget.Create()"))
 							{
-								Target = UEBuildTarget.CreateTarget(TargetDesc, bSkipRulesCompile, BuildConfiguration.SingleFileToCompile != null, BuildConfiguration.bUsePrecompiled);
+								Target = UEBuildTarget.Create(TargetDesc, bSkipRulesCompile, BuildConfiguration.SingleFileToCompile != null, BuildConfiguration.bUsePrecompiled);
 							}
 
 							// Build the target
@@ -239,9 +239,12 @@ namespace UnrealBuildTool
 							}
 
 							// Save the makefile for next time
-							using(Timeline.ScopeEvent("UBTMakefile.SaveUBTMakefile()"))
+							if(BuildConfiguration.bUseUBTMakefiles)
 							{
-								TargetMakefile.SaveUBTMakefile(TargetDesc, Makefile);
+								using(Timeline.ScopeEvent("TargetMakefile.Save()"))
+								{
+									Makefile.Save(MakefileLocation);
+								}
 							}
 						}
 						else
diff --git a/Engine/Source/Programs/UnrealBuildTool/Modes/JsonExportMode.cs b/Engine/Source/Programs/UnrealBuildTool/Modes/JsonExportMode.cs
--- a/Engine/Source/Programs/UnrealBuildTool/Modes/JsonExportMode.cs
+++ b/Engine/Source/Programs/UnrealBuildTool/Modes/JsonExportMode.cs
@@ -26,7 +26,7 @@ namespace UnrealBuildTool
 			foreach(TargetDescriptor TargetDescriptor in TargetDescriptors)
 			{
 				// Create the target
-				UEBuildTarget Target = UEBuildTarget.CreateTarget(TargetDescriptor, false, false, false);
+				UEBuildTarget Target = UEBuildTarget.Create(TargetDescriptor, false, false, false);
 
 				// Get the output file
 				FileReference OutputFile = TargetDescriptor.AdditionalArguments.GetFileReferenceOrDefault("-OutputFile=", null);
diff --git a/Engine/Source/Programs/UnrealBuildTool/ProjectFiles/ProjectFileGenerator.cs b/Engine/Source/Programs/UnrealBuildTool/ProjectFiles/ProjectFileGenerator.cs
--- a/Engine/Source/Programs/UnrealBuildTool/ProjectFiles/ProjectFileGenerator.cs
+++ b/Engine/Source/Programs/UnrealBuildTool/ProjectFiles/ProjectFileGenerator.cs
@@ -1708,7 +1708,7 @@ namespace UnrealBuildTool
 							TargetDescriptor TargetDesc = new TargetDescriptor(CurTarget.UnrealProjectFilePath, CurTarget.Name, BuildHostPlatform.Current.Platform, UnrealTargetConfiguration.Development, DefaultArchitecture, null, new CommandLineArguments(NewArguments.ToArray()));
 
 							// Create the target
-							UEBuildTarget Target = UEBuildTarget.CreateTarget(TargetDesc, false, false, bUsePrecompiled);
+							UEBuildTarget Target = UEBuildTarget.Create(TargetDesc, false, false, bUsePrecompiled);
 
 							// Generate a compile environment for each module in the binary
 							CppCompileEnvironment GlobalCompileEnvironment = Target.CreateCompileEnvironmentForProjectFiles();
diff --git a/Engine/Source/Programs/UnrealBuildTool/System/DynamicCompilation.cs b/Engine/Source/Programs/UnrealBuildTool/System/DynamicCompilation.cs
--- a/Engine/Source/Programs/UnrealBuildTool/System/DynamicCompilation.cs
+++ b/Engine/Source/Programs/UnrealBuildTool/System/DynamicCompilation.cs
@@ -26,112 +26,89 @@ namespace UnrealBuildTool
 	/// </summary>
 	public class DynamicCompilation
 	{
-		/// File information for UnrealBuildTool.exe, cached at program start
-		private static FileInfo UBTExecutableFileInfo = new FileInfo(Assembly.GetEntryAssembly().GetOriginalLocation());
-
-		/*
-		 * Checks to see if the assembly needs compilation
-		 */
-		private static bool RequiresCompilation(List<FileReference> SourceFileNames, FileReference AssemblySourceListFilePath, FileReference OutputAssemblyPath)
+		/// <summary>
+		/// Checks to see if the assembly needs compilation
+		/// </summary>
+		/// <param name="SourceFiles">Set of source files</param>
+		/// <param name="AssemblySourceListFilePath">File to use to cache source file names</param>
+		/// <param name="OutputAssemblyPath">Output path for the assembly</param>
+		/// <returns>True if the assembly needs to be built</returns>
+		private static bool RequiresCompilation(HashSet<FileReference> SourceFiles, FileReference AssemblySourceListFilePath, FileReference OutputAssemblyPath)
 		{
 			// Check to see if we already have a compiled assembly file on disk
-			FileInfo OutputAssemblyInfo = new FileInfo(OutputAssemblyPath.FullName);
-			if (OutputAssemblyInfo.Exists)
+			FileItem OutputAssemblyInfo = FileItem.GetItemByFileReference(OutputAssemblyPath);
+			if (!OutputAssemblyInfo.Exists)
 			{
-				// Check the time stamp of the UnrealBuildTool.exe file.  If Unreal Build Tool was compiled more
-				// recently than the dynamically-compiled assembly, then we'll always recompile it.  This is
-				// because Unreal Build Tool's code may have changed in such a way that invalidate these
-				// previously-compiled assembly files.
-				if (UBTExecutableFileInfo.LastWriteTimeUtc > OutputAssemblyInfo.LastWriteTimeUtc)
-				{
-					// UnrealBuildTool.exe has been recompiled more recently than our cached assemblies
-					Log.TraceVerbose("UnrealBuildTool.exe has been recompiled more recently than " + OutputAssemblyInfo.Name);
+				Log.TraceLog("Compiling {0}: Assembly does not exist", OutputAssemblyPath);
+				return true;
+			}
+
+			// Check the time stamp of the UnrealBuildTool.exe file.  If Unreal Build Tool was compiled more
+			// recently than the dynamically-compiled assembly, then we'll always recompile it.  This is
+			// because Unreal Build Tool's code may have changed in such a way that invalidate these
+			// previously-compiled assembly files.
+			FileItem ExecutableItem = FileItem.GetItemByPath(UnrealBuildTool.GetUBTPath());
+			if (ExecutableItem.LastWriteTimeUtc > OutputAssemblyInfo.LastWriteTimeUtc)
+			{
+				Log.TraceLog("Compiling {0}: {1} is newer", OutputAssemblyPath, ExecutableItem.Name);
+				return true;
+			}
+
+			// Make sure we have a manifest of source files used to compile the output assembly.  If it doesn't exist
+			// for some reason (not an expected case) then we'll need to recompile.
+			FileItem AssemblySourceListFile = FileItem.GetItemByFileReference(AssemblySourceListFilePath);
+			if (!AssemblySourceListFile.Exists)
+			{
+				Log.TraceLog("Compiling {0}: Missing source file list ({1})", OutputAssemblyPath, AssemblySourceListFilePath);
+				return true;
+			}
+
+			// Make sure the source files we're compiling are the same as the source files that were compiled
+			// for the assembly that we want to load
+			HashSet<FileItem> CurrentSourceFileItems = new HashSet<FileItem>();
+			foreach(string Line in FileReference.ReadAllLines(AssemblySourceListFile.Location))
+			{
+				CurrentSourceFileItems.Add(FileItem.GetItemByPath(Line));
+			}
+
+			// Get the new source files
+			HashSet<FileItem> SourceFileItems = new HashSet<FileItem>();
+			foreach(FileReference SourceFile in SourceFiles)
+			{
+				SourceFileItems.Add(FileItem.GetItemByFileReference(SourceFile));
+			}
 
+			// Check if there are any differences between the sets
+			foreach(FileItem CurrentSourceFileItem in CurrentSourceFileItems)
+			{
+				if(!SourceFileItems.Contains(CurrentSourceFileItem))
+				{
+					Log.TraceLog("Compiling {0}: Removed source file ({1})", OutputAssemblyPath, AssemblySourceListFilePath);
 					return true;
 				}
-				else
+			}
+			foreach(FileItem SourceFileItem in SourceFileItems)
+			{
+				if(!CurrentSourceFileItems.Contains(SourceFileItem))
 				{
-					// Make sure we have a manifest of source files used to compile the output assembly.  If it doesn't exist
-					// for some reason (not an expected case) then we'll need to recompile.
-					FileInfo AssemblySourceListFile = new FileInfo(AssemblySourceListFilePath.FullName);
-					if (!AssemblySourceListFile.Exists)
-					{
-						return true;
-					}
-					else
-					{
-						// Make sure the source files we're compiling are the same as the source files that were compiled
-						// for the assembly that we want to load
-						List<FileReference> ExistingAssemblySourceFileNames = new List<FileReference>();
-						{
-							using (FileStream Reader = AssemblySourceListFile.OpenRead())
-							{
-								using (StreamReader TextReader = new StreamReader(Reader))
-								{
-									for (string ExistingSourceFileName = TextReader.ReadLine(); ExistingSourceFileName != null; ExistingSourceFileName = TextReader.ReadLine())
-									{
-										FileReference FullExistingSourceFileName = new FileReference(ExistingSourceFileName);
-
-										ExistingAssemblySourceFileNames.Add(FullExistingSourceFileName);
-
-										// Was the existing assembly compiled with a source file that we aren't interested in?  If so, then it needs to be recompiled.
-										if (!SourceFileNames.Contains(FullExistingSourceFileName))
-										{
-											return true;
-										}
-									}
-								}
-							}
-						}
-
-						// Test against source file time stamps
-						foreach (FileReference SourceFileName in SourceFileNames)
-						{
-							// Was the existing assembly compiled without this source file?  If so, then we definitely need to recompile it!
-							if (!ExistingAssemblySourceFileNames.Contains(SourceFileName))
-							{
-								return true;
-							}
-
-							FileInfo SourceFileInfo = new FileInfo(SourceFileName.FullName);
-
-							// Check to see if the source file exists
-							if (!SourceFileInfo.Exists)
-							{
-								throw new BuildException("Could not locate source file for dynamic compilation: {0}", SourceFileName);
-							}
-
-							// Ignore temp files
-							if (!SourceFileInfo.Extension.Equals(".tmp", StringComparison.CurrentCultureIgnoreCase))
-							{
-								// Check to see if the source file is newer than the compiled assembly file.  We don't want to
-								// bother recompiling it if it hasn't changed.
-								if (SourceFileInfo.LastWriteTimeUtc > OutputAssemblyInfo.LastWriteTimeUtc)
-								{
-									// Source file has changed since we last compiled the assembly, so we'll need to recompile it now!
-									Log.TraceVerbose(SourceFileInfo.Name + " has been modified more recently than " + OutputAssemblyInfo.Name);
-
-									return true;
-								}
-							}
-						}
-					}
+					Log.TraceLog("Compiling {0}: Added source file ({1})", OutputAssemblyPath, AssemblySourceListFilePath);
+					return true;
 				}
 			}
-			else
+
+			// Check if any of the timestamps are newer
+			foreach(FileItem SourceFileItem in SourceFileItems)
 			{
-				// File doesn't exist, so we'll definitely have to compile it!
-				Log.TraceVerbose(OutputAssemblyInfo.Name + " doesn't exist yet");
-				return true;
+				if(SourceFileItem.LastWriteTimeUtc > OutputAssemblyInfo.LastWriteTimeUtc)
+				{
+					Log.TraceLog("Compiling {0}: {1} is newer", OutputAssemblyPath, SourceFileItem);
+					return true;
+				}
 			}
 
 			return false;
 		}
 
-		/*
-		 * Compiles an assembly from source files
-		 */
-
 #if NET_CORE
 		private static void LogDiagnostics(IEnumerable<Diagnostic> Diagnostics)
 		{
@@ -259,7 +236,7 @@ namespace UnrealBuildTool
 
 #else
 
-		private static Assembly CompileAssembly(FileReference OutputAssemblyPath, List<FileReference> SourceFileNames, List<string> ReferencedAssembies, List<string> PreprocessorDefines = null, bool TreatWarningsAsErrors = false)
+		private static Assembly CompileAssembly(FileReference OutputAssemblyPath, HashSet<FileReference> SourceFileNames, List<string> ReferencedAssembies, List<string> PreprocessorDefines = null, bool TreatWarningsAsErrors = false)
 		{
 			TempFileCollection TemporaryFiles = new TempFileCollection();
 
@@ -403,7 +380,7 @@ namespace UnrealBuildTool
 		/// <param name="DoNotCompile"></param>
 		/// <param name="TreatWarningsAsErrors"></param>
 		/// <returns>The assembly that was loaded</returns>
-		public static Assembly CompileAndLoadAssembly(FileReference OutputAssemblyPath, List<FileReference> SourceFileNames, List<string> ReferencedAssembies = null, List<string> PreprocessorDefines = null, bool DoNotCompile = false, bool TreatWarningsAsErrors = false)
+		public static Assembly CompileAndLoadAssembly(FileReference OutputAssemblyPath, HashSet<FileReference> SourceFileNames, List<string> ReferencedAssembies = null, List<string> PreprocessorDefines = null, bool DoNotCompile = false, bool TreatWarningsAsErrors = false)
 		{
 			// Check to see if the resulting assembly is compiled and up to date
 			FileReference AssemblySourcesListFilePath = FileReference.Combine(OutputAssemblyPath.Directory, Path.GetFileNameWithoutExtension(OutputAssemblyPath.FullName) + "SourceFiles.txt");
@@ -445,17 +422,14 @@ namespace UnrealBuildTool
 			// Compile the assembly if me
 			if (bNeedsCompilation)
 			{
-				CompiledAssembly = CompileAssembly(OutputAssemblyPath, SourceFileNames, ReferencedAssembies, PreprocessorDefines, TreatWarningsAsErrors);
+				using(Timeline.ScopeEvent(String.Format("Compiling rules assembly ({0})", OutputAssemblyPath.GetFileName())))
+				{
+					CompiledAssembly = CompileAssembly(OutputAssemblyPath, SourceFileNames, ReferencedAssembies, PreprocessorDefines, TreatWarningsAsErrors);
+				}
 
 				// Save out a list of all the source files we compiled.  This is so that we can tell if whole files were added or removed
 				// since the previous time we compiled the assembly.  In that case, we'll always want to recompile it!
-				{
-					FileInfo AssemblySourcesListFile = new FileInfo(AssemblySourcesListFilePath.FullName);
-					using (StreamWriter Writer = AssemblySourcesListFile.CreateText())
-					{
-						SourceFileNames.ForEach(x => Writer.WriteLine(x));
-					}
-				}
+				FileReference.WriteAllLines(AssemblySourcesListFilePath, SourceFileNames.Select(x => x.FullName));
 			}
 
 #if !NET_CORE
diff --git a/Engine/Source/Programs/UnrealBuildTool/System/Plugins.cs b/Engine/Source/Programs/UnrealBuildTool/System/Plugins.cs
--- a/Engine/Source/Programs/UnrealBuildTool/System/Plugins.cs
+++ b/Engine/Source/Programs/UnrealBuildTool/System/Plugins.cs
@@ -330,7 +330,10 @@ namespace UnrealBuildTool
 				DirectoryItem ParentDirectoryItem = DirectoryItem.GetItemByDirectoryReference(ParentDirectory);
 				if (ParentDirectoryItem.Exists)
 				{
-					EnumeratePluginsInternal(ParentDirectoryItem, FileNames);
+					using(ThreadPoolWorkQueue Queue = new ThreadPoolWorkQueue())
+					{
+						EnumeratePluginsInternal(ParentDirectoryItem, FileNames, Queue);
+					}
 				}
 
 				PluginFileCache.Add(ParentDirectory, FileNames);
@@ -343,7 +346,8 @@ namespace UnrealBuildTool
 		/// </summary>
 		/// <param name="ParentDirectory">Parent directory to look in. Plugins will be found in any *subfolders* of this directory.</param>
 		/// <param name="FileNames">List of filenames. Will have all the discovered .uplugin files appended to it.</param>
-		static void EnumeratePluginsInternal(DirectoryItem ParentDirectory, List<FileReference> FileNames)
+		/// <param name="Queue">Queue for tasks to be executed</param>
+		static void EnumeratePluginsInternal(DirectoryItem ParentDirectory, List<FileReference> FileNames, ThreadPoolWorkQueue Queue)
 		{
 			foreach (DirectoryItem ChildDirectory in ParentDirectory.EnumerateDirectories())
 			{
@@ -352,19 +356,21 @@ namespace UnrealBuildTool
 				{
 					if(PluginFile.HasExtension(".uplugin"))
 					{
-						FileNames.Add(PluginFile.Location);
+						lock(FileNames)
+						{
+							FileNames.Add(PluginFile.Location);
+						}
 						bSearchSubDirectories = false;
 					}
 				}
 
 				if (bSearchSubDirectories)
 				{
-					EnumeratePluginsInternal(ChildDirectory, FileNames);
+					Queue.Enqueue(() => EnumeratePluginsInternal(ChildDirectory, FileNames, Queue));
 				}
 			}
 		}
 
-
 		/// <summary>
 		/// Determine if a plugin is enabled for a given project
 		/// </summary>
diff --git a/Engine/Source/Programs/UnrealBuildTool/System/RulesAssembly.cs b/Engine/Source/Programs/UnrealBuildTool/System/RulesAssembly.cs
--- a/Engine/Source/Programs/UnrealBuildTool/System/RulesAssembly.cs
+++ b/Engine/Source/Programs/UnrealBuildTool/System/RulesAssembly.cs
@@ -98,9 +98,9 @@ namespace UnrealBuildTool
 			this.Parent = Parent;
 
 			// Find all the source files
-			List<FileReference> AssemblySourceFiles = new List<FileReference>();
-			AssemblySourceFiles.AddRange(ModuleFiles);
-			AssemblySourceFiles.AddRange(TargetFiles);
+			HashSet<FileReference> AssemblySourceFiles = new HashSet<FileReference>();
+			AssemblySourceFiles.UnionWith(ModuleFiles);
+			AssemblySourceFiles.UnionWith(TargetFiles);
 
 			// Compile the assembly
 			if (AssemblySourceFiles.Count > 0)
diff --git a/Engine/Source/Programs/UnrealBuildTool/System/RulesCompiler.cs b/Engine/Source/Programs/UnrealBuildTool/System/RulesCompiler.cs
--- a/Engine/Source/Programs/UnrealBuildTool/System/RulesCompiler.cs
+++ b/Engine/Source/Programs/UnrealBuildTool/System/RulesCompiler.cs
@@ -8,6 +8,7 @@ using System.Reflection;
 using System.Diagnostics;
 using System.Runtime.Serialization;
 using Tools.DotNETCommon;
+using System.Collections.Concurrent;
 
 namespace UnrealBuildTool
 {
@@ -179,6 +180,48 @@ namespace UnrealBuildTool
             DirectoryLookupCache.InvalidateCachedDirectory(Directory);
         }
 
+		/// <summary>
+		/// Prefetch multiple directories in parallel
+		/// </summary>
+		/// <param name="Directories">The directories to cache</param>
+		private static void PrefetchRulesFiles(IEnumerable<DirectoryReference> Directories)
+		{
+			ThreadPoolWorkQueue Queue = null;
+			try
+			{
+				foreach(DirectoryReference Directory in Directories)
+				{
+					if(!RootFolderToRulesFileCache.ContainsKey(Directory))
+					{
+						RulesFileCache Cache = new RulesFileCache();
+						RootFolderToRulesFileCache[Directory] = Cache;
+
+						if(Queue == null)
+						{
+							Queue = new ThreadPoolWorkQueue();
+						}
+
+						DirectoryItem DirectoryItem = DirectoryItem.GetItemByDirectoryReference(Directory);
+						Queue.Enqueue(() => FindAllRulesFilesRecursively(DirectoryItem, Cache, Queue));
+					}
+				}
+			}
+			finally
+			{
+				if(Queue != null)
+				{
+					Queue.Dispose();
+					Queue = null;
+				}
+			}
+		}
+
+		/// <summary>
+		/// Finds all the rules of the given type under a given directory
+		/// </summary>
+		/// <param name="Directory">Directory to search</param>
+		/// <param name="Type">Type of rules to return</param>
+		/// <returns>List of rules files of the given type</returns>
 		private static IReadOnlyList<FileReference> FindAllRulesFiles(DirectoryReference Directory, RulesFileType Type)
 		{
 			// Check to see if we've already cached source files for this folder
@@ -186,7 +229,11 @@ namespace UnrealBuildTool
 			if (!RootFolderToRulesFileCache.TryGetValue(Directory, out Cache))
 			{
 				Cache = new RulesFileCache();
-				FindAllRulesFilesRecursively(Directory, Cache);
+				using(ThreadPoolWorkQueue Queue = new ThreadPoolWorkQueue())
+				{
+					DirectoryItem BaseDirectory = DirectoryItem.GetItemByDirectoryReference(Directory);
+					Queue.Enqueue(() => FindAllRulesFilesRecursively(BaseDirectory, Cache, Queue));
+				}
 				RootFolderToRulesFileCache[Directory] = Cache;
 			}
 
@@ -209,24 +256,39 @@ namespace UnrealBuildTool
 			}
 		}
 
-		private static void FindAllRulesFilesRecursively(DirectoryReference Directory, RulesFileCache Cache)
+		/// <summary>
+		/// Search through a directory tree for any rules files
+		/// </summary>
+		/// <param name="Directory">The root directory to search from</param>
+		/// <param name="Cache">Receives all the discovered rules files</param>
+		/// <param name="Queue">Queue for adding additional tasks to</param>
+		private static void FindAllRulesFilesRecursively(DirectoryItem Directory, RulesFileCache Cache, ThreadPoolWorkQueue Queue)
 		{
 			// Scan all the files in this directory
 			bool bSearchSubFolders = true;
-			foreach (FileReference File in DirectoryLookupCache.EnumerateFiles(Directory))
+			foreach (FileItem File in Directory.EnumerateFiles())
 			{
 				if (File.HasExtension(".build.cs"))
 				{
-					Cache.ModuleRules.Add(File);
+					lock(Cache.ModuleRules)
+					{
+						Cache.ModuleRules.Add(File.Location);
+					}
 					bSearchSubFolders = false;
 				}
 				else if (File.HasExtension(".target.cs"))
 				{
-					Cache.TargetRules.Add(File);
+					lock(Cache.TargetRules)
+					{
+						Cache.TargetRules.Add(File.Location);
+					}
 				}
 				else if (File.HasExtension(".automation.csproj"))
 				{
-					Cache.AutomationModules.Add(File);
+					lock(Cache.AutomationModules)
+					{
+						Cache.AutomationModules.Add(File.Location);
+					}
 					bSearchSubFolders = false;
 				}
 			}
@@ -234,9 +296,9 @@ namespace UnrealBuildTool
 			// If we didn't find anything to stop the search, search all the subdirectories too
 			if (bSearchSubFolders)
 			{
-				foreach (DirectoryReference SubDirectory in DirectoryLookupCache.EnumerateDirectories(Directory))
+				foreach (DirectoryItem SubDirectory in Directory.EnumerateDirectories())
 				{
-					FindAllRulesFilesRecursively(SubDirectory, Cache);
+					Queue.Enqueue(() => FindAllRulesFilesRecursively(SubDirectory, Cache, Queue));
 				}
 			}
 		}
@@ -251,8 +313,10 @@ namespace UnrealBuildTool
 		/// </summary>
 		private static RulesAssembly EnterpriseRulesAssembly;
 
+		/// <summary>
 		/// Map of assembly names we've already compiled and loaded to their Assembly and list of game folders.  This is used to prevent
 		/// trying to recompile the same assembly when ping-ponging between different types of targets
+		/// </summary>
 		private static Dictionary<FileReference, RulesAssembly> LoadedAssemblyMap = new Dictionary<FileReference, RulesAssembly>();
 
 		/// <summary>
@@ -304,25 +368,39 @@ namespace UnrealBuildTool
 
 			// Find the shared modules, excluding the programs directory. These are used to create an assembly with the bContainsEngineModules flag set to true.
 			List<FileReference> EngineModuleFiles = new List<FileReference>();
-			foreach (DirectoryReference SubDirectory in DirectoryLookupCache.EnumerateDirectories(SourceDirectory))
+			using(Timeline.ScopeEvent("Finding engine modules"))
 			{
-				if(SubDirectory != ProgramsDirectory)
+				foreach (DirectoryReference SubDirectory in DirectoryLookupCache.EnumerateDirectories(SourceDirectory))
 				{
-					EngineModuleFiles.AddRange(FindAllRulesFiles(SubDirectory, RulesFileType.Module));
+					if(SubDirectory != ProgramsDirectory)
+					{
+						EngineModuleFiles.AddRange(FindAllRulesFiles(SubDirectory, RulesFileType.Module));
+					}
 				}
 			}
 
 			// Add all the plugin modules too
 			Dictionary<FileReference, PluginInfo> ModuleFileToPluginInfo = new Dictionary<FileReference, PluginInfo>();
-			FindModuleRulesForPlugins(Plugins, EngineModuleFiles, ModuleFileToPluginInfo);
+			using(Timeline.ScopeEvent("Finding plugin modules"))
+			{
+				FindModuleRulesForPlugins(Plugins, EngineModuleFiles, ModuleFileToPluginInfo);
+			}
 
 			// Create the assembly
 			FileReference EngineAssemblyFileName = FileReference.Combine(RootDirectory, "Intermediate", "Build", "BuildRules", AssemblyPrefix + "Rules" + FrameworkAssemblyExtension);
 			RulesAssembly EngineAssembly = new RulesAssembly(RootDirectory, Plugins, EngineModuleFiles, new List<FileReference>(), ModuleFileToPluginInfo, EngineAssemblyFileName, bContainsEngineModules: true, bUseBackwardsCompatibleDefaults: false, bReadOnly: bReadOnly, bSkipCompile: bSkipCompile, Parent: Parent);
 
 			// Find all the rules files
-			List<FileReference> ProgramModuleFiles = new List<FileReference>(FindAllRulesFiles(ProgramsDirectory, RulesFileType.Module));
-			List<FileReference> ProgramTargetFiles = new List<FileReference>(FindAllRulesFiles(SourceDirectory, RulesFileType.Target));
+			List<FileReference> ProgramModuleFiles;
+			using(Timeline.ScopeEvent("Finding program modules"))
+			{
+				ProgramModuleFiles = new List<FileReference>(FindAllRulesFiles(ProgramsDirectory, RulesFileType.Module));
+			}
+			List<FileReference> ProgramTargetFiles;
+			using(Timeline.ScopeEvent("Finding program targets"))
+			{
+				ProgramTargetFiles = new List<FileReference>(FindAllRulesFiles(SourceDirectory, RulesFileType.Target));
+			}
 
 			// Create a path to the assembly that we'll either load or compile
 			FileReference ProgramAssemblyFileName = FileReference.Combine(RootDirectory, "Intermediate", "Build", "BuildRules", AssemblyPrefix + "ProgramRules" + FrameworkAssemblyExtension);
@@ -478,6 +556,8 @@ namespace UnrealBuildTool
 		/// <param name="ModuleFileToPluginInfo">Dictionary which is filled with mappings from the module file to its corresponding plugin file</param>
 		private static void FindModuleRulesForPlugins(IReadOnlyList<PluginInfo> Plugins, List<FileReference> ModuleFiles, Dictionary<FileReference, PluginInfo> ModuleFileToPluginInfo)
 		{
+			PrefetchRulesFiles(Plugins.Select(x => DirectoryReference.Combine(x.Directory, "Source")));
+
 			foreach (PluginInfo Plugin in Plugins)
 			{
 				IReadOnlyList<FileReference> PluginModuleFiles = FindAllRulesFiles(DirectoryReference.Combine(Plugin.Directory, "Source"), RulesFileType.Module);
diff --git a/Engine/Source/Programs/UnrealBuildTool/System/TargetMakefile.cs b/Engine/Source/Programs/UnrealBuildTool/System/TargetMakefile.cs
--- a/Engine/Source/Programs/UnrealBuildTool/System/TargetMakefile.cs
+++ b/Engine/Source/Programs/UnrealBuildTool/System/TargetMakefile.cs
@@ -150,29 +150,22 @@ namespace UnrealBuildTool
 
 
 		/// <summary>
-		/// Saves a UBTMakefile to disk
+		/// Saves a makefile to disk
 		/// </summary>
-		/// <param name="TargetDesc">Target being built</param>
-		/// <param name="UBTMakefile">The UBT makefile</param>
-		public static void SaveUBTMakefile(TargetDescriptor TargetDesc, TargetMakefile UBTMakefile)
+		/// <param name="Location">Path to save the makefile to</param>
+		public void Save(FileReference Location)
 		{
-			if (!UBTMakefile.IsValidMakefile())
+			if (!IsValidMakefile())
 			{
 				throw new BuildException("Can't save a makefile that has invalid contents.  See UBTMakefile.IsValidMakefile()");
 			}
 
-			FileItem UBTMakefileItem = FileItem.GetItemByFileReference(GetLocation(TargetDesc.ProjectFile, TargetDesc.Name, TargetDesc.Platform, TargetDesc.Configuration));
-
-			// @todo ubtmake: Optimization: The UBTMakefile saved for game projects is upwards of 9 MB.  We should try to shrink its content if possible
-			// @todo ubtmake: Optimization: C# Serialization may be too slow for these big Makefiles.  Loading these files often shows up as the slower part of the assembling phase.
-
-			// Serialize the cache to disk.
 			try
 			{
-				Directory.CreateDirectory(Path.GetDirectoryName(UBTMakefileItem.AbsolutePath));
-				using(BinaryArchiveWriter Writer = new BinaryArchiveWriter(UBTMakefileItem.Location))
+				DirectoryReference.CreateDirectory(Location.Directory);
+				using(BinaryArchiveWriter Writer = new BinaryArchiveWriter(Location))
 				{
-					UBTMakefile.Write(Writer);
+					Write(Writer);
 				}
 			}
 			catch (Exception Ex)
