UBT: Move the setting for whether to skip rules compilation into the BuildConfiguation class.

#rnx

diff --git a/Engine/Source/Programs/UnrealBuildTool/Configuration/BuildConfiguration.cs b/Engine/Source/Programs/UnrealBuildTool/Configuration/BuildConfiguration.cs
--- a/Engine/Source/Programs/UnrealBuildTool/Configuration/BuildConfiguration.cs
+++ b/Engine/Source/Programs/UnrealBuildTool/Configuration/BuildConfiguration.cs
@@ -160,6 +160,12 @@ namespace UnrealBuildTool
 		[XmlConfigFile(Category = "UEBuildConfiguration")]
 		public bool bForceDebugUnrealHeaderTool = false;
 
+		/// <summary>
+		/// Whether to skip compiling rules assemblies and just assume they are valid
+		/// </summary>
+		[CommandLine("-SkipRulesCompile")]
+		public bool bSkipRulesCompile = false;
+
 		/// <summary>
 		/// Compiler arguments to be forwarded to UnrealHeaderTool
 		/// </summary>
diff --git a/Engine/Source/Programs/UnrealBuildTool/Modes/BuildMode.cs b/Engine/Source/Programs/UnrealBuildTool/Modes/BuildMode.cs
--- a/Engine/Source/Programs/UnrealBuildTool/Modes/BuildMode.cs
+++ b/Engine/Source/Programs/UnrealBuildTool/Modes/BuildMode.cs
@@ -100,15 +100,10 @@ namespace UnrealBuildTool
 
 			try
 			{
-				// action graph implies using the dependency resolve cache
-				bool GeneratingActionGraph = Arguments.HasOption("-Graph");
-
-				bool bSkipRulesCompile = Arguments.Any(x => x.Equals("-skiprulescompile", StringComparison.InvariantCultureIgnoreCase));
-
 				List<TargetDescriptor> TargetDescriptors = new List<TargetDescriptor>();
 				using(Timeline.ScopeEvent("TargetDescriptor.ParseCommandLine()"))
 				{
-					TargetDescriptors.AddRange(TargetDescriptor.ParseCommandLine(Arguments, BuildConfiguration.bUsePrecompiled, bSkipRulesCompile));
+					TargetDescriptors.AddRange(TargetDescriptor.ParseCommandLine(Arguments, BuildConfiguration.bUsePrecompiled, BuildConfiguration.bSkipRulesCompile));
 				}
 
 				// Handle remote builds
@@ -195,7 +190,7 @@ namespace UnrealBuildTool
 							UEBuildTarget Target;
 							using(Timeline.ScopeEvent("UEBuildTarget.Create()"))
 							{
-								Target = UEBuildTarget.Create(TargetDesc, bSkipRulesCompile, BuildConfiguration.SingleFileToCompile != null, BuildConfiguration.bUsePrecompiled);
+								Target = UEBuildTarget.Create(TargetDesc, BuildConfiguration.bSkipRulesCompile, BuildConfiguration.SingleFileToCompile != null, BuildConfiguration.bUsePrecompiled);
 							}
 
 							// Build the target
