UBT: If UnrealHeaderTool needs rebuilding, do so from the same process rather than spawning a child UBT instance. This is significantly faster, since we're already fully initialized and have cached a lot of state.

diff --git a/Engine/Source/Programs/UnrealBuildTool/Configuration/BuildConfiguration.cs b/Engine/Source/Programs/UnrealBuildTool/Configuration/BuildConfiguration.cs
--- a/Engine/Source/Programs/UnrealBuildTool/Configuration/BuildConfiguration.cs
+++ b/Engine/Source/Programs/UnrealBuildTool/Configuration/BuildConfiguration.cs
@@ -32,12 +32,6 @@ namespace UnrealBuildTool
 		[XmlConfigFile]
 		public bool bPrintDebugInfo = false;
 
-		/// <summary>
-		/// Prints performance diagnostics about include dependencies and other bits
-		/// </summary>
-		[XmlConfigFile]
-		public bool bPrintPerformanceInfo = false;
-
 		/// <summary>
 		/// Whether to log detailed action stats. This forces local execution.
 		/// </summary>
@@ -166,13 +160,6 @@ namespace UnrealBuildTool
 		[CommandLine("-SkipRulesCompile")]
 		public bool bSkipRulesCompile = false;
 
-		/// <summary>
-		/// Compiler arguments to be forwarded to UnrealHeaderTool
-		/// </summary>
-		[CommandLine("-2015", Value = "-2015")]
-		[CommandLine("-2017", Value = "-2017")]
-		public string CompilerArgumentForUnrealHeaderTool = null;
-
 		/// <summary>
 		/// Default constructor. Reads settings from the XmlConfig files.
 		/// </summary>
diff --git a/Engine/Source/Programs/UnrealBuildTool/Configuration/TargetDescriptor.cs b/Engine/Source/Programs/UnrealBuildTool/Configuration/TargetDescriptor.cs
--- a/Engine/Source/Programs/UnrealBuildTool/Configuration/TargetDescriptor.cs
+++ b/Engine/Source/Programs/UnrealBuildTool/Configuration/TargetDescriptor.cs
@@ -40,7 +40,7 @@ namespace UnrealBuildTool
 			this.Configuration = Configuration;
 			this.Architecture = Architecture;
 			this.ForeignPlugin = ForeignPlugin;
-			this.AdditionalArguments = AdditionalArguments;
+			this.AdditionalArguments = AdditionalArguments ?? new CommandLineArguments(new string[0]);
 		}
 
 		/// <summary>
diff --git a/Engine/Source/Programs/UnrealBuildTool/Configuration/UEBuildTarget.cs b/Engine/Source/Programs/UnrealBuildTool/Configuration/UEBuildTarget.cs
--- a/Engine/Source/Programs/UnrealBuildTool/Configuration/UEBuildTarget.cs
+++ b/Engine/Source/Programs/UnrealBuildTool/Configuration/UEBuildTarget.cs
@@ -1293,7 +1293,7 @@ namespace UnrealBuildTool
 				FileReference ModuleInfoFileName = FileReference.Combine(ProjectIntermediateDirectory, TargetName + ".uhtmanifest");
 
 				// Execute the header tool
-				ECompilationResult UHTResult = ExternalExecution.ExecuteHeaderToolIfNecessary(BuildConfiguration, ProjectFile, TargetName, TargetType, bHasProjectScriptPlugin, UObjectModules, ModuleInfoFileName, true, bIsAssemblingBuild);
+				ECompilationResult UHTResult = ExternalExecution.ExecuteHeaderToolIfNecessary(BuildConfiguration, ProjectFile, TargetName, TargetType, bHasProjectScriptPlugin, UObjectModules, ModuleInfoFileName, true, bIsAssemblingBuild, WorkingSet);
 				if(!UHTResult.Succeeded())
 				{
 					Log.TraceInformation(String.Format("Error: UnrealHeaderTool failed for target '{0}' (platform: {1}, module info: {2}, exit code: {3} ({4})).", TargetName, Platform.ToString(), ModuleInfoFileName, UHTResult.ToString(), (int)UHTResult));
diff --git a/Engine/Source/Programs/UnrealBuildTool/Modes/BuildMode.cs b/Engine/Source/Programs/UnrealBuildTool/Modes/BuildMode.cs
--- a/Engine/Source/Programs/UnrealBuildTool/Modes/BuildMode.cs
+++ b/Engine/Source/Programs/UnrealBuildTool/Modes/BuildMode.cs
@@ -92,12 +92,6 @@ namespace UnrealBuildTool
 			}
 
 			Stopwatch BuildTimer = Stopwatch.StartNew();
-
-			// Reset global configurations
-			string ExecutorName = "Unknown";
-
-			Stopwatch ExecutorTimer = new Stopwatch();
-
 			try
 			{
 				List<TargetDescriptor> TargetDescriptors = new List<TargetDescriptor>();
@@ -159,250 +153,258 @@ namespace UnrealBuildTool
 				// Create the working set provider
 				using (ISourceFileWorkingSet WorkingSet = SourceFileWorkingSet.Create(UnrealBuildTool.RootDirectory, ProjectDirs))
 				{
-					// Create a makefile for each target
-					TargetMakefile[] Makefiles = new TargetMakefile[TargetDescriptors.Count];
-					for(int TargetIdx = 0; TargetIdx < TargetDescriptors.Count; TargetIdx++)
-					{
-						TargetDescriptor TargetDesc = TargetDescriptors[TargetIdx];
-
-						// Get the path to the makefile for this target
-						FileReference MakefileLocation = TargetMakefile.GetLocation(TargetDesc.ProjectFile, TargetDesc.Name, TargetDesc.Platform, TargetDesc.Configuration);
-
-						// Try to load an existing makefile
-						TargetMakefile Makefile = null;
-						if(BuildConfiguration.bUseUBTMakefiles)
-						{
-							using(Timeline.ScopeEvent("TargetMakefile.Load()"))
-							{
-								string ReasonNotLoaded;
-								Makefile = TargetMakefile.Load(MakefileLocation, TargetDesc.ProjectFile, TargetDesc.Platform, WorkingSet, out ReasonNotLoaded);
-								if (Makefile == null)
-								{
-									Log.TraceInformation("Creating makefile for {0} ({1})", TargetDesc.Name, ReasonNotLoaded);
-								}
-							}
-						}
-
-						// If we couldn't load a makefile, create a new one
-						if(Makefile == null)
-						{
-							// Create the target
-							UEBuildTarget Target;
-							using(Timeline.ScopeEvent("UEBuildTarget.Create()"))
-							{
-								Target = UEBuildTarget.Create(TargetDesc, BuildConfiguration.bSkipRulesCompile, BuildConfiguration.SingleFileToCompile != null, BuildConfiguration.bUsePrecompiled);
-							}
-
-							// Build the target
-							const bool bIsAssemblingBuild = true;
+					ECompilationResult Result = Build(TargetDescriptors, BuildConfiguration, WorkingSet);
+					return (int)Result;
+				}
+			}
+			catch (Exception Ex)
+			{
+				Log.WriteException(Ex, (LogFile == null) ? null : LogFile.FullName);
+				return (int)ECompilationResult.OtherCompilationError;
+			}
+			finally
+			{
+				// Save all the caches
+				SourceFileMetadataCache.SaveAll();
+				CppDependencyCache.SaveAll();
 
-							List<FileItem> OutputItems = new List<FileItem>();
-							Dictionary<string, FileItem[]> ModuleNameToOutputItems = new Dictionary<string, FileItem[]>();
-							List<UHTModuleInfo> UObjectModules = new List<UHTModuleInfo>();
-							List<Action> Actions = new List<Action>();
-							BuildPrerequisites Prerequisites = new BuildPrerequisites();
+				// Figure out how long we took to execute.
+				Log.TraceInformation("Total build time: {0:0.00} seconds", BuildTimer.Elapsed.TotalSeconds);
+			}
+		}
 
-							ECompilationResult BuildResult;
-							using(Timeline.ScopeEvent("UEBuildTarget.Build()"))
-							{
-								BuildResult = Target.Build(BuildConfiguration, OutputItems, ModuleNameToOutputItems, UObjectModules, WorkingSet, Actions, Prerequisites, bIsAssemblingBuild);
-							}
-							if (BuildResult != ECompilationResult.Succeeded)
-							{
-								return (int)BuildResult;
-							}
+		/// <summary>
+		/// Build a list of targets
+		/// </summary>
+		/// <param name="TargetDescriptors">Target descriptors</param>
+		/// <param name="BuildConfiguration">Current build configuration</param>
+		/// <param name="WorkingSet">The source file working set</param>
+		/// <returns>Result from the compilation</returns>
+		public static ECompilationResult Build(List<TargetDescriptor> TargetDescriptors, BuildConfiguration BuildConfiguration, ISourceFileWorkingSet WorkingSet)
+		{
+			// Create a makefile for each target
+			TargetMakefile[] Makefiles = new TargetMakefile[TargetDescriptors.Count];
+			for(int TargetIdx = 0; TargetIdx < TargetDescriptors.Count; TargetIdx++)
+			{
+				TargetDescriptor TargetDesc = TargetDescriptors[TargetIdx];
 
-							// Create the makefile
-							Makefile = new TargetMakefile(Target.TargetType);
-							Makefile.ReceiptFile = Target.ReceiptFileName;
-							Makefile.Actions = Actions;
-							Makefile.OutputItems = OutputItems;
-							Makefile.ModuleNameToOutputItems = ModuleNameToOutputItems;
-							Makefile.UObjectModules = UObjectModules;
-							Makefile.Prerequisites = Prerequisites;
-							Makefile.HotReloadModuleNames = Target.GetHotReloadModuleNames();
-							Makefile.Prerequisites = Prerequisites;
-							Makefile.ProjectIntermediateDirectory = Target.ProjectIntermediateDirectory;
-							Makefile.bDeployAfterCompile = Target.bDeployAfterCompile;
-							Makefile.bHasProjectScriptPlugin = Target.bHasProjectScriptPlugin;
-							Makefile.PreBuildScripts = Target.PreBuildStepScripts;
-
-							// Save the environment variables
-							foreach (System.Collections.DictionaryEntry EnvironmentVariable in Environment.GetEnvironmentVariables())
-							{
-								Makefile.EnvironmentVariables.Add(Tuple.Create((string)EnvironmentVariable.Key, (string)EnvironmentVariable.Value));
-							}
+				// Get the path to the makefile for this target
+				FileReference MakefileLocation = TargetMakefile.GetLocation(TargetDesc.ProjectFile, TargetDesc.Name, TargetDesc.Platform, TargetDesc.Configuration);
 
-							// Save the makefile for next time
-							if(BuildConfiguration.bUseUBTMakefiles)
-							{
-								using(Timeline.ScopeEvent("TargetMakefile.Save()"))
-								{
-									Makefile.Save(MakefileLocation);
-								}
-							}
-						}
-						else
+				// Try to load an existing makefile
+				TargetMakefile Makefile = null;
+				if(BuildConfiguration.bUseUBTMakefiles)
+				{
+					using(Timeline.ScopeEvent("TargetMakefile.Load()"))
+					{
+						string ReasonNotLoaded;
+						Makefile = TargetMakefile.Load(MakefileLocation, TargetDesc.ProjectFile, TargetDesc.Platform, WorkingSet, out ReasonNotLoaded);
+						if (Makefile == null)
 						{
-							// Restore the environment variables
-							foreach (Tuple<string, string> EnvironmentVariable in Makefile.EnvironmentVariables)
-							{
-								Environment.SetEnvironmentVariable(EnvironmentVariable.Item1, EnvironmentVariable.Item2);
-							}
+							Log.TraceInformation("Creating makefile for {0} ({1})", TargetDesc.Name, ReasonNotLoaded);
+						}
+					}
+				}
 
-							// Execute all the pre-build steps
-							if (!BuildConfiguration.bXGEExport)
-							{
-								if (!Utils.ExecuteCustomBuildSteps(Makefile.PreBuildScripts))
-								{
-									return (int)ECompilationResult.OtherCompilationError;
-								}
-							}
+				// If we couldn't load a makefile, create a new one
+				if(Makefile == null)
+				{
+					// Create the target
+					UEBuildTarget Target;
+					using(Timeline.ScopeEvent("UEBuildTarget.Create()"))
+					{
+						Target = UEBuildTarget.Create(TargetDesc, BuildConfiguration.bSkipRulesCompile, BuildConfiguration.SingleFileToCompile != null, BuildConfiguration.bUsePrecompiled);
+					}
 
-							// If the target needs UHT to be run, we'll go ahead and do that now
-							if (Makefile.UObjectModules.Count > 0)
-							{
-								const bool bIsGatheringBuild = false;
-								const bool bIsAssemblingBuild = true;
+					// Build the target
+					const bool bIsAssemblingBuild = true;
 
-								FileReference ModuleInfoFileName = FileReference.Combine(Makefile.ProjectIntermediateDirectory, TargetDesc.Name + ".uhtmanifest");
-								ECompilationResult UHTResult = ExternalExecution.ExecuteHeaderToolIfNecessary(BuildConfiguration, TargetDesc.ProjectFile, TargetDesc.Name, Makefile.TargetType, Makefile.bHasProjectScriptPlugin, UObjectModules: Makefile.UObjectModules, ModuleInfoFileName: ModuleInfoFileName, bIsGatheringBuild: bIsGatheringBuild, bIsAssemblingBuild: bIsAssemblingBuild);
-								if(!UHTResult.Succeeded())
-								{
-									Log.TraceInformation("UnrealHeaderTool failed for target '" + TargetDesc.Name + "' (platform: " + TargetDesc.Platform.ToString() + ", module info: " + ModuleInfoFileName + ").");
-									return (int)UHTResult;
-								}
-							}
-						}
+					List<FileItem> OutputItems = new List<FileItem>();
+					Dictionary<string, FileItem[]> ModuleNameToOutputItems = new Dictionary<string, FileItem[]>();
+					List<UHTModuleInfo> UObjectModules = new List<UHTModuleInfo>();
+					List<Action> Actions = new List<Action>();
+					BuildPrerequisites Prerequisites = new BuildPrerequisites();
 
-						// Add the makefile to the list to build
-						Makefiles[TargetIdx] = Makefile;
+					ECompilationResult BuildResult;
+					using(Timeline.ScopeEvent("UEBuildTarget.Build()"))
+					{
+						BuildResult = Target.Build(BuildConfiguration, OutputItems, ModuleNameToOutputItems, UObjectModules, WorkingSet, Actions, Prerequisites, bIsAssemblingBuild);
+					}
+					if (BuildResult != ECompilationResult.Succeeded)
+					{
+						return BuildResult;
 					}
 
-					// Execute the build
-					if(!BuildConfiguration.bSkipBuild)
+					// Create the makefile
+					Makefile = new TargetMakefile(Target.TargetType);
+					Makefile.ReceiptFile = Target.ReceiptFileName;
+					Makefile.Actions = Actions;
+					Makefile.OutputItems = OutputItems;
+					Makefile.ModuleNameToOutputItems = ModuleNameToOutputItems;
+					Makefile.UObjectModules = UObjectModules;
+					Makefile.Prerequisites = Prerequisites;
+					Makefile.HotReloadModuleNames = Target.GetHotReloadModuleNames();
+					Makefile.Prerequisites = Prerequisites;
+					Makefile.ProjectIntermediateDirectory = Target.ProjectIntermediateDirectory;
+					Makefile.bDeployAfterCompile = Target.bDeployAfterCompile;
+					Makefile.bHasProjectScriptPlugin = Target.bHasProjectScriptPlugin;
+					Makefile.PreBuildScripts = Target.PreBuildStepScripts;
+
+					// Save the environment variables
+					foreach (System.Collections.DictionaryEntry EnvironmentVariable in Environment.GetEnvironmentVariables())
 					{
-						// Make sure that none of the actions conflict with any other (producing output files differently, etc...)
-						if(!ActionGraph.CheckForConflicts(Makefiles.SelectMany(x => x.Actions)))
-						{
-							Log.TraceInformation("Check log for additional details.");
-							return (int)ECompilationResult.OtherCompilationError;
-						}
+						Makefile.EnvironmentVariables.Add(Tuple.Create((string)EnvironmentVariable.Key, (string)EnvironmentVariable.Value));
+					}
 
-						// Find all the actions to be executed
-						HashSet<Action>[] ActionsToExecute = new HashSet<Action>[TargetDescriptors.Count];
-						for(int TargetIdx = 0; TargetIdx < TargetDescriptors.Count; TargetIdx++)
+					// Save the makefile for next time
+					if(BuildConfiguration.bUseUBTMakefiles)
+					{
+						using(Timeline.ScopeEvent("TargetMakefile.Save()"))
 						{
-							ActionsToExecute[TargetIdx] = GetActionsForTarget(BuildConfiguration, TargetDescriptors[TargetIdx], Makefiles[TargetIdx]);
+							Makefile.Save(MakefileLocation);
 						}
+					}
+				}
+				else
+				{
+					// Restore the environment variables
+					foreach (Tuple<string, string> EnvironmentVariable in Makefile.EnvironmentVariables)
+					{
+						Environment.SetEnvironmentVariable(EnvironmentVariable.Item1, EnvironmentVariable.Item2);
+					}
 
-						// If there are multiple targets being built, make sure the actions for each one 
-						List<Action> MergedActionsToExecute = new List<Action>();
-						if(TargetDescriptors.Count == 1)
-						{
-							// Just take the actions from the first target
-							MergedActionsToExecute.AddRange(ActionsToExecute[0]);
-						}
-						else
+					// Execute all the pre-build steps
+					if (!BuildConfiguration.bXGEExport)
+					{
+						if (!Utils.ExecuteCustomBuildSteps(Makefile.PreBuildScripts))
 						{
-							// Set of all output items. Knowing that there are no conflicts in produced items, we use this to eliminate duplicate actions.
-							Dictionary<FileItem, Action> OutputItemToProducingAction = new Dictionary<FileItem, Action>();
-							for(int TargetIdx = 0; TargetIdx < TargetDescriptors.Count; TargetIdx++)
-							{
-								string GroupPrefix = String.Format("{0}-{1}-{2}", TargetDescriptors[TargetIdx].Name, TargetDescriptors[TargetIdx].Platform, TargetDescriptors[TargetIdx].Configuration);
-								foreach(Action Action in ActionsToExecute[TargetIdx])
-								{
-									Action ExistingAction;
-									if(!OutputItemToProducingAction.TryGetValue(Action.ProducedItems[0], out ExistingAction))
-									{
-										OutputItemToProducingAction[Action.ProducedItems[0]] = Action;
-										ExistingAction = Action;
-									}
-									ExistingAction.GroupNames.Add(GroupPrefix);
-								}
-							}
-							MergedActionsToExecute.AddRange(OutputItemToProducingAction.Values);
+							return ECompilationResult.OtherCompilationError;
 						}
+					}
 
-						// Link all the actions together
-						ActionGraph.Link(MergedActionsToExecute);
+					// If the target needs UHT to be run, we'll go ahead and do that now
+					if (Makefile.UObjectModules.Count > 0)
+					{
+						const bool bIsGatheringBuild = false;
+						const bool bIsAssemblingBuild = true;
 
-						// Make sure the appropriate executor is selected
-						foreach(TargetDescriptor TargetDescriptor in TargetDescriptors)
+						FileReference ModuleInfoFileName = FileReference.Combine(Makefile.ProjectIntermediateDirectory, TargetDesc.Name + ".uhtmanifest");
+						ECompilationResult UHTResult = ExternalExecution.ExecuteHeaderToolIfNecessary(BuildConfiguration, TargetDesc.ProjectFile, TargetDesc.Name, Makefile.TargetType, Makefile.bHasProjectScriptPlugin, UObjectModules: Makefile.UObjectModules, ModuleInfoFileName: ModuleInfoFileName, bIsGatheringBuild: bIsGatheringBuild, bIsAssemblingBuild: bIsAssemblingBuild, WorkingSet: WorkingSet);
+						if(!UHTResult.Succeeded())
 						{
-							UEBuildPlatform BuildPlatform = UEBuildPlatform.GetBuildPlatform(TargetDescriptor.Platform);
-							BuildConfiguration.bAllowXGE &= BuildPlatform.CanUseXGE();
-							BuildConfiguration.bAllowDistcc &= BuildPlatform.CanUseDistcc();
-							BuildConfiguration.bAllowSNDBS &= BuildPlatform.CanUseSNDBS();
+							Log.TraceInformation("UnrealHeaderTool failed for target '" + TargetDesc.Name + "' (platform: " + TargetDesc.Platform.ToString() + ", module info: " + ModuleInfoFileName + ").");
+							return UHTResult;
 						}
+					}
+				}
 
-						// Delete produced items that are outdated.
-						ActionGraph.DeleteOutdatedProducedItems(MergedActionsToExecute);
+				// Add the makefile to the list to build
+				Makefiles[TargetIdx] = Makefile;
+			}
 
-						// Save all the action histories now that files have been removed. We have to do this after deleting produced items to ensure that any
-						// items created during the build don't have the wrong command line.
-						ActionHistory.SaveAll();
+			// Execute the build
+			if(!BuildConfiguration.bSkipBuild)
+			{
+				// Make sure that none of the actions conflict with any other (producing output files differently, etc...)
+				if(!ActionGraph.CheckForConflicts(Makefiles.SelectMany(x => x.Actions)))
+				{
+					Log.TraceInformation("Check log for additional details.");
+					return ECompilationResult.OtherCompilationError;
+				}
 
-						// Create directories for the outdated produced items.
-						ActionGraph.CreateDirectoriesForProducedItems(MergedActionsToExecute);
+				// Find all the actions to be executed
+				HashSet<Action>[] ActionsToExecute = new HashSet<Action>[TargetDescriptors.Count];
+				for(int TargetIdx = 0; TargetIdx < TargetDescriptors.Count; TargetIdx++)
+				{
+					ActionsToExecute[TargetIdx] = GetActionsForTarget(BuildConfiguration, TargetDescriptors[TargetIdx], Makefiles[TargetIdx]);
+				}
 
-						// Execute the actions.
-						ExecutorTimer.Start();
-						if (BuildConfiguration.bXGEExport)
-						{
-							using(Timeline.ScopeEvent("XGE.ExportActions()"))
-							{
-								XGE.ExportActions(MergedActionsToExecute);
-							}
-						}
-						else
+				// If there are multiple targets being built, make sure the actions for each one 
+				List<Action> MergedActionsToExecute = new List<Action>();
+				if(TargetDescriptors.Count == 1)
+				{
+					// Just take the actions from the first target
+					MergedActionsToExecute.AddRange(ActionsToExecute[0]);
+				}
+				else
+				{
+					// Set of all output items. Knowing that there are no conflicts in produced items, we use this to eliminate duplicate actions.
+					Dictionary<FileItem, Action> OutputItemToProducingAction = new Dictionary<FileItem, Action>();
+					for(int TargetIdx = 0; TargetIdx < TargetDescriptors.Count; TargetIdx++)
+					{
+						string GroupPrefix = String.Format("{0}-{1}-{2}", TargetDescriptors[TargetIdx].Name, TargetDescriptors[TargetIdx].Platform, TargetDescriptors[TargetIdx].Configuration);
+						foreach(Action Action in ActionsToExecute[TargetIdx])
 						{
-							using(Timeline.ScopeEvent("ActionGraph.ExecuteActions()"))
+							Action ExistingAction;
+							if(!OutputItemToProducingAction.TryGetValue(Action.ProducedItems[0], out ExistingAction))
 							{
-								if(!ActionGraph.ExecuteActions(BuildConfiguration, MergedActionsToExecute, out ExecutorName))
-								{
-									return (int)ECompilationResult.OtherCompilationError;
-								}
+								OutputItemToProducingAction[Action.ProducedItems[0]] = Action;
+								ExistingAction = Action;
 							}
+							ExistingAction.GroupNames.Add(GroupPrefix);
 						}
-						ExecutorTimer.Stop();
+					}
+					MergedActionsToExecute.AddRange(OutputItemToProducingAction.Values);
+				}
+
+				// Link all the actions together
+				ActionGraph.Link(MergedActionsToExecute);
+
+				// Make sure the appropriate executor is selected
+				foreach(TargetDescriptor TargetDescriptor in TargetDescriptors)
+				{
+					UEBuildPlatform BuildPlatform = UEBuildPlatform.GetBuildPlatform(TargetDescriptor.Platform);
+					BuildConfiguration.bAllowXGE &= BuildPlatform.CanUseXGE();
+					BuildConfiguration.bAllowDistcc &= BuildPlatform.CanUseDistcc();
+					BuildConfiguration.bAllowSNDBS &= BuildPlatform.CanUseSNDBS();
+				}
+
+				// Delete produced items that are outdated.
+				ActionGraph.DeleteOutdatedProducedItems(MergedActionsToExecute);
+
+				// Save all the action histories now that files have been removed. We have to do this after deleting produced items to ensure that any
+				// items created during the build don't have the wrong command line.
+				ActionHistory.SaveAll();
+
+				// Create directories for the outdated produced items.
+				ActionGraph.CreateDirectoriesForProducedItems(MergedActionsToExecute);
 
-						// Run the deployment steps
-						if (BuildConfiguration.SingleFileToCompile == null
-							&& !BuildConfiguration.bGenerateManifest
-							&& !BuildConfiguration.bXGEExport)
+				// Execute the actions.
+				if (BuildConfiguration.bXGEExport)
+				{
+					using(Timeline.ScopeEvent("XGE.ExportActions()"))
+					{
+						XGE.ExportActions(MergedActionsToExecute);
+					}
+				}
+				else
+				{
+					using(Timeline.ScopeEvent("ActionGraph.ExecuteActions()"))
+					{
+						if(!ActionGraph.ExecuteActions(BuildConfiguration, MergedActionsToExecute))
 						{
-							foreach(TargetMakefile Makefile in Makefiles)
-							{
-								if (Makefile.bDeployAfterCompile)
-								{
-									TargetReceipt Receipt = TargetReceipt.Read(Makefile.ReceiptFile);
-									Log.TraceInformation("Deploying {0} {1} {2}...", Receipt.TargetName, Receipt.Platform, Receipt.Configuration);
-									UEBuildPlatform.GetBuildPlatform(Receipt.Platform).Deploy(Receipt);
-								}
-							}
+							return ECompilationResult.OtherCompilationError;
 						}
 					}
 				}
-				return (int)ECompilationResult.Succeeded;
-			}
-			catch (Exception Ex)
-			{
-				Log.WriteException(Ex, (LogFile == null) ? null : LogFile.FullName);
-				return (int)ECompilationResult.OtherCompilationError;
-			}
-			finally
-			{
-				// Save all the caches
-				SourceFileMetadataCache.SaveAll();
-				CppDependencyCache.SaveAll();
 
-				// Figure out how long we took to execute.
-				if (ExecutorName != "Unknown")
+				// Run the deployment steps
+				if (BuildConfiguration.SingleFileToCompile == null
+					&& !BuildConfiguration.bGenerateManifest
+					&& !BuildConfiguration.bXGEExport)
 				{
-					Log.TraceInformation("Total build time: {0:0.00} seconds ({1} executor: {2:0.00} seconds)", BuildTimer.Elapsed.TotalSeconds, ExecutorName, ExecutorTimer.Elapsed.TotalSeconds);
+					foreach(TargetMakefile Makefile in Makefiles)
+					{
+						if (Makefile.bDeployAfterCompile)
+						{
+							TargetReceipt Receipt = TargetReceipt.Read(Makefile.ReceiptFile);
+							Log.TraceInformation("Deploying {0} {1} {2}...", Receipt.TargetName, Receipt.Platform, Receipt.Configuration);
+							UEBuildPlatform.GetBuildPlatform(Receipt.Platform).Deploy(Receipt);
+						}
+					}
 				}
 			}
+			return ECompilationResult.Succeeded;
 		}
 
 		/// <summary>
diff --git a/Engine/Source/Programs/UnrealBuildTool/System/ActionGraph.cs b/Engine/Source/Programs/UnrealBuildTool/System/ActionGraph.cs
--- a/Engine/Source/Programs/UnrealBuildTool/System/ActionGraph.cs
+++ b/Engine/Source/Programs/UnrealBuildTool/System/ActionGraph.cs
@@ -109,14 +109,15 @@ namespace UnrealBuildTool
 		/// <summary>
 		/// Executes a list of actions.
 		/// </summary>
-		public static bool ExecuteActions(BuildConfiguration BuildConfiguration, List<Action> ActionsToExecute, out string ExecutorName)
+		public static bool ExecuteActions(BuildConfiguration BuildConfiguration, List<Action> ActionsToExecute)
 		{
-			bool Result = true;
-			ExecutorName = "";
-			if (ActionsToExecute.Count > 0)
+			bool bResult = true;
+			if(ActionsToExecute.Count == 0)
+			{
+				Log.TraceInformation("Target is up to date");
+			}
+			else
 			{
-				DateTime StartTime = DateTime.UtcNow;
-
 				ActionExecutor Executor;
 				if ((XGE.IsAvailable() && BuildConfiguration.bAllowXGE) || BuildConfiguration.bXGEExport)
 				{
@@ -139,8 +140,9 @@ namespace UnrealBuildTool
 					Executor = new LocalExecutor();
 				}
 
-				ExecutorName = Executor.Name;
-				Result = Executor.ExecuteActions(ActionsToExecute, BuildConfiguration.bLogDetailedActionStats);
+				Stopwatch Timer = Stopwatch.StartNew();
+				bResult = Executor.ExecuteActions(ActionsToExecute, BuildConfiguration.bLogDetailedActionStats);
+				Log.TraceInformation("Total time in {0} executor: {1:0.00} seconds", Executor.Name, Timer.Elapsed.TotalSeconds);
 
 				if (!BuildConfiguration.bXGEExport)
 				{
@@ -151,8 +153,8 @@ namespace UnrealBuildTool
 						{
 							foreach (FileItem Item in BuildAction.ProducedItems)
 							{
-								bool bExists = File.Exists(Item.AbsolutePath) || Directory.Exists(Item.AbsolutePath);
-								if (!bExists)
+								Item.ResetFileInfo();
+								if(!Item.Exists)
 								{
 									throw new BuildException("UBT ERROR: Failed to produce item: " + Item.AbsolutePath);
 								}
@@ -161,13 +163,7 @@ namespace UnrealBuildTool
 					}
 				}
 			}
-			else
-			{
-				ExecutorName = "NoActionsToExecute";
-				Log.TraceInformation("Target is up to date");
-			}
-
-			return Result;
+			return bResult;
 		}
 
 		/// <summary>
diff --git a/Engine/Source/Programs/UnrealBuildTool/System/ExternalExecution.cs b/Engine/Source/Programs/UnrealBuildTool/System/ExternalExecution.cs
--- a/Engine/Source/Programs/UnrealBuildTool/System/ExternalExecution.cs
+++ b/Engine/Source/Programs/UnrealBuildTool/System/ExternalExecution.cs
@@ -1108,7 +1108,7 @@ namespace UnrealBuildTool
 		/// Builds and runs the header tool and touches the header directories.
 		/// Performs any early outs if headers need no changes, given the UObject modules, tool path, game name, and configuration
 		/// </summary>
-		public static ECompilationResult ExecuteHeaderToolIfNecessary(BuildConfiguration BuildConfiguration, FileReference ProjectFile, string TargetName, TargetType TargetType, bool bHasProjectScriptPlugin, List<UHTModuleInfo> UObjectModules, FileReference ModuleInfoFileName, bool bIsGatheringBuild, bool bIsAssemblingBuild)
+		public static ECompilationResult ExecuteHeaderToolIfNecessary(BuildConfiguration BuildConfiguration, FileReference ProjectFile, string TargetName, TargetType TargetType, bool bHasProjectScriptPlugin, List<UHTModuleInfo> UObjectModules, FileReference ModuleInfoFileName, bool bIsGatheringBuild, bool bIsAssemblingBuild, ISourceFileWorkingSet WorkingSet)
 		{
 			if (ProgressWriter.bWriteMarkup)
 			{
@@ -1161,62 +1161,45 @@ namespace UnrealBuildTool
 						// If it is there and up to date, it will add 0.8 seconds to the build time.
 						Log.TraceInformation("Building UnrealHeaderTool...");
 
-						StringBuilder UBTArguments = new StringBuilder();
-
-						UBTArguments.Append("UnrealHeaderTool");
-
 						// Which desktop platform do we need to compile UHT for?
-						UBTArguments.Append(" " + BuildHostPlatform.Current.Platform.ToString());
+						UnrealTargetPlatform Platform = BuildHostPlatform.Current.Platform;
 
 						// NOTE: We force Development configuration for UHT so that it runs quickly, even when compiling debug, unless we say so explicitly
+						UnrealTargetConfiguration Configuration;
 						if (BuildConfiguration.bForceDebugUnrealHeaderTool)
 						{
-							UBTArguments.Append(" " + UnrealTargetConfiguration.Debug.ToString());
+							Configuration = UnrealTargetConfiguration.Debug;
 						}
 						else
 						{
-							UBTArguments.Append(" " + UnrealTargetConfiguration.Development.ToString());
-						}
-
-						// NOTE: We disable mutex when launching UBT from within UBT to compile UHT
-						UBTArguments.Append(" -NoMutex");
-
-						if (!BuildConfiguration.bAllowXGE)
-						{
-							UBTArguments.Append(" -noxge");
+							Configuration = UnrealTargetConfiguration.Development;
 						}
 
-						// Always ignore the junk manifest on recursive invocations; it will have been run by this process if necessary
-						UBTArguments.Append(" -ignorejunk");
+						// Get the default architecture
+						string Architecture = UEBuildPlatform.GetBuildPlatform(Platform).GetDefaultArchitecture(null);
 
 						// Add UHT plugins to UBT command line as external plugins
+						FileReference ScriptProjectFile = null;
 						if(bHasProjectScriptPlugin && ProjectFile != null)
 						{
-							UBTArguments.AppendFormat(" -project=\"{0}\"", ProjectFile);
-						}
-
-						// Add any global override for the compiler
-						if(!String.IsNullOrEmpty(BuildConfiguration.CompilerArgumentForUnrealHeaderTool))
-						{
-							UBTArguments.AppendFormat(" {0}", BuildConfiguration.CompilerArgumentForUnrealHeaderTool);
+							ScriptProjectFile = ProjectFile;
 						}
 
-						// Output the log next to the current log
-						if(Environment.CommandLine.IndexOf("-NoLog", StringComparison.OrdinalIgnoreCase) != -1)
-						{
-							UBTArguments.Append(" -nolog");
-						}
-						else
-						{
-							UBTArguments.Append(" -logsuffix=UHT");
-						}
+						// Create the target descriptor
+						List<TargetDescriptor> TargetDescriptors = new List<TargetDescriptor>();
+						TargetDescriptors.Add(new TargetDescriptor(ScriptProjectFile, "UnrealHeaderTool", Platform, Configuration, Architecture, null, null));
 
-						int ExitCode;
+						ECompilationResult Result;
 						using(Timeline.ScopeEvent("Buildng UnrealHeaderTool"))
 						{
-							ExitCode = RunExternalDotNETExecutable(UnrealBuildTool.GetUBTPath(), UBTArguments.ToString());
+							bool bPrevXGEExport = BuildConfiguration.bXGEExport;
+							BuildConfiguration.bXGEExport = false;
+
+							Result = BuildMode.Build(TargetDescriptors, BuildConfiguration, WorkingSet);
+
+							BuildConfiguration.bXGEExport = bPrevXGEExport;
 						}
-						if(ExitCode != 0)
+						if(Result != ECompilationResult.Succeeded)
 						{
 							return ECompilationResult.OtherCompilationError;
 						}
