UBT: Make the -IgnoreJunk argument apply only to the build mode, since it's not used during the actual build.

#rnx

diff --git a/Engine/Source/Programs/UnrealBuildTool/Configuration/BuildConfiguration.cs b/Engine/Source/Programs/UnrealBuildTool/Configuration/BuildConfiguration.cs
--- a/Engine/Source/Programs/UnrealBuildTool/Configuration/BuildConfiguration.cs
+++ b/Engine/Source/Programs/UnrealBuildTool/Configuration/BuildConfiguration.cs
@@ -9,6 +9,9 @@ using Tools.DotNETCommon;
 
 namespace UnrealBuildTool
 {
+	/// <summary>
+	/// Global settings for building. Should not contain any target-specific settings.
+	/// </summary>
 	class BuildConfiguration
 	{
 		/// <summary>
@@ -107,13 +110,6 @@ namespace UnrealBuildTool
 		[CommandLine("-SingleFile=")]
 		public FileReference SingleFileToCompile = null;
 
-		/// <summary>
-		/// Whether to skip checking for files identified by the junk manifest
-		/// </summary>
-		[XmlConfigFile]
-		[CommandLine("-IgnoreJunk")]
-		public bool bIgnoreJunk = false;
-
 		/// <summary>
 		/// Whether to generate a manifest file that contains the files to add to Perforce
 		/// </summary>
diff --git a/Engine/Source/Programs/UnrealBuildTool/Modes/BuildMode.cs b/Engine/Source/Programs/UnrealBuildTool/Modes/BuildMode.cs
--- a/Engine/Source/Programs/UnrealBuildTool/Modes/BuildMode.cs
+++ b/Engine/Source/Programs/UnrealBuildTool/Modes/BuildMode.cs
@@ -24,6 +24,13 @@ namespace UnrealBuildTool
 		[XmlConfigFile(Category = "BuildConfiguration")]
 		public string BaseLogFileName = "../Programs/UnrealBuildTool/Log.txt";
 
+		/// <summary>
+		/// Whether to skip checking for files identified by the junk manifest
+		/// </summary>
+		[XmlConfigFile]
+		[CommandLine("-IgnoreJunk")]
+		public bool bIgnoreJunk = false;
+
 		/// <summary>
 		/// Main entry point
 		/// </summary>
@@ -86,9 +93,12 @@ namespace UnrealBuildTool
 			UnrealBuildTool.SetRemoteIniPath(RemoteIniPath);
 
 			// now that we know the available platforms, we can delete other platforms' junk. if we're only building specific modules from the editor, don't touch anything else (it may be in use).
-			if (!BuildConfiguration.bIgnoreJunk)
+			if (!bIgnoreJunk && !UnrealBuildTool.IsEngineInstalled())
 			{
-				JunkDeleter.DeleteJunk();
+				using(Timeline.ScopeEvent("DeleteJunk()"))
+				{
+					JunkDeleter.DeleteJunk();
+				}
 			}
 
 			Stopwatch BuildTimer = Stopwatch.StartNew();
diff --git a/Engine/Source/Programs/UnrealBuildTool/System/JunkDeleter.cs b/Engine/Source/Programs/UnrealBuildTool/System/JunkDeleter.cs
--- a/Engine/Source/Programs/UnrealBuildTool/System/JunkDeleter.cs
+++ b/Engine/Source/Programs/UnrealBuildTool/System/JunkDeleter.cs
@@ -16,14 +16,8 @@ namespace UnrealBuildTool
 		/// </summary>
 		static public void DeleteJunk()
 		{
-			if (UnrealBuildTool.IsEngineInstalled() == false)
-			{
-				using(Timeline.ScopeEvent("DeleteAllJunk()"))
-				{
-					List<string> JunkManifest = LoadJunkManifest();
-					DeleteAllJunk(JunkManifest);
-				}
-			}
+			List<string> JunkManifest = LoadJunkManifest();
+			DeleteAllJunk(JunkManifest);
 		}
 
 		/// <summary>
