UBT: Convert the -SingleFile= argument to be per-target.

#rnx

diff --git a/Engine/Source/Programs/UnrealBuildTool/Configuration/BuildConfiguration.cs b/Engine/Source/Programs/UnrealBuildTool/Configuration/BuildConfiguration.cs
--- a/Engine/Source/Programs/UnrealBuildTool/Configuration/BuildConfiguration.cs
+++ b/Engine/Source/Programs/UnrealBuildTool/Configuration/BuildConfiguration.cs
@@ -104,12 +104,6 @@ namespace UnrealBuildTool
 		[XmlConfigFile]
 		public bool bAllowParallelExecutor = true;
 
-		/// <summary>
-		/// If specified, we will only build this particular source file, ignore all other outputs.  Useful for testing non-Unity builds.
-		/// </summary>
-		[CommandLine("-SingleFile=")]
-		public FileReference SingleFileToCompile = null;
-
 		/// <summary>
 		/// Whether to generate a manifest file that contains the files to add to Perforce
 		/// </summary>
diff --git a/Engine/Source/Programs/UnrealBuildTool/Configuration/TargetDescriptor.cs b/Engine/Source/Programs/UnrealBuildTool/Configuration/TargetDescriptor.cs
--- a/Engine/Source/Programs/UnrealBuildTool/Configuration/TargetDescriptor.cs
+++ b/Engine/Source/Programs/UnrealBuildTool/Configuration/TargetDescriptor.cs
@@ -19,9 +19,20 @@ namespace UnrealBuildTool
 		public UnrealTargetPlatform Platform;
 		public UnrealTargetConfiguration Configuration;
 		public string Architecture;
-		public FileReference ForeignPlugin;
 		public CommandLineArguments AdditionalArguments;
 
+		/// <summary>
+		/// Foreign plugin to compile against this target
+		/// </summary>
+		[CommandLine("-Plugin=")]
+		public FileReference ForeignPlugin = null;
+
+		/// <summary>
+		/// Single file to compile
+		/// </summary>
+		[CommandLine("-SingleFile=")]
+		public FileReference SingleFileToCompile = null;
+
 		/// <summary>
 		/// Constructor
 		/// </summary>
@@ -30,17 +41,32 @@ namespace UnrealBuildTool
 		/// <param name="Platform">Platform to build for</param>
 		/// <param name="Configuration">Configuration to build</param>
 		/// <param name="Architecture">Architecture to build for</param>
-		/// <param name="ForeignPlugin">Path to a plugin to be compiled against this target. Usually null.</param>
-		/// <param name="AdditionalArguments">Additional arguments for the target</param>
-		public TargetDescriptor(FileReference ProjectFile, string TargetName, UnrealTargetPlatform Platform, UnrealTargetConfiguration Configuration, string Architecture, FileReference ForeignPlugin, CommandLineArguments AdditionalArguments)
+		/// <param name="Arguments">Other command-line arguments for the target</param>
+		public TargetDescriptor(FileReference ProjectFile, string TargetName, UnrealTargetPlatform Platform, UnrealTargetConfiguration Configuration, string Architecture, CommandLineArguments Arguments)
 		{
 			this.ProjectFile = ProjectFile;
 			this.Name = TargetName;
 			this.Platform = Platform;
 			this.Configuration = Configuration;
 			this.Architecture = Architecture;
-			this.ForeignPlugin = ForeignPlugin;
-			this.AdditionalArguments = AdditionalArguments ?? new CommandLineArguments(new string[0]);
+
+			// If there are any additional command line arguments
+			List<string> AdditionalArguments = new List<string>();
+			if(Arguments != null)
+			{
+				// Apply the arguments to this object
+				Arguments.ApplyTo(this);
+
+				// Pull out all the arguments that haven't been used so far
+				for(int Idx = 0; Idx < Arguments.Count; Idx++)
+				{
+					if(!Arguments.HasBeenUsed(Idx))
+					{
+						AdditionalArguments.Add(Arguments[Idx]);
+					}
+				}
+			}
+			this.AdditionalArguments = new CommandLineArguments(AdditionalArguments.ToArray());
 		}
 
 		/// <summary>
@@ -188,8 +214,7 @@ namespace UnrealBuildTool
 				throw new BuildException("No configurations specified for target");
 			}
 
-			FileReference ForeignPlugin = Arguments.GetFileReferenceOrDefault("-Plugin=", null);
-
+			// Expand all the platforms, architectures and configurations
 			foreach(UnrealTargetPlatform Platform in Platforms)
 			{
 				// Parse the architecture parameter, or get the default for the platform
@@ -237,18 +262,8 @@ namespace UnrealBuildTool
 								Log.TraceVerbose("Found project file for {0} - {1}", TargetName, ProjectFile);
 							}
 
-							// Pull out all the arguments that haven't been used so far
-							List<string> AdditionalArguments = new List<string>();
-							for(int Idx = 0; Idx < Arguments.Count; Idx++)
-							{
-								if(!Arguments.HasBeenUsed(Idx))
-								{
-									AdditionalArguments.Add(Arguments[Idx]);
-								}
-							}
-
 							// Create the target descriptor
-							TargetDescriptors.Add(new TargetDescriptor(ProjectFile, TargetName, Platform, Configuration, Architecture, ForeignPlugin, new CommandLineArguments(AdditionalArguments.ToArray())));
+							TargetDescriptors.Add(new TargetDescriptor(ProjectFile, TargetName, Platform, Configuration, Architecture, Arguments));
 						}
 					}
 				}
diff --git a/Engine/Source/Programs/UnrealBuildTool/Configuration/UEBuildTarget.cs b/Engine/Source/Programs/UnrealBuildTool/Configuration/UEBuildTarget.cs
--- a/Engine/Source/Programs/UnrealBuildTool/Configuration/UEBuildTarget.cs
+++ b/Engine/Source/Programs/UnrealBuildTool/Configuration/UEBuildTarget.cs
@@ -263,10 +263,9 @@ namespace UnrealBuildTool
 		/// </summary>
 		/// <param name="Descriptor">Information about the target</param>
 		/// <param name="bSkipRulesCompile">Whether to skip compiling any rules assemblies</param>
-		/// <param name="bCompilingSingleFile">Whether we're compiling a single file</param>
 		/// <param name="bUsePrecompiled">Whether to use a precompiled engine/enterprise build</param>
 		/// <returns>The build target object for the specified build rules source file</returns>
-		public static UEBuildTarget Create(TargetDescriptor Descriptor, bool bSkipRulesCompile, bool bCompilingSingleFile, bool bUsePrecompiled)
+		public static UEBuildTarget Create(TargetDescriptor Descriptor, bool bSkipRulesCompile, bool bUsePrecompiled)
 		{
 			RulesAssembly RulesAssembly;
 			using(Timeline.ScopeEvent("RulesCompiler.CreateTargetRulesAssembly()"))
@@ -317,7 +316,7 @@ namespace UnrealBuildTool
 			}
 
 			// If we're compiling just a single file, we need to prevent unity builds from running
-			if(bCompilingSingleFile)
+			if(Descriptor.SingleFileToCompile != null)
 			{
 				RulesObject.bUseUnityBuild = false;
 				RulesObject.bForceUnityBuild = false;
@@ -352,6 +351,18 @@ namespace UnrealBuildTool
 			{
 				Target.PreBuildSetup();
 			}
+
+			// If we're just compiling a single file, filter the list of binaries to only include the file we're interested in.
+			FileReference SingleFileToCompile = Descriptor.SingleFileToCompile;
+			if (SingleFileToCompile != null)
+			{
+				Target.Binaries.RemoveAll(x => !x.Modules.Any(y => SingleFileToCompile.IsUnderDirectory(y.ModuleDirectory)));
+				if(Target.Binaries.Count == 0)
+				{
+					throw new BuildException("Couldn't find any module containing {0} in {1}.", SingleFileToCompile, Target.TargetName);
+				}
+			}
+
 			return Target;
 		}
 
@@ -626,7 +637,7 @@ namespace UnrealBuildTool
 			RulesAssembly = InRulesAssembly;
 			TargetType = Rules.Type;
 			ForeignPlugin = InDescriptor.ForeignPlugin;
-			bDeployAfterCompile = InRules.bDeployAfterCompile && !InRules.bDisableLinking;
+			bDeployAfterCompile = InRules.bDeployAfterCompile && !InRules.bDisableLinking && InDescriptor.SingleFileToCompile == null;
 
 			// now that we have the platform, we can set the intermediate path to include the platform/architecture name
 			PlatformIntermediateFolder = GetPlatformIntermediateFolder(Platform, Architecture);
@@ -1141,16 +1152,6 @@ namespace UnrealBuildTool
 				}
 			}
 
-			// If we're just compiling a single file, filter the list of binaries to only include the file we're interested in.
-			if (BuildConfiguration.SingleFileToCompile != null)
-			{
-				Binaries.RemoveAll(x => !x.Modules.Any(y => BuildConfiguration.SingleFileToCompile.IsUnderDirectory(y.ModuleDirectory)));
-				if(Binaries.Count == 0)
-				{
-					throw new BuildException("Couldn't find any module containing {0} in {1}.", BuildConfiguration.SingleFileToCompile, TargetName);
-				}
-			}
-
 			// Build a mapping from module to its plugin
 			Dictionary<UEBuildModule, UEBuildPlugin> ModuleToPlugin = new Dictionary<UEBuildModule, UEBuildPlugin>();
 			foreach(UEBuildPlugin Plugin in BuildPlugins)
diff --git a/Engine/Source/Programs/UnrealBuildTool/Modes/BuildMode.cs b/Engine/Source/Programs/UnrealBuildTool/Modes/BuildMode.cs
--- a/Engine/Source/Programs/UnrealBuildTool/Modes/BuildMode.cs
+++ b/Engine/Source/Programs/UnrealBuildTool/Modes/BuildMode.cs
@@ -82,10 +82,6 @@ namespace UnrealBuildTool
 			{
 				BuildConfiguration.bAllowXGE = true;
 			}
-			if(BuildConfiguration.SingleFileToCompile != null)
-			{
-				BuildConfiguration.bUseUBTMakefiles = false;
-			}
 
 			// Parse the remote INI setting
 			string RemoteIniPath;
@@ -199,11 +195,15 @@ namespace UnrealBuildTool
 				TargetDescriptor TargetDesc = TargetDescriptors[TargetIdx];
 
 				// Get the path to the makefile for this target
-				FileReference MakefileLocation = TargetMakefile.GetLocation(TargetDesc.ProjectFile, TargetDesc.Name, TargetDesc.Platform, TargetDesc.Configuration);
+				FileReference MakefileLocation = null;
+				if(BuildConfiguration.bUseUBTMakefiles && TargetDesc.SingleFileToCompile == null)
+				{
+					MakefileLocation = TargetMakefile.GetLocation(TargetDesc.ProjectFile, TargetDesc.Name, TargetDesc.Platform, TargetDesc.Configuration);
+				}
 
 				// Try to load an existing makefile
 				TargetMakefile Makefile = null;
-				if(BuildConfiguration.bUseUBTMakefiles)
+				if(MakefileLocation != null)
 				{
 					using(Timeline.ScopeEvent("TargetMakefile.Load()"))
 					{
@@ -223,7 +223,7 @@ namespace UnrealBuildTool
 					UEBuildTarget Target;
 					using(Timeline.ScopeEvent("UEBuildTarget.Create()"))
 					{
-						Target = UEBuildTarget.Create(TargetDesc, BuildConfiguration.bSkipRulesCompile, BuildConfiguration.SingleFileToCompile != null, BuildConfiguration.bUsePrecompiled);
+						Target = UEBuildTarget.Create(TargetDesc, BuildConfiguration.bSkipRulesCompile, BuildConfiguration.bUsePrecompiled);
 					}
 
 					// Build the target
@@ -267,7 +267,7 @@ namespace UnrealBuildTool
 					}
 
 					// Save the makefile for next time
-					if(BuildConfiguration.bUseUBTMakefiles)
+					if(MakefileLocation != null)
 					{
 						using(Timeline.ScopeEvent("TargetMakefile.Save()"))
 						{
@@ -399,8 +399,7 @@ namespace UnrealBuildTool
 				}
 
 				// Run the deployment steps
-				if (BuildConfiguration.SingleFileToCompile == null
-					&& !BuildConfiguration.bGenerateManifest
+				if (!BuildConfiguration.bGenerateManifest
 					&& !BuildConfiguration.bXGEExport)
 				{
 					foreach(TargetMakefile Makefile in Makefiles)
@@ -488,9 +487,9 @@ namespace UnrealBuildTool
 			}
 
 			// If we're just compiling a single file, set the target items to be all the derived items
-			if(BuildConfiguration.SingleFileToCompile != null)
+			if(TargetDescriptor.SingleFileToCompile != null)
 			{
-				FileItem FileToCompile = FileItem.GetItemByFileReference(BuildConfiguration.SingleFileToCompile);
+				FileItem FileToCompile = FileItem.GetItemByFileReference(TargetDescriptor.SingleFileToCompile);
 				TargetOutputItems = Makefile.Actions.Where(x => x.PrerequisiteItems.Contains(FileToCompile)).SelectMany(x => x.ProducedItems).ToList();
 			}
 
@@ -577,7 +576,7 @@ namespace UnrealBuildTool
 
 			// Plan the actions to execute for the build.
 			HashSet<Action> TargetActionsToExecute;
-			if (BuildConfiguration.SingleFileToCompile == null)
+			if (TargetDescriptor.SingleFileToCompile == null)
 			{
 				TargetActionsToExecute = ActionGraph.GetActionsToExecute(Makefile.Actions, PrerequisiteActions, CppDependencies, History, BuildConfiguration.bIgnoreOutdatedImportLibraries);
 			}
diff --git a/Engine/Source/Programs/UnrealBuildTool/Modes/CleanMode.cs b/Engine/Source/Programs/UnrealBuildTool/Modes/CleanMode.cs
--- a/Engine/Source/Programs/UnrealBuildTool/Modes/CleanMode.cs
+++ b/Engine/Source/Programs/UnrealBuildTool/Modes/CleanMode.cs
@@ -74,7 +74,7 @@ namespace UnrealBuildTool
 					string Architecture = UEBuildPlatform.GetBuildPlatform(BuildHostPlatform.Current.Platform).GetDefaultArchitecture(null);
 					foreach(FileReference ProjectFile in ProjectFiles)
 					{
-						TargetDescriptors.Add(new TargetDescriptor(ProjectFile, UnrealHeaderToolTarget, BuildHostPlatform.Current.Platform, Configuration, Architecture, null, null));
+						TargetDescriptors.Add(new TargetDescriptor(ProjectFile, UnrealHeaderToolTarget, BuildHostPlatform.Current.Platform, Configuration, Architecture, null));
 					}
 				}
 			}
diff --git a/Engine/Source/Programs/UnrealBuildTool/Modes/JsonExportMode.cs b/Engine/Source/Programs/UnrealBuildTool/Modes/JsonExportMode.cs
--- a/Engine/Source/Programs/UnrealBuildTool/Modes/JsonExportMode.cs
+++ b/Engine/Source/Programs/UnrealBuildTool/Modes/JsonExportMode.cs
@@ -26,7 +26,7 @@ namespace UnrealBuildTool
 			foreach(TargetDescriptor TargetDescriptor in TargetDescriptors)
 			{
 				// Create the target
-				UEBuildTarget Target = UEBuildTarget.Create(TargetDescriptor, false, false, false);
+				UEBuildTarget Target = UEBuildTarget.Create(TargetDescriptor, false, false);
 
 				// Get the output file
 				FileReference OutputFile = TargetDescriptor.AdditionalArguments.GetFileReferenceOrDefault("-OutputFile=", null);
diff --git a/Engine/Source/Programs/UnrealBuildTool/ProjectFiles/ProjectFileGenerator.cs b/Engine/Source/Programs/UnrealBuildTool/ProjectFiles/ProjectFileGenerator.cs
--- a/Engine/Source/Programs/UnrealBuildTool/ProjectFiles/ProjectFileGenerator.cs
+++ b/Engine/Source/Programs/UnrealBuildTool/ProjectFiles/ProjectFileGenerator.cs
@@ -1705,10 +1705,10 @@ namespace UnrealBuildTool
 							const string DefaultArchitecture = "";
 
 							// Create the target descriptor
-							TargetDescriptor TargetDesc = new TargetDescriptor(CurTarget.UnrealProjectFilePath, CurTarget.Name, BuildHostPlatform.Current.Platform, UnrealTargetConfiguration.Development, DefaultArchitecture, null, new CommandLineArguments(NewArguments.ToArray()));
+							TargetDescriptor TargetDesc = new TargetDescriptor(CurTarget.UnrealProjectFilePath, CurTarget.Name, BuildHostPlatform.Current.Platform, UnrealTargetConfiguration.Development, DefaultArchitecture, new CommandLineArguments(NewArguments.ToArray()));
 
 							// Create the target
-							UEBuildTarget Target = UEBuildTarget.Create(TargetDesc, false, false, bUsePrecompiled);
+							UEBuildTarget Target = UEBuildTarget.Create(TargetDesc, false, bUsePrecompiled);
 
 							// Generate a compile environment for each module in the binary
 							CppCompileEnvironment GlobalCompileEnvironment = Target.CreateCompileEnvironmentForProjectFiles();
diff --git a/Engine/Source/Programs/UnrealBuildTool/System/ExternalExecution.cs b/Engine/Source/Programs/UnrealBuildTool/System/ExternalExecution.cs
--- a/Engine/Source/Programs/UnrealBuildTool/System/ExternalExecution.cs
+++ b/Engine/Source/Programs/UnrealBuildTool/System/ExternalExecution.cs
@@ -1187,7 +1187,7 @@ namespace UnrealBuildTool
 
 						// Create the target descriptor
 						List<TargetDescriptor> TargetDescriptors = new List<TargetDescriptor>();
-						TargetDescriptors.Add(new TargetDescriptor(ScriptProjectFile, "UnrealHeaderTool", Platform, Configuration, Architecture, null, null));
+						TargetDescriptors.Add(new TargetDescriptor(ScriptProjectFile, "UnrealHeaderTool", Platform, Configuration, Architecture, null));
 
 						ECompilationResult Result;
 						using(Timeline.ScopeEvent("Buildng UnrealHeaderTool"))
