UBT: Remove the BuildConfiguration constructor. XML config is already applied in every place that an instance is constructed.

#rnx

diff --git a/Engine/Source/Programs/UnrealBuildTool/Configuration/BuildConfiguration.cs b/Engine/Source/Programs/UnrealBuildTool/Configuration/BuildConfiguration.cs
--- a/Engine/Source/Programs/UnrealBuildTool/Configuration/BuildConfiguration.cs
+++ b/Engine/Source/Programs/UnrealBuildTool/Configuration/BuildConfiguration.cs
@@ -149,13 +149,5 @@ namespace UnrealBuildTool
 		/// </summary>
 		[CommandLine("-SkipRulesCompile")]
 		public bool bSkipRulesCompile = false;
-
-		/// <summary>
-		/// Default constructor. Reads settings from the XmlConfig files.
-		/// </summary>
-		public BuildConfiguration()
-		{
-			XmlConfig.ApplyTo(this);
-		}
 	}
 }
