UBT: Fix incorrect timezone in makefile out-of-date checks.

#rnx

diff --git a/Engine/Source/Programs/UnrealBuildTool/System/TargetMakefile.cs b/Engine/Source/Programs/UnrealBuildTool/System/TargetMakefile.cs
--- a/Engine/Source/Programs/UnrealBuildTool/System/TargetMakefile.cs
+++ b/Engine/Source/Programs/UnrealBuildTool/System/TargetMakefile.cs
@@ -341,7 +341,7 @@ namespace UnrealBuildTool
 						ReasonNotLoaded = string.Format("{0} deleted", AdditionalDependency.Location.GetFileName());
 						return null;
 					}
-					if(AdditionalDependency.LastWriteTimeUtc > UBTMakefileInfo.LastWriteTime)
+					if(AdditionalDependency.LastWriteTimeUtc > UBTMakefileInfo.LastWriteTimeUtc)
 					{
 						Log.TraceLog("{0} has been modified since makefile was built.", AdditionalDependency.Location);
 						ReasonNotLoaded = string.Format("{0} modified", AdditionalDependency.Location.GetFileName());
