Minor cleanup.

#rnx

diff --git a/Engine/Source/Programs/UnrealBuildTool/Modes/BuildMode.cs b/Engine/Source/Programs/UnrealBuildTool/Modes/BuildMode.cs
--- a/Engine/Source/Programs/UnrealBuildTool/Modes/BuildMode.cs
+++ b/Engine/Source/Programs/UnrealBuildTool/Modes/BuildMode.cs
@@ -38,6 +38,7 @@ namespace UnrealBuildTool
 		/// <returns>One of the values of ECompilationResult</returns>
 		public override int Execute(CommandLineArguments Arguments)
 		{
+			Stopwatch BuildTimer = Stopwatch.StartNew();
 			Arguments.ApplyTo(this);
 
 			// Initialize the log system, buffering the output until we can create the log file
@@ -97,9 +98,10 @@ namespace UnrealBuildTool
 				}
 			}
 
-			Stopwatch BuildTimer = Stopwatch.StartNew();
+			// Parse and build the targets
 			try
 			{
+				// Parse all the target descriptors
 				List<TargetDescriptor> TargetDescriptors = new List<TargetDescriptor>();
 				using(Timeline.ScopeEvent("TargetDescriptor.ParseCommandLine()"))
 				{
@@ -125,43 +127,32 @@ namespace UnrealBuildTool
 						}
 
 						TargetDescriptors.RemoveAt(Idx--);
-						if(TargetDescriptors.Count == 0)
-						{
-							return (int)ECompilationResult.Succeeded;
-						}
 					}
 				}
 
-				if (Arguments.Any(x => x.Equals("-InvalidateMakefilesOnly", StringComparison.InvariantCultureIgnoreCase)))
+				// Handle local builds
+				if(TargetDescriptors.Count > 0)
 				{
-					Log.TraceInformation("Invalidating makefiles only in this run.");
+					// Get a set of all the project directories.
+					HashSet<DirectoryReference> ProjectDirs = new HashSet<DirectoryReference>();
 					foreach(TargetDescriptor TargetDesc in TargetDescriptors)
 					{
-						FileReference MakefileLocation = TargetMakefile.GetLocation(TargetDesc.ProjectFile, TargetDesc.Name, TargetDesc.Platform, TargetDesc.Configuration);
-						if(FileReference.Exists(MakefileLocation))
-						{ 
-							FileReference.Delete(MakefileLocation);
+						if(TargetDesc.ProjectFile != null)
+						{
+							ProjectDirs.Add(TargetDesc.ProjectFile.Directory);
 						}
 					}
-					return (int)ECompilationResult.Succeeded;
-				}
 
-				// Get a set of all the project directories.
-				HashSet<DirectoryReference> ProjectDirs = new HashSet<DirectoryReference>();
-				foreach(TargetDescriptor TargetDesc in TargetDescriptors)
-				{
-					if(TargetDesc.ProjectFile != null)
+					// Create the working set provider
+					using (ISourceFileWorkingSet WorkingSet = SourceFileWorkingSet.Create(UnrealBuildTool.RootDirectory, ProjectDirs))
 					{
-						ProjectDirs.Add(TargetDesc.ProjectFile.Directory);
+						ECompilationResult Result = Build(TargetDescriptors, BuildConfiguration, WorkingSet);
+						if(Result != ECompilationResult.Succeeded)
+						{
+							return (int)Result;
+						}
 					}
 				}
-
-				// Create the working set provider
-				using (ISourceFileWorkingSet WorkingSet = SourceFileWorkingSet.Create(UnrealBuildTool.RootDirectory, ProjectDirs))
-				{
-					ECompilationResult Result = Build(TargetDescriptors, BuildConfiguration, WorkingSet);
-					return (int)Result;
-				}
 			}
 			catch (Exception Ex)
 			{
@@ -173,10 +164,11 @@ namespace UnrealBuildTool
 				// Save all the caches
 				SourceFileMetadataCache.SaveAll();
 				CppDependencyCache.SaveAll();
-
-				// Figure out how long we took to execute.
-				Log.TraceInformation("Total build time: {0:0.00} seconds", BuildTimer.Elapsed.TotalSeconds);
 			}
+
+			// Figure out how long we took to execute.
+			Log.TraceInformation("Total build time: {0:0.00} seconds", BuildTimer.Elapsed.TotalSeconds);
+			return 0;
 		}
 
 		/// <summary>
@@ -399,8 +391,7 @@ namespace UnrealBuildTool
 				}
 
 				// Run the deployment steps
-				if (!BuildConfiguration.bGenerateManifest
-					&& !BuildConfiguration.bXGEExport)
+				if (!BuildConfiguration.bGenerateManifest && !BuildConfiguration.bXGEExport)
 				{
 					foreach(TargetMakefile Makefile in Makefiles)
 					{
