UBT: Remove the -GenerateManifest argument. The -Manifest argument already implements the same functionality with an explicit output file, and nothing uses -GenerateManifest any more.

#rnx

diff --git a/Engine/Source/Programs/UnrealBuildTool/Configuration/BuildConfiguration.cs b/Engine/Source/Programs/UnrealBuildTool/Configuration/BuildConfiguration.cs
--- a/Engine/Source/Programs/UnrealBuildTool/Configuration/BuildConfiguration.cs
+++ b/Engine/Source/Programs/UnrealBuildTool/Configuration/BuildConfiguration.cs
@@ -104,12 +104,6 @@ namespace UnrealBuildTool
 		[XmlConfigFile]
 		public bool bAllowParallelExecutor = true;
 
-		/// <summary>
-		/// Whether to generate a manifest file that contains the files to add to Perforce
-		/// </summary>
-		[CommandLine("-GenerateManifest")]
-		public bool bGenerateManifest = false;
-
 		/// <summary>
 		/// If true, force header regeneration. Intended for the build machine
 		/// </summary>
diff --git a/Engine/Source/Programs/UnrealBuildTool/Configuration/UEBuildTarget.cs b/Engine/Source/Programs/UnrealBuildTool/Configuration/UEBuildTarget.cs
--- a/Engine/Source/Programs/UnrealBuildTool/Configuration/UEBuildTarget.cs
+++ b/Engine/Source/Programs/UnrealBuildTool/Configuration/UEBuildTarget.cs
@@ -878,23 +878,6 @@ namespace UnrealBuildTool
 			FileReference.WriteAllLines(Location, Files.Where(x => x.IsUnderDirectory(UnrealBuildTool.RootDirectory)).Select(x => x.MakeRelativeTo(UnrealBuildTool.RootDirectory).Replace(Path.DirectorySeparatorChar, '/')).OrderBy(x => x));
 		}
 
-		/// <summary>
-		/// Generates a public manifest file for writing out
-		/// </summary>
-		public void GenerateManifest(List<KeyValuePair<FileReference, BuildProductType>> BuildProducts)
-		{
-			FileReference ManifestPath;
-			if (UnrealBuildTool.IsEngineInstalled() && ProjectFile != null)
-			{
-				ManifestPath = FileReference.Combine(ProjectFile.Directory, "Intermediate", "Build", "Manifest.xml");
-			}
-			else
-			{
-				ManifestPath = FileReference.Combine(UnrealBuildTool.EngineDirectory, "Intermediate", "Build", "Manifest.xml");
-			}
-			GenerateManifest(ManifestPath, BuildProducts);
-		}
-
 		/// <summary>
 		/// Generates a public manifest file for writing out
 		/// </summary>
@@ -1495,10 +1478,6 @@ namespace UnrealBuildTool
 			}
 
 			// If we're only generating the manifest, return now
-			if (BuildConfiguration.bGenerateManifest)
-			{
-				GenerateManifest(BuildProducts);
-			}
 			foreach(FileReference ManifestFileName in Rules.ManifestFileNames)
 			{
 				GenerateManifest(ManifestFileName, BuildProducts);
diff --git a/Engine/Source/Programs/UnrealBuildTool/Modes/BuildMode.cs b/Engine/Source/Programs/UnrealBuildTool/Modes/BuildMode.cs
--- a/Engine/Source/Programs/UnrealBuildTool/Modes/BuildMode.cs
+++ b/Engine/Source/Programs/UnrealBuildTool/Modes/BuildMode.cs
@@ -391,7 +391,7 @@ namespace UnrealBuildTool
 				}
 
 				// Run the deployment steps
-				if (!BuildConfiguration.bGenerateManifest && !BuildConfiguration.bXGEExport)
+				if (!BuildConfiguration.bXGEExport)
 				{
 					foreach(TargetMakefile Makefile in Makefiles)
 					{
diff --git a/Engine/Source/Programs/UnrealBuildTool/System/HotReload.cs b/Engine/Source/Programs/UnrealBuildTool/System/HotReload.cs
--- a/Engine/Source/Programs/UnrealBuildTool/System/HotReload.cs
+++ b/Engine/Source/Programs/UnrealBuildTool/System/HotReload.cs
@@ -140,9 +140,7 @@ namespace UnrealBuildTool
 
 			// @todo ubtmake: Kind of cheating here to figure out if an editor target.  At this point we don't have access to the actual target description, and
 			// this code must be able to execute before we create or load module rules DLLs so that hot reload can work with bUseUBTMakefiles
-			bool bIsEditorTarget = TargetDesc.Name.EndsWith("Editor", StringComparison.InvariantCultureIgnoreCase);
-
-			if (!BuildConfiguration.bGenerateManifest && bIsEditorTarget)
+			if (TargetDesc.Name.EndsWith("Editor", StringComparison.OrdinalIgnoreCase))
 			{
 				string EditorBaseFileName = "UE4Editor";
 				if (TargetDesc.Configuration != UnrealTargetConfiguration.Development && TargetDesc.Configuration != UnrealTargetConfiguration.DebugGame)
diff --git a/Engine/Source/Programs/UnrealBuildTool/ToolChain/RemoteMac.cs b/Engine/Source/Programs/UnrealBuildTool/ToolChain/RemoteMac.cs
--- a/Engine/Source/Programs/UnrealBuildTool/ToolChain/RemoteMac.cs
+++ b/Engine/Source/Programs/UnrealBuildTool/ToolChain/RemoteMac.cs
@@ -526,19 +526,12 @@ namespace UnrealBuildTool
 			}
 
 			// Download the files from the remote
-			if(TargetDesc.AdditionalArguments.Any(x => x.Equals("-GenerateManifest", StringComparison.InvariantCultureIgnoreCase)))
-			{
-				LocalManifestFiles.Add(FileReference.Combine(UnrealBuildTool.EngineDirectory, "Intermediate", "Build", "Manifest.xml"));
-			}
-			else
-			{
-				Log.TraceInformation("[Remote] Downloading build products");
+			Log.TraceInformation("[Remote] Downloading build products");
 
-				List<FileReference> FilesToDownload = new List<FileReference>();
-				FilesToDownload.Add(RemoteLogFile);
-				FilesToDownload.AddRange(Manifest.BuildProducts.Select(x => new FileReference(x)));
-				DownloadFiles(FilesToDownload);
-			}
+			List<FileReference> FilesToDownload = new List<FileReference>();
+			FilesToDownload.Add(RemoteLogFile);
+			FilesToDownload.AddRange(Manifest.BuildProducts.Select(x => new FileReference(x)));
+			DownloadFiles(FilesToDownload);
 
 			// Write out all the local manifests
 			foreach(FileReference LocalManifestFile in LocalManifestFiles)
