UBT: Rename ECompilationResult to CompilationResult.

#rnx

diff --git a/Engine/Source/Programs/UnrealBuildTool/Configuration/UEBuildTarget.cs b/Engine/Source/Programs/UnrealBuildTool/Configuration/UEBuildTarget.cs
--- a/Engine/Source/Programs/UnrealBuildTool/Configuration/UEBuildTarget.cs
+++ b/Engine/Source/Programs/UnrealBuildTool/Configuration/UEBuildTarget.cs
@@ -1096,7 +1096,7 @@ namespace UnrealBuildTool
 		/// <summary>
 		/// Builds the target, appending list of output files and returns building result.
 		/// </summary>
-		public ECompilationResult Build(BuildConfiguration BuildConfiguration, List<FileItem> OutputItems, Dictionary<string, FileItem[]> ModuleNameToOutputItems, List<UHTModuleInfo> UObjectModules, ISourceFileWorkingSet WorkingSet, List<Action> Actions, BuildPrerequisites Prerequisites, bool bIsAssemblingBuild)
+		public CompilationResult Build(BuildConfiguration BuildConfiguration, List<FileItem> OutputItems, Dictionary<string, FileItem[]> ModuleNameToOutputItems, List<UHTModuleInfo> UObjectModules, ISourceFileWorkingSet WorkingSet, List<Action> Actions, BuildPrerequisites Prerequisites, bool bIsAssemblingBuild)
 		{
 			CppPlatform CppPlatform = UEBuildPlatform.GetBuildPlatform(Platform).DefaultCppPlatform;
 			CppConfiguration CppConfiguration = GetCppConfiguration(Configuration);
@@ -1209,7 +1209,7 @@ namespace UnrealBuildTool
 			// Execute the pre-build steps
 			if(!ExecuteCustomPreBuildSteps())
 			{
-				return ECompilationResult.OtherCompilationError;
+				return CompilationResult.OtherCompilationError;
 			}
 
 			// If we're compiling monolithic, make sure the executable knows about all referenced modules
@@ -1277,7 +1277,7 @@ namespace UnrealBuildTool
 				FileReference ModuleInfoFileName = FileReference.Combine(ProjectIntermediateDirectory, TargetName + ".uhtmanifest");
 
 				// Execute the header tool
-				ECompilationResult UHTResult = ExternalExecution.ExecuteHeaderToolIfNecessary(BuildConfiguration, ProjectFile, TargetName, TargetType, bHasProjectScriptPlugin, UObjectModules, ModuleInfoFileName, true, bIsAssemblingBuild, WorkingSet);
+				CompilationResult UHTResult = ExternalExecution.ExecuteHeaderToolIfNecessary(BuildConfiguration, ProjectFile, TargetName, TargetType, bHasProjectScriptPlugin, UObjectModules, ModuleInfoFileName, true, bIsAssemblingBuild, WorkingSet);
 				if(!UHTResult.Succeeded())
 				{
 					Log.TraceInformation(String.Format("Error: UnrealHeaderTool failed for target '{0}' (platform: {1}, module info: {2}, exit code: {3} ({4})).", TargetName, Platform.ToString(), ModuleInfoFileName, UHTResult.ToString(), (int)UHTResult));
@@ -1379,7 +1379,7 @@ namespace UnrealBuildTool
 					Log.WriteLine(0, LogEventType.Error, LogFormatOptions.NoSeverityPrefix, "{0}", InvalidIncludeDirectiveMessage);
 				}
 				Log.TraceError("Build canceled.");
-				return ECompilationResult.Canceled;
+				return CompilationResult.Canceled;
 			}
 
 			// Finalize and generate metadata for this target
@@ -1521,7 +1521,7 @@ namespace UnrealBuildTool
 
 			// Clean any stale modules which exist in multiple output directories. This can lead to the wrong DLL being loaded on Windows.
 			CleanStaleModules();
-			return ECompilationResult.Succeeded;
+			return CompilationResult.Succeeded;
 		}
 
 		/// <summary>
diff --git a/Engine/Source/Programs/UnrealBuildTool/Modes/BuildMode.cs b/Engine/Source/Programs/UnrealBuildTool/Modes/BuildMode.cs
--- a/Engine/Source/Programs/UnrealBuildTool/Modes/BuildMode.cs
+++ b/Engine/Source/Programs/UnrealBuildTool/Modes/BuildMode.cs
@@ -123,7 +123,7 @@ namespace UnrealBuildTool
 						RemoteMac RemoteMac = new RemoteMac(TargetDesc.ProjectFile);
 						if(!RemoteMac.Build(TargetDesc, LogFile))
 						{
-							return (int)ECompilationResult.Unknown;
+							return (int)CompilationResult.Unknown;
 						}
 
 						TargetDescriptors.RemoveAt(Idx--);
@@ -146,8 +146,8 @@ namespace UnrealBuildTool
 					// Create the working set provider
 					using (ISourceFileWorkingSet WorkingSet = SourceFileWorkingSet.Create(UnrealBuildTool.RootDirectory, ProjectDirs))
 					{
-						ECompilationResult Result = Build(TargetDescriptors, BuildConfiguration, WorkingSet);
-						if(Result != ECompilationResult.Succeeded)
+						CompilationResult Result = Build(TargetDescriptors, BuildConfiguration, WorkingSet);
+						if(Result != CompilationResult.Succeeded)
 						{
 							return (int)Result;
 						}
@@ -157,7 +157,7 @@ namespace UnrealBuildTool
 			catch (Exception Ex)
 			{
 				Log.WriteException(Ex, (LogFile == null) ? null : LogFile.FullName);
-				return (int)ECompilationResult.OtherCompilationError;
+				return (int)CompilationResult.OtherCompilationError;
 			}
 			finally
 			{
@@ -178,7 +178,7 @@ namespace UnrealBuildTool
 		/// <param name="BuildConfiguration">Current build configuration</param>
 		/// <param name="WorkingSet">The source file working set</param>
 		/// <returns>Result from the compilation</returns>
-		public static ECompilationResult Build(List<TargetDescriptor> TargetDescriptors, BuildConfiguration BuildConfiguration, ISourceFileWorkingSet WorkingSet)
+		public static CompilationResult Build(List<TargetDescriptor> TargetDescriptors, BuildConfiguration BuildConfiguration, ISourceFileWorkingSet WorkingSet)
 		{
 			// Create a makefile for each target
 			TargetMakefile[] Makefiles = new TargetMakefile[TargetDescriptors.Count];
@@ -227,12 +227,12 @@ namespace UnrealBuildTool
 					List<Action> Actions = new List<Action>();
 					BuildPrerequisites Prerequisites = new BuildPrerequisites();
 
-					ECompilationResult BuildResult;
+					CompilationResult BuildResult;
 					using(Timeline.ScopeEvent("UEBuildTarget.Build()"))
 					{
 						BuildResult = Target.Build(BuildConfiguration, OutputItems, ModuleNameToOutputItems, UObjectModules, WorkingSet, Actions, Prerequisites, bIsAssemblingBuild);
 					}
-					if (BuildResult != ECompilationResult.Succeeded)
+					if (BuildResult != CompilationResult.Succeeded)
 					{
 						return BuildResult;
 					}
@@ -280,7 +280,7 @@ namespace UnrealBuildTool
 					{
 						if (!Utils.ExecuteCustomBuildSteps(Makefile.PreBuildScripts))
 						{
-							return ECompilationResult.OtherCompilationError;
+							return CompilationResult.OtherCompilationError;
 						}
 					}
 
@@ -291,7 +291,7 @@ namespace UnrealBuildTool
 						const bool bIsAssemblingBuild = true;
 
 						FileReference ModuleInfoFileName = FileReference.Combine(Makefile.ProjectIntermediateDirectory, TargetDesc.Name + ".uhtmanifest");
-						ECompilationResult UHTResult = ExternalExecution.ExecuteHeaderToolIfNecessary(BuildConfiguration, TargetDesc.ProjectFile, TargetDesc.Name, Makefile.TargetType, Makefile.bHasProjectScriptPlugin, UObjectModules: Makefile.UObjectModules, ModuleInfoFileName: ModuleInfoFileName, bIsGatheringBuild: bIsGatheringBuild, bIsAssemblingBuild: bIsAssemblingBuild, WorkingSet: WorkingSet);
+						CompilationResult UHTResult = ExternalExecution.ExecuteHeaderToolIfNecessary(BuildConfiguration, TargetDesc.ProjectFile, TargetDesc.Name, Makefile.TargetType, Makefile.bHasProjectScriptPlugin, UObjectModules: Makefile.UObjectModules, ModuleInfoFileName: ModuleInfoFileName, bIsGatheringBuild: bIsGatheringBuild, bIsAssemblingBuild: bIsAssemblingBuild, WorkingSet: WorkingSet);
 						if(!UHTResult.Succeeded())
 						{
 							Log.TraceInformation("UnrealHeaderTool failed for target '" + TargetDesc.Name + "' (platform: " + TargetDesc.Platform.ToString() + ", module info: " + ModuleInfoFileName + ").");
@@ -311,7 +311,7 @@ namespace UnrealBuildTool
 				if(!ActionGraph.CheckForConflicts(Makefiles.SelectMany(x => x.Actions)))
 				{
 					Log.TraceInformation("Check log for additional details.");
-					return ECompilationResult.OtherCompilationError;
+					return CompilationResult.OtherCompilationError;
 				}
 
 				// Find all the actions to be executed
@@ -385,7 +385,7 @@ namespace UnrealBuildTool
 					{
 						if(!ActionGraph.ExecuteActions(BuildConfiguration, MergedActionsToExecute))
 						{
-							return ECompilationResult.OtherCompilationError;
+							return CompilationResult.OtherCompilationError;
 						}
 					}
 				}
@@ -404,7 +404,7 @@ namespace UnrealBuildTool
 					}
 				}
 			}
-			return ECompilationResult.Succeeded;
+			return CompilationResult.Succeeded;
 		}
 
 		/// <summary>
diff --git a/Engine/Source/Programs/UnrealBuildTool/Modes/DeployMode.cs b/Engine/Source/Programs/UnrealBuildTool/Modes/DeployMode.cs
--- a/Engine/Source/Programs/UnrealBuildTool/Modes/DeployMode.cs
+++ b/Engine/Source/Programs/UnrealBuildTool/Modes/DeployMode.cs
@@ -37,7 +37,7 @@ namespace UnrealBuildTool
 			Log.WriteLine(LogEventType.Console, "Deploying {0} {1} {2}...", Receipt.TargetName, Receipt.Platform, Receipt.Configuration);
 			UEBuildPlatform.GetBuildPlatform(Receipt.Platform).Deploy(Receipt);
 
-			return (int)ECompilationResult.Succeeded;
+			return (int)CompilationResult.Succeeded;
 		}
 	}
 }
diff --git a/Engine/Source/Programs/UnrealBuildTool/Modes/GenerateProjectFilesMode.cs b/Engine/Source/Programs/UnrealBuildTool/Modes/GenerateProjectFilesMode.cs
--- a/Engine/Source/Programs/UnrealBuildTool/Modes/GenerateProjectFilesMode.cs
+++ b/Engine/Source/Programs/UnrealBuildTool/Modes/GenerateProjectFilesMode.cs
@@ -160,10 +160,10 @@ namespace UnrealBuildTool
 			{
 				if (!Generator.GenerateProjectFiles(PlatformProjectGenerators, Arguments.GetRawArray()))
 				{
-					return (int)ECompilationResult.OtherCompilationError;
+					return (int)CompilationResult.OtherCompilationError;
 				}
 			}
-			return (int)ECompilationResult.Succeeded;
+			return (int)CompilationResult.Succeeded;
 		}
 
 		/// <summary>
diff --git a/Engine/Source/Programs/UnrealBuildTool/System/ExternalExecution.cs b/Engine/Source/Programs/UnrealBuildTool/System/ExternalExecution.cs
--- a/Engine/Source/Programs/UnrealBuildTool/System/ExternalExecution.cs
+++ b/Engine/Source/Programs/UnrealBuildTool/System/ExternalExecution.cs
@@ -18,7 +18,7 @@ namespace UnrealBuildTool
 	// UE4\Engine\Source\Runtime\Core\Public\Misc\ComplilationResult.h 
 	// to keep communication between UHT, UBT and Editor compiling
 	// processes valid.
-	enum ECompilationResult
+	enum CompilationResult
 	{
 		/// <summary>
 		/// Compilation succeeded
@@ -60,11 +60,12 @@ namespace UnrealBuildTool
 		/// </summary>
 		Unknown
 	}
+
 	static class CompilationResultExtensions
 	{
-		public static bool Succeeded(this ECompilationResult Result)
+		public static bool Succeeded(this CompilationResult Result)
 		{
-			return Result == ECompilationResult.Succeeded || Result == ECompilationResult.UpToDate;
+			return Result == CompilationResult.Succeeded || Result == CompilationResult.UpToDate;
 		}
 	}
 
@@ -1108,7 +1109,7 @@ namespace UnrealBuildTool
 		/// Builds and runs the header tool and touches the header directories.
 		/// Performs any early outs if headers need no changes, given the UObject modules, tool path, game name, and configuration
 		/// </summary>
-		public static ECompilationResult ExecuteHeaderToolIfNecessary(BuildConfiguration BuildConfiguration, FileReference ProjectFile, string TargetName, TargetType TargetType, bool bHasProjectScriptPlugin, List<UHTModuleInfo> UObjectModules, FileReference ModuleInfoFileName, bool bIsGatheringBuild, bool bIsAssemblingBuild, ISourceFileWorkingSet WorkingSet)
+		public static CompilationResult ExecuteHeaderToolIfNecessary(BuildConfiguration BuildConfiguration, FileReference ProjectFile, string TargetName, TargetType TargetType, bool bHasProjectScriptPlugin, List<UHTModuleInfo> UObjectModules, FileReference ModuleInfoFileName, bool bIsGatheringBuild, bool bIsAssemblingBuild, ISourceFileWorkingSet WorkingSet)
 		{
 			if (ProgressWriter.bWriteMarkup)
 			{
@@ -1189,7 +1190,7 @@ namespace UnrealBuildTool
 						List<TargetDescriptor> TargetDescriptors = new List<TargetDescriptor>();
 						TargetDescriptors.Add(new TargetDescriptor(ScriptProjectFile, "UnrealHeaderTool", Platform, Configuration, Architecture, null));
 
-						ECompilationResult Result;
+						CompilationResult Result;
 						using(Timeline.ScopeEvent("Buildng UnrealHeaderTool"))
 						{
 							bool bPrevXGEExport = BuildConfiguration.bXGEExport;
@@ -1199,9 +1200,9 @@ namespace UnrealBuildTool
 
 							BuildConfiguration.bXGEExport = bPrevXGEExport;
 						}
-						if(Result != ECompilationResult.Succeeded)
+						if(Result != CompilationResult.Succeeded)
 						{
-							return ECompilationResult.OtherCompilationError;
+							return CompilationResult.OtherCompilationError;
 						}
 					}
 
@@ -1237,11 +1238,11 @@ namespace UnrealBuildTool
 					Stopwatch s = new Stopwatch();
 					s.Start();
 					ITimelineEvent Timer = Timeline.ScopeEvent("Executing UnrealHeaderTool");
-					ECompilationResult UHTResult = (ECompilationResult)RunExternalNativeExecutable(ExternalExecution.GetHeaderToolPath(HeaderToolReceipt), CmdLine);
+					CompilationResult UHTResult = (CompilationResult)RunExternalNativeExecutable(ExternalExecution.GetHeaderToolPath(HeaderToolReceipt), CmdLine);
 					Timer.Finish();
 					s.Stop();
 
-					if (UHTResult != ECompilationResult.Succeeded)
+					if (UHTResult != CompilationResult.Succeeded)
 					{
 						// On Linux and Mac, the shell will return 128+signal number exit codes if UHT gets a signal (e.g. crashes or is interrupted)
 						if ((BuildHostPlatform.Current.Platform == UnrealTargetPlatform.Linux ||
@@ -1250,7 +1251,7 @@ namespace UnrealBuildTool
 							)
 						{
 							// SIGINT is 2, so 128 + SIGINT is 130
-							UHTResult = ((int)(UHTResult) == 130) ? ECompilationResult.Canceled : ECompilationResult.CrashOrAssert;
+							UHTResult = ((int)(UHTResult) == 130) ? CompilationResult.Canceled : CompilationResult.CrashOrAssert;
 						}
 
 						if ((BuildHostPlatform.Current.Platform == UnrealTargetPlatform.Win32 || 
@@ -1287,7 +1288,7 @@ namespace UnrealBuildTool
 			{
 				Log.WriteLine(LogEventType.Console, "@progress pop");
 			}
-			return ECompilationResult.Succeeded;
+			return CompilationResult.Succeeded;
 		}
 	}
 }
diff --git a/Engine/Source/Programs/UnrealBuildTool/UnrealBuildTool.cs b/Engine/Source/Programs/UnrealBuildTool/UnrealBuildTool.cs
--- a/Engine/Source/Programs/UnrealBuildTool/UnrealBuildTool.cs
+++ b/Engine/Source/Programs/UnrealBuildTool/UnrealBuildTool.cs
@@ -404,14 +404,14 @@ namespace UnrealBuildTool
 				// BuildExceptions should have nicely formatted messages. We can log these directly.
 				Log.TraceError(Ex.Message.ToString());
 				Log.TraceLog(ExceptionUtils.FormatExceptionDetails(Ex));
-				return (int)ECompilationResult.OtherCompilationError;
+				return (int)CompilationResult.OtherCompilationError;
 			}
 			catch (Exception Ex)
 			{
 				// Unhandled exception. 
 				Log.TraceError("Unhandled exception: {0}", Ex);
 				Log.TraceLog(ExceptionUtils.FormatExceptionDetails(Ex));
-				return (int)ECompilationResult.OtherCompilationError;
+				return (int)CompilationResult.OtherCompilationError;
 			}
 			finally
 			{
