UBT: Add an exception for propagating error codes to the calling process, rather than having to pass them by return values. Also suppress "produced item not found" errors when the compile has already failed.

diff --git a/Engine/Source/Programs/UnrealBuildTool/Configuration/UEBuildTarget.cs b/Engine/Source/Programs/UnrealBuildTool/Configuration/UEBuildTarget.cs
--- a/Engine/Source/Programs/UnrealBuildTool/Configuration/UEBuildTarget.cs
+++ b/Engine/Source/Programs/UnrealBuildTool/Configuration/UEBuildTarget.cs
@@ -1096,7 +1096,7 @@ namespace UnrealBuildTool
 		/// <summary>
 		/// Builds the target, appending list of output files and returns building result.
 		/// </summary>
-		public CompilationResult Build(BuildConfiguration BuildConfiguration, List<FileItem> OutputItems, Dictionary<string, FileItem[]> ModuleNameToOutputItems, List<UHTModuleInfo> UObjectModules, ISourceFileWorkingSet WorkingSet, List<Action> Actions, BuildPrerequisites Prerequisites, bool bIsAssemblingBuild)
+		public void Build(BuildConfiguration BuildConfiguration, List<FileItem> OutputItems, Dictionary<string, FileItem[]> ModuleNameToOutputItems, List<UHTModuleInfo> UObjectModules, ISourceFileWorkingSet WorkingSet, List<Action> Actions, BuildPrerequisites Prerequisites, bool bIsAssemblingBuild)
 		{
 			CppPlatform CppPlatform = UEBuildPlatform.GetBuildPlatform(Platform).DefaultCppPlatform;
 			CppConfiguration CppConfiguration = GetCppConfiguration(Configuration);
@@ -1207,10 +1207,7 @@ namespace UnrealBuildTool
 			}
 
 			// Execute the pre-build steps
-			if(!ExecuteCustomPreBuildSteps())
-			{
-				return CompilationResult.OtherCompilationError;
-			}
+			Utils.ExecuteCustomBuildSteps(PreBuildStepScripts);
 
 			// If we're compiling monolithic, make sure the executable knows about all referenced modules
 			if (ShouldCompileMonolithic())
@@ -1275,14 +1272,7 @@ namespace UnrealBuildTool
 			if (UObjectModules.Count > 0)
 			{
 				FileReference ModuleInfoFileName = FileReference.Combine(ProjectIntermediateDirectory, TargetName + ".uhtmanifest");
-
-				// Execute the header tool
-				CompilationResult UHTResult = ExternalExecution.ExecuteHeaderToolIfNecessary(BuildConfiguration, ProjectFile, TargetName, TargetType, bHasProjectScriptPlugin, UObjectModules, ModuleInfoFileName, true, bIsAssemblingBuild, WorkingSet);
-				if(!UHTResult.Succeeded())
-				{
-					Log.TraceInformation(String.Format("Error: UnrealHeaderTool failed for target '{0}' (platform: {1}, module info: {2}, exit code: {3} ({4})).", TargetName, Platform.ToString(), ModuleInfoFileName, UHTResult.ToString(), (int)UHTResult));
-					return UHTResult;
-				}
+				ExternalExecution.ExecuteHeaderToolIfNecessary(BuildConfiguration, ProjectFile, TargetName, TargetType, bHasProjectScriptPlugin, UObjectModules, ModuleInfoFileName, true, bIsAssemblingBuild, WorkingSet);
 			}
 
 			// Find all the shared PCHs.
@@ -1376,10 +1366,8 @@ namespace UnrealBuildTool
 			{
 				foreach (string InvalidIncludeDirectiveMessage in InvalidIncludeDirectiveMessages)
 				{
-					Log.WriteLine(0, LogEventType.Error, LogFormatOptions.NoSeverityPrefix, "{0}", InvalidIncludeDirectiveMessage);
+					Log.WriteLine(0, LogEventType.Warning, LogFormatOptions.NoSeverityPrefix, "{0}", InvalidIncludeDirectiveMessage);
 				}
-				Log.TraceError("Build canceled.");
-				return CompilationResult.Canceled;
 			}
 
 			// Finalize and generate metadata for this target
@@ -1521,7 +1509,6 @@ namespace UnrealBuildTool
 
 			// Clean any stale modules which exist in multiple output directories. This can lead to the wrong DLL being loaded on Windows.
 			CleanStaleModules();
-			return CompilationResult.Succeeded;
 		}
 
 		/// <summary>
@@ -1973,22 +1960,6 @@ namespace UnrealBuildTool
 			return ScriptFiles.ToArray();
 		}
 
-		/// <summary>
-		/// Executes the custom pre-build steps
-		/// </summary>
-		public bool ExecuteCustomPreBuildSteps()
-		{
-			return Utils.ExecuteCustomBuildSteps(PreBuildStepScripts);
-		}
-
-		/// <summary>
-		/// Executes the custom post-build steps
-		/// </summary>
-		public bool ExecuteCustomPostBuildSteps()
-		{
-			return Utils.ExecuteCustomBuildSteps(PostBuildStepScripts);
-		}
-
 		private static FileReference AddModuleFilenameSuffix(string ModuleName, FileReference FilePath, string Suffix)
 		{
 			int MatchPos = FilePath.FullName.LastIndexOf(ModuleName, StringComparison.InvariantCultureIgnoreCase);
diff --git a/Engine/Source/Programs/UnrealBuildTool/Modes/BuildMode.cs b/Engine/Source/Programs/UnrealBuildTool/Modes/BuildMode.cs
--- a/Engine/Source/Programs/UnrealBuildTool/Modes/BuildMode.cs
+++ b/Engine/Source/Programs/UnrealBuildTool/Modes/BuildMode.cs
@@ -146,19 +146,10 @@ namespace UnrealBuildTool
 					// Create the working set provider
 					using (ISourceFileWorkingSet WorkingSet = SourceFileWorkingSet.Create(UnrealBuildTool.RootDirectory, ProjectDirs))
 					{
-						CompilationResult Result = Build(TargetDescriptors, BuildConfiguration, WorkingSet);
-						if(Result != CompilationResult.Succeeded)
-						{
-							return (int)Result;
-						}
+						Build(TargetDescriptors, BuildConfiguration, WorkingSet);
 					}
 				}
 			}
-			catch (Exception Ex)
-			{
-				Log.WriteException(Ex, (LogFile == null) ? null : LogFile.FullName);
-				return (int)CompilationResult.OtherCompilationError;
-			}
 			finally
 			{
 				// Save all the caches
@@ -178,7 +169,7 @@ namespace UnrealBuildTool
 		/// <param name="BuildConfiguration">Current build configuration</param>
 		/// <param name="WorkingSet">The source file working set</param>
 		/// <returns>Result from the compilation</returns>
-		public static CompilationResult Build(List<TargetDescriptor> TargetDescriptors, BuildConfiguration BuildConfiguration, ISourceFileWorkingSet WorkingSet)
+		public static void Build(List<TargetDescriptor> TargetDescriptors, BuildConfiguration BuildConfiguration, ISourceFileWorkingSet WorkingSet)
 		{
 			// Create a makefile for each target
 			TargetMakefile[] Makefiles = new TargetMakefile[TargetDescriptors.Count];
@@ -227,14 +218,9 @@ namespace UnrealBuildTool
 					List<Action> Actions = new List<Action>();
 					BuildPrerequisites Prerequisites = new BuildPrerequisites();
 
-					CompilationResult BuildResult;
 					using(Timeline.ScopeEvent("UEBuildTarget.Build()"))
 					{
-						BuildResult = Target.Build(BuildConfiguration, OutputItems, ModuleNameToOutputItems, UObjectModules, WorkingSet, Actions, Prerequisites, bIsAssemblingBuild);
-					}
-					if (BuildResult != CompilationResult.Succeeded)
-					{
-						return BuildResult;
+						Target.Build(BuildConfiguration, OutputItems, ModuleNameToOutputItems, UObjectModules, WorkingSet, Actions, Prerequisites, bIsAssemblingBuild);
 					}
 
 					// Create the makefile
@@ -278,10 +264,7 @@ namespace UnrealBuildTool
 					// Execute all the pre-build steps
 					if (!BuildConfiguration.bXGEExport)
 					{
-						if (!Utils.ExecuteCustomBuildSteps(Makefile.PreBuildScripts))
-						{
-							return CompilationResult.OtherCompilationError;
-						}
+						Utils.ExecuteCustomBuildSteps(Makefile.PreBuildScripts);
 					}
 
 					// If the target needs UHT to be run, we'll go ahead and do that now
@@ -291,12 +274,7 @@ namespace UnrealBuildTool
 						const bool bIsAssemblingBuild = true;
 
 						FileReference ModuleInfoFileName = FileReference.Combine(Makefile.ProjectIntermediateDirectory, TargetDesc.Name + ".uhtmanifest");
-						CompilationResult UHTResult = ExternalExecution.ExecuteHeaderToolIfNecessary(BuildConfiguration, TargetDesc.ProjectFile, TargetDesc.Name, Makefile.TargetType, Makefile.bHasProjectScriptPlugin, UObjectModules: Makefile.UObjectModules, ModuleInfoFileName: ModuleInfoFileName, bIsGatheringBuild: bIsGatheringBuild, bIsAssemblingBuild: bIsAssemblingBuild, WorkingSet: WorkingSet);
-						if(!UHTResult.Succeeded())
-						{
-							Log.TraceInformation("UnrealHeaderTool failed for target '" + TargetDesc.Name + "' (platform: " + TargetDesc.Platform.ToString() + ", module info: " + ModuleInfoFileName + ").");
-							return UHTResult;
-						}
+						ExternalExecution.ExecuteHeaderToolIfNecessary(BuildConfiguration, TargetDesc.ProjectFile, TargetDesc.Name, Makefile.TargetType, Makefile.bHasProjectScriptPlugin, UObjectModules: Makefile.UObjectModules, ModuleInfoFileName: ModuleInfoFileName, bIsGatheringBuild: bIsGatheringBuild, bIsAssemblingBuild: bIsAssemblingBuild, WorkingSet: WorkingSet);
 					}
 				}
 
@@ -308,12 +286,8 @@ namespace UnrealBuildTool
 			if(!BuildConfiguration.bSkipBuild)
 			{
 				// Make sure that none of the actions conflict with any other (producing output files differently, etc...)
-				if(!ActionGraph.CheckForConflicts(Makefiles.SelectMany(x => x.Actions)))
-				{
-					Log.TraceInformation("Check log for additional details.");
-					return CompilationResult.OtherCompilationError;
-				}
-
+				ActionGraph.CheckForConflicts(Makefiles.SelectMany(x => x.Actions));
+				
 				// Find all the actions to be executed
 				HashSet<Action>[] ActionsToExecute = new HashSet<Action>[TargetDescriptors.Count];
 				for(int TargetIdx = 0; TargetIdx < TargetDescriptors.Count; TargetIdx++)
@@ -371,7 +345,7 @@ namespace UnrealBuildTool
 				// Create directories for the outdated produced items.
 				ActionGraph.CreateDirectoriesForProducedItems(MergedActionsToExecute);
 
-				// Execute the actions.
+				// Execute the actions
 				if (BuildConfiguration.bXGEExport)
 				{
 					using(Timeline.ScopeEvent("XGE.ExportActions()"))
@@ -383,10 +357,7 @@ namespace UnrealBuildTool
 				{
 					using(Timeline.ScopeEvent("ActionGraph.ExecuteActions()"))
 					{
-						if(!ActionGraph.ExecuteActions(BuildConfiguration, MergedActionsToExecute))
-						{
-							return CompilationResult.OtherCompilationError;
-						}
+						ActionGraph.ExecuteActions(BuildConfiguration, MergedActionsToExecute);
 					}
 				}
 
@@ -404,7 +375,6 @@ namespace UnrealBuildTool
 					}
 				}
 			}
-			return CompilationResult.Succeeded;
 		}
 
 		/// <summary>
diff --git a/Engine/Source/Programs/UnrealBuildTool/System/ActionGraph.cs b/Engine/Source/Programs/UnrealBuildTool/System/ActionGraph.cs
--- a/Engine/Source/Programs/UnrealBuildTool/System/ActionGraph.cs
+++ b/Engine/Source/Programs/UnrealBuildTool/System/ActionGraph.cs
@@ -55,8 +55,7 @@ namespace UnrealBuildTool
 		/// Checks a set of actions for conflicts (ie. different actions producing the same output items)
 		/// </summary>
 		/// <param name="Actions">The set of actions to check</param>
-		/// <returns>True if the actions are ok, false if there are conflicts</returns>
-		public static bool CheckForConflicts(IEnumerable<Action> Actions)
+		public static void CheckForConflicts(IEnumerable<Action> Actions)
 		{
 			bool bResult = true;
 
@@ -77,7 +76,10 @@ namespace UnrealBuildTool
 				}
 			}
 
-			return bResult;
+			if(!bResult)
+			{
+				throw new BuildException("Action graph is invalid; unable to continue. See log for additional details.");
+			}
 		}
 
 		/// <summary>
@@ -109,17 +111,17 @@ namespace UnrealBuildTool
 		/// <summary>
 		/// Executes a list of actions.
 		/// </summary>
-		public static bool ExecuteActions(BuildConfiguration BuildConfiguration, List<Action> ActionsToExecute)
+		public static void ExecuteActions(BuildConfiguration BuildConfiguration, List<Action> ActionsToExecute)
 		{
-			bool bResult = true;
 			if(ActionsToExecute.Count == 0)
 			{
 				Log.TraceInformation("Target is up to date");
 			}
 			else
 			{
+				// Figure out which executor to use
 				ActionExecutor Executor;
-				if ((XGE.IsAvailable() && BuildConfiguration.bAllowXGE) || BuildConfiguration.bXGEExport)
+				if (BuildConfiguration.bAllowXGE && XGE.IsAvailable())
 				{
 					Executor = new XGE();
 				}
@@ -127,11 +129,11 @@ namespace UnrealBuildTool
 				{
 					Executor = new Distcc();
 				}
-				else if(SNDBS.IsAvailable() && BuildConfiguration.bAllowSNDBS)
+				else if(BuildConfiguration.bAllowSNDBS && SNDBS.IsAvailable())
 				{
 					Executor = new SNDBS();
 				}
-				else if(ParallelExecutor.IsAvailable() && BuildConfiguration.bAllowParallelExecutor)
+				else if(BuildConfiguration.bAllowParallelExecutor && ParallelExecutor.IsAvailable())
 				{
 					Executor = new ParallelExecutor();
 				}
@@ -140,30 +142,30 @@ namespace UnrealBuildTool
 					Executor = new LocalExecutor();
 				}
 
+				// Execute the build
 				Stopwatch Timer = Stopwatch.StartNew();
-				bResult = Executor.ExecuteActions(ActionsToExecute, BuildConfiguration.bLogDetailedActionStats);
+				if(!Executor.ExecuteActions(ActionsToExecute, BuildConfiguration.bLogDetailedActionStats))
+				{
+					throw new CompilationResultException(CompilationResult.OtherCompilationError);
+				}
 				Log.TraceInformation("Total time in {0} executor: {1:0.00} seconds", Executor.Name, Timer.Elapsed.TotalSeconds);
 
-				if (!BuildConfiguration.bXGEExport)
+				// Verify the link outputs were created (seems to happen with Win64 compiles)
+				foreach (Action BuildAction in ActionsToExecute)
 				{
-					// Verify the link outputs were created (seems to happen with Win64 compiles)
-					foreach (Action BuildAction in ActionsToExecute)
+					if (BuildAction.ActionType == ActionType.Link)
 					{
-						if (BuildAction.ActionType == ActionType.Link)
+						foreach (FileItem Item in BuildAction.ProducedItems)
 						{
-							foreach (FileItem Item in BuildAction.ProducedItems)
+							Item.ResetFileInfo();
+							if(!Item.Exists)
 							{
-								Item.ResetFileInfo();
-								if(!Item.Exists)
-								{
-									throw new BuildException("UBT ERROR: Failed to produce item: " + Item.AbsolutePath);
-								}
+								throw new BuildException("Failed to produce item: {0}", Item.AbsolutePath);
 							}
 						}
 					}
 				}
 			}
-			return bResult;
 		}
 
 		/// <summary>
diff --git a/Engine/Source/Programs/UnrealBuildTool/System/ExternalExecution.cs b/Engine/Source/Programs/UnrealBuildTool/System/ExternalExecution.cs
--- a/Engine/Source/Programs/UnrealBuildTool/System/ExternalExecution.cs
+++ b/Engine/Source/Programs/UnrealBuildTool/System/ExternalExecution.cs
@@ -69,6 +69,17 @@ namespace UnrealBuildTool
 		}
 	}
 
+	class CompilationResultException : BuildException
+	{
+		public readonly CompilationResult Result;
+
+		public CompilationResultException(CompilationResult Result)
+			: base("Error: {0}", Result)
+		{
+			this.Result = Result;
+		}
+	}
+
 	/// <summary>
 	/// Type of module. Mirrored in UHT as EBuildModuleType.
 	/// This should be sorted by the order in which we expect modules to be built.
@@ -1109,7 +1120,7 @@ namespace UnrealBuildTool
 		/// Builds and runs the header tool and touches the header directories.
 		/// Performs any early outs if headers need no changes, given the UObject modules, tool path, game name, and configuration
 		/// </summary>
-		public static CompilationResult ExecuteHeaderToolIfNecessary(BuildConfiguration BuildConfiguration, FileReference ProjectFile, string TargetName, TargetType TargetType, bool bHasProjectScriptPlugin, List<UHTModuleInfo> UObjectModules, FileReference ModuleInfoFileName, bool bIsGatheringBuild, bool bIsAssemblingBuild, ISourceFileWorkingSet WorkingSet)
+		public static void ExecuteHeaderToolIfNecessary(BuildConfiguration BuildConfiguration, FileReference ProjectFile, string TargetName, TargetType TargetType, bool bHasProjectScriptPlugin, List<UHTModuleInfo> UObjectModules, FileReference ModuleInfoFileName, bool bIsGatheringBuild, bool bIsAssemblingBuild, ISourceFileWorkingSet WorkingSet)
 		{
 			if (ProgressWriter.bWriteMarkup)
 			{
@@ -1190,20 +1201,15 @@ namespace UnrealBuildTool
 						List<TargetDescriptor> TargetDescriptors = new List<TargetDescriptor>();
 						TargetDescriptors.Add(new TargetDescriptor(ScriptProjectFile, "UnrealHeaderTool", Platform, Configuration, Architecture, null));
 
-						CompilationResult Result;
 						using(Timeline.ScopeEvent("Buildng UnrealHeaderTool"))
 						{
 							bool bPrevXGEExport = BuildConfiguration.bXGEExport;
 							BuildConfiguration.bXGEExport = false;
 
-							Result = BuildMode.Build(TargetDescriptors, BuildConfiguration, WorkingSet);
+							BuildMode.Build(TargetDescriptors, BuildConfiguration, WorkingSet);
 
 							BuildConfiguration.bXGEExport = bPrevXGEExport;
 						}
-						if(Result != CompilationResult.Succeeded)
-						{
-							return CompilationResult.OtherCompilationError;
-						}
 					}
 
 					Progress.Write(1, 3);
@@ -1258,9 +1264,10 @@ namespace UnrealBuildTool
 							BuildHostPlatform.Current.Platform == UnrealTargetPlatform.Win64) && 
 							(int)(UHTResult) < 0)
 						{
-							Log.TraceInformation(String.Format("UnrealHeaderTool failed with exit code 0x{0:X} - check that UE4 prerequisites are installed.", (int)UHTResult));
+							Log.TraceError(String.Format("UnrealHeaderTool failed with exit code 0x{0:X} - check that UE4 prerequisites are installed.", (int)UHTResult));
 						}
-						return UHTResult;
+
+						throw new CompilationResultException(UHTResult);
 					}
 
 					Log.TraceInformation("Reflection code generated for {0} in {1} seconds", ActualTargetName, s.Elapsed.TotalSeconds);
@@ -1288,7 +1295,6 @@ namespace UnrealBuildTool
 			{
 				Log.WriteLine(LogEventType.Console, "@progress pop");
 			}
-			return CompilationResult.Succeeded;
 		}
 	}
 }
diff --git a/Engine/Source/Programs/UnrealBuildTool/System/Utils.cs b/Engine/Source/Programs/UnrealBuildTool/System/Utils.cs
--- a/Engine/Source/Programs/UnrealBuildTool/System/Utils.cs
+++ b/Engine/Source/Programs/UnrealBuildTool/System/Utils.cs
@@ -1051,7 +1051,7 @@ namespace UnrealBuildTool
 		/// </summary>
 		/// <param name="ScriptFiles">List of script files to execute</param>
 		/// <returns>True if the steps succeeded, false otherwise</returns>
-		public static bool ExecuteCustomBuildSteps(FileReference[] ScriptFiles)
+		public static void ExecuteCustomBuildSteps(FileReference[] ScriptFiles)
 		{
 			UnrealTargetPlatform HostPlatform = BuildHostPlatform.Current.Platform;
 			foreach(FileReference ScriptFile in ScriptFiles)
@@ -1071,11 +1071,9 @@ namespace UnrealBuildTool
 				int ReturnCode = Utils.RunLocalProcessAndLogOutput(StartInfo);
 				if(ReturnCode != 0)
 				{
-					Log.TraceError("Custom build step terminated with exit code {0}", ReturnCode);
-					return false;
+					throw new BuildException("Custom build step terminated with exit code {0}", ReturnCode);
 				}
 			}
-			return true;
 		}
 	}
 }
diff --git a/Engine/Source/Programs/UnrealBuildTool/UnrealBuildTool.cs b/Engine/Source/Programs/UnrealBuildTool/UnrealBuildTool.cs
--- a/Engine/Source/Programs/UnrealBuildTool/UnrealBuildTool.cs
+++ b/Engine/Source/Programs/UnrealBuildTool/UnrealBuildTool.cs
@@ -399,6 +399,12 @@ namespace UnrealBuildTool
 				ToolMode Mode = (ToolMode)Activator.CreateInstance(ModeType);
 				return Mode.Execute(Arguments);
 			}
+			catch (CompilationResultException Ex)
+			{
+				// Used to return a propagate a specific exit code after an error has occurred. Does not log any message.
+				Log.TraceLog(ExceptionUtils.FormatExceptionDetails(Ex));
+				return (int)Ex.Result;
+			}
 			catch (BuildException Ex)
 			{
 				// BuildExceptions should have nicely formatted messages. We can log these directly.
