UBT: Change UEBuildTarget.Build() to return a construct and return a makefile.

#rnx

diff --git a/Engine/Source/Programs/UnrealBuildTool/Configuration/UEBuildBinary.cs b/Engine/Source/Programs/UnrealBuildTool/Configuration/UEBuildBinary.cs
--- a/Engine/Source/Programs/UnrealBuildTool/Configuration/UEBuildBinary.cs
+++ b/Engine/Source/Programs/UnrealBuildTool/Configuration/UEBuildBinary.cs
@@ -156,13 +156,11 @@ namespace UnrealBuildTool
 		/// <param name="ToolChain">The toolchain which to use for building</param>
 		/// <param name="CompileEnvironment">The environment to compile the binary in</param>
 		/// <param name="LinkEnvironment">The environment to link the binary in</param>
-		/// <param name="SharedPCHs">List of templates for shared PCHs</param>
 		/// <param name="WorkingSet">The working set of source files</param>
-		/// <param name="Prerequisites">Predicates for the generated actions to be valid</param>
 		/// <param name="ExeDir">Directory containing the output executable</param>
-		/// <param name="Actions">List of actions to be executed. Additional actions will be added to this list.</param>
+		/// <param name="Makefile">The makefile to add actions to</param>
 		/// <returns>Set of built products</returns>
-		public List<FileItem> Build(ReadOnlyTargetRules Target, UEToolChain ToolChain, CppCompileEnvironment CompileEnvironment, LinkEnvironment LinkEnvironment, List<PrecompiledHeaderTemplate> SharedPCHs, ISourceFileWorkingSet WorkingSet, BuildPrerequisites Prerequisites, DirectoryReference ExeDir, List<Action> Actions)
+		public List<FileItem> Build(ReadOnlyTargetRules Target, UEToolChain ToolChain, CppCompileEnvironment CompileEnvironment, LinkEnvironment LinkEnvironment, ISourceFileWorkingSet WorkingSet, DirectoryReference ExeDir, TargetMakefile Makefile)
 		{
 			// Return nothing if we're using precompiled binaries. If we're not linking, we might want just one module to be compiled (eg. a foreign plugin), so allow any actions to run.
 			if (bUsePrecompiled && !Target.bDisableLinking)
@@ -171,7 +169,7 @@ namespace UnrealBuildTool
 			}
 
 			// Setup linking environment.
-			LinkEnvironment BinaryLinkEnvironment = SetupBinaryLinkEnvironment(Target, ToolChain, LinkEnvironment, CompileEnvironment, SharedPCHs, WorkingSet, Prerequisites, ExeDir, Actions);
+			LinkEnvironment BinaryLinkEnvironment = SetupBinaryLinkEnvironment(Target, ToolChain, LinkEnvironment, CompileEnvironment, WorkingSet, ExeDir, Makefile);
 
 			// If we're generating projects, we only need include paths and definitions, there is no need to run the linking logic.
 			if (ProjectFileGenerator.bGenerateProjectFiles)
@@ -195,12 +193,12 @@ namespace UnrealBuildTool
 				if (BinaryLinkEnvironment.Platform != CppPlatform.Mac && BinaryLinkEnvironment.Platform != CppPlatform.Linux)
 				{
 					// Create the import library.
-					OutputFiles.AddRange(ToolChain.LinkAllFiles(BinaryLinkEnvironment, true, Actions));
+					OutputFiles.AddRange(ToolChain.LinkAllFiles(BinaryLinkEnvironment, true, Makefile.Actions));
 				}
 			}
 
 			// Link the binary.
-			FileItem[] Executables = ToolChain.LinkAllFiles(BinaryLinkEnvironment, false, Actions);
+			FileItem[] Executables = ToolChain.LinkAllFiles(BinaryLinkEnvironment, false, Makefile.Actions);
 			OutputFiles.AddRange(Executables);
 
 			// Produce additional console app if requested
@@ -213,12 +211,12 @@ namespace UnrealBuildTool
 				ConsoleAppLinkEvironment.OutputFilePaths = ConsoleAppLinkEvironment.OutputFilePaths.Select(Path => GetAdditionalConsoleAppPath(Path)).ToList();
 
 				// Link the console app executable
-				OutputFiles.AddRange(ToolChain.LinkAllFiles(ConsoleAppLinkEvironment, false, Actions));
+				OutputFiles.AddRange(ToolChain.LinkAllFiles(ConsoleAppLinkEvironment, false, Makefile.Actions));
 			}
 
 			foreach (FileItem Executable in Executables)
 			{
-				OutputFiles.AddRange(ToolChain.PostBuild(Executable, BinaryLinkEnvironment, Actions));
+				OutputFiles.AddRange(ToolChain.PostBuild(Executable, BinaryLinkEnvironment, Makefile.Actions));
 			}
 
 			return OutputFiles;
@@ -650,7 +648,7 @@ namespace UnrealBuildTool
 			return BinaryCompileEnvironment;
 		}
 
-		private LinkEnvironment SetupBinaryLinkEnvironment(ReadOnlyTargetRules Target, UEToolChain ToolChain, LinkEnvironment LinkEnvironment, CppCompileEnvironment CompileEnvironment, List<PrecompiledHeaderTemplate> SharedPCHs, ISourceFileWorkingSet WorkingSet, BuildPrerequisites Prerequisites, DirectoryReference ExeDir, List<Action> Actions)
+		private LinkEnvironment SetupBinaryLinkEnvironment(ReadOnlyTargetRules Target, UEToolChain ToolChain, LinkEnvironment LinkEnvironment, CppCompileEnvironment CompileEnvironment, ISourceFileWorkingSet WorkingSet, DirectoryReference ExeDir, TargetMakefile Makefile)
 		{
 			LinkEnvironment BinaryLinkEnvironment = new LinkEnvironment(LinkEnvironment);
 			HashSet<UEBuildModule> LinkEnvironmentVisitedModules = new HashSet<UEBuildModule>();
@@ -670,7 +668,7 @@ namespace UnrealBuildTool
 				{
 					// Compile each module.
 					Log.TraceVerbose("Compile module: " + Module.Name);
-					LinkInputFiles = Module.Compile(Target, ToolChain, BinaryCompileEnvironment, SharedPCHs, WorkingSet, Prerequisites, Actions);
+					LinkInputFiles = Module.Compile(Target, ToolChain, BinaryCompileEnvironment, WorkingSet, Makefile);
 
 					// NOTE: Because of 'Shared PCHs', in monolithic builds the same PCH file may appear as a link input
 					// multiple times for a single binary.  We'll check for that here, and only add it once.  This avoids
@@ -746,7 +744,7 @@ namespace UnrealBuildTool
 
 						// Otherwise compile the default resource file per-binary, so that it gets the correct ORIGINAL_FILE_NAME macro.
 						FileItem DefaultResourceFile = FileItem.GetItemByFileReference(FileReference.Combine(UnrealBuildTool.EngineDirectory, "Build", "Windows", "Default.rc2"));
-						CPPOutput DefaultResourceOutput = ToolChain.CompileRCFiles(ResourceCompileEnvironment, new List<FileItem> { DefaultResourceFile }, ResourceIntermediateDirectory, Actions);
+						CPPOutput DefaultResourceOutput = ToolChain.CompileRCFiles(ResourceCompileEnvironment, new List<FileItem> { DefaultResourceFile }, ResourceIntermediateDirectory, Makefile.Actions);
 						BinaryLinkEnvironment.InputFiles.AddRange(DefaultResourceOutput.ObjectFiles);
 					}
 				}
diff --git a/Engine/Source/Programs/UnrealBuildTool/Configuration/UEBuildModule.cs b/Engine/Source/Programs/UnrealBuildTool/Configuration/UEBuildModule.cs
--- a/Engine/Source/Programs/UnrealBuildTool/Configuration/UEBuildModule.cs
+++ b/Engine/Source/Programs/UnrealBuildTool/Configuration/UEBuildModule.cs
@@ -725,7 +725,7 @@ namespace UnrealBuildTool
 		/// <summary>
 		/// Compiles the module, and returns a list of files output by the compiler.
 		/// </summary>
-		public abstract List<FileItem> Compile(ReadOnlyTargetRules Target, UEToolChain ToolChain, CppCompileEnvironment CompileEnvironment, List<PrecompiledHeaderTemplate> SharedPCHModules, ISourceFileWorkingSet WorkingSet, BuildPrerequisites Predicates, List<Action> Actions);
+		public abstract List<FileItem> Compile(ReadOnlyTargetRules Target, UEToolChain ToolChain, CppCompileEnvironment CompileEnvironment, ISourceFileWorkingSet WorkingSet, TargetMakefile Makefile);
 
 		// Object interface.
 		public override string ToString()
diff --git a/Engine/Source/Programs/UnrealBuildTool/Configuration/UEBuildModuleCPP.cs b/Engine/Source/Programs/UnrealBuildTool/Configuration/UEBuildModuleCPP.cs
--- a/Engine/Source/Programs/UnrealBuildTool/Configuration/UEBuildModuleCPP.cs
+++ b/Engine/Source/Programs/UnrealBuildTool/Configuration/UEBuildModuleCPP.cs
@@ -248,7 +248,7 @@ namespace UnrealBuildTool
 		}
 
 		// UEBuildModule interface.
-		public override List<FileItem> Compile(ReadOnlyTargetRules Target, UEToolChain ToolChain, CppCompileEnvironment BinaryCompileEnvironment, List<PrecompiledHeaderTemplate> SharedPCHs, ISourceFileWorkingSet WorkingSet, BuildPrerequisites Prerequisites, List<Action> Actions)
+		public override List<FileItem> Compile(ReadOnlyTargetRules Target, UEToolChain ToolChain, CppCompileEnvironment BinaryCompileEnvironment, ISourceFileWorkingSet WorkingSet, TargetMakefile Makefile)
 		{
 			UEBuildPlatform BuildPlatform = UEBuildPlatform.GetBuildPlatformForCPPTargetPlatform(BinaryCompileEnvironment.Platform);
 
@@ -272,7 +272,7 @@ namespace UnrealBuildTool
 				return LinkInputFiles;
 			}
 
-			SourceFileCollection SourceFilesToBuild = FindSourceFiles(Target.Platform, Prerequisites.SourceDirectories);
+			SourceFileCollection SourceFilesToBuild = FindSourceFiles(Target.Platform, Makefile.SourceDirectories);
 
 			SourceDirectories = new HashSet<DirectoryReference>(SourceFilesToBuild.AllFiles.Select(x => x.Location.Directory));
 
@@ -348,7 +348,7 @@ namespace UnrealBuildTool
 				// If this module doesn't need a shared PCH, configure that
 				if(Rules.PrivatePCHHeaderFile != null && (Rules.PCHUsage == ModuleRules.PCHUsageMode.NoSharedPCHs || Rules.PCHUsage == ModuleRules.PCHUsageMode.UseExplicitOrSharedPCHs))
 				{
-					PrecompiledHeaderInstance Instance = CreatePrivatePCH(ToolChain, FileItem.GetItemByFileReference(FileReference.Combine(ModuleDirectory, Rules.PrivatePCHHeaderFile)), CompileEnvironment, Actions);
+					PrecompiledHeaderInstance Instance = CreatePrivatePCH(ToolChain, FileItem.GetItemByFileReference(FileReference.Combine(ModuleDirectory, Rules.PrivatePCHHeaderFile)), CompileEnvironment, Makefile.Actions);
 
 					CompileEnvironment = new CppCompileEnvironment(CompileEnvironment);
 					CompileEnvironment.Definitions.Clear();
@@ -360,17 +360,17 @@ namespace UnrealBuildTool
 				}
 
 				// Try to find a suitable shared PCH for this module
-				if (CompileEnvironment.PrecompiledHeaderFile == null && SharedPCHs.Count > 0 && !CompileEnvironment.bIsBuildingLibrary && Rules.PCHUsage != ModuleRules.PCHUsageMode.NoSharedPCHs)
+				if (CompileEnvironment.PrecompiledHeaderFile == null && CompileEnvironment.SharedPCHs.Count > 0 && !CompileEnvironment.bIsBuildingLibrary && Rules.PCHUsage != ModuleRules.PCHUsageMode.NoSharedPCHs)
 				{
 					// Find all the dependencies of this module
 					HashSet<UEBuildModule> ReferencedModules = new HashSet<UEBuildModule>();
 					GetAllDependencyModules(new List<UEBuildModule>(), ReferencedModules, bIncludeDynamicallyLoaded: false, bForceCircular: false, bOnlyDirectDependencies: true);
 
 					// Find the first shared PCH module we can use
-					PrecompiledHeaderTemplate Template = SharedPCHs.FirstOrDefault(x => ReferencedModules.Contains(x.Module));
+					PrecompiledHeaderTemplate Template = CompileEnvironment.SharedPCHs.FirstOrDefault(x => ReferencedModules.Contains(x.Module));
 					if(Template != null && Template.IsValidFor(CompileEnvironment))
 					{
-						PrecompiledHeaderInstance Instance = FindOrCreateSharedPCH(ToolChain, Template, ModuleCompileEnvironment.bOptimizeCode, ModuleCompileEnvironment.bUseRTTI, ModuleCompileEnvironment.bEnableExceptions, Actions);
+						PrecompiledHeaderInstance Instance = FindOrCreateSharedPCH(ToolChain, Template, ModuleCompileEnvironment.bOptimizeCode, ModuleCompileEnvironment.bUseRTTI, ModuleCompileEnvironment.bEnableExceptions, Makefile.Actions);
 
 						FileReference PrivateDefinitionsFile = FileReference.Combine(IntermediateDirectory, String.Format("Definitions.{0}.h", Name));
 
@@ -413,12 +413,12 @@ namespace UnrealBuildTool
 			List<FileItem> CPPFilesToCompile = SourceFilesToBuild.CPPFiles;
 			if (bModuleUsesUnityBuild)
 			{
-				CPPFilesToCompile = Unity.GenerateUnityCPPs(Target, CPPFilesToCompile, CompileEnvironment, WorkingSet, Prerequisites, Rules.ShortName ?? Name, IntermediateDirectory);
-				LinkInputFiles.AddRange(CompileUnityFilesWithToolChain(Target, ToolChain, CompileEnvironment, ModuleCompileEnvironment, CPPFilesToCompile, Actions).ObjectFiles);
+				CPPFilesToCompile = Unity.GenerateUnityCPPs(Target, CPPFilesToCompile, CompileEnvironment, WorkingSet, Rules.ShortName ?? Name, IntermediateDirectory, Makefile);
+				LinkInputFiles.AddRange(CompileUnityFilesWithToolChain(Target, ToolChain, CompileEnvironment, ModuleCompileEnvironment, CPPFilesToCompile, Makefile.Actions).ObjectFiles);
 			}
 			else
 			{
-				LinkInputFiles.AddRange(ToolChain.CompileCPPFiles(CompileEnvironment, CPPFilesToCompile, IntermediateDirectory, Name, Actions).ObjectFiles);
+				LinkInputFiles.AddRange(ToolChain.CompileCPPFiles(CompileEnvironment, CPPFilesToCompile, IntermediateDirectory, Name, Makefile.Actions).ObjectFiles);
 			}
 
 			// Compile all the generated CPP files
@@ -460,12 +460,12 @@ namespace UnrealBuildTool
 
 					if (bModuleUsesUnityBuild)
 					{
-						GeneratedFileItems = Unity.GenerateUnityCPPs(Target, GeneratedFileItems, GeneratedCPPCompileEnvironment, WorkingSet, Prerequisites, (Rules.ShortName ?? Name) + ".gen", IntermediateDirectory);
-						LinkInputFiles.AddRange(CompileUnityFilesWithToolChain(Target, ToolChain, GeneratedCPPCompileEnvironment, ModuleCompileEnvironment, GeneratedFileItems, Actions).ObjectFiles);
+						GeneratedFileItems = Unity.GenerateUnityCPPs(Target, GeneratedFileItems, GeneratedCPPCompileEnvironment, WorkingSet, (Rules.ShortName ?? Name) + ".gen", IntermediateDirectory, Makefile);
+						LinkInputFiles.AddRange(CompileUnityFilesWithToolChain(Target, ToolChain, GeneratedCPPCompileEnvironment, ModuleCompileEnvironment, GeneratedFileItems, Makefile.Actions).ObjectFiles);
 					}
 					else
 					{
-						LinkInputFiles.AddRange(ToolChain.CompileCPPFiles(GeneratedCPPCompileEnvironment, GeneratedFileItems, IntermediateDirectory, Name, Actions).ObjectFiles);
+						LinkInputFiles.AddRange(ToolChain.CompileCPPFiles(GeneratedCPPCompileEnvironment, GeneratedFileItems, IntermediateDirectory, Name, Makefile.Actions).ObjectFiles);
 					}
 				}
 			}
@@ -473,19 +473,19 @@ namespace UnrealBuildTool
 			// Compile C files directly. Do not use a PCH here, because a C++ PCH is not compatible with C source files.
 			if(SourceFilesToBuild.CFiles.Count > 0)
 			{
-				LinkInputFiles.AddRange(ToolChain.CompileCPPFiles(ModuleCompileEnvironment, SourceFilesToBuild.CFiles, IntermediateDirectory, Name, Actions).ObjectFiles);
+				LinkInputFiles.AddRange(ToolChain.CompileCPPFiles(ModuleCompileEnvironment, SourceFilesToBuild.CFiles, IntermediateDirectory, Name, Makefile.Actions).ObjectFiles);
 			}
 
 			// Compile CC files directly.
 			if(SourceFilesToBuild.CCFiles.Count > 0)
 			{
-				LinkInputFiles.AddRange(ToolChain.CompileCPPFiles(CompileEnvironment, SourceFilesToBuild.CCFiles, IntermediateDirectory, Name, Actions).ObjectFiles);
+				LinkInputFiles.AddRange(ToolChain.CompileCPPFiles(CompileEnvironment, SourceFilesToBuild.CCFiles, IntermediateDirectory, Name, Makefile.Actions).ObjectFiles);
 			}
 
 			// Compile MM files directly.
 			if(SourceFilesToBuild.MMFiles.Count > 0)
 			{
-				LinkInputFiles.AddRange(ToolChain.CompileCPPFiles(CompileEnvironment, SourceFilesToBuild.MMFiles, IntermediateDirectory, Name, Actions).ObjectFiles);
+				LinkInputFiles.AddRange(ToolChain.CompileCPPFiles(CompileEnvironment, SourceFilesToBuild.MMFiles, IntermediateDirectory, Name, Makefile.Actions).ObjectFiles);
 			}
 
 			// Compile RC files. The resource compiler does not work with response files, and using the regular compile environment can easily result in the 
@@ -495,7 +495,7 @@ namespace UnrealBuildTool
 			{
 				CppCompileEnvironment ResourceCompileEnvironment = new CppCompileEnvironment(BinaryCompileEnvironment);
 				WindowsPlatform.SetupResourceCompileEnvironment(ResourceCompileEnvironment, Target);
-				LinkInputFiles.AddRange(ToolChain.CompileRCFiles(ResourceCompileEnvironment, SourceFilesToBuild.RCFiles, IntermediateDirectory, Actions).ObjectFiles);
+				LinkInputFiles.AddRange(ToolChain.CompileRCFiles(ResourceCompileEnvironment, SourceFilesToBuild.RCFiles, IntermediateDirectory, Makefile.Actions).ObjectFiles);
 			}
 
 			// Write the compiled manifest
diff --git a/Engine/Source/Programs/UnrealBuildTool/Configuration/UEBuildModuleExternal.cs b/Engine/Source/Programs/UnrealBuildTool/Configuration/UEBuildModuleExternal.cs
--- a/Engine/Source/Programs/UnrealBuildTool/Configuration/UEBuildModuleExternal.cs
+++ b/Engine/Source/Programs/UnrealBuildTool/Configuration/UEBuildModuleExternal.cs
@@ -20,7 +20,7 @@ namespace UnrealBuildTool
 		}
 
 		// UEBuildModule interface.
-		public override List<FileItem> Compile(ReadOnlyTargetRules Target, UEToolChain ToolChain, CppCompileEnvironment CompileEnvironment, List<PrecompiledHeaderTemplate> SharedPCHs, ISourceFileWorkingSet WorkingSet, BuildPrerequisites Prerequisites, List<Action> Actions)
+		public override List<FileItem> Compile(ReadOnlyTargetRules Target, UEToolChain ToolChain, CppCompileEnvironment CompileEnvironment, ISourceFileWorkingSet WorkingSet, TargetMakefile Makefile)
 		{
 			return new List<FileItem>();
 		}
diff --git a/Engine/Source/Programs/UnrealBuildTool/Configuration/UEBuildTarget.cs b/Engine/Source/Programs/UnrealBuildTool/Configuration/UEBuildTarget.cs
--- a/Engine/Source/Programs/UnrealBuildTool/Configuration/UEBuildTarget.cs
+++ b/Engine/Source/Programs/UnrealBuildTool/Configuration/UEBuildTarget.cs
@@ -1096,7 +1096,7 @@ namespace UnrealBuildTool
 		/// <summary>
 		/// Builds the target, appending list of output files and returns building result.
 		/// </summary>
-		public void Build(BuildConfiguration BuildConfiguration, List<FileItem> OutputItems, Dictionary<string, FileItem[]> ModuleNameToOutputItems, List<UHTModuleInfo> UObjectModules, ISourceFileWorkingSet WorkingSet, List<Action> Actions, BuildPrerequisites Prerequisites, bool bIsAssemblingBuild)
+		public TargetMakefile Build(BuildConfiguration BuildConfiguration, ISourceFileWorkingSet WorkingSet, bool bIsAssemblingBuild)
 		{
 			CppPlatform CppPlatform = UEBuildPlatform.GetBuildPlatform(Platform).DefaultCppPlatform;
 			CppConfiguration CppConfiguration = GetCppConfiguration(Configuration);
@@ -1260,26 +1260,28 @@ namespace UnrealBuildTool
 			// On Mac and Linux we have actions that should be executed after all the binaries are created
 			TargetToolChain.SetupBundleDependencies(Binaries, TargetName);
 
+			// Create the makefile
+			TargetMakefile Makefile = new TargetMakefile(ReceiptFileName, ProjectIntermediateDirectory, TargetType, bDeployAfterCompile, bHasProjectScriptPlugin);
+
 			// Generate headers
 			HashSet<UEBuildModuleCPP> ModulesToGenerateHeadersFor = GatherDependencyModules(OriginalBinaries.ToList());
 			using(Timeline.ScopeEvent("ExternalExecution.SetupUObjectModules()"))
 			{
-				ExternalExecution.SetupUObjectModules(ModulesToGenerateHeadersFor, Rules.Platform, ProjectDescriptor, UObjectModules, Prerequisites.UObjectModuleHeaders, Rules.GeneratedCodeVersion, bIsAssemblingBuild, MetadataCache);
+				ExternalExecution.SetupUObjectModules(ModulesToGenerateHeadersFor, Rules.Platform, ProjectDescriptor, Makefile.UObjectModules, Makefile.UObjectModuleHeaders, Rules.GeneratedCodeVersion, bIsAssemblingBuild, MetadataCache);
 			}
 
 			// NOTE: Even in Gather mode, we need to run UHT to make sure the files exist for the static action graph to be setup correctly.  This is because UHT generates .cpp
 			// files that are injected as top level prerequisites.  If UHT only emitted included header files, we wouldn't need to run it during the Gather phase at all.
-			if (UObjectModules.Count > 0)
+			if (Makefile.UObjectModules.Count > 0)
 			{
 				FileReference ModuleInfoFileName = FileReference.Combine(ProjectIntermediateDirectory, TargetName + ".uhtmanifest");
-				ExternalExecution.ExecuteHeaderToolIfNecessary(BuildConfiguration, ProjectFile, TargetName, TargetType, bHasProjectScriptPlugin, UObjectModules, ModuleInfoFileName, true, bIsAssemblingBuild, WorkingSet);
+				ExternalExecution.ExecuteHeaderToolIfNecessary(BuildConfiguration, ProjectFile, TargetName, TargetType, bHasProjectScriptPlugin, Makefile.UObjectModules, ModuleInfoFileName, true, bIsAssemblingBuild, WorkingSet);
 			}
 
 			// Find all the shared PCHs.
-			List<PrecompiledHeaderTemplate> SharedPCHs = new List<PrecompiledHeaderTemplate>();
 			if (Rules.bUseSharedPCHs)
 			{
-				SharedPCHs = FindSharedPCHs(OriginalBinaries, GlobalCompileEnvironment);
+				FindSharedPCHs(OriginalBinaries, GlobalCompileEnvironment);
 			}
 
 			// Compile the resource files common to all DLLs on Windows
@@ -1295,7 +1297,7 @@ namespace UnrealBuildTool
 
 						WindowsPlatform.SetupResourceCompileEnvironment(DefaultResourceCompileEnvironment, Rules);
 
-						CPPOutput DefaultResourceOutput = TargetToolChain.CompileRCFiles(DefaultResourceCompileEnvironment, new List<FileItem> { DefaultResourceFile }, EngineIntermediateDirectory, Actions);
+						CPPOutput DefaultResourceOutput = TargetToolChain.CompileRCFiles(DefaultResourceCompileEnvironment, new List<FileItem> { DefaultResourceFile }, EngineIntermediateDirectory, Makefile.Actions);
 						GlobalLinkEnvironment.DefaultResourceFiles.AddRange(DefaultResourceOutput.ObjectFiles);
 					}
 				}
@@ -1311,12 +1313,12 @@ namespace UnrealBuildTool
 			{
 				foreach (UEBuildBinary Binary in Binaries)
 				{
-					List<FileItem> BinaryOutputItems = Binary.Build(Rules, TargetToolChain, GlobalCompileEnvironment, GlobalLinkEnvironment, SharedPCHs, WorkingSet, Prerequisites, ExeDir, Actions);
+					List<FileItem> BinaryOutputItems = Binary.Build(Rules, TargetToolChain, GlobalCompileEnvironment, GlobalLinkEnvironment, WorkingSet, ExeDir, Makefile);
 					if(!bCompileMonolithic)
 					{
-						ModuleNameToOutputItems[Binary.PrimaryModule.Name] = BinaryOutputItems.ToArray();
+						Makefile.ModuleNameToOutputItems[Binary.PrimaryModule.Name] = BinaryOutputItems.ToArray();
 					}
-					OutputItems.AddRange(BinaryOutputItems);
+					Makefile.OutputItems.AddRange(BinaryOutputItems);
 				}
 			}
 
@@ -1331,18 +1333,18 @@ namespace UnrealBuildTool
 			{
 				if(!UnrealBuildTool.IsFileInstalled(Pair.Key))
 				{
-					OutputItems.Add(CreateCopyAction(Pair.Value, Pair.Key, Actions));
+					Makefile.OutputItems.Add(CreateCopyAction(Pair.Value, Pair.Key, Makefile.Actions));
 				}
 			}
 
 			// If we're just precompiling a plugin, only include output items which are under that directory
 			if(ForeignPlugin != null)
 			{
-				OutputItems.RemoveAll(x => !x.Location.IsUnderDirectory(ForeignPlugin.Directory));
+				Makefile.OutputItems.RemoveAll(x => !x.Location.IsUnderDirectory(ForeignPlugin.Directory));
 			}
 
 			// Allow the toolchain to modify the final output items
-			TargetToolChain.FinalizeOutput(Rules, OutputItems, Actions);
+			TargetToolChain.FinalizeOutput(Rules, Makefile);
 
 			// Get all the regular build products
 			List<KeyValuePair<FileReference, BuildProductType>> BuildProducts = new List<KeyValuePair<FileReference, BuildProductType>>();
@@ -1425,12 +1427,12 @@ namespace UnrealBuildTool
 				WriteMetadataAction.WorkingDirectory = UnrealBuildTool.EngineSourceDirectory.FullName;
 				WriteMetadataAction.StatusDescription = ReceiptFileName.GetFileName();
 				WriteMetadataAction.bCanExecuteRemotely = false;
-				WriteMetadataAction.PrerequisiteItems.AddRange(OutputItems);
+				WriteMetadataAction.PrerequisiteItems.AddRange(Makefile.OutputItems);
 				WriteMetadataAction.PrerequisiteItems.Add(FileItem.GetItemByPath(Assembly.GetExecutingAssembly().Location));
 				WriteMetadataAction.ProducedItems.Add(FileItem.GetItemByFileReference(ReceiptFileName));
-				Actions.Add(WriteMetadataAction);
+				Makefile.Actions.Add(WriteMetadataAction);
 
-				OutputItems.AddRange(WriteMetadataAction.ProducedItems);
+				Makefile.OutputItems.AddRange(WriteMetadataAction.ProducedItems);
 
 				// Create actions to run the post build steps
 				foreach(FileReference PostBuildStepScript in PostBuildStepScripts)
@@ -1453,9 +1455,9 @@ namespace UnrealBuildTool
 					PostBuildStepAction.bCanExecuteRemotely = false;
 					PostBuildStepAction.PrerequisiteItems.Add(FileItem.GetItemByFileReference(ReceiptFileName));
 					PostBuildStepAction.ProducedItems.Add(FileItem.GetItemByFileReference(OutputFile));
-					Actions.Add(PostBuildStepAction);
+					Makefile.Actions.Add(PostBuildStepAction);
 
-					OutputItems.AddRange(PostBuildStepAction.ProducedItems);
+					Makefile.OutputItems.AddRange(PostBuildStepAction.ProducedItems);
 				}
 			}
 
@@ -1496,19 +1498,21 @@ namespace UnrealBuildTool
 			}
 
 			// Add all the input files to the predicate store
-			Prerequisites.AdditionalDependencies.Add(FileItem.GetItemByFileReference(TargetRulesFile));
+			Makefile.AdditionalDependencies.Add(FileItem.GetItemByFileReference(TargetRulesFile));
 			foreach(UEBuildModule Module in Modules.Values)
 			{
-				Prerequisites.AdditionalDependencies.Add(FileItem.GetItemByFileReference(Module.RulesFile));
+				Makefile.AdditionalDependencies.Add(FileItem.GetItemByFileReference(Module.RulesFile));
 				foreach(string ExternalDependency in Module.Rules.ExternalDependencies)
 				{
 					FileReference Location = FileReference.Combine(Module.RulesFile.Directory, ExternalDependency);
-					Prerequisites.AdditionalDependencies.Add(FileItem.GetItemByFileReference(Location));
+					Makefile.AdditionalDependencies.Add(FileItem.GetItemByFileReference(Location));
 				}
 			}
 
 			// Clean any stale modules which exist in multiple output directories. This can lead to the wrong DLL being loaded on Windows.
 			CleanStaleModules();
+
+			return Makefile;
 		}
 
 		/// <summary>
@@ -1975,7 +1979,7 @@ namespace UnrealBuildTool
 		/// Finds a list of module names which can be hot-reloaded
 		/// </summary>
 		/// <returns>Set of module names</returns>
-		public HashSet<string> GetHotReloadModuleNames()
+		private HashSet<string> GetHotReloadModuleNames()
 		{
 			HashSet<string> HotReloadModuleNames = new HashSet<string>(StringComparer.OrdinalIgnoreCase);
 			foreach (UEBuildBinary Binary in Binaries)
@@ -1995,8 +1999,9 @@ namespace UnrealBuildTool
 		/// <summary>
 		/// Determines which modules can be used to create shared PCHs
 		/// </summary>
-		/// <returns>List of shared PCH modules, in order of preference</returns>
-		List<PrecompiledHeaderTemplate> FindSharedPCHs(List<UEBuildBinary> OriginalBinaries, CppCompileEnvironment GlobalCompileEnvironment)
+		/// <param name="OriginalBinaries">The list of binaries</param>
+		/// <param name="GlobalCompileEnvironment">The compile environment. The shared PCHs will be added to the SharedPCHs list in this.</param>
+		void FindSharedPCHs(List<UEBuildBinary> OriginalBinaries, CppCompileEnvironment GlobalCompileEnvironment)
 		{
 			// Find how many other shared PCH modules each module depends on, and use that to sort the shared PCHs by reverse order of size.
 			HashSet<UEBuildModuleCPP> SharedPCHModules = new HashSet<UEBuildModuleCPP>();
@@ -2028,7 +2033,7 @@ namespace UnrealBuildTool
 			}
 
 			// Create the shared PCH modules, in order
-			List<PrecompiledHeaderTemplate> OrderedSharedPCHModules = new List<PrecompiledHeaderTemplate>();
+			List<PrecompiledHeaderTemplate> OrderedSharedPCHModules = GlobalCompileEnvironment.SharedPCHs;
 			foreach(UEBuildModuleCPP Module in SharedPCHModuleToPriority.OrderByDescending(x => x.Value).Select(x => x.Key))
 			{
 				OrderedSharedPCHModules.Add(Module.CreateSharedPCHTemplate(this, GlobalCompileEnvironment));
@@ -2043,7 +2048,6 @@ namespace UnrealBuildTool
 					Log.TraceVerbose("	" + SharedPCHModule.Module.Name);
 				}
 			}
-			return OrderedSharedPCHModules;
 		}
 
 		/// <summary>
diff --git a/Engine/Source/Programs/UnrealBuildTool/Modes/BuildMode.cs b/Engine/Source/Programs/UnrealBuildTool/Modes/BuildMode.cs
--- a/Engine/Source/Programs/UnrealBuildTool/Modes/BuildMode.cs
+++ b/Engine/Source/Programs/UnrealBuildTool/Modes/BuildMode.cs
@@ -175,111 +175,7 @@ namespace UnrealBuildTool
 			TargetMakefile[] Makefiles = new TargetMakefile[TargetDescriptors.Count];
 			for(int TargetIdx = 0; TargetIdx < TargetDescriptors.Count; TargetIdx++)
 			{
-				TargetDescriptor TargetDesc = TargetDescriptors[TargetIdx];
-
-				// Get the path to the makefile for this target
-				FileReference MakefileLocation = null;
-				if(BuildConfiguration.bUseUBTMakefiles && TargetDesc.SingleFileToCompile == null)
-				{
-					MakefileLocation = TargetMakefile.GetLocation(TargetDesc.ProjectFile, TargetDesc.Name, TargetDesc.Platform, TargetDesc.Configuration);
-				}
-
-				// Try to load an existing makefile
-				TargetMakefile Makefile = null;
-				if(MakefileLocation != null)
-				{
-					using(Timeline.ScopeEvent("TargetMakefile.Load()"))
-					{
-						string ReasonNotLoaded;
-						Makefile = TargetMakefile.Load(MakefileLocation, TargetDesc.ProjectFile, TargetDesc.Platform, WorkingSet, out ReasonNotLoaded);
-						if (Makefile == null)
-						{
-							Log.TraceInformation("Creating makefile for {0} ({1})", TargetDesc.Name, ReasonNotLoaded);
-						}
-					}
-				}
-
-				// If we couldn't load a makefile, create a new one
-				if(Makefile == null)
-				{
-					// Create the target
-					UEBuildTarget Target;
-					using(Timeline.ScopeEvent("UEBuildTarget.Create()"))
-					{
-						Target = UEBuildTarget.Create(TargetDesc, BuildConfiguration.bSkipRulesCompile, BuildConfiguration.bUsePrecompiled);
-					}
-
-					// Build the target
-					const bool bIsAssemblingBuild = true;
-
-					List<FileItem> OutputItems = new List<FileItem>();
-					Dictionary<string, FileItem[]> ModuleNameToOutputItems = new Dictionary<string, FileItem[]>();
-					List<UHTModuleInfo> UObjectModules = new List<UHTModuleInfo>();
-					List<Action> Actions = new List<Action>();
-					BuildPrerequisites Prerequisites = new BuildPrerequisites();
-
-					using(Timeline.ScopeEvent("UEBuildTarget.Build()"))
-					{
-						Target.Build(BuildConfiguration, OutputItems, ModuleNameToOutputItems, UObjectModules, WorkingSet, Actions, Prerequisites, bIsAssemblingBuild);
-					}
-
-					// Create the makefile
-					Makefile = new TargetMakefile(Target.TargetType);
-					Makefile.ReceiptFile = Target.ReceiptFileName;
-					Makefile.Actions = Actions;
-					Makefile.OutputItems = OutputItems;
-					Makefile.ModuleNameToOutputItems = ModuleNameToOutputItems;
-					Makefile.UObjectModules = UObjectModules;
-					Makefile.Prerequisites = Prerequisites;
-					Makefile.HotReloadModuleNames = Target.GetHotReloadModuleNames();
-					Makefile.Prerequisites = Prerequisites;
-					Makefile.ProjectIntermediateDirectory = Target.ProjectIntermediateDirectory;
-					Makefile.bDeployAfterCompile = Target.bDeployAfterCompile;
-					Makefile.bHasProjectScriptPlugin = Target.bHasProjectScriptPlugin;
-					Makefile.PreBuildScripts = Target.PreBuildStepScripts;
-
-					// Save the environment variables
-					foreach (System.Collections.DictionaryEntry EnvironmentVariable in Environment.GetEnvironmentVariables())
-					{
-						Makefile.EnvironmentVariables.Add(Tuple.Create((string)EnvironmentVariable.Key, (string)EnvironmentVariable.Value));
-					}
-
-					// Save the makefile for next time
-					if(MakefileLocation != null)
-					{
-						using(Timeline.ScopeEvent("TargetMakefile.Save()"))
-						{
-							Makefile.Save(MakefileLocation);
-						}
-					}
-				}
-				else
-				{
-					// Restore the environment variables
-					foreach (Tuple<string, string> EnvironmentVariable in Makefile.EnvironmentVariables)
-					{
-						Environment.SetEnvironmentVariable(EnvironmentVariable.Item1, EnvironmentVariable.Item2);
-					}
-
-					// Execute all the pre-build steps
-					if (!BuildConfiguration.bXGEExport)
-					{
-						Utils.ExecuteCustomBuildSteps(Makefile.PreBuildScripts);
-					}
-
-					// If the target needs UHT to be run, we'll go ahead and do that now
-					if (Makefile.UObjectModules.Count > 0)
-					{
-						const bool bIsGatheringBuild = false;
-						const bool bIsAssemblingBuild = true;
-
-						FileReference ModuleInfoFileName = FileReference.Combine(Makefile.ProjectIntermediateDirectory, TargetDesc.Name + ".uhtmanifest");
-						ExternalExecution.ExecuteHeaderToolIfNecessary(BuildConfiguration, TargetDesc.ProjectFile, TargetDesc.Name, Makefile.TargetType, Makefile.bHasProjectScriptPlugin, UObjectModules: Makefile.UObjectModules, ModuleInfoFileName: ModuleInfoFileName, bIsGatheringBuild: bIsGatheringBuild, bIsAssemblingBuild: bIsAssemblingBuild, WorkingSet: WorkingSet);
-					}
-				}
-
-				// Add the makefile to the list to build
-				Makefiles[TargetIdx] = Makefile;
+				Makefiles[TargetIdx] = CreateMakefile(BuildConfiguration, TargetDescriptors[TargetIdx], WorkingSet);
 			}
 
 			// Execute the build
@@ -377,11 +273,102 @@ namespace UnrealBuildTool
 			}
 		}
 
+		/// <summary>
+		/// Creates the makefile for a target. If an existing, valid makefile already exists on disk, loads that instead.
+		/// </summary>
+		/// <param name="BuildConfiguration">The build configuration</param>
+		/// <param name="TargetDescriptor">Target being built</param>
+		/// <param name="WorkingSet">Set of source files which are part of the working set</param>
+		/// <returns>Makefile for the given target</returns>
+		static TargetMakefile CreateMakefile(BuildConfiguration BuildConfiguration, TargetDescriptor TargetDescriptor, ISourceFileWorkingSet WorkingSet)
+		{
+			// Get the path to the makefile for this target
+			FileReference MakefileLocation = null;
+			if(BuildConfiguration.bUseUBTMakefiles && TargetDescriptor.SingleFileToCompile == null)
+			{
+				MakefileLocation = TargetMakefile.GetLocation(TargetDescriptor.ProjectFile, TargetDescriptor.Name, TargetDescriptor.Platform, TargetDescriptor.Configuration);
+			}
+
+			// Try to load an existing makefile
+			TargetMakefile Makefile = null;
+			if(MakefileLocation != null)
+			{
+				using(Timeline.ScopeEvent("TargetMakefile.Load()"))
+				{
+					string ReasonNotLoaded;
+					Makefile = TargetMakefile.Load(MakefileLocation, TargetDescriptor.ProjectFile, TargetDescriptor.Platform, WorkingSet, out ReasonNotLoaded);
+					if (Makefile == null)
+					{
+						Log.TraceInformation("Creating makefile for {0} ({1})", TargetDescriptor.Name, ReasonNotLoaded);
+					}
+				}
+			}
+
+			// If we couldn't load a makefile, create a new one
+			if(Makefile == null)
+			{
+				// Create the target
+				UEBuildTarget Target;
+				using(Timeline.ScopeEvent("UEBuildTarget.Create()"))
+				{
+					Target = UEBuildTarget.Create(TargetDescriptor, BuildConfiguration.bSkipRulesCompile, BuildConfiguration.bUsePrecompiled);
+				}
+
+				// Build the target
+				const bool bIsAssemblingBuild = true;
+
+				using(Timeline.ScopeEvent("UEBuildTarget.Build()"))
+				{
+					Makefile = Target.Build(BuildConfiguration, WorkingSet, bIsAssemblingBuild);
+				}
+
+				// Save the environment variables
+				foreach (System.Collections.DictionaryEntry EnvironmentVariable in Environment.GetEnvironmentVariables())
+				{
+					Makefile.EnvironmentVariables.Add(Tuple.Create((string)EnvironmentVariable.Key, (string)EnvironmentVariable.Value));
+				}
+
+				// Save the makefile for next time
+				if(MakefileLocation != null)
+				{
+					using(Timeline.ScopeEvent("TargetMakefile.Save()"))
+					{
+						Makefile.Save(MakefileLocation);
+					}
+				}
+			}
+			else
+			{
+				// Restore the environment variables
+				foreach (Tuple<string, string> EnvironmentVariable in Makefile.EnvironmentVariables)
+				{
+					Environment.SetEnvironmentVariable(EnvironmentVariable.Item1, EnvironmentVariable.Item2);
+				}
+
+				// Execute all the pre-build steps
+				if (!BuildConfiguration.bXGEExport)
+				{
+					Utils.ExecuteCustomBuildSteps(Makefile.PreBuildScripts);
+				}
+
+				// If the target needs UHT to be run, we'll go ahead and do that now
+				if (Makefile.UObjectModules.Count > 0)
+				{
+					const bool bIsGatheringBuild = false;
+					const bool bIsAssemblingBuild = true;
+
+					FileReference ModuleInfoFileName = FileReference.Combine(Makefile.ProjectIntermediateDirectory, TargetDescriptor.Name + ".uhtmanifest");
+					ExternalExecution.ExecuteHeaderToolIfNecessary(BuildConfiguration, TargetDescriptor.ProjectFile, TargetDescriptor.Name, Makefile.TargetType, Makefile.bHasProjectScriptPlugin, UObjectModules: Makefile.UObjectModules, ModuleInfoFileName: ModuleInfoFileName, bIsGatheringBuild: bIsGatheringBuild, bIsAssemblingBuild: bIsAssemblingBuild, WorkingSet: WorkingSet);
+				}
+			}
+			return Makefile;
+		}
+
 		/// <summary>
 		/// Determine what needs to be built for a target
 		/// </summary>
 		/// <param name="BuildConfiguration">The build configuration</param>
-		/// <param name="TargetDescriptor">Target being build</param>
+		/// <param name="TargetDescriptor">Target being built</param>
 		/// <param name="Makefile">Makefile generated for this target</param>
 		/// <returns>Set of actions to execute</returns>
 		static HashSet<Action> GetActionsForTarget(BuildConfiguration BuildConfiguration, TargetDescriptor TargetDescriptor, TargetMakefile Makefile)
diff --git a/Engine/Source/Programs/UnrealBuildTool/Platform/Windows/PVSToolChain.cs b/Engine/Source/Programs/UnrealBuildTool/Platform/Windows/PVSToolChain.cs
--- a/Engine/Source/Programs/UnrealBuildTool/Platform/Windows/PVSToolChain.cs
+++ b/Engine/Source/Programs/UnrealBuildTool/Platform/Windows/PVSToolChain.cs
@@ -329,7 +329,7 @@ namespace UnrealBuildTool
 			throw new BuildException("Unable to link with PVS toolchain.");
 		}
 
-		public override void FinalizeOutput(ReadOnlyTargetRules Target, List<FileItem> OutputItems, List<Action> Actions)
+		public override void FinalizeOutput(ReadOnlyTargetRules Target, TargetMakefile Makefile)
 		{
 			FileReference OutputFile;
 			if (Target.ProjectFile == null)
@@ -341,7 +341,7 @@ namespace UnrealBuildTool
 				OutputFile = FileReference.Combine(Target.ProjectFile.Directory, "Saved", "PVS-Studio", String.Format("{0}.pvslog", Target.Name));
 			}
 
-			List<FileReference> InputFiles = OutputItems.Select(x => x.Location).Where(x => x.HasExtension(".pvslog")).ToList();
+			List<FileReference> InputFiles = Makefile.OutputItems.Select(x => x.Location).Where(x => x.HasExtension(".pvslog")).ToList();
 
 			FileItem InputFileListItem = FileItem.CreateIntermediateTextFile(OutputFile.ChangeExtension(".input"), InputFiles.Select(x => x.FullName));
 
@@ -350,12 +350,12 @@ namespace UnrealBuildTool
 			AnalyzeAction.CommandArguments = String.Format("-Mode=PVSGather -Input=\"{0}\" -Output=\"{1}\"", InputFileListItem.Location, OutputFile);
 			AnalyzeAction.WorkingDirectory = UnrealBuildTool.EngineSourceDirectory.FullName;
 			AnalyzeAction.PrerequisiteItems.Add(InputFileListItem);
-			AnalyzeAction.PrerequisiteItems.AddRange(OutputItems);
+			AnalyzeAction.PrerequisiteItems.AddRange(Makefile.OutputItems);
 			AnalyzeAction.ProducedItems.Add(FileItem.GetItemByFileReference(OutputFile));
 			AnalyzeAction.DeleteItems.AddRange(AnalyzeAction.ProducedItems);
-			Actions.Add(AnalyzeAction);
+			Makefile.Actions.Add(AnalyzeAction);
 
-			OutputItems.AddRange(AnalyzeAction.ProducedItems);
+			Makefile.OutputItems.AddRange(AnalyzeAction.ProducedItems);
 		}
 	}
 }
diff --git a/Engine/Source/Programs/UnrealBuildTool/System/BuildPrerequisites.cs b/Engine/Source/Programs/UnrealBuildTool/System/BuildPrerequisites.cs
deleted file mode 100644
diff --git a/Engine/Source/Programs/UnrealBuildTool/System/CppCompileEnvironment.cs b/Engine/Source/Programs/UnrealBuildTool/System/CppCompileEnvironment.cs
--- a/Engine/Source/Programs/UnrealBuildTool/System/CppCompileEnvironment.cs
+++ b/Engine/Source/Programs/UnrealBuildTool/System/CppCompileEnvironment.cs
@@ -95,6 +95,11 @@ namespace UnrealBuildTool
 		/// </summary>
 		public readonly SourceFileMetadataCache MetadataCache;
 
+		/// <summary>
+		/// Templates for shared precompiled headers
+		/// </summary>
+		public readonly List<PrecompiledHeaderTemplate> SharedPCHs;
+
 		/// <summary>
 		/// The name of the header file which is precompiled.
 		/// </summary>
@@ -341,6 +346,7 @@ namespace UnrealBuildTool
 			this.Configuration = Configuration;
 			this.Architecture = Architecture;
 			this.MetadataCache = MetadataCache;
+			this.SharedPCHs = new List<PrecompiledHeaderTemplate>();
 			this.UserIncludePaths = new HashSet<DirectoryReference>();
 			this.SystemIncludePaths = new HashSet<DirectoryReference>();
 		}
@@ -355,6 +361,7 @@ namespace UnrealBuildTool
 			Configuration = Other.Configuration;
 			Architecture = Other.Architecture;
 			MetadataCache = Other.MetadataCache;
+			SharedPCHs = Other.SharedPCHs;
 			PrecompiledHeaderIncludeFilename = Other.PrecompiledHeaderIncludeFilename;
 			PrecompiledHeaderAction = Other.PrecompiledHeaderAction;
 			bUseSharedBuildEnvironment = Other.bUseSharedBuildEnvironment;
diff --git a/Engine/Source/Programs/UnrealBuildTool/System/TargetMakefile.cs b/Engine/Source/Programs/UnrealBuildTool/System/TargetMakefile.cs
--- a/Engine/Source/Programs/UnrealBuildTool/System/TargetMakefile.cs
+++ b/Engine/Source/Programs/UnrealBuildTool/System/TargetMakefile.cs
@@ -63,11 +63,6 @@ namespace UnrealBuildTool
 		// @todo ubtmake: Really we want to allow a different set of environment variables for every Action.  This would allow for targets on multiple platforms to be built in a single assembling phase.  We'd only have unique variables for each platform that has actions, so we'd want to make sure we only store the minimum set.
 		public readonly List<Tuple<string, string>> EnvironmentVariables = new List<Tuple<string, string>>();
 
-		/// <summary>
-		/// Maps each target to a list of UObject module info structures
-		/// </summary>
-		public List<UHTModuleInfo> UObjectModules;
-
 		/// <summary>
 		/// The final output items for all target
 		/// </summary>
@@ -84,23 +79,68 @@ namespace UnrealBuildTool
 		public HashSet<string> HotReloadModuleNames;
 
 		/// <summary>
-		/// List of predicates for the action graph to be valid.
+		/// Set of all source directories. Any files being added or removed from these directories will invalidate the makefile.
+		/// </summary>
+		public HashSet<DirectoryReference> SourceDirectories = new HashSet<DirectoryReference>();
+
+		/// <summary>
+		/// The set of source files that UnrealBuildTool determined to be part of the programmer's "working set". Used for adaptive non-unity builds.
+		/// </summary>
+		public HashSet<FileItem> WorkingSet = new HashSet<FileItem>();
+
+		/// <summary>
+		/// Set of files which are currently not part of the working set, but could be.
 		/// </summary>
-		public BuildPrerequisites Prerequisites;
+		public HashSet<FileItem> CandidatesForWorkingSet = new HashSet<FileItem>();
 
-		public TargetMakefile(TargetType TargetType)
+		/// <summary>
+		/// Maps each target to a list of UObject module info structures
+		/// </summary>
+		public List<UHTModuleInfo> UObjectModules;
+
+		/// <summary>
+		/// Used to map names of modules to their .Build.cs filename
+		/// </summary>
+		public List<UHTModuleHeaderInfo> UObjectModuleHeaders = new List<UHTModuleHeaderInfo>();
+
+		/// <summary>
+		/// Additional files which are required 
+		/// </summary>
+		public HashSet<FileItem> AdditionalDependencies = new HashSet<FileItem>(); 
+
+		/// <summary>
+		/// Constructor
+		/// </summary>
+		/// <param name="ReceiptFile">Path to the receipt file</param>
+		/// <param name="ProjectIntermediateDirectory">Path to the project intermediate directory</param>
+		/// <param name="TargetType">The type of target</param>
+		/// <param name="bDeployAfterCompile">Whether to deploy the target after compiling</param>
+		/// <param name="bHasProjectScriptPlugin">Whether the target has a project script plugin</param>
+		public TargetMakefile(FileReference ReceiptFile, DirectoryReference ProjectIntermediateDirectory, TargetType TargetType, bool bDeployAfterCompile, bool bHasProjectScriptPlugin)
 		{
+			this.ReceiptFile = ReceiptFile;
+			this.ProjectIntermediateDirectory = ProjectIntermediateDirectory;
 			this.TargetType = TargetType;
+			this.bDeployAfterCompile = bDeployAfterCompile;
+			this.bHasProjectScriptPlugin = bHasProjectScriptPlugin;
+			this.Actions = new List<Action>();
+			this.OutputItems = new List<FileItem>();
+			this.ModuleNameToOutputItems = new Dictionary<string, FileItem[]>(StringComparer.OrdinalIgnoreCase);
+			this.HotReloadModuleNames = new HashSet<string>(StringComparer.OrdinalIgnoreCase);
+			this.SourceDirectories = new HashSet<DirectoryReference>();
+			this.WorkingSet = new HashSet<FileItem>();
+			this.CandidatesForWorkingSet = new HashSet<FileItem>();
+			this.UObjectModules = new List<UHTModuleInfo>();
+			this.UObjectModuleHeaders = new List<UHTModuleHeaderInfo>();
+			this.AdditionalDependencies = new HashSet<FileItem>();
 		}
 
+		/// <summary>
+		/// Constructor. Reads a makefile from disk.
+		/// </summary>
+		/// <param name="Reader">The archive to read from</param>
 		public TargetMakefile(BinaryArchiveReader Reader)
 		{
-			int Version = Reader.ReadInt();
-			if(Version != CurrentVersion)
-			{
-				throw new Exception(string.Format("Makefile version does not match - found {0}, expected: {1}", Version, CurrentVersion));
-			}
-
 			ReceiptFile = Reader.ReadFileReference();
 			ProjectIntermediateDirectory = Reader.ReadDirectoryReference();
 			TargetType = (TargetType)Reader.ReadInt();
@@ -109,17 +149,23 @@ namespace UnrealBuildTool
 			PreBuildScripts = Reader.ReadArray(() => Reader.ReadFileReference());
 			Actions = Reader.ReadList(() => new Action(Reader));
 			EnvironmentVariables = Reader.ReadList(() => Tuple.Create(Reader.ReadString(), Reader.ReadString()));
-			UObjectModules = Reader.ReadList(() => new UHTModuleInfo(Reader));
 			OutputItems = Reader.ReadList(() => Reader.ReadFileItem());
-			ModuleNameToOutputItems = Reader.ReadDictionary(() => Reader.ReadString(), () => Reader.ReadArray(() => Reader.ReadFileItem()));
-			HotReloadModuleNames = Reader.ReadHashSet(() => Reader.ReadString());
-			Prerequisites = new BuildPrerequisites(Reader);
+			ModuleNameToOutputItems = Reader.ReadDictionary(() => Reader.ReadString(), () => Reader.ReadArray(() => Reader.ReadFileItem()), StringComparer.OrdinalIgnoreCase);
+			HotReloadModuleNames = Reader.ReadHashSet(() => Reader.ReadString(), StringComparer.OrdinalIgnoreCase);
+			SourceDirectories = Reader.ReadHashSet(() => Reader.ReadDirectoryReference());
+			WorkingSet = Reader.ReadHashSet(() => Reader.ReadFileItem());
+			CandidatesForWorkingSet = Reader.ReadHashSet(() => Reader.ReadFileItem());
+			UObjectModules = Reader.ReadList(() => new UHTModuleInfo(Reader));
+			UObjectModuleHeaders = Reader.ReadList(() => new UHTModuleHeaderInfo(Reader));
+			AdditionalDependencies = Reader.ReadHashSet(() => Reader.ReadFileItem());
 		}
 
+		/// <summary>
+		/// Write the makefile to the given archive
+		/// </summary>
+		/// <param name="Writer">The archive to write to</param>
 		public void Write(BinaryArchiveWriter Writer)
 		{
-			Writer.WriteInt(CurrentVersion);
-
 			Writer.WriteFileReference(ReceiptFile);
 			Writer.WriteDirectoryReference(ProjectIntermediateDirectory);
 			Writer.WriteInt((int)TargetType);
@@ -132,49 +178,27 @@ namespace UnrealBuildTool
 			Writer.WriteList(OutputItems, Item => Writer.WriteFileItem(Item));
 			Writer.WriteDictionary(ModuleNameToOutputItems, k => Writer.WriteString(k), v => Writer.WriteArray(v, e => Writer.WriteFileItem(e)));
 			Writer.WriteHashSet(HotReloadModuleNames, x => Writer.WriteString(x));
-            Prerequisites.Write(Writer);
+			Writer.WriteHashSet(SourceDirectories, x => Writer.WriteDirectoryReference(x));
+			Writer.WriteHashSet(WorkingSet, x => Writer.WriteFileItem(x));
+			Writer.WriteHashSet(CandidatesForWorkingSet, x => Writer.WriteFileItem(x));
+			Writer.WriteHashSet(AdditionalDependencies, x => Writer.WriteFileItem(x));
+			Writer.WriteList(UObjectModuleHeaders, x => x.Write(Writer));
 		}
 
-		/// <returns> True if this makefile's contents look valid.  Called after loading the file to make sure it is legit.</returns>
-		public bool IsValidMakefile()
-		{
-			return
-				Actions != null && Actions.Count > 0 &&
-				EnvironmentVariables != null &&
-				UObjectModules != null && 
-				OutputItems != null && 
-				ModuleNameToOutputItems != null && 
-				HotReloadModuleNames != null &&
-				Prerequisites != null;
-		}
-
-
 		/// <summary>
 		/// Saves a makefile to disk
 		/// </summary>
 		/// <param name="Location">Path to save the makefile to</param>
 		public void Save(FileReference Location)
 		{
-			if (!IsValidMakefile())
+			DirectoryReference.CreateDirectory(Location.Directory);
+			using(BinaryArchiveWriter Writer = new BinaryArchiveWriter(Location))
 			{
-				throw new BuildException("Can't save a makefile that has invalid contents.  See UBTMakefile.IsValidMakefile()");
-			}
-
-			try
-			{
-				DirectoryReference.CreateDirectory(Location.Directory);
-				using(BinaryArchiveWriter Writer = new BinaryArchiveWriter(Location))
-				{
-					Write(Writer);
-				}
-			}
-			catch (Exception Ex)
-			{
-				Log.TraceError("Failed to write makefile: {0}", Ex.Message);
+				Writer.WriteInt(CurrentVersion);
+				Write(Writer);
 			}
 		}
 
-
 		/// <summary>
 		/// Loads a UBTMakefile from disk
 		/// </summary>
@@ -283,6 +307,12 @@ namespace UnrealBuildTool
 				{
 					using(BinaryArchiveReader Reader = new BinaryArchiveReader(new FileReference(UBTMakefileInfo)))
 					{
+						int Version = Reader.ReadInt();
+						if(Version != CurrentVersion)
+						{
+							ReasonNotLoaded = "makefile version does not match";
+							return null;
+						}
 						LoadedUBTMakefile = new TargetMakefile(Reader);
 					}
 				}
@@ -297,13 +327,6 @@ namespace UnrealBuildTool
 
 			using(Timeline.ScopeEvent("Checking makefile validity"))
 			{
-				if (!LoadedUBTMakefile.IsValidMakefile())
-				{
-					Log.TraceWarning("Loaded makefile appears to have invalid contents, ignoring it ({0})", UBTMakefileInfo.FullName);
-					ReasonNotLoaded = "existing makefile appears to be invalid";
-					return null;
-				}
-
 				// Check if ini files are newer. Ini files contain build settings too.
 				DirectoryReference ProjectDirectory = DirectoryReference.FromFile(ProjectFile);
 				foreach (ConfigHierarchyType IniType in (ConfigHierarchyType[])Enum.GetValues(typeof(ConfigHierarchyType)))
@@ -320,9 +343,7 @@ namespace UnrealBuildTool
 					}
 				}
 
-				BuildPrerequisites Prerequisites = LoadedUBTMakefile.Prerequisites;
-
-				foreach(DirectoryReference SourceDirectory in Prerequisites.SourceDirectories)
+				foreach(DirectoryReference SourceDirectory in LoadedUBTMakefile.SourceDirectories)
 				{
 					DirectoryInfo SourceDirectoryInfo = new DirectoryInfo(SourceDirectory.FullName);
 					if(!SourceDirectoryInfo.Exists || SourceDirectoryInfo.LastWriteTimeUtc > UBTMakefileInfo.LastWriteTimeUtc)
@@ -333,7 +354,7 @@ namespace UnrealBuildTool
 					}
 				}
 
-				foreach(FileItem AdditionalDependency in Prerequisites.AdditionalDependencies)
+				foreach(FileItem AdditionalDependency in LoadedUBTMakefile.AdditionalDependencies)
 				{
 					if (!AdditionalDependency.Exists)
 					{
@@ -356,14 +377,14 @@ namespace UnrealBuildTool
 				// Get all H files in processed modules newer than the makefile itself
 				HashSet<FileReference> HFilesNewerThanMakefile =
 					new HashSet<FileReference>(
-						Prerequisites.UObjectModuleHeaders
+						LoadedUBTMakefile.UObjectModuleHeaders
 						.SelectMany(x => x.SourceFolder != null ? DirectoryReference.EnumerateFiles(x.SourceFolder, "*.h", SearchOption.AllDirectories) : Enumerable.Empty<FileReference>())
 						.Where(y => FileReference.GetLastWriteTimeUtc(y) > UBTMakefileInfo.LastWriteTimeUtc)
 						.OrderBy(z => z).Distinct()
 					);
 
 				// Get all H files in all modules processed in the last makefile build
-				HashSet<FileReference> AllUHTHeaders = new HashSet<FileReference>(Prerequisites.UObjectModuleHeaders.SelectMany(x => x.HeaderFiles).Select(x => x.Location));
+				HashSet<FileReference> AllUHTHeaders = new HashSet<FileReference>(LoadedUBTMakefile.UObjectModuleHeaders.SelectMany(x => x.HeaderFiles).Select(x => x.Location));
 
 				// Check whether any headers have been deleted. If they have, we need to regenerate the makefile since the module might now be empty. If we don't,
 				// and the file has been moved to a different module, we may include stale generated headers.
@@ -400,7 +421,7 @@ namespace UnrealBuildTool
 				// iteration times.)
 
 				// Check if any source files in the working set no longer belong in it
-				foreach (FileItem SourceFile in Prerequisites.WorkingSet)
+				foreach (FileItem SourceFile in LoadedUBTMakefile.WorkingSet)
 				{
 					if (!WorkingSet.Contains(SourceFile.Location) && File.GetLastWriteTimeUtc(SourceFile.AbsolutePath) > UBTMakefileInfo.LastWriteTimeUtc)
 					{
@@ -411,7 +432,7 @@ namespace UnrealBuildTool
 				}
 
 				// Check if any source files that are eligible for being in the working set have been modified
-				foreach (FileItem SourceFile in Prerequisites.CandidatesForWorkingSet)
+				foreach (FileItem SourceFile in LoadedUBTMakefile.CandidatesForWorkingSet)
 				{
 					if (WorkingSet.Contains(SourceFile.Location) && File.GetLastWriteTimeUtc(SourceFile.AbsolutePath) > UBTMakefileInfo.LastWriteTimeUtc)
 					{
diff --git a/Engine/Source/Programs/UnrealBuildTool/System/Unity.cs b/Engine/Source/Programs/UnrealBuildTool/System/Unity.cs
--- a/Engine/Source/Programs/UnrealBuildTool/System/Unity.cs
+++ b/Engine/Source/Programs/UnrealBuildTool/System/Unity.cs
@@ -152,18 +152,18 @@ namespace UnrealBuildTool
 		/// <param name="CPPFiles">The C++ files to #include.</param>
 		/// <param name="CompileEnvironment">The environment that is used to compile the C++ files.</param>
 		/// <param name="WorkingSet">Interface to query files which belong to the working set</param>
-		/// <param name="Prerequisites">Dependencies for this target</param>
 		/// <param name="BaseName">Base name to use for the Unity files</param>
 		/// <param name="IntermediateDirectory">Intermediate directory for unity cpp files</param>
+		/// <param name="Makefile">The makefile being built</param>
 		/// <returns>The "unity" C++ files.</returns>
 		public static List<FileItem> GenerateUnityCPPs(
 			ReadOnlyTargetRules Target,
 			List<FileItem> CPPFiles,
 			CppCompileEnvironment CompileEnvironment,
 			ISourceFileWorkingSet WorkingSet,
-			BuildPrerequisites Prerequisites,
 			string BaseName,
-			DirectoryReference IntermediateDirectory
+			DirectoryReference IntermediateDirectory,
+			TargetMakefile Makefile
 			)
 		{
 			List<FileItem> NewCPPFiles = new List<FileItem>();
@@ -218,7 +218,7 @@ namespace UnrealBuildTool
 							// Mark this file as part of the working set.  This will be saved into the UBT Makefile so that
 							// the assembler can automatically invalidate the Makefile when the working set changes (allowing this
 							// code to run again, to build up new unity blobs.)
-							Prerequisites.WorkingSet.Add(CPPFile);
+							Makefile.WorkingSet.Add(CPPFile);
 						}
 					}
 
@@ -240,7 +240,7 @@ namespace UnrealBuildTool
 					}
 
 					// When adaptive unity is enabled, go ahead and exclude any source files that we're actively working with
-					if (bUseAdaptiveUnityBuild && Prerequisites.WorkingSet.Contains(CPPFile))
+					if (bUseAdaptiveUnityBuild && Makefile.WorkingSet.Contains(CPPFile))
 					{
 						// Just compile this file normally, not as part of the unity blob
 						NewCPPFiles.Add(CPPFile);
@@ -266,7 +266,7 @@ namespace UnrealBuildTool
 						// If adaptive unity build is enabled for this module, add this source file to the set that will invalidate the makefile
 						if(bUseAdaptiveUnityBuild)
 						{
-							Prerequisites.CandidatesForWorkingSet.Add(CPPFile);
+							Makefile.CandidatesForWorkingSet.Add(CPPFile);
 						}
 
 						// Compile this file as part of the unity blob
diff --git a/Engine/Source/Programs/UnrealBuildTool/ToolChain/UEToolChain.cs b/Engine/Source/Programs/UnrealBuildTool/ToolChain/UEToolChain.cs
--- a/Engine/Source/Programs/UnrealBuildTool/ToolChain/UEToolChain.cs
+++ b/Engine/Source/Programs/UnrealBuildTool/ToolChain/UEToolChain.cs
@@ -66,7 +66,7 @@ namespace UnrealBuildTool
 		{
 		}
 
-		public virtual void FinalizeOutput(ReadOnlyTargetRules Target, List<FileItem> OutputItems, List<Action> Actions)
+		public virtual void FinalizeOutput(ReadOnlyTargetRules Target, TargetMakefile Makefile)
 		{
 		}
 
