UBT: Reduce false positives when detecting that source directories have changed.

If the directory timestamp has changed, we now rebuild a list of source files and compare it to the previous list. Visual Studio seems to write out a temporary file during a save operation, which modifes the directory timestamp and invalidates the makefile.

#rnx

diff --git a/Engine/Source/Programs/UnrealBuildTool/Configuration/UEBuildModuleCPP.cs b/Engine/Source/Programs/UnrealBuildTool/Configuration/UEBuildModuleCPP.cs
--- a/Engine/Source/Programs/UnrealBuildTool/Configuration/UEBuildModuleCPP.cs
+++ b/Engine/Source/Programs/UnrealBuildTool/Configuration/UEBuildModuleCPP.cs
@@ -19,19 +19,15 @@ namespace UnrealBuildTool
 		/// <summary>
 		/// Stores a list of all source files, of different types
 		/// </summary>
-		class SourceFileCollection
+		class InputFileCollection
 		{
-			public readonly List<FileItem> HFiles = new List<FileItem>();
+			public readonly List<FileItem> HeaderFiles = new List<FileItem>();
+
 			public readonly List<FileItem> CPPFiles = new List<FileItem>();
 			public readonly List<FileItem> CFiles = new List<FileItem>();
 			public readonly List<FileItem> CCFiles = new List<FileItem>();
 			public readonly List<FileItem> MMFiles = new List<FileItem>();
 			public readonly List<FileItem> RCFiles = new List<FileItem>();
-
-			public IEnumerable<FileItem> AllFiles
-			{
-				get { return CPPFiles.Concat(CFiles).Concat(CCFiles).Concat(MMFiles).Concat(RCFiles); }
-			}
 		}
 
 		/// <summary>
@@ -272,12 +268,11 @@ namespace UnrealBuildTool
 				return LinkInputFiles;
 			}
 
-			SourceFileCollection SourceFilesToBuild = FindSourceFiles(Target.Platform, Makefile.SourceDirectories);
-
-			SourceDirectories = new HashSet<DirectoryReference>(SourceFilesToBuild.AllFiles.Select(x => x.Location.Directory));
+			// Find all the input files
+			InputFileCollection InputFiles = FindInputFiles(Target.Platform, Makefile);
 
 			// Process all of the header file dependencies for this module
-			CheckFirstIncludeMatchesEachCppFile(Target, ModuleCompileEnvironment, SourceFilesToBuild.HFiles, SourceFilesToBuild.CPPFiles);
+			CheckFirstIncludeMatchesEachCppFile(Target, ModuleCompileEnvironment, InputFiles.HeaderFiles, InputFiles.CPPFiles);
 
 			// Should we force a precompiled header to be generated for this module?  Usually, we only bother with a
 			// precompiled header if there are at least several source files in the module (after combining them for unity
@@ -325,7 +320,7 @@ namespace UnrealBuildTool
 					Log.TraceVerbose("Module '{0}' not using unity build mode (bFasterWithoutUnity enabled for this module)", this.Name);
 					bModuleUsesUnityBuild = false;
 				}
-				else if (SourceFilesToBuild.CPPFiles.Count < MinSourceFilesForUnityBuild)
+				else if (InputFiles.CPPFiles.Count < MinSourceFilesForUnityBuild)
 				{
 					Log.TraceVerbose("Module '{0}' not using unity build mode (module with fewer than {1} source files)", this.Name, MinSourceFilesForUnityBuild);
 					bModuleUsesUnityBuild = false;
@@ -410,7 +405,7 @@ namespace UnrealBuildTool
 			CreateHeaderForDefinitions(CompileEnvironment, IntermediateDirectory, null);
 
 			// Compile CPP files
-			List<FileItem> CPPFilesToCompile = SourceFilesToBuild.CPPFiles;
+			List<FileItem> CPPFilesToCompile = InputFiles.CPPFiles;
 			if (bModuleUsesUnityBuild)
 			{
 				CPPFilesToCompile = Unity.GenerateUnityCPPs(Target, CPPFilesToCompile, CompileEnvironment, WorkingSet, Rules.ShortName ?? Name, IntermediateDirectory, Makefile);
@@ -471,31 +466,31 @@ namespace UnrealBuildTool
 			}
 
 			// Compile C files directly. Do not use a PCH here, because a C++ PCH is not compatible with C source files.
-			if(SourceFilesToBuild.CFiles.Count > 0)
+			if(InputFiles.CFiles.Count > 0)
 			{
-				LinkInputFiles.AddRange(ToolChain.CompileCPPFiles(ModuleCompileEnvironment, SourceFilesToBuild.CFiles, IntermediateDirectory, Name, Makefile.Actions).ObjectFiles);
+				LinkInputFiles.AddRange(ToolChain.CompileCPPFiles(ModuleCompileEnvironment, InputFiles.CFiles, IntermediateDirectory, Name, Makefile.Actions).ObjectFiles);
 			}
 
 			// Compile CC files directly.
-			if(SourceFilesToBuild.CCFiles.Count > 0)
+			if(InputFiles.CCFiles.Count > 0)
 			{
-				LinkInputFiles.AddRange(ToolChain.CompileCPPFiles(CompileEnvironment, SourceFilesToBuild.CCFiles, IntermediateDirectory, Name, Makefile.Actions).ObjectFiles);
+				LinkInputFiles.AddRange(ToolChain.CompileCPPFiles(CompileEnvironment, InputFiles.CCFiles, IntermediateDirectory, Name, Makefile.Actions).ObjectFiles);
 			}
 
 			// Compile MM files directly.
-			if(SourceFilesToBuild.MMFiles.Count > 0)
+			if(InputFiles.MMFiles.Count > 0)
 			{
-				LinkInputFiles.AddRange(ToolChain.CompileCPPFiles(CompileEnvironment, SourceFilesToBuild.MMFiles, IntermediateDirectory, Name, Makefile.Actions).ObjectFiles);
+				LinkInputFiles.AddRange(ToolChain.CompileCPPFiles(CompileEnvironment, InputFiles.MMFiles, IntermediateDirectory, Name, Makefile.Actions).ObjectFiles);
 			}
 
 			// Compile RC files. The resource compiler does not work with response files, and using the regular compile environment can easily result in the 
 			// command line length exceeding the OS limit. Use the binary compile environment to keep the size down, and require that all include paths
 			// must be specified relative to the resource file itself or Engine/Source.
-			if(SourceFilesToBuild.RCFiles.Count > 0)
+			if(InputFiles.RCFiles.Count > 0)
 			{
 				CppCompileEnvironment ResourceCompileEnvironment = new CppCompileEnvironment(BinaryCompileEnvironment);
 				WindowsPlatform.SetupResourceCompileEnvironment(ResourceCompileEnvironment, Target);
-				LinkInputFiles.AddRange(ToolChain.CompileRCFiles(ResourceCompileEnvironment, SourceFilesToBuild.RCFiles, IntermediateDirectory, Makefile.Actions).ObjectFiles);
+				LinkInputFiles.AddRange(ToolChain.CompileRCFiles(ResourceCompileEnvironment, InputFiles.RCFiles, IntermediateDirectory, Makefile.Actions).ObjectFiles);
 			}
 
 			// Write the compiled manifest
@@ -1100,17 +1095,19 @@ namespace UnrealBuildTool
 		/// Finds all the source files that should be built for this module
 		/// </summary>
 		/// <param name="Platform">The platform the module is being built for</param>
-		/// <param name="SearchedDirectories">Receives a set of directories that have been searched</param>
+		/// <param name="Makefile">Makefile for the target being built</param>
 		/// <returns>Set of source files that should be built</returns>
-		SourceFileCollection FindSourceFiles(UnrealTargetPlatform Platform, HashSet<DirectoryReference> SearchedDirectories)
+		InputFileCollection FindInputFiles(UnrealTargetPlatform Platform, TargetMakefile Makefile)
 		{
 			ReadOnlyHashSet<string> ExcludedNames = UEBuildPlatform.GetBuildPlatform(Platform).GetExcludedFolderNames();
 
-			SourceFileCollection SourceFilesToBuild = new SourceFileCollection();
+			InputFileCollection InputFiles = new InputFileCollection();
 			DirectoryItem ModuleDirectoryItem = DirectoryItem.GetItemByDirectoryReference(ModuleDirectory);
-			FindSourceFilesRecursive(ModuleDirectoryItem, ExcludedNames, SearchedDirectories, SourceFilesToBuild);
 
-			return SourceFilesToBuild;
+			SourceDirectories = new HashSet<DirectoryReference>();
+			FindInputFilesFromDirectoryRecursive(ModuleDirectoryItem, ExcludedNames, SourceDirectories, Makefile.DirectoryToSourceFiles, InputFiles);
+
+			return InputFiles;
 		}
 
 		/// <summary>
@@ -1118,47 +1115,79 @@ namespace UnrealBuildTool
 		/// </summary>
 		/// <param name="BaseDirectory">Directory to search from</param>
 		/// <param name="ExcludedNames">Set of excluded directory names (eg. other platforms)</param>
-		/// <param name="SearchedDirectories">Receives a set of directories that have been searched</param>
-		/// <param name="SourceFiles">Collection of source files, categorized by type</param>
-		static void FindSourceFilesRecursive(DirectoryItem BaseDirectory, ReadOnlyHashSet<string> ExcludedNames, HashSet<DirectoryReference> SearchedDirectories, SourceFileCollection SourceFiles)
+		/// <param name="SourceDirectories">Set of all non-empty source directories.</param>
+		/// <param name="DirectoryToSourceFiles">Map from directory to source files inside it</param>
+		/// <param name="InputFiles">Collection of source files, categorized by type</param>
+		static void FindInputFilesFromDirectoryRecursive(DirectoryItem BaseDirectory, ReadOnlyHashSet<string> ExcludedNames, HashSet<DirectoryReference> SourceDirectories, Dictionary<DirectoryItem, FileItem[]> DirectoryToSourceFiles, InputFileCollection InputFiles)
 		{
-			SearchedDirectories.Add(BaseDirectory.Location);
-
 			foreach(DirectoryItem SubDirectory in BaseDirectory.EnumerateDirectories())
 			{
 				if(!ExcludedNames.Contains(SubDirectory.Name))
 				{
-					FindSourceFilesRecursive(SubDirectory, ExcludedNames, SearchedDirectories, SourceFiles);
+					FindInputFilesFromDirectoryRecursive(SubDirectory, ExcludedNames, SourceDirectories, DirectoryToSourceFiles, InputFiles);
 				}
 			}
 
-			foreach(FileItem SourceFile in BaseDirectory.EnumerateFiles())
+			FileItem[] SourceFiles = FindInputFilesFromDirectory(BaseDirectory, InputFiles);
+			if(SourceFiles.Length > 0)
 			{
-				if (SourceFile.HasExtension(".h"))
+				SourceDirectories.Add(BaseDirectory.Location);
+			}
+			DirectoryToSourceFiles.Add(BaseDirectory, SourceFiles);
+		}
+
+		/// <summary>
+		/// Finds the input files that should be built for this module, from a given directory
+		/// </summary>
+		/// <param name="BaseDirectory"></param>
+		/// <param name="InputFiles"></param>
+		/// <returns>Array of source files</returns>
+		static FileItem[] FindInputFilesFromDirectory(DirectoryItem BaseDirectory, InputFileCollection InputFiles)
+		{
+			List<FileItem> SourceFiles = new List<FileItem>();
+			foreach(FileItem InputFile in BaseDirectory.EnumerateFiles())
+			{
+				if (InputFile.HasExtension(".h"))
 				{
-					SourceFiles.HFiles.Add(SourceFile);
+					InputFiles.HeaderFiles.Add(InputFile);
 				}
-				if (SourceFile.HasExtension(".cpp"))
+				if (InputFile.HasExtension(".cpp"))
 				{
-					SourceFiles.CPPFiles.Add(SourceFile);
+					SourceFiles.Add(InputFile);
+					InputFiles.CPPFiles.Add(InputFile);
 				}
-				else if (SourceFile.HasExtension(".c"))
+				else if (InputFile.HasExtension(".c"))
 				{
-					SourceFiles.CFiles.Add(SourceFile);
+					SourceFiles.Add(InputFile);
+					InputFiles.CFiles.Add(InputFile);
 				}
-				else if (SourceFile.HasExtension(".cc"))
+				else if (InputFile.HasExtension(".cc"))
 				{
-					SourceFiles.CCFiles.Add(SourceFile);
+					SourceFiles.Add(InputFile);
+					InputFiles.CCFiles.Add(InputFile);
 				}
-				else if (SourceFile.HasExtension(".m") || SourceFile.HasExtension(".mm"))
+				else if (InputFile.HasExtension(".m") || InputFile.HasExtension(".mm"))
 				{
-					SourceFiles.MMFiles.Add(SourceFile);
+					SourceFiles.Add(InputFile);
+					InputFiles.MMFiles.Add(InputFile);
 				}
-				else if (SourceFile.HasExtension(".rc"))
+				else if (InputFile.HasExtension(".rc"))
 				{
-					SourceFiles.RCFiles.Add(SourceFile);
+					SourceFiles.Add(InputFile);
+					InputFiles.RCFiles.Add(InputFile);
 				}
 			}
+			return SourceFiles.ToArray();
+		}
+
+		/// <summary>
+		/// Gets a set of source files for the given directory. Used to detect when the makefile is out of date.
+		/// </summary>
+		/// <param name="Directory"></param>
+		/// <returns>Array of source files</returns>
+		public static FileItem[] GetSourceFiles(DirectoryItem Directory)
+		{
+			return FindInputFilesFromDirectory(Directory, new InputFileCollection());
 		}
 	}
 }
diff --git a/Engine/Source/Programs/UnrealBuildTool/Configuration/UEBuildTarget.cs b/Engine/Source/Programs/UnrealBuildTool/Configuration/UEBuildTarget.cs
--- a/Engine/Source/Programs/UnrealBuildTool/Configuration/UEBuildTarget.cs
+++ b/Engine/Source/Programs/UnrealBuildTool/Configuration/UEBuildTarget.cs
@@ -594,16 +594,6 @@ namespace UnrealBuildTool
 		/// </summary>
 		readonly FileReference TargetRulesFile;
 
-		/// <summary>
-		/// List of scripts to run before building
-		/// </summary>
-		public FileReference[] PreBuildStepScripts;
-
-		/// <summary>
-		/// List of scripts to run after building
-		/// </summary>
-		FileReference[] PostBuildStepScripts;
-
 		/// <summary>
 		/// Whether to deploy this target after compilation
 		/// </summary>
@@ -1206,8 +1196,12 @@ namespace UnrealBuildTool
 				}
 			}
 
+			// Create the makefile
+			TargetMakefile Makefile = new TargetMakefile(ReceiptFileName, ProjectIntermediateDirectory, TargetType, bDeployAfterCompile, bHasProjectScriptPlugin);
+
 			// Execute the pre-build steps
-			Utils.ExecuteCustomBuildSteps(PreBuildStepScripts);
+			Makefile.PreBuildScripts = CreatePreBuildScripts();
+			Utils.ExecuteCustomBuildSteps(Makefile.PreBuildScripts);
 
 			// If we're compiling monolithic, make sure the executable knows about all referenced modules
 			if (ShouldCompileMonolithic())
@@ -1260,9 +1254,6 @@ namespace UnrealBuildTool
 			// On Mac and Linux we have actions that should be executed after all the binaries are created
 			TargetToolChain.SetupBundleDependencies(Binaries, TargetName);
 
-			// Create the makefile
-			TargetMakefile Makefile = new TargetMakefile(ReceiptFileName, ProjectIntermediateDirectory, TargetType, bDeployAfterCompile, bHasProjectScriptPlugin);
-
 			// Generate headers
 			HashSet<UEBuildModuleCPP> ModulesToGenerateHeadersFor = GatherDependencyModules(OriginalBinaries.ToList());
 			using(Timeline.ScopeEvent("ExternalExecution.SetupUObjectModules()"))
@@ -1435,23 +1426,24 @@ namespace UnrealBuildTool
 				Makefile.OutputItems.AddRange(WriteMetadataAction.ProducedItems);
 
 				// Create actions to run the post build steps
-				foreach(FileReference PostBuildStepScript in PostBuildStepScripts)
+				FileReference[] PostBuildScripts = CreatePostBuildScripts();
+				foreach(FileReference PostBuildScript in PostBuildScripts)
 				{
-					FileReference OutputFile = new FileReference(PostBuildStepScript.FullName + ".ran");
+					FileReference OutputFile = new FileReference(PostBuildScript.FullName + ".ran");
 
 					Action PostBuildStepAction = new Action(ActionType.PostBuildStep);
 					if(BuildHostPlatform.Current.Platform == UnrealTargetPlatform.Win64)
 					{
 						PostBuildStepAction.CommandPath = "cmd.exe";
-						PostBuildStepAction.CommandArguments = String.Format("/C \"call \"{0}\" && type NUL >\"{1}\"\"", PostBuildStepScript, OutputFile);
+						PostBuildStepAction.CommandArguments = String.Format("/C \"call \"{0}\" && type NUL >\"{1}\"\"", PostBuildScript, OutputFile);
 					}
 					else
 					{
 						PostBuildStepAction.CommandPath = "/bin/sh";
-						PostBuildStepAction.CommandArguments = String.Format("\"{0}\" && touch \"{1}\"", PostBuildStepScript, OutputFile);
+						PostBuildStepAction.CommandArguments = String.Format("\"{0}\" && touch \"{1}\"", PostBuildScript, OutputFile);
 					}
 					PostBuildStepAction.WorkingDirectory = UnrealBuildTool.EngineSourceDirectory.FullName;
-					PostBuildStepAction.StatusDescription = String.Format("Executing post build script ({0})", PostBuildStepScript.GetFileName());
+					PostBuildStepAction.StatusDescription = String.Format("Executing post build script ({0})", PostBuildScript.GetFileName());
 					PostBuildStepAction.bCanExecuteRemotely = false;
 					PostBuildStepAction.PrerequisiteItems.Add(FileItem.GetItemByFileReference(ReceiptFileName));
 					PostBuildStepAction.ProducedItems.Add(FileItem.GetItemByFileReference(OutputFile));
@@ -1754,9 +1746,6 @@ namespace UnrealBuildTool
 			// Setup the target's plugins
 			SetupPlugins();
 
-			// Setup the custom build steps for this target
-			SetupCustomBuildSteps();
-
 			// Add the plugin binaries to the build
 			foreach (UEBuildPlugin Plugin in BuildPlugins)
 			{
@@ -1837,13 +1826,10 @@ namespace UnrealBuildTool
 		}
 
 		/// <summary>
-		/// Writes scripts for all the custom build steps
+		/// Creates scripts for executing the pre-build scripts
 		/// </summary>
-		private void SetupCustomBuildSteps()
+		private FileReference[] CreatePreBuildScripts()
 		{
-			// Make sure the intermediate directory exists
-			DirectoryReference ScriptDirectory = ProjectIntermediateDirectory;
-
 			// Find all the pre-build steps
 			List<Tuple<string[], UEBuildPlugin>> PreBuildCommandBatches = new List<Tuple<string[], UEBuildPlugin>>();
 			if(ProjectDescriptor != null && ProjectDescriptor.PreBuildSteps != null)
@@ -1858,8 +1844,15 @@ namespace UnrealBuildTool
 			{
 				AddCustomBuildSteps(BuildPlugin.Descriptor.PreBuildSteps, BuildPlugin, PreBuildCommandBatches);
 			}
-			PreBuildStepScripts = WriteCustomBuildStepScripts(BuildHostPlatform.Current.Platform, ScriptDirectory, "PreBuild", PreBuildCommandBatches);
+			return WriteCustomBuildStepScripts(BuildHostPlatform.Current.Platform, ProjectIntermediateDirectory, "PreBuild", PreBuildCommandBatches);
+		}
 
+		/// <summary>
+		/// Creates scripts for executing post-build steps
+		/// </summary>
+		/// <returns>Array of post-build scripts</returns>
+		private FileReference[] CreatePostBuildScripts()
+		{
 			// Find all the post-build steps
 			List<Tuple<string[], UEBuildPlugin>> PostBuildCommandBatches = new List<Tuple<string[], UEBuildPlugin>>();
 			if(!Rules.bDisableLinking)
@@ -1877,7 +1870,7 @@ namespace UnrealBuildTool
 					AddCustomBuildSteps(BuildPlugin.Descriptor.PostBuildSteps, BuildPlugin, PostBuildCommandBatches);
 				}
 			}
-			PostBuildStepScripts = WriteCustomBuildStepScripts(BuildHostPlatform.Current.Platform, ScriptDirectory, "PostBuild", PostBuildCommandBatches);
+			return WriteCustomBuildStepScripts(BuildHostPlatform.Current.Platform, ProjectIntermediateDirectory, "PostBuild", PostBuildCommandBatches);
 		}
 
 		/// <summary>
diff --git a/Engine/Source/Programs/UnrealBuildTool/System/DirectoryItem.cs b/Engine/Source/Programs/UnrealBuildTool/System/DirectoryItem.cs
--- a/Engine/Source/Programs/UnrealBuildTool/System/DirectoryItem.cs
+++ b/Engine/Source/Programs/UnrealBuildTool/System/DirectoryItem.cs
@@ -222,4 +222,30 @@ namespace UnrealBuildTool
 			return Location.FullName;
 		}
 	}
+
+	/// <summary>
+	/// Helper functions for serialization
+	/// </summary>
+	static class DirectoryItemExtensionMethods
+	{
+		/// <summary>
+		/// Read a directory item from a binary archive
+		/// </summary>
+		/// <param name="Reader">Reader to serialize data from</param>
+		/// <returns>Instance of the serialized directory item</returns>
+		public static DirectoryItem ReadDirectoryItem(this BinaryArchiveReader Reader)
+		{
+			return Reader.ReadObjectReference<DirectoryItem>(() => DirectoryItem.GetItemByDirectoryReference(Reader.ReadDirectoryReference()));
+		}
+
+		/// <summary>
+		/// Write a directory item to a binary archive
+		/// </summary>
+		/// <param name="Writer">Writer to serialize data to</param>
+		/// <param name="DirectoryItem">Directory item to write</param>
+		public static void WriteDirectoryItem(this BinaryArchiveWriter Writer, DirectoryItem DirectoryItem)
+		{
+			Writer.WriteObjectReference<DirectoryItem>(DirectoryItem, () => Writer.WriteDirectoryReference(DirectoryItem.Location));
+		}
+	}
 }
diff --git a/Engine/Source/Programs/UnrealBuildTool/System/TargetMakefile.cs b/Engine/Source/Programs/UnrealBuildTool/System/TargetMakefile.cs
--- a/Engine/Source/Programs/UnrealBuildTool/System/TargetMakefile.cs
+++ b/Engine/Source/Programs/UnrealBuildTool/System/TargetMakefile.cs
@@ -81,7 +81,7 @@ namespace UnrealBuildTool
 		/// <summary>
 		/// Set of all source directories. Any files being added or removed from these directories will invalidate the makefile.
 		/// </summary>
-		public HashSet<DirectoryReference> SourceDirectories = new HashSet<DirectoryReference>();
+		public Dictionary<DirectoryItem, FileItem[]> DirectoryToSourceFiles;
 
 		/// <summary>
 		/// The set of source files that UnrealBuildTool determined to be part of the programmer's "working set". Used for adaptive non-unity builds.
@@ -127,7 +127,7 @@ namespace UnrealBuildTool
 			this.OutputItems = new List<FileItem>();
 			this.ModuleNameToOutputItems = new Dictionary<string, FileItem[]>(StringComparer.OrdinalIgnoreCase);
 			this.HotReloadModuleNames = new HashSet<string>(StringComparer.OrdinalIgnoreCase);
-			this.SourceDirectories = new HashSet<DirectoryReference>();
+			this.DirectoryToSourceFiles = new Dictionary<DirectoryItem, FileItem[]>();
 			this.WorkingSet = new HashSet<FileItem>();
 			this.CandidatesForWorkingSet = new HashSet<FileItem>();
 			this.UObjectModules = new List<UHTModuleInfo>();
@@ -152,7 +152,7 @@ namespace UnrealBuildTool
 			OutputItems = Reader.ReadList(() => Reader.ReadFileItem());
 			ModuleNameToOutputItems = Reader.ReadDictionary(() => Reader.ReadString(), () => Reader.ReadArray(() => Reader.ReadFileItem()), StringComparer.OrdinalIgnoreCase);
 			HotReloadModuleNames = Reader.ReadHashSet(() => Reader.ReadString(), StringComparer.OrdinalIgnoreCase);
-			SourceDirectories = Reader.ReadHashSet(() => Reader.ReadDirectoryReference());
+			DirectoryToSourceFiles = Reader.ReadDictionary(() => Reader.ReadDirectoryItem(), () => Reader.ReadArray(() => Reader.ReadFileItem()));
 			WorkingSet = Reader.ReadHashSet(() => Reader.ReadFileItem());
 			CandidatesForWorkingSet = Reader.ReadHashSet(() => Reader.ReadFileItem());
 			UObjectModules = Reader.ReadList(() => new UHTModuleInfo(Reader));
@@ -174,15 +174,15 @@ namespace UnrealBuildTool
 			Writer.WriteArray(PreBuildScripts, Item => Writer.WriteFileReference(Item));
 			Writer.WriteList(Actions, Action => Action.Write(Writer));
 			Writer.WriteList(EnvironmentVariables, x => { Writer.WriteString(x.Item1); Writer.WriteString(x.Item2); });
-			Writer.WriteList(UObjectModules, e => e.Write(Writer));
 			Writer.WriteList(OutputItems, Item => Writer.WriteFileItem(Item));
 			Writer.WriteDictionary(ModuleNameToOutputItems, k => Writer.WriteString(k), v => Writer.WriteArray(v, e => Writer.WriteFileItem(e)));
 			Writer.WriteHashSet(HotReloadModuleNames, x => Writer.WriteString(x));
-			Writer.WriteHashSet(SourceDirectories, x => Writer.WriteDirectoryReference(x));
+			Writer.WriteDictionary(DirectoryToSourceFiles, k => Writer.WriteDirectoryItem(k), v => Writer.WriteArray(v, e => Writer.WriteFileItem(e)));
 			Writer.WriteHashSet(WorkingSet, x => Writer.WriteFileItem(x));
 			Writer.WriteHashSet(CandidatesForWorkingSet, x => Writer.WriteFileItem(x));
-			Writer.WriteHashSet(AdditionalDependencies, x => Writer.WriteFileItem(x));
+			Writer.WriteList(UObjectModules, e => e.Write(Writer));
 			Writer.WriteList(UObjectModuleHeaders, x => x.Write(Writer));
+			Writer.WriteHashSet(AdditionalDependencies, x => Writer.WriteFileItem(x));
 		}
 
 		/// <summary>
@@ -343,14 +343,27 @@ namespace UnrealBuildTool
 					}
 				}
 
-				foreach(DirectoryReference SourceDirectory in LoadedUBTMakefile.SourceDirectories)
+				foreach(KeyValuePair<DirectoryItem, FileItem[]> Pair in LoadedUBTMakefile.DirectoryToSourceFiles)
 				{
-					DirectoryInfo SourceDirectoryInfo = new DirectoryInfo(SourceDirectory.FullName);
-					if(!SourceDirectoryInfo.Exists || SourceDirectoryInfo.LastWriteTimeUtc > UBTMakefileInfo.LastWriteTimeUtc)
+					DirectoryItem InputDirectory = Pair.Key;
+					if(!InputDirectory.Exists || InputDirectory.LastWriteTimeUtc > UBTMakefileInfo.LastWriteTimeUtc)
 					{
-						Log.TraceLog("Timestamp of {0} ({1}) is newer than makefile ({2})", SourceDirectory, SourceDirectoryInfo.LastWriteTimeUtc, UBTMakefileInfo.LastWriteTimeUtc);
-						ReasonNotLoaded = "source directory changed";
-						return null;
+						FileItem[] SourceFiles = UEBuildModuleCPP.GetSourceFiles(InputDirectory);
+						if(SourceFiles.Length < Pair.Value.Length)
+						{
+							ReasonNotLoaded = "source file removed";
+							return null;
+						}
+						else if(SourceFiles.Length > Pair.Value.Length)
+						{
+							ReasonNotLoaded = "source file added";
+							return null;
+						}
+						else if(SourceFiles.Intersect(Pair.Value).Count() != SourceFiles.Length)
+						{
+							ReasonNotLoaded = "source file modified";
+							return null;
+						}
 					}
 				}
 
