UBT: Improve performance of timestamp checking in makefile; used cached file states rather than querying the filesystem directly.

#rnx

diff --git a/Engine/Source/Programs/UnrealBuildTool/System/TargetMakefile.cs b/Engine/Source/Programs/UnrealBuildTool/System/TargetMakefile.cs
--- a/Engine/Source/Programs/UnrealBuildTool/System/TargetMakefile.cs
+++ b/Engine/Source/Programs/UnrealBuildTool/System/TargetMakefile.cs
@@ -388,24 +388,28 @@ namespace UnrealBuildTool
 				// we'll just invalidate the entire makefile and force it to be rebuilt.
 
 				// Get all H files in processed modules newer than the makefile itself
-				HashSet<FileReference> HFilesNewerThanMakefile =
-					new HashSet<FileReference>(
-						LoadedUBTMakefile.UObjectModuleHeaders
-						.SelectMany(x => x.SourceFolder != null ? DirectoryReference.EnumerateFiles(x.SourceFolder, "*.h", SearchOption.AllDirectories) : Enumerable.Empty<FileReference>())
-						.Where(y => FileReference.GetLastWriteTimeUtc(y) > UBTMakefileInfo.LastWriteTimeUtc)
-						.OrderBy(z => z).Distinct()
-					);
+				HashSet<FileItem> HFilesNewerThanMakefile = new HashSet<FileItem>();
+				foreach(UHTModuleHeaderInfo ModuleHeaderInfo in LoadedUBTMakefile.UObjectModuleHeaders)
+				{
+					foreach(FileItem HeaderFile in ModuleHeaderInfo.SourceFolder.EnumerateFiles())
+					{
+						if(HeaderFile.HasExtension(".h") && HeaderFile.LastWriteTimeUtc > UBTMakefileInfo.LastWriteTimeUtc)
+						{
+							HFilesNewerThanMakefile.Add(HeaderFile);
+						}
+					}
+				}
 
 				// Get all H files in all modules processed in the last makefile build
-				HashSet<FileReference> AllUHTHeaders = new HashSet<FileReference>(LoadedUBTMakefile.UObjectModuleHeaders.SelectMany(x => x.HeaderFiles).Select(x => x.Location));
+				HashSet<FileItem> AllUHTHeaders = new HashSet<FileItem>(LoadedUBTMakefile.UObjectModuleHeaders.SelectMany(x => x.HeaderFiles));
 
 				// Check whether any headers have been deleted. If they have, we need to regenerate the makefile since the module might now be empty. If we don't,
 				// and the file has been moved to a different module, we may include stale generated headers.
-				foreach (FileReference FileName in AllUHTHeaders)
+				foreach (FileItem HeaderFile in AllUHTHeaders)
 				{
-					if (!FileReference.Exists(FileName))
+					if (!HeaderFile.Exists)
 					{
-						Log.TraceLog("File processed by UHT was deleted ({0}); invalidating makefile", FileName);
+						Log.TraceLog("File processed by UHT was deleted ({0}); invalidating makefile", HeaderFile);
 						ReasonNotLoaded = string.Format("UHT file was deleted");
 						return null;
 					}
@@ -415,13 +419,13 @@ namespace UnrealBuildTool
 				// * There are any newer files which contain no UHT data, but were previously in the makefile
 				// * There are any newer files contain data which needs processing by UHT, but weren't not previously in the makefile
 				SourceFileMetadataCache MetadataCache = SourceFileMetadataCache.CreateHierarchy(ProjectFile);
-				foreach (FileReference Filename in HFilesNewerThanMakefile)
+				foreach (FileItem HeaderFile in HFilesNewerThanMakefile)
 				{
-					bool bContainsUHTData = MetadataCache.ContainsReflectionMarkup(FileItem.GetItemByFileReference(Filename));
-					bool bWasProcessed = AllUHTHeaders.Contains(Filename);
+					bool bContainsUHTData = MetadataCache.ContainsReflectionMarkup(HeaderFile);
+					bool bWasProcessed = AllUHTHeaders.Contains(HeaderFile);
 					if (bContainsUHTData != bWasProcessed)
 					{
-						Log.TraceLog("{0} {1} contain UHT types and now {2} , ignoring it ({3})", Filename, bWasProcessed ? "used to" : "didn't", bWasProcessed ? "doesn't" : "does", UBTMakefileInfo.FullName);
+						Log.TraceLog("{0} {1} contain UHT types and now {2} , ignoring it ({3})", HeaderFile, bWasProcessed ? "used to" : "didn't", bWasProcessed ? "doesn't" : "does", UBTMakefileInfo.FullName);
 						ReasonNotLoaded = string.Format("new files with reflected types");
 						return null;
 					}
@@ -436,7 +440,7 @@ namespace UnrealBuildTool
 				// Check if any source files in the working set no longer belong in it
 				foreach (FileItem SourceFile in LoadedUBTMakefile.WorkingSet)
 				{
-					if (!WorkingSet.Contains(SourceFile.Location) && File.GetLastWriteTimeUtc(SourceFile.AbsolutePath) > UBTMakefileInfo.LastWriteTimeUtc)
+					if (!WorkingSet.Contains(SourceFile.Location) && SourceFile.LastWriteTimeUtc > UBTMakefileInfo.LastWriteTimeUtc)
 					{
 						Log.TraceLog("{0} was part of source working set and now is not; invalidating makefile ({1})", SourceFile.AbsolutePath, UBTMakefileInfo.FullName);
 						ReasonNotLoaded = string.Format("working set of source files changed");
@@ -447,7 +451,7 @@ namespace UnrealBuildTool
 				// Check if any source files that are eligible for being in the working set have been modified
 				foreach (FileItem SourceFile in LoadedUBTMakefile.CandidatesForWorkingSet)
 				{
-					if (WorkingSet.Contains(SourceFile.Location) && File.GetLastWriteTimeUtc(SourceFile.AbsolutePath) > UBTMakefileInfo.LastWriteTimeUtc)
+					if (WorkingSet.Contains(SourceFile.Location) && SourceFile.LastWriteTimeUtc > UBTMakefileInfo.LastWriteTimeUtc)
 					{
 						Log.TraceLog("{0} was part of source working set and now is not; invalidating makefile ({1})", SourceFile.AbsolutePath, UBTMakefileInfo.FullName);
 						ReasonNotLoaded = string.Format("working set of source files changed");
