Missing file.

#rnx

diff --git a/Engine/Source/Programs/UnrealBuildTool/System/ExternalExecution.cs b/Engine/Source/Programs/UnrealBuildTool/System/ExternalExecution.cs
--- a/Engine/Source/Programs/UnrealBuildTool/System/ExternalExecution.cs
+++ b/Engine/Source/Programs/UnrealBuildTool/System/ExternalExecution.cs
@@ -336,10 +336,10 @@ namespace UnrealBuildTool
 
 	class UHTModuleHeaderInfo
 	{
-		public DirectoryReference SourceFolder;
+		public DirectoryItem SourceFolder;
 		public List<FileItem> HeaderFiles;
 
-		public UHTModuleHeaderInfo(DirectoryReference SourceFolder, List<FileItem> HeaderFiles)
+		public UHTModuleHeaderInfo(DirectoryItem SourceFolder, List<FileItem> HeaderFiles)
 		{
 			this.SourceFolder = SourceFolder;
 			this.HeaderFiles = HeaderFiles;
@@ -347,13 +347,13 @@ namespace UnrealBuildTool
 
 		public UHTModuleHeaderInfo(BinaryArchiveReader Reader)
 		{
-			SourceFolder = Reader.ReadDirectoryReference();
+			SourceFolder = Reader.ReadDirectoryItem();
 			HeaderFiles = Reader.ReadList(() => Reader.ReadFileItem());
 		}
 
 		public void Write(BinaryArchiveWriter Writer)
 		{
-			Writer.WriteDirectoryReference(SourceFolder);
+			Writer.WriteDirectoryItem(SourceFolder);
 			Writer.WriteList(HeaderFiles, Item => Writer.WriteFileItem(Item));
 		}
 	}
@@ -638,10 +638,10 @@ namespace UnrealBuildTool
 			ReadOnlyHashSet<string> ExcludedFolders = UEBuildPlatform.GetBuildPlatform(Platform, true).GetExcludedFolderNames();
 			foreach (UEBuildModuleCPP Module in ModulesSortedByType)
 			{
-				DirectoryItem ModuleDirectory = DirectoryItem.GetItemByDirectoryReference(Module.ModuleDirectory);
+				DirectoryItem ModuleDirectoryItem = DirectoryItem.GetItemByDirectoryReference(Module.ModuleDirectory);
 
 				List<FileItem> HeaderFiles = new List<FileItem>();
-				FindHeaders(ModuleDirectory, ExcludedFolders, HeaderFiles);
+				FindHeaders(ModuleDirectoryItem, ExcludedFolders, HeaderFiles);
 
 				UHTModuleInfo Info = ExternalExecution.CreateUHTModuleInfo(HeaderFiles, Module.Name, Module.RulesFile, Module.ModuleDirectory, ModuleToType[Module], GeneratedCodeVersion, Module.Rules.bUsePrecompiled, MetadataCache);
 				if (Info.PublicUObjectClassesHeaders.Count > 0 || Info.PrivateUObjectHeaders.Count > 0 || Info.PublicUObjectHeaders.Count > 0)
@@ -663,7 +663,7 @@ namespace UnrealBuildTool
 					ReflectedHeaderFiles.AddRange(Info.PublicUObjectHeaders);
 					ReflectedHeaderFiles.AddRange(Info.PublicUObjectClassesHeaders);
 					ReflectedHeaderFiles.AddRange(Info.PrivateUObjectHeaders);
-					UObjectModuleHeaders.Add(new UHTModuleHeaderInfo(Module.ModuleDirectory, ReflectedHeaderFiles));
+					UObjectModuleHeaders.Add(new UHTModuleHeaderInfo(ModuleDirectoryItem, ReflectedHeaderFiles));
 
 					Log.TraceVerbose("Detected UObject module: " + Info.ModuleName);
 				}
