UBT: Move code for merging action graphs into a separate method.

#rnx

diff --git a/Engine/Source/Programs/UnrealBuildTool/Modes/BuildMode.cs b/Engine/Source/Programs/UnrealBuildTool/Modes/BuildMode.cs
--- a/Engine/Source/Programs/UnrealBuildTool/Modes/BuildMode.cs
+++ b/Engine/Source/Programs/UnrealBuildTool/Modes/BuildMode.cs
@@ -191,32 +191,15 @@ namespace UnrealBuildTool
 					ActionsToExecute[TargetIdx] = GetActionsForTarget(BuildConfiguration, TargetDescriptors[TargetIdx], Makefiles[TargetIdx]);
 				}
 
-				// If there are multiple targets being built, make sure the actions for each one 
-				List<Action> MergedActionsToExecute = new List<Action>();
+				// If there are multiple targets being built, merge the actions together
+				List<Action> MergedActionsToExecute;
 				if(TargetDescriptors.Count == 1)
 				{
-					// Just take the actions from the first target
-					MergedActionsToExecute.AddRange(ActionsToExecute[0]);
+					MergedActionsToExecute = new List<Action>(ActionsToExecute[0]);
 				}
 				else
 				{
-					// Set of all output items. Knowing that there are no conflicts in produced items, we use this to eliminate duplicate actions.
-					Dictionary<FileItem, Action> OutputItemToProducingAction = new Dictionary<FileItem, Action>();
-					for(int TargetIdx = 0; TargetIdx < TargetDescriptors.Count; TargetIdx++)
-					{
-						string GroupPrefix = String.Format("{0}-{1}-{2}", TargetDescriptors[TargetIdx].Name, TargetDescriptors[TargetIdx].Platform, TargetDescriptors[TargetIdx].Configuration);
-						foreach(Action Action in ActionsToExecute[TargetIdx])
-						{
-							Action ExistingAction;
-							if(!OutputItemToProducingAction.TryGetValue(Action.ProducedItems[0], out ExistingAction))
-							{
-								OutputItemToProducingAction[Action.ProducedItems[0]] = Action;
-								ExistingAction = Action;
-							}
-							ExistingAction.GroupNames.Add(GroupPrefix);
-						}
-					}
-					MergedActionsToExecute.AddRange(OutputItemToProducingAction.Values);
+					MergedActionsToExecute = MergeActionGraphs(TargetDescriptors, ActionsToExecute);
 				}
 
 				// Link all the actions together
@@ -639,6 +622,33 @@ namespace UnrealBuildTool
 
 			return TargetActionsToExecute;
 		}
+
+		/// <summary>
+		/// Merge action graphs for multiple targets into a single set of actions. Sets group names on merged actions to indicate which target they belong to.
+		/// </summary>
+		/// <param name="TargetDescriptors">List of target descriptors</param>
+		/// <param name="ActionsToExecute">Set of actions to execute for each target</param>
+		/// <returns>List of merged actions</returns>
+		static List<Action> MergeActionGraphs(List<TargetDescriptor> TargetDescriptors, HashSet<Action>[] ActionsToExecute)
+		{
+			// Set of all output items. Knowing that there are no conflicts in produced items, we use this to eliminate duplicate actions.
+			Dictionary<FileItem, Action> OutputItemToProducingAction = new Dictionary<FileItem, Action>();
+			for(int TargetIdx = 0; TargetIdx < TargetDescriptors.Count; TargetIdx++)
+			{
+				string GroupPrefix = String.Format("{0}-{1}-{2}", TargetDescriptors[TargetIdx].Name, TargetDescriptors[TargetIdx].Platform, TargetDescriptors[TargetIdx].Configuration);
+				foreach(Action Action in ActionsToExecute[TargetIdx])
+				{
+					Action ExistingAction;
+					if(!OutputItemToProducingAction.TryGetValue(Action.ProducedItems[0], out ExistingAction))
+					{
+						OutputItemToProducingAction[Action.ProducedItems[0]] = Action;
+						ExistingAction = Action;
+					}
+					ExistingAction.GroupNames.Add(GroupPrefix);
+				}
+			}
+			return new List<Action>(OutputItemToProducingAction.Values);
+		}
 	}
 }
 
