UBT: When a response file is updated, rename the old one rather than dumping it to the log. Some of the trace listeners are very slow (eg. the OutputDebugString one), and writing a large amount of text can significantly affect performance.

#rnx

diff --git a/Engine/Source/Programs/UnrealBuildTool/System/FileItem.cs b/Engine/Source/Programs/UnrealBuildTool/System/FileItem.cs
--- a/Engine/Source/Programs/UnrealBuildTool/System/FileItem.cs
+++ b/Engine/Source/Programs/UnrealBuildTool/System/FileItem.cs
@@ -168,29 +168,42 @@ namespace UnrealBuildTool
 		/// Creates a text file with the given contents.  If the contents of the text file aren't changed, it won't write the new contents to
 		/// the file to avoid causing an action to be considered outdated.
 		/// </summary>
-		/// <param name="AbsolutePath">Path to the intermediate file to create</param>
+		/// <param name="Location">Path to the intermediate file to create</param>
 		/// <param name="Contents">Contents of the new file</param>
 		/// <returns>File item for the newly created file</returns>
-		public static FileItem CreateIntermediateTextFile(FileReference AbsolutePath, string Contents)
+		public static FileItem CreateIntermediateTextFile(FileReference Location, string Contents)
 		{
-			// Create the directory if it doesn't exist.
-			Directory.CreateDirectory(Path.GetDirectoryName(AbsolutePath.FullName));
-
 			// Only write the file if its contents have changed.
-			if (!FileReference.Exists(AbsolutePath))
+			if (!FileReference.Exists(Location))
 			{
-				File.WriteAllText(AbsolutePath.FullName, Contents, GetEncodingForString(Contents));
+				DirectoryReference.CreateDirectory(Location.Directory);
+				FileReference.WriteAllText(Location, Contents, GetEncodingForString(Contents));
 			}
 			else
 			{
-				string CurrentContents = Utils.ReadAllText(AbsolutePath.FullName);
+				string CurrentContents = Utils.ReadAllText(Location.FullName);
 				if(!String.Equals(CurrentContents, Contents, StringComparison.InvariantCultureIgnoreCase))
 				{
-					Log.TraceLog("Updating {0} - contents have changed. Previous:\n  {1}\nNew:\n  {2}", AbsolutePath.FullName, CurrentContents.Replace("\n", "\n  "), Contents.Replace("\n", "\n  "));
-					File.WriteAllText(AbsolutePath.FullName, Contents, GetEncodingForString(Contents));
+					FileReference BackupFile = new FileReference(Location.FullName + ".old");
+					try
+					{
+						Log.TraceLog("Updating {0}: contents have changed. Saving previous version to {1}.", Location, BackupFile);
+						FileReference.Delete(BackupFile);
+						FileReference.Move(Location, BackupFile);
+					}
+					catch(Exception Ex)
+					{
+						Log.TraceWarning("Unable to rename {0} to {1}", Location, BackupFile);
+						Log.TraceLog("{0}", ExceptionUtils.FormatExceptionDetails(Ex));
+					}
+					FileReference.WriteAllText(Location, Contents, GetEncodingForString(Contents));
 				}
 			}
-			return GetItemByFileReference(AbsolutePath);
+
+			// Reset the file info, in case it already knows about the old file
+			FileItem Item = GetItemByFileReference(Location);
+			Item.ResetFileInfo();
+			return Item;
 		}
 
 		/// <summary>
