Fix platform search paths.

#rnx

diff --git a/Engine/Source/Programs/UnrealBuildTool/Configuration/UEBuildPlatform.cs b/Engine/Source/Programs/UnrealBuildTool/Configuration/UEBuildPlatform.cs
--- a/Engine/Source/Programs/UnrealBuildTool/Configuration/UEBuildPlatform.cs
+++ b/Engine/Source/Programs/UnrealBuildTool/Configuration/UEBuildPlatform.cs
@@ -119,7 +119,7 @@ namespace UnrealBuildTool
 		/// <summary>
 		/// Finds a list of folder names to include when building for this platform
 		/// </summary>
-		virtual public ReadOnlyHashSet<string> GetIncludedFolderNames()
+		public ReadOnlyHashSet<string> GetIncludedFolderNames()
 		{
 			if(CachedIncludedFolderNames == null)
 			{
