UBT: Parse hot reload parameters in the target descriptor.

#rnx

diff --git a/Engine/Source/Programs/UnrealBuildTool/Configuration/TargetDescriptor.cs b/Engine/Source/Programs/UnrealBuildTool/Configuration/TargetDescriptor.cs
--- a/Engine/Source/Programs/UnrealBuildTool/Configuration/TargetDescriptor.cs
+++ b/Engine/Source/Programs/UnrealBuildTool/Configuration/TargetDescriptor.cs
@@ -27,12 +27,30 @@ namespace UnrealBuildTool
 		[CommandLine("-Plugin=")]
 		public FileReference ForeignPlugin = null;
 
+		/// <summary>
+		/// Set of module names to compile.
+		/// </summary>
+		[CommandLine("-Module=")]
+		public HashSet<string> OnlyModuleNames = new HashSet<string>(StringComparer.OrdinalIgnoreCase);
+
 		/// <summary>
 		/// Single file to compile
 		/// </summary>
 		[CommandLine("-SingleFile=")]
 		public FileReference SingleFileToCompile = null;
 
+		/// <summary>
+		/// Whether to perform hot reload for this target
+		/// </summary>
+		[CommandLine("-NoHotReload", Value = nameof(HotReloadMode.Disabled))]
+		[CommandLine("-ForceHotReload", Value = nameof(HotReloadMode.FromIDE))]
+		public HotReloadMode HotReloadMode = HotReloadMode.Default;
+
+		/// <summary>
+		/// Map of module name to suffix for hot reloading from the editor
+		/// </summary>
+		public Dictionary<string, int> HotReloadModuleNameToSuffix = new Dictionary<string, int>(StringComparer.OrdinalIgnoreCase);
+
 		/// <summary>
 		/// Constructor
 		/// </summary>
@@ -50,6 +68,26 @@ namespace UnrealBuildTool
 			this.Configuration = Configuration;
 			this.Architecture = Architecture;
 
+			// Parse the list of hot-reload module names
+			foreach(string ModuleWithSuffix in Arguments.GetValues("-ModuleWithSuffix="))
+			{
+				int SuffixIdx = ModuleWithSuffix.LastIndexOf(',');
+				if(SuffixIdx == -1)
+				{
+					throw new BuildException("Missing suffix argument from -ModuleWithSuffix=Name,Suffix");
+				}
+
+				string ModuleName = ModuleWithSuffix.Substring(0, SuffixIdx);
+
+				int Suffix;
+				if(!Int32.TryParse(ModuleWithSuffix.Substring(SuffixIdx + 1), out Suffix))
+				{
+					throw new BuildException("Suffix for modules must be an integer");
+				}
+
+				HotReloadModuleNameToSuffix[ModuleName] = Suffix;
+			}
+
 			// If there are any additional command line arguments
 			List<string> AdditionalArguments = new List<string>();
 			if(Arguments != null)
diff --git a/Engine/Source/Programs/UnrealBuildTool/Modes/BuildMode.cs b/Engine/Source/Programs/UnrealBuildTool/Modes/BuildMode.cs
--- a/Engine/Source/Programs/UnrealBuildTool/Modes/BuildMode.cs
+++ b/Engine/Source/Programs/UnrealBuildTool/Modes/BuildMode.cs
@@ -227,6 +227,7 @@ namespace UnrealBuildTool
 				// Execute the actions
 				if (BuildConfiguration.bXGEExport)
 				{
+					// Just export to an XML file
 					using(Timeline.ScopeEvent("XGE.ExportActions()"))
 					{
 						XGE.ExportActions(MergedActionsToExecute);
@@ -234,15 +235,13 @@ namespace UnrealBuildTool
 				}
 				else
 				{
+					// Execute the actions
 					using(Timeline.ScopeEvent("ActionGraph.ExecuteActions()"))
 					{
 						ActionGraph.ExecuteActions(BuildConfiguration, MergedActionsToExecute);
 					}
-				}
 
-				// Run the deployment steps
-				if (!BuildConfiguration.bXGEExport)
-				{
+					// Run the deployment steps
 					foreach(TargetMakefile Makefile in Makefiles)
 					{
 						if (Makefile.bDeployAfterCompile)
@@ -359,54 +358,32 @@ namespace UnrealBuildTool
 			// Create the action graph
 			ActionGraph.Link(Makefile.Actions);
 
-			// Parse the list of hot-reload module names
-			Dictionary<string, int> HotReloadModuleNameToSuffix = new Dictionary<string, int>(StringComparer.OrdinalIgnoreCase);
-			foreach(string ModuleWithSuffix in TargetDescriptor.AdditionalArguments.GetValues("-ModuleWithSuffix="))
-			{
-				int SuffixIdx = ModuleWithSuffix.LastIndexOf(',');
-				if(SuffixIdx == -1)
-				{
-					throw new BuildException("Missing suffix argument from -ModuleWithSuffix=Name,Suffix");
-				}
-
-				string ModuleName = ModuleWithSuffix.Substring(0, SuffixIdx);
-
-				int Suffix;
-				if(!Int32.TryParse(ModuleWithSuffix.Substring(SuffixIdx + 1), out Suffix))
-				{
-					throw new BuildException("Suffix for modules must be an integer");
-				}
-
-				HotReloadModuleNameToSuffix[ModuleName] = Suffix;
-			}
-
 			// Get the hot-reload mode
-			HotReloadMode HotReloadMode = HotReloadMode.Disabled;
-			if(TargetDescriptor.AdditionalArguments.Any(x => x.Equals("-ForceHotReload", StringComparison.InvariantCultureIgnoreCase)))
-			{
-				HotReloadMode = HotReloadMode.FromIDE;
-			}
-			else if (!TargetDescriptor.AdditionalArguments.Any(x => x.Equals("-NoHotReload", StringComparison.InvariantCultureIgnoreCase)))
+			HotReloadMode HotReloadMode = TargetDescriptor.HotReloadMode;
+			if(HotReloadMode == HotReloadMode.Default)
 			{
 				if (BuildConfiguration.bAllowHotReloadFromIDE && HotReload.ShouldDoHotReloadFromIDE(BuildConfiguration, TargetDescriptor))
 				{
 					HotReloadMode = HotReloadMode.FromIDE;
 				}
-				else if (HotReloadModuleNameToSuffix.Count > 0 && TargetDescriptor.ForeignPlugin == null)
+				else if (TargetDescriptor.HotReloadModuleNameToSuffix.Count > 0 && TargetDescriptor.ForeignPlugin == null)
 				{
 					HotReloadMode = HotReloadMode.FromEditor;
 				}
+				else
+				{
+					HotReloadMode = HotReloadMode.Disabled;
+				}
 			}
 
 			// Get the output items we're trying to build
 			List<FileItem> TargetOutputItems = Makefile.OutputItems;
 
 			// Parse the list of modules to build
-			HashSet<string> OnlyModuleNames = new HashSet<string>(TargetDescriptor.AdditionalArguments.GetValues("-Module="));
-			if (OnlyModuleNames.Count > 0)
+			if (TargetDescriptor.OnlyModuleNames.Count > 0)
 			{
 				TargetOutputItems = new List<FileItem>();
-				foreach(string OnlyModuleName in OnlyModuleNames)
+				foreach(string OnlyModuleName in TargetDescriptor.OnlyModuleNames)
 				{
 					FileItem[] OutputItemsForModule;
 					if(!Makefile.ModuleNameToOutputItems.TryGetValue(OnlyModuleName, out OutputItemsForModule))
@@ -467,13 +444,13 @@ namespace UnrealBuildTool
 				}
 
 				// If we want a specific suffix on any modules, apply that now. We'll track the outputs later, but the suffix has to be forced  (and is always out of date if it doesn't exist).
-				if(HotReloadModuleNameToSuffix.Count > 0)
+				if(TargetDescriptor.HotReloadModuleNameToSuffix.Count > 0)
 				{
 					Dictionary<FileReference, FileReference> OldLocationToNewLocation = new Dictionary<FileReference, FileReference>();
 					foreach(string HotReloadModuleName in Makefile.HotReloadModuleNames)
 					{
 						int ModuleSuffix;
-						if(HotReloadModuleNameToSuffix.TryGetValue(HotReloadModuleName, out ModuleSuffix))
+						if(TargetDescriptor.HotReloadModuleNameToSuffix.TryGetValue(HotReloadModuleName, out ModuleSuffix))
 						{
 							FileItem[] ModuleOutputItems = Makefile.ModuleNameToOutputItems[HotReloadModuleName];
 							foreach(FileItem ModuleOutputItem in ModuleOutputItems)
@@ -526,7 +503,7 @@ namespace UnrealBuildTool
 				foreach(string HotReloadModuleName in Makefile.HotReloadModuleNames)
 				{
 					int ModuleSuffix;
-					if(!HotReloadModuleNameToSuffix.TryGetValue(HotReloadModuleName, out ModuleSuffix) || ModuleSuffix == -1)
+					if(!TargetDescriptor.HotReloadModuleNameToSuffix.TryGetValue(HotReloadModuleName, out ModuleSuffix) || ModuleSuffix == -1)
 					{
 						FileItem[] ModuleOutputItems;
 						if(Makefile.ModuleNameToOutputItems.TryGetValue(HotReloadModuleName, out ModuleOutputItems))
diff --git a/Engine/Source/Programs/UnrealBuildTool/System/HotReload.cs b/Engine/Source/Programs/UnrealBuildTool/System/HotReload.cs
--- a/Engine/Source/Programs/UnrealBuildTool/System/HotReload.cs
+++ b/Engine/Source/Programs/UnrealBuildTool/System/HotReload.cs
@@ -16,6 +16,7 @@ namespace UnrealBuildTool
 	/// </summary>
 	enum HotReloadMode
 	{
+		Default,
 		Disabled,
 		FromIDE,
 		FromEditor
