UBT: Move some hot reload functionality into HotReload.cs.

#rnx

diff --git a/Engine/Source/Programs/UnrealBuildTool/Modes/BuildMode.cs b/Engine/Source/Programs/UnrealBuildTool/Modes/BuildMode.cs
--- a/Engine/Source/Programs/UnrealBuildTool/Modes/BuildMode.cs
+++ b/Engine/Source/Programs/UnrealBuildTool/Modes/BuildMode.cs
@@ -426,43 +426,11 @@ namespace UnrealBuildTool
 					HotReloadState = new HotReloadState();
 				}
 
-				// Update the action graph to produce these new files
-				HotReload.PatchActionGraph(PrerequisiteActions, HotReloadState.OriginalFileToHotReloadFile);
+				// Apply the old state to the makefile
+				HotReload.ApplyState(HotReloadState, Makefile);
 
-				// Update the module to output file mapping
-				foreach(string HotReloadModuleName in Makefile.HotReloadModuleNames)
-				{
-					FileItem[] ModuleOutputItems = Makefile.ModuleNameToOutputItems[HotReloadModuleName];
-					for(int Idx = 0; Idx < ModuleOutputItems.Length; Idx++)
-					{
-						FileReference NewLocation;
-						if(HotReloadState.OriginalFileToHotReloadFile.TryGetValue(ModuleOutputItems[Idx].Location, out NewLocation))
-						{
-							ModuleOutputItems[Idx] = FileItem.GetItemByFileReference(NewLocation);
-						}
-					}
-				}
-
-				// If we want a specific suffix on any modules, apply that now. We'll track the outputs later, but the suffix has to be forced  (and is always out of date if it doesn't exist).
-				if(TargetDescriptor.HotReloadModuleNameToSuffix.Count > 0)
-				{
-					Dictionary<FileReference, FileReference> OldLocationToNewLocation = new Dictionary<FileReference, FileReference>();
-					foreach(string HotReloadModuleName in Makefile.HotReloadModuleNames)
-					{
-						int ModuleSuffix;
-						if(TargetDescriptor.HotReloadModuleNameToSuffix.TryGetValue(HotReloadModuleName, out ModuleSuffix))
-						{
-							FileItem[] ModuleOutputItems = Makefile.ModuleNameToOutputItems[HotReloadModuleName];
-							foreach(FileItem ModuleOutputItem in ModuleOutputItems)
-							{
-								FileReference OldLocation = ModuleOutputItem.Location;
-								FileReference NewLocation = HotReload.ReplaceSuffix(OldLocation, ModuleSuffix);
-								OldLocationToNewLocation[OldLocation] = NewLocation;
-							}
-						}
-					}
-					HotReload.PatchActionGraph(PrerequisiteActions, OldLocationToNewLocation);
-				}
+				// If we want a specific suffix on any modules, apply that now. We'll track the outputs later, but the suffix has to be forced (and is always out of date if it doesn't exist).
+				HotReload.PatchActionGraphWithNames(PrerequisiteActions, TargetDescriptor.HotReloadModuleNameToSuffix, Makefile);
 
 				// Re-link the action graph
 				ActionGraph.Link(PrerequisiteActions);
diff --git a/Engine/Source/Programs/UnrealBuildTool/System/HotReload.cs b/Engine/Source/Programs/UnrealBuildTool/System/HotReload.cs
--- a/Engine/Source/Programs/UnrealBuildTool/System/HotReload.cs
+++ b/Engine/Source/Programs/UnrealBuildTool/System/HotReload.cs
@@ -287,6 +287,31 @@ namespace UnrealBuildTool
 			}
 		}
 
+		/// <summary>
+		/// Apply a saved hot reload state to a makefile
+		/// </summary>
+		/// <param name="HotReloadState">The hot-reload state</param>
+		/// <param name="Makefile">Makefile to apply the state</param>
+		public static void ApplyState(HotReloadState HotReloadState, TargetMakefile Makefile)
+		{
+			// Update the action graph to produce these new files
+			HotReload.PatchActionGraph(Makefile.Actions, HotReloadState.OriginalFileToHotReloadFile);
+
+			// Update the module to output file mapping
+			foreach(string HotReloadModuleName in Makefile.HotReloadModuleNames)
+			{
+				FileItem[] ModuleOutputItems = Makefile.ModuleNameToOutputItems[HotReloadModuleName];
+				for(int Idx = 0; Idx < ModuleOutputItems.Length; Idx++)
+				{
+					FileReference NewLocation;
+					if(HotReloadState.OriginalFileToHotReloadFile.TryGetValue(ModuleOutputItems[Idx].Location, out NewLocation))
+					{
+						ModuleOutputItems[Idx] = FileItem.GetItemByFileReference(NewLocation);
+					}
+				}
+			}
+		}
+
 		/// <summary>
 		/// Replaces a hot reload suffix in a filename.
 		/// </summary>
@@ -636,5 +661,132 @@ namespace UnrealBuildTool
 				}
 			}
 		}
+
+		/// <summary>
+		/// Patches a set of actions to use a specific list of suffixes for each module name
+		/// </summary>
+		/// <param name="PrerequisiteActions">Actions in the target being built</param>
+		/// <param name="ModuleNameToSuffix">Map of module name to suffix</param>
+		/// <param name="Makefile">Makefile for the target being built</param>
+		public static void PatchActionGraphWithNames(List<Action> PrerequisiteActions, Dictionary<string, int> ModuleNameToSuffix, TargetMakefile Makefile)
+		{
+			if(ModuleNameToSuffix.Count > 0)
+			{
+				Dictionary<FileReference, FileReference> OldLocationToNewLocation = new Dictionary<FileReference, FileReference>();
+				foreach(string HotReloadModuleName in Makefile.HotReloadModuleNames)
+				{
+					int ModuleSuffix;
+					if(ModuleNameToSuffix.TryGetValue(HotReloadModuleName, out ModuleSuffix))
+					{
+						FileItem[] ModuleOutputItems = Makefile.ModuleNameToOutputItems[HotReloadModuleName];
+						foreach(FileItem ModuleOutputItem in ModuleOutputItems)
+						{
+							FileReference OldLocation = ModuleOutputItem.Location;
+							FileReference NewLocation = HotReload.ReplaceSuffix(OldLocation, ModuleSuffix);
+							OldLocationToNewLocation[OldLocation] = NewLocation;
+						}
+					}
+				}
+				HotReload.PatchActionGraph(PrerequisiteActions, OldLocationToNewLocation);
+			}
+		}
+
+
+		public static void UpdateState(HotReloadState HotReloadState, List<Action> PrerequisiteActions, List<Action> TargetActionsToExecute, Dictionary<string, int> ModuleNameToSuffix, TargetMakefile Makefile)
+		{
+			// For all the hot-reloadable modules that may need a unique suffix appended, build a mapping from output item to all the output items in that module. We can't 
+			// apply a suffix to one without applying a suffix to all of them.
+			Dictionary<FileItem, FileItem[]> HotReloadItemToDependentItems = new Dictionary<FileItem, FileItem[]>();
+			foreach(string HotReloadModuleName in Makefile.HotReloadModuleNames)
+			{
+				int ModuleSuffix;
+				if(!ModuleNameToSuffix.TryGetValue(HotReloadModuleName, out ModuleSuffix) || ModuleSuffix == -1)
+				{
+					FileItem[] ModuleOutputItems;
+					if(Makefile.ModuleNameToOutputItems.TryGetValue(HotReloadModuleName, out ModuleOutputItems))
+					{
+						foreach(FileItem ModuleOutputItem in ModuleOutputItems)
+						{
+							HotReloadItemToDependentItems[ModuleOutputItem] = ModuleOutputItems;
+						}
+					}
+				} 
+			}
+
+			// Expand the list of actions to execute to include everything that references any files with a new suffix. Unlike a regular build, we can't ignore
+			// dependencies on import libraries under the assumption that a header would change if the API changes, because the dependency will be on a different DLL.
+			HashSet<FileItem> FilesRequiringSuffix = new HashSet<FileItem>(TargetActionsToExecute.SelectMany(x => x.ProducedItems).Where(x => HotReloadItemToDependentItems.ContainsKey(x)));
+			for(int LastNumFilesWithNewSuffix = 0; FilesRequiringSuffix.Count > LastNumFilesWithNewSuffix;)
+			{
+				LastNumFilesWithNewSuffix = FilesRequiringSuffix.Count;
+				foreach(Action PrerequisiteAction in PrerequisiteActions)
+				{
+					if(!TargetActionsToExecute.Contains(PrerequisiteAction))
+					{
+						foreach(FileItem ProducedItem in PrerequisiteAction.ProducedItems)
+						{
+							FileItem[] DependentItems;
+							if(HotReloadItemToDependentItems.TryGetValue(ProducedItem, out DependentItems))
+							{
+								TargetActionsToExecute.Add(PrerequisiteAction);
+								FilesRequiringSuffix.UnionWith(DependentItems);
+							}
+						}
+					}
+				}
+			}
+
+			// Build a list of file mappings
+			Dictionary<FileReference, FileReference> OldLocationToNewLocation = new Dictionary<FileReference, FileReference>();
+			foreach(FileItem FileRequiringSuffix in FilesRequiringSuffix)
+			{
+				FileReference OldLocation = FileRequiringSuffix.Location;
+				FileReference NewLocation = HotReload.ReplaceSuffix(OldLocation, HotReloadState.NextSuffix);
+				OldLocationToNewLocation[OldLocation] = NewLocation;
+			}
+
+			// Update the action graph with these new paths
+			HotReload.PatchActionGraph(PrerequisiteActions, OldLocationToNewLocation);
+
+			// Get a new list of actions to execute now that the graph has been modified
+			TargetActionsToExecute = ActionGraph.GetActionsToExecute(Makefile.Actions, PrerequisiteActions, CppDependencies, History, BuildConfiguration.bIgnoreOutdatedImportLibraries);
+
+			// Build a mapping of all file items to their original
+			Dictionary<FileReference, FileReference> HotReloadFileToOriginalFile = new Dictionary<FileReference, FileReference>();
+			foreach(KeyValuePair<FileReference, FileReference> Pair in HotReloadState.OriginalFileToHotReloadFile)
+			{
+				HotReloadFileToOriginalFile[Pair.Value] = Pair.Key;
+			}
+			foreach(KeyValuePair<FileReference, FileReference> Pair in OldLocationToNewLocation)
+			{
+				FileReference OriginalLocation;
+				if(!HotReloadFileToOriginalFile.TryGetValue(Pair.Key, out OriginalLocation))
+				{
+					OriginalLocation = Pair.Key;
+				}
+				HotReloadFileToOriginalFile[Pair.Value] = OriginalLocation;
+			}
+
+			// Now filter out all the hot reload files and update the state
+			foreach(Action Action in TargetActionsToExecute)
+			{
+				foreach(FileItem ProducedItem in Action.ProducedItems)
+				{
+					FileReference OriginalLocation;
+					if(HotReloadFileToOriginalFile.TryGetValue(ProducedItem.Location, out OriginalLocation))
+					{
+						HotReloadState.OriginalFileToHotReloadFile[OriginalLocation] = ProducedItem.Location;
+						HotReloadState.TemporaryFiles.Add(ProducedItem.Location);
+					}
+				}
+			}
+
+			// Increment the suffix for the next iteration
+			if(TargetActionsToExecute.Count > 0)
+			{
+				HotReloadState.NextSuffix++;
+			}
+		}
+
 	}
 }
