Fix CIS.

#rnx

diff --git a/Engine/Source/Programs/UnrealBuildTool/System/HotReload.cs b/Engine/Source/Programs/UnrealBuildTool/System/HotReload.cs
--- a/Engine/Source/Programs/UnrealBuildTool/System/HotReload.cs
+++ b/Engine/Source/Programs/UnrealBuildTool/System/HotReload.cs
@@ -690,103 +690,5 @@ namespace UnrealBuildTool
 				HotReload.PatchActionGraph(PrerequisiteActions, OldLocationToNewLocation);
 			}
 		}
-
-
-		public static void UpdateState(HotReloadState HotReloadState, List<Action> PrerequisiteActions, List<Action> TargetActionsToExecute, Dictionary<string, int> ModuleNameToSuffix, TargetMakefile Makefile)
-		{
-			// For all the hot-reloadable modules that may need a unique suffix appended, build a mapping from output item to all the output items in that module. We can't 
-			// apply a suffix to one without applying a suffix to all of them.
-			Dictionary<FileItem, FileItem[]> HotReloadItemToDependentItems = new Dictionary<FileItem, FileItem[]>();
-			foreach(string HotReloadModuleName in Makefile.HotReloadModuleNames)
-			{
-				int ModuleSuffix;
-				if(!ModuleNameToSuffix.TryGetValue(HotReloadModuleName, out ModuleSuffix) || ModuleSuffix == -1)
-				{
-					FileItem[] ModuleOutputItems;
-					if(Makefile.ModuleNameToOutputItems.TryGetValue(HotReloadModuleName, out ModuleOutputItems))
-					{
-						foreach(FileItem ModuleOutputItem in ModuleOutputItems)
-						{
-							HotReloadItemToDependentItems[ModuleOutputItem] = ModuleOutputItems;
-						}
-					}
-				} 
-			}
-
-			// Expand the list of actions to execute to include everything that references any files with a new suffix. Unlike a regular build, we can't ignore
-			// dependencies on import libraries under the assumption that a header would change if the API changes, because the dependency will be on a different DLL.
-			HashSet<FileItem> FilesRequiringSuffix = new HashSet<FileItem>(TargetActionsToExecute.SelectMany(x => x.ProducedItems).Where(x => HotReloadItemToDependentItems.ContainsKey(x)));
-			for(int LastNumFilesWithNewSuffix = 0; FilesRequiringSuffix.Count > LastNumFilesWithNewSuffix;)
-			{
-				LastNumFilesWithNewSuffix = FilesRequiringSuffix.Count;
-				foreach(Action PrerequisiteAction in PrerequisiteActions)
-				{
-					if(!TargetActionsToExecute.Contains(PrerequisiteAction))
-					{
-						foreach(FileItem ProducedItem in PrerequisiteAction.ProducedItems)
-						{
-							FileItem[] DependentItems;
-							if(HotReloadItemToDependentItems.TryGetValue(ProducedItem, out DependentItems))
-							{
-								TargetActionsToExecute.Add(PrerequisiteAction);
-								FilesRequiringSuffix.UnionWith(DependentItems);
-							}
-						}
-					}
-				}
-			}
-
-			// Build a list of file mappings
-			Dictionary<FileReference, FileReference> OldLocationToNewLocation = new Dictionary<FileReference, FileReference>();
-			foreach(FileItem FileRequiringSuffix in FilesRequiringSuffix)
-			{
-				FileReference OldLocation = FileRequiringSuffix.Location;
-				FileReference NewLocation = HotReload.ReplaceSuffix(OldLocation, HotReloadState.NextSuffix);
-				OldLocationToNewLocation[OldLocation] = NewLocation;
-			}
-
-			// Update the action graph with these new paths
-			HotReload.PatchActionGraph(PrerequisiteActions, OldLocationToNewLocation);
-
-			// Get a new list of actions to execute now that the graph has been modified
-			TargetActionsToExecute = ActionGraph.GetActionsToExecute(Makefile.Actions, PrerequisiteActions, CppDependencies, History, BuildConfiguration.bIgnoreOutdatedImportLibraries);
-
-			// Build a mapping of all file items to their original
-			Dictionary<FileReference, FileReference> HotReloadFileToOriginalFile = new Dictionary<FileReference, FileReference>();
-			foreach(KeyValuePair<FileReference, FileReference> Pair in HotReloadState.OriginalFileToHotReloadFile)
-			{
-				HotReloadFileToOriginalFile[Pair.Value] = Pair.Key;
-			}
-			foreach(KeyValuePair<FileReference, FileReference> Pair in OldLocationToNewLocation)
-			{
-				FileReference OriginalLocation;
-				if(!HotReloadFileToOriginalFile.TryGetValue(Pair.Key, out OriginalLocation))
-				{
-					OriginalLocation = Pair.Key;
-				}
-				HotReloadFileToOriginalFile[Pair.Value] = OriginalLocation;
-			}
-
-			// Now filter out all the hot reload files and update the state
-			foreach(Action Action in TargetActionsToExecute)
-			{
-				foreach(FileItem ProducedItem in Action.ProducedItems)
-				{
-					FileReference OriginalLocation;
-					if(HotReloadFileToOriginalFile.TryGetValue(ProducedItem.Location, out OriginalLocation))
-					{
-						HotReloadState.OriginalFileToHotReloadFile[OriginalLocation] = ProducedItem.Location;
-						HotReloadState.TemporaryFiles.Add(ProducedItem.Location);
-					}
-				}
-			}
-
-			// Increment the suffix for the next iteration
-			if(TargetActionsToExecute.Count > 0)
-			{
-				HotReloadState.NextSuffix++;
-			}
-		}
-
 	}
 }
