UBT: Separate code for determining actions which are part of a target into a separate function.

#rnx

diff --git a/Engine/Source/Programs/UnrealBuildTool/Modes/BuildMode.cs b/Engine/Source/Programs/UnrealBuildTool/Modes/BuildMode.cs
--- a/Engine/Source/Programs/UnrealBuildTool/Modes/BuildMode.cs
+++ b/Engine/Source/Programs/UnrealBuildTool/Modes/BuildMode.cs
@@ -376,33 +376,8 @@ namespace UnrealBuildTool
 				}
 			}
 
-			// Get the output items we're trying to build
-			List<FileItem> TargetOutputItems = Makefile.OutputItems;
-
-			// Parse the list of modules to build
-			if (TargetDescriptor.OnlyModuleNames.Count > 0)
-			{
-				TargetOutputItems = new List<FileItem>();
-				foreach(string OnlyModuleName in TargetDescriptor.OnlyModuleNames)
-				{
-					FileItem[] OutputItemsForModule;
-					if(!Makefile.ModuleNameToOutputItems.TryGetValue(OnlyModuleName, out OutputItemsForModule))
-					{
-						throw new BuildException("Unable to find output items for module '{0}'", OnlyModuleName);
-					}
-					TargetOutputItems.AddRange(OutputItemsForModule);
-				}
-			}
-
-			// If we're just compiling a single file, set the target items to be all the derived items
-			if(TargetDescriptor.SingleFileToCompile != null)
-			{
-				FileItem FileToCompile = FileItem.GetItemByFileReference(TargetDescriptor.SingleFileToCompile);
-				TargetOutputItems = Makefile.Actions.Where(x => x.PrerequisiteItems.Contains(FileToCompile)).SelectMany(x => x.ProducedItems).ToList();
-			}
-
 			// Get the root prerequisite actions
-			List<Action> PrerequisiteActions = ActionGraph.GatherPrerequisiteActions(Makefile.Actions, new HashSet<FileItem>(TargetOutputItems));
+			List<Action> PrerequisiteActions = GatherPrerequisiteActions(TargetDescriptor, Makefile);
 
 			// Get the path to the hot reload state file for this target
 			FileReference HotReloadStateFile = global::UnrealBuildTool.HotReloadState.GetLocation(TargetDescriptor.ProjectFile, TargetDescriptor.Name, TargetDescriptor.Platform, TargetDescriptor.Configuration, TargetDescriptor.Architecture);
@@ -450,7 +425,7 @@ namespace UnrealBuildTool
 				History = ActionHistory.CreateHierarchy(TargetDescriptor.ProjectFile, TargetDescriptor.Name, TargetDescriptor.Platform, Makefile.TargetType);
 			}
 
-			// Plan the actions to execute for the build.
+			// Plan the actions to execute for the build. For single file compiles, always rebuild the source file regardless of whether it's out of date.
 			HashSet<Action> TargetActionsToExecute;
 			if (TargetDescriptor.SingleFileToCompile == null)
 			{
@@ -568,6 +543,44 @@ namespace UnrealBuildTool
 			return TargetActionsToExecute;
 		}
 
+		/// <summary>
+		/// Determines all the actions that should be executed for a target (filtering for single module/file, etc..)
+		/// </summary>
+		/// <param name="TargetDescriptor">The target being built</param>
+		/// <param name="Makefile">Makefile for the target</param>
+		/// <returns>List of actions that need to be executed</returns>
+		static List<Action> GatherPrerequisiteActions(TargetDescriptor TargetDescriptor, TargetMakefile Makefile)
+		{
+			List<Action> PrerequisiteActions;
+			if(TargetDescriptor.SingleFileToCompile != null)
+			{
+				// If we're just compiling a single file, set the target items to be all the derived items
+				FileItem FileToCompile = FileItem.GetItemByFileReference(TargetDescriptor.SingleFileToCompile);
+				PrerequisiteActions = Makefile.Actions.Where(x => x.PrerequisiteItems.Contains(FileToCompile)).ToList();
+			}
+			else if(TargetDescriptor.OnlyModuleNames.Count > 0)
+			{
+				// Find the output items for this module
+				HashSet<FileItem> ModuleOutputItems = new HashSet<FileItem>();
+				foreach(string OnlyModuleName in TargetDescriptor.OnlyModuleNames)
+				{
+					FileItem[] OutputItemsForModule;
+					if(!Makefile.ModuleNameToOutputItems.TryGetValue(OnlyModuleName, out OutputItemsForModule))
+					{
+						throw new BuildException("Unable to find output items for module '{0}'", OnlyModuleName);
+					}
+					ModuleOutputItems.UnionWith(OutputItemsForModule);
+				}
+				PrerequisiteActions = ActionGraph.GatherPrerequisiteActions(Makefile.Actions, ModuleOutputItems);
+			}
+			else
+			{
+				// Use all the output items from the target
+				PrerequisiteActions = ActionGraph.GatherPrerequisiteActions(Makefile.Actions, new HashSet<FileItem>(Makefile.OutputItems));
+			}
+			return PrerequisiteActions;
+		}
+
 		/// <summary>
 		/// Merge action graphs for multiple targets into a single set of actions. Sets group names on merged actions to indicate which target they belong to.
 		/// </summary>
