UBT: When building a target, start prefetching the contents of engine/project/plugin source folders in parallel on background threads as soon as UBT starts up (by instantiating FileItem and DirectoryItem objects). This reduces I/O waits on the main thread when validating the contents of the makefile or building a target. Measured ~20% reduction in UBT times when using a makefile.

Other directories are probably worth prefetching too; needs measuring.

diff --git a/Engine/Source/Programs/UnrealBuildTool/Modes/BuildMode.cs b/Engine/Source/Programs/UnrealBuildTool/Modes/BuildMode.cs
--- a/Engine/Source/Programs/UnrealBuildTool/Modes/BuildMode.cs
+++ b/Engine/Source/Programs/UnrealBuildTool/Modes/BuildMode.cs
@@ -15,7 +15,7 @@ namespace UnrealBuildTool
 	/// <summary>
 	/// Builds a target
 	/// </summary>
-	[ToolMode("Build", ToolModeOptions.XmlConfig | ToolModeOptions.BuildPlatforms | ToolModeOptions.SingleInstance)]
+	[ToolMode("Build", ToolModeOptions.XmlConfig | ToolModeOptions.BuildPlatforms | ToolModeOptions.SingleInstance | ToolModeOptions.StartPrefetchingEngine)]
 	class BuildMode : ToolMode
 	{
 		/// <summary>
@@ -102,10 +102,10 @@ namespace UnrealBuildTool
 			try
 			{
 				// Parse all the target descriptors
-				List<TargetDescriptor> TargetDescriptors = new List<TargetDescriptor>();
+				List<TargetDescriptor> TargetDescriptors;
 				using(Timeline.ScopeEvent("TargetDescriptor.ParseCommandLine()"))
 				{
-					TargetDescriptors.AddRange(TargetDescriptor.ParseCommandLine(Arguments, BuildConfiguration.bUsePrecompiled, BuildConfiguration.bSkipRulesCompile));
+					TargetDescriptors = TargetDescriptor.ParseCommandLine(Arguments, BuildConfiguration.bUsePrecompiled, BuildConfiguration.bSkipRulesCompile);
 				}
 
 				// Handle remote builds
@@ -133,13 +133,15 @@ namespace UnrealBuildTool
 				// Handle local builds
 				if(TargetDescriptors.Count > 0)
 				{
-					// Get a set of all the project directories.
+					// Get a set of all the project directories
 					HashSet<DirectoryReference> ProjectDirs = new HashSet<DirectoryReference>();
 					foreach(TargetDescriptor TargetDesc in TargetDescriptors)
 					{
 						if(TargetDesc.ProjectFile != null)
 						{
-							ProjectDirs.Add(TargetDesc.ProjectFile.Directory);
+							DirectoryReference ProjectDirectory = TargetDesc.ProjectFile.Directory;
+							FileMetadataPrefetch.QueueProjectDirectory(ProjectDirectory);
+							ProjectDirs.Add(ProjectDirectory);
 						}
 					}
 
diff --git a/Engine/Source/Programs/UnrealBuildTool/System/CppDependencyCache.cs b/Engine/Source/Programs/UnrealBuildTool/System/CppDependencyCache.cs
--- a/Engine/Source/Programs/UnrealBuildTool/System/CppDependencyCache.cs
+++ b/Engine/Source/Programs/UnrealBuildTool/System/CppDependencyCache.cs
@@ -171,13 +171,13 @@ namespace UnrealBuildTool
 					AppName = UEBuildTarget.GetAppNameForTargetType(TargetType);
 				}
 
-				FileReference EngineCacheLocation = FileReference.Combine(UnrealBuildTool.EngineDirectory, "Intermediate", "Build", Platform.ToString(), AppName, Configuration.ToString(), "DependencyCache.dat");
+				FileReference EngineCacheLocation = FileReference.Combine(UnrealBuildTool.EngineDirectory, "Intermediate", "Build", Platform.ToString(), AppName, Configuration.ToString(), "DependencyCache.bin");
 				Cache = FindOrAddCache(EngineCacheLocation, UnrealBuildTool.EngineDirectory, Cache);
 			}
 
 			if(ProjectFile != null)
 			{
-				FileReference ProjectCacheLocation = FileReference.Combine(ProjectFile.Directory, "Intermediate", "Build", Platform.ToString(), TargetName, Configuration.ToString(), "DependencyCache.dat");
+				FileReference ProjectCacheLocation = FileReference.Combine(ProjectFile.Directory, "Intermediate", "Build", Platform.ToString(), TargetName, Configuration.ToString(), "DependencyCache.bin");
 				Cache = FindOrAddCache(ProjectCacheLocation, ProjectFile.Directory, Cache);
 			}
 
diff --git a/Engine/Source/Programs/UnrealBuildTool/System/FileMetadataPrefetch.cs b/Engine/Source/Programs/UnrealBuildTool/System/FileMetadataPrefetch.cs
new file mode 100644
--- /dev/null
+++ b/Engine/Source/Programs/UnrealBuildTool/System/FileMetadataPrefetch.cs
@@ -0,0 +1,172 @@
+using System;
+using System.Collections.Generic;
+using System.Linq;
+using System.Text;
+using System.Threading;
+using System.Threading.Tasks;
+using Tools.DotNETCommon;
+
+namespace UnrealBuildTool
+{
+	/// <summary>
+	/// Prefetches metadata from the filesystem, by populating FileItem and DirectoryItem objects for requested directory trees. Since 
+	/// </summary>
+	static class FileMetadataPrefetch
+	{
+		/// <summary>
+		/// Queue for tasks added to the thread pool
+		/// </summary>
+		static ThreadPoolWorkQueue Queue = new ThreadPoolWorkQueue();
+
+		/// <summary>
+		/// Used to cancel any queued tasks
+		/// </summary>
+		static CancellationTokenSource CancelSource = new CancellationTokenSource();
+
+		/// <summary>
+		/// The cancellation token
+		/// </summary>
+		static CancellationToken CancelToken = CancelSource.Token;
+
+		/// <summary>
+		/// Set of all the directory trees that have been queued up, to save adding any more than once.
+		/// </summary>
+		static HashSet<DirectoryReference> QueuedDirectories = new HashSet<DirectoryReference>();
+
+		/// <summary>
+		/// Enqueue the engine directory for prefetching
+		/// </summary>
+		public static void QueueEngineDirectory()
+		{
+			lock(QueuedDirectories)
+			{
+				if(QueuedDirectories.Add(UnrealBuildTool.EngineDirectory))
+				{
+					Enqueue(() => ScanEngineDirectory());
+				}
+			}
+		}
+
+		/// <summary>
+		/// Enqueue a project directory for prefetching
+		/// </summary>
+		/// <param name="ProjectDirectory">The project directory to prefetch</param>
+		public static void QueueProjectDirectory(DirectoryReference ProjectDirectory)
+		{
+			lock(QueuedDirectories)
+			{
+				if(QueuedDirectories.Add(ProjectDirectory))
+				{
+					Enqueue(() => ScanProjectDirectory(DirectoryItem.GetItemByDirectoryReference(ProjectDirectory)));
+				}
+			}
+		}
+
+		/// <summary>
+		/// Enqueue a directory tree for prefetching
+		/// </summary>
+		/// <param name="Directory">Directory to start searching from</param>
+		public static void QueueDirectoryTree(DirectoryReference Directory)
+		{
+			lock(QueuedDirectories)
+			{
+				if(QueuedDirectories.Add(Directory))
+				{
+					Enqueue(() => ScanDirectoryTree(DirectoryItem.GetItemByDirectoryReference(Directory)));
+				}
+			}
+		}
+
+		/// <summary>
+		/// Wait for the prefetcher to complete all reqeusted tasks
+		/// </summary>
+		public static void Wait()
+		{
+			Queue.Wait();
+		}
+
+		/// <summary>
+		/// Stop prefetching items, and cancel all pending tasks. synchronous.
+		/// </summary>
+		public static void Stop()
+		{
+			CancelSource.Cancel();
+			Queue.Wait();
+		}
+
+		/// <summary>
+		/// Enqueue a task which checks for the cancellation token first
+		/// </summary>
+		/// <param name="Action">Action to enqueue</param>
+		static void Enqueue(System.Action Action)
+		{
+			Queue.Enqueue(() => { if(!CancelToken.IsCancellationRequested){ Action(); } });
+		}
+
+		/// <summary>
+		/// Scans the engine directory, adding tasks for subdirectories
+		/// </summary>
+		static void ScanEngineDirectory()
+		{
+			DirectoryItem EngineDirectory = DirectoryItem.GetItemByDirectoryReference(UnrealBuildTool.EngineDirectory);
+			EngineDirectory.CacheDirectories();
+
+			DirectoryItem EnginePluginsDirectory = DirectoryItem.Combine(EngineDirectory, "Plugins");
+			Enqueue(() => ScanPluginFolder(EnginePluginsDirectory));
+
+			DirectoryItem EngineRuntimeDirectory = DirectoryItem.GetItemByDirectoryReference(UnrealBuildTool.EngineSourceRuntimeDirectory);
+			Enqueue(() => ScanDirectoryTree(EngineRuntimeDirectory));
+
+			DirectoryItem EngineDeveloperDirectory = DirectoryItem.GetItemByDirectoryReference(UnrealBuildTool.EngineSourceDeveloperDirectory);
+			Enqueue(() => ScanDirectoryTree(EngineDeveloperDirectory));
+
+			DirectoryItem EngineEditorDirectory = DirectoryItem.GetItemByDirectoryReference(UnrealBuildTool.EngineSourceEditorDirectory);
+			Enqueue(() => ScanDirectoryTree(EngineEditorDirectory));
+		}
+
+		/// <summary>
+		/// Scans a project directory, adding tasks for subdirectories
+		/// </summary>
+		/// <param name="ProjectDirectory">The project directory to search</param>
+		static void ScanProjectDirectory(DirectoryItem ProjectDirectory)
+		{
+			DirectoryItem ProjectPluginsDirectory = DirectoryItem.Combine(ProjectDirectory, "Plugins");
+			Enqueue(() => ScanPluginFolder(ProjectPluginsDirectory));
+
+			DirectoryItem ProjectSourceDirectory = DirectoryItem.Combine(ProjectDirectory, "Source");
+			Enqueue(() => ScanDirectoryTree(ProjectSourceDirectory));
+		}
+
+		/// <summary>
+		/// Scans a plugin parent directory, adding tasks for subdirectories
+		/// </summary>
+		/// <param name="Directory">The directory which may contain plugin directories</param>
+		static void ScanPluginFolder(DirectoryItem Directory)
+		{
+			foreach(DirectoryItem SubDirectory in Directory.EnumerateDirectories())
+			{
+				if(SubDirectory.EnumerateFiles().Any(x => x.HasExtension(".uplugin")))
+				{
+					Enqueue(() => ScanDirectoryTree(DirectoryItem.Combine(SubDirectory, "Source")));
+				}
+				else
+				{
+					Enqueue(() => ScanPluginFolder(SubDirectory));
+				}
+			}
+		}
+
+		/// <summary>
+		/// Scans an arbitrary directory tree
+		/// </summary>
+		/// <param name="Directory">Root of the directory tree</param>
+		static void ScanDirectoryTree(DirectoryItem Directory)
+		{
+			foreach(DirectoryItem SubDirectory in Directory.EnumerateDirectories())
+			{
+				Enqueue(() => ScanDirectoryTree(SubDirectory));
+			}
+			Directory.CacheFiles();
+		}
+	}
+}
diff --git a/Engine/Source/Programs/UnrealBuildTool/System/ToolMode.cs b/Engine/Source/Programs/UnrealBuildTool/System/ToolMode.cs
--- a/Engine/Source/Programs/UnrealBuildTool/System/ToolMode.cs
+++ b/Engine/Source/Programs/UnrealBuildTool/System/ToolMode.cs
@@ -20,25 +20,30 @@ namespace UnrealBuildTool
 		/// </summary>
 		None = 0,
 
+		/// <summary>
+		/// Start prefetching metadata for the engine folder as early as possible
+		/// </summary>
+		StartPrefetchingEngine = 1,
+
 		/// <summary>
 		/// Initializes the XmlConfig system
 		/// </summary>
-		XmlConfig = 1,
+		XmlConfig = 2,
 
 		/// <summary>
 		/// Registers build platforms
 		/// </summary>
-		BuildPlatforms = 2,
+		BuildPlatforms = 4,
 
 		/// <summary>
 		/// Registers build platforms for validation
 		/// </summary>
-		BuildPlatformsForValidation = 4,
+		BuildPlatformsForValidation = 8,
 
 		/// <summary>
 		/// Only allow a single instance running in the branch at once
 		/// </summary>
-		SingleInstance = 8,
+		SingleInstance = 16,
 	}
 
 	/// <summary>
diff --git a/Engine/Source/Programs/UnrealBuildTool/UnrealBuildTool.cs b/Engine/Source/Programs/UnrealBuildTool/UnrealBuildTool.cs
--- a/Engine/Source/Programs/UnrealBuildTool/UnrealBuildTool.cs
+++ b/Engine/Source/Programs/UnrealBuildTool/UnrealBuildTool.cs
@@ -282,7 +282,7 @@ namespace UnrealBuildTool
 			[CommandLine("-VSCode", Value="GenerateProjectFiles")]
 			[CommandLine("-VSMac", Value="GenerateProjectFiles")]
 			[CommandLine("-CLion", Value="GenerateProjectFiles")]
-			public string Mode = "Build";
+			public string Mode = null;
 
 			/// <summary>
 			/// Initialize the options with the given commnad line arguments
@@ -327,21 +327,6 @@ namespace UnrealBuildTool
 					Log.AddFileWriter("LogTraceListener", Options.LogFileName);
 				}
 
-				// Find all the valid modes
-				Dictionary<string, Type> ModeNameToType = new Dictionary<string, Type>(StringComparer.OrdinalIgnoreCase);
-				foreach(Type Type in Assembly.GetExecutingAssembly().GetTypes())
-				{
-					if(Type.IsClass && !Type.IsAbstract && Type.IsSubclassOf(typeof(ToolMode)))
-					{
-						ToolModeAttribute Attribute = Type.GetCustomAttribute<ToolModeAttribute>();
-						if(Attribute == null)
-						{
-							throw new BuildException("Class '{0}' should have a ToolModeAttribute", Type.Name);
-						}
-						ModeNameToType.Add(Attribute.Name, Type);
-					}
-				}
-
 				// Ensure we can resolve any external assemblies that are not in the same folder as our assembly.
 				AssemblyUtils.InstallAssemblyResolver(Path.GetDirectoryName(Assembly.GetEntryAssembly().GetOriginalLocation()));
 
@@ -349,17 +334,45 @@ namespace UnrealBuildTool
 				// This is critical to be done early so any code that relies on the current directory being Engine/Source will work.
 				DirectoryReference.SetCurrentDirectory(UnrealBuildTool.EngineSourceDirectory);
 
-				// Try to get the correct mode
-				Type ModeType;
-				if(!ModeNameToType.TryGetValue(Options.Mode, out ModeType))
+				// Get the type of the mode to execute, using a fast-path for the build mode.
+				Type ModeType = typeof(BuildMode);
+				if(Options.Mode != null)
 				{
-					Log.TraceError("No mode named '{0}'. Available modes are:\n  {1}", Options.Mode, String.Join("\n  ", ModeNameToType.Keys));
-					return 1;
+					// Find all the valid modes
+					Dictionary<string, Type> ModeNameToType = new Dictionary<string, Type>(StringComparer.OrdinalIgnoreCase);
+					foreach(Type Type in Assembly.GetExecutingAssembly().GetTypes())
+					{
+						if(Type.IsClass && !Type.IsAbstract && Type.IsSubclassOf(typeof(ToolMode)))
+						{
+							ToolModeAttribute Attribute = Type.GetCustomAttribute<ToolModeAttribute>();
+							if(Attribute == null)
+							{
+								throw new BuildException("Class '{0}' should have a ToolModeAttribute", Type.Name);
+							}
+							ModeNameToType.Add(Attribute.Name, Type);
+						}
+					}
+
+					// Try to get the correct mode
+					if(!ModeNameToType.TryGetValue(Options.Mode, out ModeType))
+					{
+						Log.TraceError("No mode named '{0}'. Available modes are:\n  {1}", Options.Mode, String.Join("\n  ", ModeNameToType.Keys));
+						return 1;
+					}
 				}
 
 				// Get the options for which systems have to be initialized for this mode
 				ToolModeOptions ModeOptions = ModeType.GetCustomAttribute<ToolModeAttribute>().Options;
 
+				// Start prefetching the contents of the engine folder
+				if((ModeOptions & ToolModeOptions.StartPrefetchingEngine) != 0)
+				{
+					using(Timeline.ScopeEvent("FileMetadataPrefetch.QueueEngineDirectory()"))
+					{
+						FileMetadataPrefetch.QueueEngineDirectory();
+					}
+				}
+
 				// Read the XML configuration files
 				if((ModeOptions & ToolModeOptions.XmlConfig) != 0)
 				{
@@ -427,6 +440,12 @@ namespace UnrealBuildTool
 					Mutex.Dispose();
 				}
 
+				// Cancel the prefetcher
+				using(Timeline.ScopeEvent("FileMetadataPrefetch.Stop()"))
+				{
+					FileMetadataPrefetch.Stop();
+				}
+
 				// Print out all the performance info
 				Timeline.Print(TimeSpan.FromMilliseconds(20.0), LogEventType.Log);
 
diff --git a/Engine/Source/Programs/UnrealBuildTool/UnrealBuildTool.csproj b/Engine/Source/Programs/UnrealBuildTool/UnrealBuildTool.csproj
--- a/Engine/Source/Programs/UnrealBuildTool/UnrealBuildTool.csproj
+++ b/Engine/Source/Programs/UnrealBuildTool/UnrealBuildTool.csproj
@@ -190,6 +190,7 @@
     <Compile Include="System\CppCompileEnvironment.cs" />
     <Compile Include="System\CppDependencyCache.cs" />
     <Compile Include="System\CppIncludeParser.cs" />
+    <Compile Include="System\FileMetadataPrefetch.cs" />
     <Compile Include="System\SourceFileMetadataCache.cs" />
     <Compile Include="System\DataDrivenPlatformInfo.cs" />
     <Compile Include="System\DirectoryItem.cs" />
