Missing copyright notice.

#rnx

diff --git a/Engine/Source/Programs/UnrealBuildTool/System/FileMetadataPrefetch.cs b/Engine/Source/Programs/UnrealBuildTool/System/FileMetadataPrefetch.cs
--- a/Engine/Source/Programs/UnrealBuildTool/System/FileMetadataPrefetch.cs
+++ b/Engine/Source/Programs/UnrealBuildTool/System/FileMetadataPrefetch.cs
@@ -1,3 +1,5 @@
+// Copyright 1998-2019 Epic Games, Inc. All Rights Reserved.
+
 using System;
 using System.Collections.Generic;
 using System.Linq;
