UBT: Remove code for querying UHT DLL versions. This pathway no longer works, because we don't compile the changelist into binaries. Instead, check the receipt is newer than any build products (which also works on all platforms).

diff --git a/Engine/Source/Programs/UnrealBuildTool/Configuration/TargetDescriptor.cs b/Engine/Source/Programs/UnrealBuildTool/Configuration/TargetDescriptor.cs
--- a/Engine/Source/Programs/UnrealBuildTool/Configuration/TargetDescriptor.cs
+++ b/Engine/Source/Programs/UnrealBuildTool/Configuration/TargetDescriptor.cs
@@ -68,26 +68,6 @@ namespace UnrealBuildTool
 			this.Configuration = Configuration;
 			this.Architecture = Architecture;
 
-			// Parse the list of hot-reload module names
-			foreach(string ModuleWithSuffix in Arguments.GetValues("-ModuleWithSuffix="))
-			{
-				int SuffixIdx = ModuleWithSuffix.LastIndexOf(',');
-				if(SuffixIdx == -1)
-				{
-					throw new BuildException("Missing suffix argument from -ModuleWithSuffix=Name,Suffix");
-				}
-
-				string ModuleName = ModuleWithSuffix.Substring(0, SuffixIdx);
-
-				int Suffix;
-				if(!Int32.TryParse(ModuleWithSuffix.Substring(SuffixIdx + 1), out Suffix))
-				{
-					throw new BuildException("Suffix for modules must be an integer");
-				}
-
-				HotReloadModuleNameToSuffix[ModuleName] = Suffix;
-			}
-
 			// If there are any additional command line arguments
 			List<string> AdditionalArguments = new List<string>();
 			if(Arguments != null)
@@ -95,6 +75,26 @@ namespace UnrealBuildTool
 				// Apply the arguments to this object
 				Arguments.ApplyTo(this);
 
+				// Parse all the hot-reload module names
+				foreach(string ModuleWithSuffix in Arguments.GetValues("-ModuleWithSuffix="))
+				{
+					int SuffixIdx = ModuleWithSuffix.LastIndexOf(',');
+					if(SuffixIdx == -1)
+					{
+						throw new BuildException("Missing suffix argument from -ModuleWithSuffix=Name,Suffix");
+					}
+
+					string ModuleName = ModuleWithSuffix.Substring(0, SuffixIdx);
+
+					int Suffix;
+					if(!Int32.TryParse(ModuleWithSuffix.Substring(SuffixIdx + 1), out Suffix))
+					{
+						throw new BuildException("Suffix for modules must be an integer");
+					}
+
+					HotReloadModuleNameToSuffix[ModuleName] = Suffix;
+				}
+
 				// Pull out all the arguments that haven't been used so far
 				for(int Idx = 0; Idx < Arguments.Count; Idx++)
 				{
diff --git a/Engine/Source/Programs/UnrealBuildTool/System/BuildHostPlatform.cs b/Engine/Source/Programs/UnrealBuildTool/System/BuildHostPlatform.cs
--- a/Engine/Source/Programs/UnrealBuildTool/System/BuildHostPlatform.cs
+++ b/Engine/Source/Programs/UnrealBuildTool/System/BuildHostPlatform.cs
@@ -69,13 +69,6 @@ namespace UnrealBuildTool
 		/// </summary>
 		abstract public UnrealTargetPlatform Platform { get; }
 
-		/// <summary>
-		/// Checks the API version of a dynamic library
-		/// </summary>
-		/// <param name="Filename">Filename of the library</param>
-		/// <returns>API version of -1 if not found.</returns>
-		abstract public int GetDllApiVersion(string Filename);
-
 		/// <summary>
 		/// Class that holds information about a running process
 		/// </summary>
@@ -227,91 +220,6 @@ namespace UnrealBuildTool
 			get { return UnrealTargetPlatform.Win64; }
 		}
 
-		[DllImport("kernel32.dll", SetLastError = true)]
-		static extern IntPtr LoadLibraryEx(string lpFileName, IntPtr hFile, uint dwFlags);
-		[DllImport("kernel32.dll", SetLastError = true)]
-		static extern IntPtr FindResource(IntPtr hModule, string lpName, string lpType);
-		[DllImport("kernel32.dll")]
-		static extern IntPtr FindResource(IntPtr hModule, IntPtr lpID, IntPtr lpType);
-		[DllImport("kernel32.dll", SetLastError = true)]
-		static extern IntPtr LoadResource(IntPtr hModule, IntPtr hResInfo);
-		[DllImport("kernel32.dll", SetLastError = true)]
-		static extern uint SizeofResource(IntPtr hModule, IntPtr hResInfo);
-		[DllImport("kernel32.dll", SetLastError = true)]
-		static extern IntPtr LockResource(IntPtr hResData);
-		[DllImport("kernel32.dll", SetLastError = true)]
-		static extern void UnlockResource(IntPtr hResInfo);
-		[DllImport("kernel32.dll", SetLastError = true)]
-		static extern void FreeLibrary(IntPtr hModule);
-
-		[Flags]
-		enum LoadLibraryFlags : uint
-		{
-			DONT_RESOLVE_DLL_REFERENCES = 0x00000001,
-			LOAD_LIBRARY_AS_DATAFILE = 0x00000002,
-			LOAD_WITH_ALTERED_SEARCH_PATH = 0x00000008,
-			LOAD_IGNORE_CODE_AUTHZ_LEVEL = 0x00000010,
-			LOAD_LIBRARY_AS_IMAGE_RESOURCE = 0x00000020,
-			LOAD_LIBRARY_AS_DATAFILE_EXCLUSIVE = 0x00000040,
-			LOAD_LIBRARY_REQUIRE_SIGNED_TARGET = 0x00000080
-		}
-
-		enum ResourceType
-		{
-			RT_CURSOR = 1,
-			RT_BITMAP = 2,
-			RT_ICON = 3,
-			RT_MENU = 4,
-			RT_DIALOG = 5,
-			RT_STRING = 6,
-			RT_FONTDIR = 7,
-			RT_FONT = 8,
-			RT_ACCELERATOR = 9,
-			RT_RCDATA = 10,
-			RT_MESSAGETABLE = 11
-		}
-
-		public override int GetDllApiVersion(string Filename)
-		{
-			int Result = -1;
-
-			try
-			{
-				const int ID_MODULE_API_VERSION_RESOURCE = 191;
-
-				// Retrieves the embedded API version from a DLL
-				IntPtr hModule = LoadLibraryEx(Filename, IntPtr.Zero, (uint)LoadLibraryFlags.LOAD_LIBRARY_AS_DATAFILE);
-				if (hModule != IntPtr.Zero)
-				{
-					IntPtr hResInfo = FindResource(hModule, new IntPtr(ID_MODULE_API_VERSION_RESOURCE), new IntPtr((int)ResourceType.RT_RCDATA));
-					if (hResInfo != IntPtr.Zero)
-					{
-						IntPtr hResGlobal = LoadResource(hModule, hResInfo);
-						if (hResGlobal != IntPtr.Zero)
-						{
-							IntPtr pResData = LockResource(hResGlobal);
-							if (pResData != IntPtr.Zero)
-							{
-								uint Length = SizeofResource(hModule, hResInfo);
-								if (Length > 0)
-								{
-									string Str = Marshal.PtrToStringAnsi(pResData);
-									Result = Int32.Parse(Str);
-								}
-							}
-						}
-					}
-					FreeLibrary(hModule);
-				}
-			}
-			catch (Exception Ex)
-			{
-				Log.TraceWarning("Failed to get DLL API version for {0}. Exception: {1}", Filename, Ex.Message);
-			}
-
-			return Result;
-		}
-
 		internal override IEnumerable<ProjectFileFormat> GetDefaultProjectFileFormats()
 		{
 			yield return ProjectFileFormat.VisualStudio;
@@ -325,12 +233,6 @@ namespace UnrealBuildTool
 			get { return UnrealTargetPlatform.Mac; }
 		}
 
-		public override int GetDllApiVersion(string Filename)
-		{
-			// @TODO: Implement GetDllApiVersion for Mac
-			return -1;
-		}
-
 		/// <summary>
 		/// Currently Mono returns incomplete process names in Process.GetProcesses() so we need to parse 'ps' output.
 		/// </summary>
@@ -445,12 +347,6 @@ namespace UnrealBuildTool
 			get { return UnrealTargetPlatform.Linux; }
 		}
 
-		public override int GetDllApiVersion(string Filename)
-		{
-			// @TODO: Implement GetDllApiVersion for Linux
-			return -1;
-		}
-
 		/// <summary>
 		/// Currently Mono returns incomplete process names in Process.GetProcesses() so we need to use /proc
 		/// (also, Mono locks up during process traversal sometimes, trying to open /dev/snd/pcm*)
diff --git a/Engine/Source/Programs/UnrealBuildTool/System/ExternalExecution.cs b/Engine/Source/Programs/UnrealBuildTool/System/ExternalExecution.cs
--- a/Engine/Source/Programs/UnrealBuildTool/System/ExternalExecution.cs
+++ b/Engine/Source/Programs/UnrealBuildTool/System/ExternalExecution.cs
@@ -710,15 +710,14 @@ namespace UnrealBuildTool
 		/// <summary>
 		/// Gets the latest write time of any of the UnrealHeaderTool binaries (including DLLs and Plugins) or DateTime.MaxValue if UnrealHeaderTool does not exist
 		/// </summary>
-		/// <returns>
-		/// Latest timestamp of UHT binaries or DateTime.MaxValue if UnrealHeaderTool is out of date and needs to be rebuilt.
-		/// </returns>
+		/// <returns>Latest timestamp of UHT binaries or DateTime.MaxValue if UnrealHeaderTool is out of date and needs to be rebuilt.</returns>
 		static bool GetHeaderToolTimestamp(FileReference ReceiptPath, out DateTime Timestamp)
 		{
 			using (ScopedTimer TimestampTimer = new ScopedTimer("GetHeaderToolTimestamp"))
 			{
 				// Try to read the receipt for UHT.
-				if (!FileReference.Exists(ReceiptPath))
+				FileItem ReceiptFile = FileItem.GetItemByFileReference(ReceiptPath);
+				if (!ReceiptFile.Exists)
 				{
 					Timestamp = DateTime.MaxValue;
 					return false;
@@ -731,95 +730,21 @@ namespace UnrealBuildTool
 					return false;
 				}
 
-				// Check all the binaries exist, and that all the DLLs are built against the right version
-				if (!CheckBinariesExist(Receipt) || !CheckDynamicLibaryVersionsMatch(Receipt))
+				// Make sure all the build products exist, and that the receipt is newer
+				foreach (BuildProduct BuildProduct in Receipt.BuildProducts)
 				{
-					Timestamp = DateTime.MaxValue;
-					return false;
-				}
-
-				// Return the timestamp for all the binaries
-				Timestamp = GetTimestampFromBinaries(Receipt);
-				return true;
-			}
-		}
-
-		/// <summary>
-		/// Checks if all the files in a receipt are present and that all the DLLs are at the same version
-		/// </summary>
-		/// <returns>
-		/// True if all the files are valid.
-		/// </returns>
-		static bool CheckBinariesExist(TargetReceipt Receipt)
-		{
-			bool bExist = true;
-			foreach (BuildProduct BuildProduct in Receipt.BuildProducts)
-			{
-				if (BuildProduct.Type == BuildProductType.Executable || BuildProduct.Type == BuildProductType.DynamicLibrary)
-				{
-					if (!FileReference.Exists(BuildProduct.Path))
+					FileItem BuildProductItem = FileItem.GetItemByFileReference(BuildProduct.Path);
+					if(!BuildProductItem.Exists || BuildProductItem.LastWriteTimeUtc > ReceiptFile.LastWriteTimeUtc)
 					{
-						Log.TraceWarning("Missing binary: {0}", BuildProduct.Path);
-						bExist = false;
+						Timestamp = DateTime.MaxValue;
+						return false;
 					}
 				}
-			}
-			return bExist;
-		}
-
-		/// <summary>
-		/// Checks if all the files in a receipt have the same version
-		/// </summary>
-		/// <returns>
-		/// True if all the files are valid.
-		/// </returns>
-		static bool CheckDynamicLibaryVersionsMatch(TargetReceipt Receipt)
-		{
-			List<Tuple<FileReference, int>> BinaryVersions = new List<Tuple<FileReference, int>>();
-			foreach (BuildProduct BuildProduct in Receipt.BuildProducts)
-			{
-				if (BuildProduct.Type == BuildProductType.DynamicLibrary)
-				{
-					int Version = BuildHostPlatform.Current.GetDllApiVersion(BuildProduct.Path.FullName);
-					BinaryVersions.Add(new Tuple<FileReference, int>(BuildProduct.Path, Version));
-				}
-			}
-
-			bool bMatch = true;
-			if (BinaryVersions.Count > 0 && !BinaryVersions.All(x => x.Item2 == BinaryVersions[0].Item2))
-			{
-				Log.TraceWarning("Detected mismatch in binary versions:");
-				foreach (Tuple<FileReference, int> BinaryVersion in BinaryVersions)
-				{
-					Log.TraceWarning("  {0} has API version {1}", BinaryVersion.Item1, BinaryVersion.Item2);
-					FileReference.Delete(BinaryVersion.Item1);
-				}
-				bMatch = false;
-			}
-			return bMatch;
-		}
 
-		/// <summary>
-		/// Checks if all the files in a receipt are present and that all the DLLs are at the same version
-		/// </summary>
-		/// <returns>
-		/// True if all the files are valid.
-		/// </returns>
-		static DateTime GetTimestampFromBinaries(TargetReceipt Receipt)
-		{
-			DateTime LatestWriteTime = DateTime.MinValue;
-			foreach (BuildProduct BuildProduct in Receipt.BuildProducts)
-			{
-				if (BuildProduct.Type == BuildProductType.Executable || BuildProduct.Type == BuildProductType.DynamicLibrary)
-				{
-					DateTime WriteTime = FileReference.GetLastWriteTime(BuildProduct.Path);
-					if (WriteTime > LatestWriteTime)
-					{
-						LatestWriteTime = WriteTime;
-					}
-				}
+				// Return the timestamp for all the binaries
+				Timestamp = ReceiptFile.LastWriteTimeUtc;
+				return true;
 			}
-			return LatestWriteTime;
 		}
 
 		/// <summary>
