UBT: Add a property to BuildHostPlatform for the type and path of the shell to use.

#rnx

diff --git a/Engine/Source/Programs/UnrealBuildTool/Configuration/UEBuildTarget.cs b/Engine/Source/Programs/UnrealBuildTool/Configuration/UEBuildTarget.cs
--- a/Engine/Source/Programs/UnrealBuildTool/Configuration/UEBuildTarget.cs
+++ b/Engine/Source/Programs/UnrealBuildTool/Configuration/UEBuildTarget.cs
@@ -1432,14 +1432,13 @@ namespace UnrealBuildTool
 					FileReference OutputFile = new FileReference(PostBuildScript.FullName + ".ran");
 
 					Action PostBuildStepAction = new Action(ActionType.PostBuildStep);
-					if(BuildHostPlatform.Current.Platform == UnrealTargetPlatform.Win64)
+					PostBuildStepAction.CommandPath = BuildHostPlatform.Current.Shell;
+					if(BuildHostPlatform.Current.ShellType == ShellType.Cmd)
 					{
-						PostBuildStepAction.CommandPath = "cmd.exe";
 						PostBuildStepAction.CommandArguments = String.Format("/C \"call \"{0}\" && type NUL >\"{1}\"\"", PostBuildScript, OutputFile);
 					}
 					else
 					{
-						PostBuildStepAction.CommandPath = "/bin/sh";
 						PostBuildStepAction.CommandArguments = String.Format("\"{0}\" && touch \"{1}\"", PostBuildScript, OutputFile);
 					}
 					PostBuildStepAction.WorkingDirectory = UnrealBuildTool.EngineSourceDirectory.FullName;
@@ -1521,14 +1520,13 @@ namespace UnrealBuildTool
 
 			Action CopyAction = new Action(ActionType.BuildProject);
 			CopyAction.CommandDescription = "Copy";
-			if(BuildHostPlatform.Current.Platform == UnrealTargetPlatform.Win64)
+			CopyAction.CommandPath = BuildHostPlatform.Current.Shell;
+			if(BuildHostPlatform.Current.ShellType == ShellType.Cmd)
 			{
-				CopyAction.CommandPath = "cmd.exe";
 				CopyAction.CommandArguments = String.Format("/C \"copy /Y \"{0}\" \"{1}\" 1>nul\"", SourceFile, TargetFile);
 			}
 			else
 			{
-				CopyAction.CommandPath = "/bin/sh";
 				CopyAction.CommandArguments = String.Format("-c 'cp -f {0} {1}'", Utils.EscapeShellArgument(SourceFile.FullName), Utils.EscapeShellArgument(TargetFile.FullName));
 			}
 			CopyAction.WorkingDirectory = Environment.CurrentDirectory;
diff --git a/Engine/Source/Programs/UnrealBuildTool/Platform/Android/AndroidToolChain.cs b/Engine/Source/Programs/UnrealBuildTool/Platform/Android/AndroidToolChain.cs
--- a/Engine/Source/Programs/UnrealBuildTool/Platform/Android/AndroidToolChain.cs
+++ b/Engine/Source/Programs/UnrealBuildTool/Platform/Android/AndroidToolChain.cs
@@ -1454,17 +1454,16 @@ namespace UnrealBuildTool
 						CompileAction.WorkingDirectory = UnrealBuildTool.EngineSourceDirectory.FullName;
 						if(bExecuteCompilerThroughShell)
 						{
-							if (BuildHostPlatform.Current.Platform == UnrealTargetPlatform.Win64)
+							CompileAction.CommandPath = BuildHostPlatform.Current.Shell;
+							if (BuildHostPlatform.Current.ShellType == ShellType.Sh)
 							{
-								CompileAction.CommandPath = "cmd.exe";
 								CompileAction.CommandArguments = String.Format("/c \"{0} {1}\"", ClangPath, ResponseArgument);
 							}
 							else
 							{
-								CompileAction.CommandPath = "/bin/sh";
 								CompileAction.CommandArguments = String.Format("-c \'{0} {1}\'", ClangPath, ResponseArgument);
-								CompileAction.CommandDescription = "Compile";
 							}
+							CompileAction.CommandDescription = "Compile";
 						}
 						else
 						{
@@ -1688,17 +1687,16 @@ namespace UnrealBuildTool
 
 					if(bExecuteCompilerThroughShell)
 					{
-						if (BuildHostPlatform.Current.Platform == UnrealTargetPlatform.Win64)
+						LinkAction.CommandPath = BuildHostPlatform.Current.Shell;
+						if (BuildHostPlatform.Current.ShellType == ShellType.Cmd)
 						{
 							LinkAction.CommandArguments = String.Format("/c \"{0} {1}\"", LinkAction.CommandPath, LinkAction.CommandArguments);
-							LinkAction.CommandPath = "cmd.exe";
 						}
 						else
 						{
 							LinkAction.CommandArguments = String.Format("-c \'{0} {1}\'", LinkAction.CommandPath, LinkAction.CommandArguments);
-							LinkAction.CommandPath = "/bin/sh";
-							LinkAction.CommandDescription = "Link";
 						}
+						LinkAction.CommandDescription = "Link";
 					}
 					Actions.Add(LinkAction);
 
diff --git a/Engine/Source/Programs/UnrealBuildTool/Platform/Android/UEDeployAndroid.cs b/Engine/Source/Programs/UnrealBuildTool/Platform/Android/UEDeployAndroid.cs
--- a/Engine/Source/Programs/UnrealBuildTool/Platform/Android/UEDeployAndroid.cs
+++ b/Engine/Source/Programs/UnrealBuildTool/Platform/Android/UEDeployAndroid.cs
@@ -3553,9 +3553,9 @@ namespace UnrealBuildTool
 
 					// Use ant to build the .apk file
 					string AntOptions = AntBuildType + " -Djava.source=1.7 -Djava.target=1.7";
-					string ShellExecutable = Utils.IsRunningOnMono ? "/bin/sh" : "cmd.exe";
-					string ShellParametersBegin = Utils.IsRunningOnMono ? "-c '" : "/c ";
-					string ShellParametersEnd = Utils.IsRunningOnMono ? "'" : "";
+					string ShellExecutable = BuildHostPlatform.Current.Shell;
+					string ShellParametersBegin = (BuildHostPlatform.Current.ShellType == ShellType.Sh) ? "-c '" : "/c ";
+					string ShellParametersEnd = (BuildHostPlatform.Current.ShellType == ShellType.Sh) ? "'" : "";
 					switch (AntVerbosity.ToLower())
 					{
 						default:
@@ -3907,9 +3907,9 @@ namespace UnrealBuildTool
 						Directory.CreateDirectory(Path.GetDirectoryName(DestApkName));
 
 						// Use gradle to build the .apk file
-						string ShellExecutable = Utils.IsRunningOnMono ? "/bin/sh" : "cmd.exe";
-						string ShellParametersBegin = Utils.IsRunningOnMono ? "-c '" : "/c ";
-						string ShellParametersEnd = Utils.IsRunningOnMono ? "'" : "";
+						string ShellExecutable = BuildHostPlatform.Current.Shell;
+						string ShellParametersBegin = (BuildHostPlatform.Current.ShellType == ShellType.Sh) ? "-c '" : "/c ";
+						string ShellParametersEnd = (BuildHostPlatform.Current.ShellType == ShellType.Sh) ? "'" : "";
 						RunCommandLineProgramWithExceptionAndFiltering(UE4BuildGradlePath, ShellExecutable, ShellParametersBegin + "\"" + GradleScriptPath + "\" " + GradleOptions + ShellParametersEnd, "Making .apk with Gradle...");
 
 						// For build machine run a clean afterward to clean up intermediate files (does not remove final APK)
diff --git a/Engine/Source/Programs/UnrealBuildTool/Platform/Linux/LinuxToolChain.cs b/Engine/Source/Programs/UnrealBuildTool/Platform/Linux/LinuxToolChain.cs
--- a/Engine/Source/Programs/UnrealBuildTool/Platform/Linux/LinuxToolChain.cs
+++ b/Engine/Source/Programs/UnrealBuildTool/Platform/Linux/LinuxToolChain.cs
@@ -1286,15 +1286,14 @@ namespace UnrealBuildTool
 			// Create an archive action
 			Action ArchiveAction = new Action(ActionType.Link);
 			ArchiveAction.WorkingDirectory = UnrealBuildTool.EngineSourceDirectory.FullName;
-			bool bUsingSh = BuildHostPlatform.Current.Platform != UnrealTargetPlatform.Win64 && BuildHostPlatform.Current.Platform != UnrealTargetPlatform.Win32;
-			if (bUsingSh)
+			ArchiveAction.CommandPath = BuildHostPlatform.Current.Shell;
+
+			if (BuildHostPlatform.Current.ShellType == ShellType.Sh)
 			{
-				ArchiveAction.CommandPath = "/bin/sh";
 				ArchiveAction.CommandArguments = "-c '";
 			}
 			else
 			{
-				ArchiveAction.CommandPath = "cmd.exe";
 				ArchiveAction.CommandArguments = "/c \"";
 			}
 
@@ -1336,7 +1335,7 @@ namespace UnrealBuildTool
 			ArchiveAction.CommandArguments += LinkEnvironment.AdditionalArguments;
 			ArchiveAction.CommandArguments = ArchiveAction.CommandArguments.Replace("\\", "/");
 
-			if (bUsingSh)
+			if (BuildHostPlatform.Current.ShellType == ShellType.Sh)
 			{
 				ArchiveAction.CommandArguments += "'";
 			}
@@ -1363,8 +1362,8 @@ namespace UnrealBuildTool
 
 				Log.TraceVerbose("Adding postlink step");
 
-				bool bUseCmdExe = BuildHostPlatform.Current.Platform == UnrealTargetPlatform.Win64 || BuildHostPlatform.Current.Platform == UnrealTargetPlatform.Win32;
-				string ShellBinary = bUseCmdExe ? "cmd.exe" : "/bin/sh";
+				bool bUseCmdExe = BuildHostPlatform.Current.ShellType == ShellType.Cmd;
+				string ShellBinary = BuildHostPlatform.Current.Shell;
 				string ExecuteSwitch = bUseCmdExe ? " /C" : ""; // avoid -c so scripts don't need +x
 				string ScriptName = bUseCmdExe ? "FixDependencies.bat" : "FixDependencies.sh";
 
diff --git a/Engine/Source/Programs/UnrealBuildTool/System/BuildHostPlatform.cs b/Engine/Source/Programs/UnrealBuildTool/System/BuildHostPlatform.cs
--- a/Engine/Source/Programs/UnrealBuildTool/System/BuildHostPlatform.cs
+++ b/Engine/Source/Programs/UnrealBuildTool/System/BuildHostPlatform.cs
@@ -11,6 +11,22 @@ using Tools.DotNETCommon;
 
 namespace UnrealBuildTool
 {
+	/// <summary>
+	/// The type of shell supported by this platform. Used to configure command line arguments.
+	/// </summary>
+	public enum ShellType
+	{
+		/// <summary>
+		/// The Bourne shell
+		/// </summary>
+		Sh,
+
+		/// <summary>
+		/// Windows command interpreter
+		/// </summary>
+		Cmd,
+	}
+
 	/// <summary>
 	/// Host platform abstraction
 	/// </summary>
@@ -69,6 +85,16 @@ namespace UnrealBuildTool
 		/// </summary>
 		abstract public UnrealTargetPlatform Platform { get; }
 
+		/// <summary>
+		/// Gets the path to the shell for this platform
+		/// </summary>
+		abstract public string Shell { get; }
+
+		/// <summary>
+		/// The type of shell returned by the Shell parameter
+		/// </summary>
+		abstract public ShellType ShellType { get; }
+
 		/// <summary>
 		/// Class that holds information about a running process
 		/// </summary>
@@ -220,6 +246,16 @@ namespace UnrealBuildTool
 			get { return UnrealTargetPlatform.Win64; }
 		}
 
+		public override string Shell
+		{
+			get { return Environment.GetEnvironmentVariable("COMSPEC"); }
+		}
+
+		public override ShellType ShellType
+		{
+			get { return ShellType.Cmd; }
+		}
+			
 		internal override IEnumerable<ProjectFileFormat> GetDefaultProjectFileFormats()
 		{
 			yield return ProjectFileFormat.VisualStudio;
@@ -233,6 +269,16 @@ namespace UnrealBuildTool
 			get { return UnrealTargetPlatform.Mac; }
 		}
 
+		public override string Shell
+		{
+			get { return "/bin/sh"; }
+		}
+
+		public override ShellType ShellType
+		{
+			get { return ShellType.Sh; }
+		}
+
 		/// <summary>
 		/// Currently Mono returns incomplete process names in Process.GetProcesses() so we need to parse 'ps' output.
 		/// </summary>
@@ -347,6 +393,16 @@ namespace UnrealBuildTool
 			get { return UnrealTargetPlatform.Linux; }
 		}
 
+		public override string Shell
+		{
+			get { return "/bin/sh"; }
+		}
+
+		public override ShellType ShellType
+		{
+			get { return ShellType.Sh; }
+		}
+
 		/// <summary>
 		/// Currently Mono returns incomplete process names in Process.GetProcesses() so we need to use /proc
 		/// (also, Mono locks up during process traversal sometimes, trying to open /dev/snd/pcm*)
diff --git a/Engine/Source/Programs/UnrealBuildTool/System/Utils.cs b/Engine/Source/Programs/UnrealBuildTool/System/Utils.cs
--- a/Engine/Source/Programs/UnrealBuildTool/System/Utils.cs
+++ b/Engine/Source/Programs/UnrealBuildTool/System/Utils.cs
@@ -1,4 +1,4 @@
-﻿// Copyright 1998-2019 Epic Games, Inc. All Rights Reserved.
+// Copyright 1998-2019 Epic Games, Inc. All Rights Reserved.
 
 using System;
 using System.Collections.Generic;
@@ -1057,14 +1057,14 @@ namespace UnrealBuildTool
 			foreach(FileReference ScriptFile in ScriptFiles)
 			{
 				ProcessStartInfo StartInfo = new ProcessStartInfo();
-				if(HostPlatform == UnrealTargetPlatform.Win64)
+				StartInfo.FileName = BuildHostPlatform.Current.Shell;
+
+				if(BuildHostPlatform.Current.ShellType == ShellType.Cmd)
 				{
-					StartInfo.FileName = "cmd.exe";
 					StartInfo.Arguments = String.Format("/C \"{0}\"", ScriptFile.FullName);
 				}
 				else
 				{
-					StartInfo.FileName = "/bin/sh";
 					StartInfo.Arguments = String.Format("\"{0}\"", ScriptFile.FullName);
 				}
 
diff --git a/Engine/Source/Programs/UnrealBuildTool/ToolChain/RemoteMac.cs b/Engine/Source/Programs/UnrealBuildTool/ToolChain/RemoteMac.cs
--- a/Engine/Source/Programs/UnrealBuildTool/ToolChain/RemoteMac.cs
+++ b/Engine/Source/Programs/UnrealBuildTool/ToolChain/RemoteMac.cs
@@ -209,7 +209,7 @@ namespace UnrealBuildTool
 				CommandLine.AppendFormat(" \"{0}\"", UnrealBuildTool.EngineDirectory);
 				CommandLine.Append("\"");
 
-				using(Process ChildProcess = Process.Start("C:\\Windows\\System32\\Cmd.exe", CommandLine.ToString()))
+				using(Process ChildProcess = Process.Start(BuildHostPlatform.Current.Shell, CommandLine.ToString()))
 				{
 					ChildProcess.WaitForExit();
 				}
