UBT: Remove support for executing batch files in actions. This is not supported by any executors other than the local executor and XGE, and is not used.

#rnx

diff --git a/Engine/Source/Programs/UnrealBuildTool/Executors/LocalExecutor.cs b/Engine/Source/Programs/UnrealBuildTool/Executors/LocalExecutor.cs
--- a/Engine/Source/Programs/UnrealBuildTool/Executors/LocalExecutor.cs
+++ b/Engine/Source/Programs/UnrealBuildTool/Executors/LocalExecutor.cs
@@ -89,31 +89,13 @@ namespace UnrealBuildTool
 			// thread start time
 			Action.StartTime = DateTimeOffset.Now;
 
-			// batch files (.bat, .cmd) need to be run via or cmd /c or shellexecute, 
-			// the latter which we can't use because we want to redirect input/output
-			bool bLaunchViaCmdExe = !Utils.IsRunningOnMono && !Path.GetExtension(Action.CommandPath).ToLower().EndsWith("exe");
-
 			// Create the action's process.
 			ProcessStartInfo ActionStartInfo = new ProcessStartInfo();
 			ActionStartInfo.WorkingDirectory = ExpandEnvironmentVariables(Action.WorkingDirectory);
 
-
 			string ExpandedCommandPath = ExpandEnvironmentVariables(Action.CommandPath);
-			if (bLaunchViaCmdExe)
-			{
-				ActionStartInfo.FileName = "cmd.exe";
-				ActionStartInfo.Arguments = string.Format
-				(
-					"/c \"{0} {1}\"",
-					ExpandEnvironmentVariables(ExpandedCommandPath),
-					ExpandEnvironmentVariables(Action.CommandArguments)
-				);
-			}
-			else
-			{
-				ActionStartInfo.FileName = ExpandedCommandPath;
-				ActionStartInfo.Arguments = ExpandEnvironmentVariables(Action.CommandArguments);
-			}
+			ActionStartInfo.FileName = ExpandedCommandPath;
+			ActionStartInfo.Arguments = ExpandEnvironmentVariables(Action.CommandArguments);
 
 			ActionStartInfo.UseShellExecute = false;
 			ActionStartInfo.RedirectStandardInput = false;
diff --git a/Engine/Source/Programs/UnrealBuildTool/Executors/XGE.cs b/Engine/Source/Programs/UnrealBuildTool/Executors/XGE.cs
--- a/Engine/Source/Programs/UnrealBuildTool/Executors/XGE.cs
+++ b/Engine/Source/Programs/UnrealBuildTool/Executors/XGE.cs
@@ -344,33 +344,8 @@ namespace UnrealBuildTool
 					ToolElement.SetAttribute("GroupPrefix", String.Format("** For {0}", String.Join(" + ", Action.GroupNames)));
 				}
 
-				// When running on Windows, differentiate between .exe and batch files.
-				// Those (.bat, .cmd) need to be run via cmd /c or shellexecute, 
-				// the latter which we can't use because we want to redirect input/output
-
-				bool bLaunchViaCmdExe = (BuildHostPlatform.Current.Platform == UnrealTargetPlatform.Win64) && (!Path.GetExtension(Action.CommandPath).ToLower().EndsWith(".exe"));
-
-				string CommandPath = "";
-				string CommandArguments = "";
-
-				if (bLaunchViaCmdExe)
-				{
-					CommandPath = "cmd.exe";
-					CommandArguments = string.Format
-					(
-						"/c \"\"{0}\" {1}\"",
-						(Action.CommandPath),
-						(Action.CommandArguments)
-					);
-				}
-				else
-				{
-					CommandPath = Action.CommandPath;
-					CommandArguments = Action.CommandArguments;
-				}
-
-				ToolElement.SetAttribute("Params", CommandArguments);
-				ToolElement.SetAttribute("Path", CommandPath);
+				ToolElement.SetAttribute("Params", Action.CommandArguments);
+				ToolElement.SetAttribute("Path", Action.CommandPath);
 				ToolElement.SetAttribute("SkipIfProjectFailed", "true");
 				if (Action.bIsGCCCompiler)
 				{
