UBT: Use FileReference and DirectoryReference instead of raw strings in actions.

#rnx

diff --git a/Engine/Source/Programs/AutomationTool/HTML5/HTML5Platform.Automation.cs b/Engine/Source/Programs/AutomationTool/HTML5/HTML5Platform.Automation.cs
--- a/Engine/Source/Programs/AutomationTool/HTML5/HTML5Platform.Automation.cs
+++ b/Engine/Source/Programs/AutomationTool/HTML5/HTML5Platform.Automation.cs
@@ -1,4 +1,4 @@
-﻿// Copyright 1998-2019 Epic Games, Inc. All Rights Reserved.
+// Copyright 1998-2019 Epic Games, Inc. All Rights Reserved.
 using System;
 using System.IO;
 using System.Net.Http;
@@ -98,7 +98,7 @@ public class HTML5Platform : Platform
 		else
 		{
 			// we need to operate in the root
-			string PythonPath = HTML5SDKInfo.Python();
+			string PythonPath = HTML5SDKInfo.Python().FullName;
 			string EmPackagerPath = HTML5SDKInfo.EmscriptenPackager();
 
 			using (new ScopedEnvVar("EM_CONFIG", HTML5SDKInfo.DOT_EMSCRIPTEN))
diff --git a/Engine/Source/Programs/AutomationTool/HTML5/HTML5Platform.PakFiles.Automation.cs b/Engine/Source/Programs/AutomationTool/HTML5/HTML5Platform.PakFiles.Automation.cs
--- a/Engine/Source/Programs/AutomationTool/HTML5/HTML5Platform.PakFiles.Automation.cs
+++ b/Engine/Source/Programs/AutomationTool/HTML5/HTML5Platform.PakFiles.Automation.cs
@@ -1,4 +1,4 @@
-﻿// Copyright 1998-2019 Epic Games, Inc. All Rights Reserved.
+// Copyright 1998-2019 Epic Games, Inc. All Rights Reserved.
 using System;
 using System.Collections.Generic;
 using System.Text.RegularExpressions;
@@ -68,7 +68,7 @@ public class HTMLPakAutomation
 	/// </summary>
 	public void CreateEmscriptenDataPackage(string PackagePath, string FinalDataLocation)
 	{
-		string PythonPath = HTML5SDKInfo.Python();
+		string PythonPath = HTML5SDKInfo.Python().FullName;
 		string PackagerPath = HTML5SDKInfo.EmscriptenPackager();
 
 		using (new ScopedEnvVar("EM_CONFIG", HTML5SDKInfo.DOT_EMSCRIPTEN))
diff --git a/Engine/Source/Programs/AutomationTool/Scripts/BuildPhysX.Automation.cs b/Engine/Source/Programs/AutomationTool/Scripts/BuildPhysX.Automation.cs
--- a/Engine/Source/Programs/AutomationTool/Scripts/BuildPhysX.Automation.cs
+++ b/Engine/Source/Programs/AutomationTool/Scripts/BuildPhysX.Automation.cs
@@ -861,7 +861,7 @@ class BuildPhysX : BuildCommand
 					Environment.GetEnvironmentVariable("EMSCRIPTEN") + ";" +
 					Environment.GetEnvironmentVariable("NODEPATH") + ";" +
 					Environment.GetEnvironmentVariable("LLVM") + ";" +
-					Path.GetDirectoryName(HTML5SDKInfo.Python()) + ";" +
+					Path.GetDirectoryName(HTML5SDKInfo.Python().FullName) + ";" +
 					Environment.GetEnvironmentVariable("PATH"));
 			//Log("set {0}={1}", "PATH", Environment.GetEnvironmentVariable("PATH"));
 		}
diff --git a/Engine/Source/Programs/UnrealBuildTool/Configuration/UEBuildTarget.cs b/Engine/Source/Programs/UnrealBuildTool/Configuration/UEBuildTarget.cs
--- a/Engine/Source/Programs/UnrealBuildTool/Configuration/UEBuildTarget.cs
+++ b/Engine/Source/Programs/UnrealBuildTool/Configuration/UEBuildTarget.cs
@@ -1415,7 +1415,7 @@ namespace UnrealBuildTool
 				}
 
 				Action WriteMetadataAction = Action.CreateRecursiveAction<WriteMetadataMode>(ActionType.WriteMetadata, WriteMetadataArguments.ToString());
-				WriteMetadataAction.WorkingDirectory = UnrealBuildTool.EngineSourceDirectory.FullName;
+				WriteMetadataAction.WorkingDirectory = UnrealBuildTool.EngineSourceDirectory;
 				WriteMetadataAction.StatusDescription = ReceiptFileName.GetFileName();
 				WriteMetadataAction.bCanExecuteRemotely = false;
 				WriteMetadataAction.PrerequisiteItems.AddRange(Makefile.OutputItems);
@@ -1441,7 +1441,7 @@ namespace UnrealBuildTool
 					{
 						PostBuildStepAction.CommandArguments = String.Format("\"{0}\" && touch \"{1}\"", PostBuildScript, OutputFile);
 					}
-					PostBuildStepAction.WorkingDirectory = UnrealBuildTool.EngineSourceDirectory.FullName;
+					PostBuildStepAction.WorkingDirectory = UnrealBuildTool.EngineSourceDirectory;
 					PostBuildStepAction.StatusDescription = String.Format("Executing post build script ({0})", PostBuildScript.GetFileName());
 					PostBuildStepAction.bCanExecuteRemotely = false;
 					PostBuildStepAction.PrerequisiteItems.Add(FileItem.GetItemByFileReference(ReceiptFileName));
@@ -1529,7 +1529,7 @@ namespace UnrealBuildTool
 			{
 				CopyAction.CommandArguments = String.Format("-c 'cp -f {0} {1}'", Utils.EscapeShellArgument(SourceFile.FullName), Utils.EscapeShellArgument(TargetFile.FullName));
 			}
-			CopyAction.WorkingDirectory = Environment.CurrentDirectory;
+			CopyAction.WorkingDirectory = UnrealBuildTool.EngineSourceDirectory;
 			CopyAction.PrerequisiteItems.Add(SourceFileItem);
 			CopyAction.ProducedItems.Add(TargetFileItem);
 			CopyAction.DeleteItems.Add(TargetFileItem);
diff --git a/Engine/Source/Programs/UnrealBuildTool/Executors/Distcc.cs b/Engine/Source/Programs/UnrealBuildTool/Executors/Distcc.cs
--- a/Engine/Source/Programs/UnrealBuildTool/Executors/Distcc.cs
+++ b/Engine/Source/Programs/UnrealBuildTool/Executors/Distcc.cs
@@ -205,8 +205,8 @@ namespace UnrealBuildTool
 					Environment.SetEnvironmentVariable("DISTCC_VERBOSE", "1");
 				}
 
-				string DistccExecutable = DistccExecutablesPath + "/distcc";
-				string GetHostExecutable = DistccExecutablesPath + "/gethost";
+				FileReference DistccExecutable = new FileReference(DistccExecutablesPath + "/distcc");
+				FileReference GetHostExecutable = new FileReference(DistccExecutablesPath + "/gethost");
 
 				Log.TraceInformation("Performing {0} actions ({1} in parallel)", Actions.Count, MaxActionsToExecuteInParallel, DistccExecutable, GetHostExecutable);
 
@@ -361,7 +361,7 @@ namespace UnrealBuildTool
 						"^{0}^{1:0.00}^{2}^{3}^{4}",
 						Action.ActionType.ToString(),
 						ThreadSeconds,
-						Path.GetFileName(Action.CommandPath),
+						Action.CommandPath.GetFileName(),
 						Action.StatusDescription,
 						Action.bIsUsingPCH);
 
diff --git a/Engine/Source/Programs/UnrealBuildTool/Executors/LocalExecutor.cs b/Engine/Source/Programs/UnrealBuildTool/Executors/LocalExecutor.cs
--- a/Engine/Source/Programs/UnrealBuildTool/Executors/LocalExecutor.cs
+++ b/Engine/Source/Programs/UnrealBuildTool/Executors/LocalExecutor.cs
@@ -91,9 +91,9 @@ namespace UnrealBuildTool
 
 			// Create the action's process.
 			ProcessStartInfo ActionStartInfo = new ProcessStartInfo();
-			ActionStartInfo.WorkingDirectory = ExpandEnvironmentVariables(Action.WorkingDirectory);
+			ActionStartInfo.WorkingDirectory = Action.WorkingDirectory.FullName;
 
-			string ExpandedCommandPath = ExpandEnvironmentVariables(Action.CommandPath);
+			string ExpandedCommandPath = ExpandEnvironmentVariables(Action.CommandPath.FullName);
 			ActionStartInfo.FileName = ExpandedCommandPath;
 			ActionStartInfo.Arguments = ExpandEnvironmentVariables(Action.CommandArguments);
 
@@ -458,7 +458,7 @@ namespace UnrealBuildTool
 					"^{0}^{1:0.00}^{2}^{3}^{4}",
 					Action.ActionType.ToString(),
 					ThreadSeconds,
-					Path.GetFileName(Action.CommandPath),
+					Action.CommandPath.GetFileName(),
 					  Action.StatusDescription,
 					Action.bIsUsingPCH);
 
diff --git a/Engine/Source/Programs/UnrealBuildTool/Executors/ParallelExecutor.cs b/Engine/Source/Programs/UnrealBuildTool/Executors/ParallelExecutor.cs
--- a/Engine/Source/Programs/UnrealBuildTool/Executors/ParallelExecutor.cs
+++ b/Engine/Source/Programs/UnrealBuildTool/Executors/ParallelExecutor.cs
@@ -249,7 +249,7 @@ namespace UnrealBuildTool
 				Action.LogLines.Add(Action.Inner.StatusDescription);
 			}
 
-			using(ManagedProcess Process = new ManagedProcess(ProcessGroup, Action.Inner.CommandPath, Action.Inner.CommandArguments, Action.Inner.WorkingDirectory, null, null, ProcessPriorityClass.BelowNormal))
+			using(ManagedProcess Process = new ManagedProcess(ProcessGroup, Action.Inner.CommandPath.FullName, Action.Inner.CommandArguments, Action.Inner.WorkingDirectory.FullName, null, null, ProcessPriorityClass.BelowNormal))
 			{
 				Action.LogLines.AddRange(Process.ReadAllLines());
 				Action.ExitCode = Process.ExitCode;
diff --git a/Engine/Source/Programs/UnrealBuildTool/Executors/SNDBS.cs b/Engine/Source/Programs/UnrealBuildTool/Executors/SNDBS.cs
--- a/Engine/Source/Programs/UnrealBuildTool/Executors/SNDBS.cs
+++ b/Engine/Source/Programs/UnrealBuildTool/Executors/SNDBS.cs
@@ -409,7 +409,7 @@ namespace UnrealBuildTool
 						"^{0}^{1:0.00}^{2}^{3}^{4}",
 						Action.ActionType.ToString(),
 						ThreadSeconds,
-						Path.GetFileName(Action.CommandPath),
+						Action.CommandPath.GetFileName(),
 						Action.StatusDescription,
 						Action.bIsUsingPCH);
 
diff --git a/Engine/Source/Programs/UnrealBuildTool/Executors/XGE.cs b/Engine/Source/Programs/UnrealBuildTool/Executors/XGE.cs
--- a/Engine/Source/Programs/UnrealBuildTool/Executors/XGE.cs
+++ b/Engine/Source/Programs/UnrealBuildTool/Executors/XGE.cs
@@ -345,7 +345,7 @@ namespace UnrealBuildTool
 				}
 
 				ToolElement.SetAttribute("Params", Action.CommandArguments);
-				ToolElement.SetAttribute("Path", Action.CommandPath);
+				ToolElement.SetAttribute("Path", Action.CommandPath.FullName);
 				ToolElement.SetAttribute("SkipIfProjectFailed", "true");
 				if (Action.bIsGCCCompiler)
 				{
@@ -391,7 +391,7 @@ namespace UnrealBuildTool
 				}
 				TaskElement.SetAttribute("Name", string.Format("Action{0}", ActionIndex));
 				TaskElement.SetAttribute("Tool", string.Format("Tool{0}", ActionIndex));
-				TaskElement.SetAttribute("WorkingDir", Action.WorkingDirectory);
+				TaskElement.SetAttribute("WorkingDir", Action.WorkingDirectory.FullName);
 				TaskElement.SetAttribute("SkipIfProjectFailed", "true");
 
 				// Create a semi-colon separated list of the other tasks this task depends on the results of.
diff --git a/Engine/Source/Programs/UnrealBuildTool/Platform/Android/AndroidToolChain.cs b/Engine/Source/Programs/UnrealBuildTool/Platform/Android/AndroidToolChain.cs
--- a/Engine/Source/Programs/UnrealBuildTool/Platform/Android/AndroidToolChain.cs
+++ b/Engine/Source/Programs/UnrealBuildTool/Platform/Android/AndroidToolChain.cs
@@ -1451,7 +1451,7 @@ namespace UnrealBuildTool
 						FileItem ResponseFileItem = FileItem.CreateIntermediateTextFile(ResponseFileName, new List<string> { AllArguments });
 						string ResponseArgument = string.Format("@\"{0}\"", ResponseFileName);
 
-						CompileAction.WorkingDirectory = UnrealBuildTool.EngineSourceDirectory.FullName;
+						CompileAction.WorkingDirectory = UnrealBuildTool.EngineSourceDirectory;
 						if(bExecuteCompilerThroughShell)
 						{
 							CompileAction.CommandPath = BuildHostPlatform.Current.Shell;
@@ -1467,7 +1467,7 @@ namespace UnrealBuildTool
 						}
 						else
 						{
-							CompileAction.CommandPath = ClangPath;
+							CompileAction.CommandPath = new FileReference(ClangPath);
 							CompileAction.CommandArguments = ResponseArgument;
 						}
 						CompileAction.PrerequisiteItems.Add(ResponseFileItem);
@@ -1547,27 +1547,27 @@ namespace UnrealBuildTool
 
 					// Create an action that invokes the linker.
 					Action LinkAction = new Action(ActionType.Link);
-					LinkAction.WorkingDirectory = UnrealBuildTool.EngineSourceDirectory.FullName;
+					LinkAction.WorkingDirectory = UnrealBuildTool.EngineSourceDirectory;
 
 					if (LinkEnvironment.bIsBuildingLibrary)
 					{
 						switch (Arch)
 						{
-							case "-armv7": LinkAction.CommandPath = ArPathArm; break;
-							case "-arm64": LinkAction.CommandPath = ArPathArm64; break;
-							case "-x86": LinkAction.CommandPath = ArPathx86; ; break;
-							case "-x64": LinkAction.CommandPath = ArPathx64; ; break;
-							default: LinkAction.CommandPath = ArPathArm; ; break;
+							case "-armv7": LinkAction.CommandPath = new FileReference(ArPathArm); break;
+							case "-arm64": LinkAction.CommandPath = new FileReference(ArPathArm64); break;
+							case "-x86": LinkAction.CommandPath = new FileReference(ArPathx86); ; break;
+							case "-x64": LinkAction.CommandPath = new FileReference(ArPathx64); ; break;
+							default: LinkAction.CommandPath = new FileReference(ArPathArm); ; break;
 						}
 					}
 					else
 					{
-						LinkAction.CommandPath = ClangPath;
+						LinkAction.CommandPath = new FileReference(ClangPath);
 					}
 
-					string LinkerPath = LinkAction.WorkingDirectory;
+					DirectoryReference LinkerPath = LinkAction.WorkingDirectory;
 
-					LinkAction.WorkingDirectory = LinkEnvironment.IntermediateDirectory.FullName;
+					LinkAction.WorkingDirectory = LinkEnvironment.IntermediateDirectory;
 
 					// Get link arguments.
 					LinkAction.CommandArguments = LinkEnvironment.bIsBuildingLibrary ? GetArArguments(LinkEnvironment) : GetLinkArguments(LinkEnvironment, Arch);
@@ -1628,7 +1628,7 @@ namespace UnrealBuildTool
 								// environment variables aren't expanded when using the $( style
 								if (Path.IsPathRooted(AbsoluteLibraryPath) == false)
 								{
-									AbsoluteLibraryPath = Path.Combine(LinkerPath, AbsoluteLibraryPath);
+									AbsoluteLibraryPath = Path.Combine(LinkerPath.FullName, AbsoluteLibraryPath);
 								}
 								LinkResponseArguments += string.Format(" -L\"{0}\"", Utils.CollapseRelativeDirectories(AbsoluteLibraryPath));
 							}
diff --git a/Engine/Source/Programs/UnrealBuildTool/Platform/Android/UEDeployAndroid.cs b/Engine/Source/Programs/UnrealBuildTool/Platform/Android/UEDeployAndroid.cs
--- a/Engine/Source/Programs/UnrealBuildTool/Platform/Android/UEDeployAndroid.cs
+++ b/Engine/Source/Programs/UnrealBuildTool/Platform/Android/UEDeployAndroid.cs
@@ -3553,7 +3553,7 @@ namespace UnrealBuildTool
 
 					// Use ant to build the .apk file
 					string AntOptions = AntBuildType + " -Djava.source=1.7 -Djava.target=1.7";
-					string ShellExecutable = BuildHostPlatform.Current.Shell;
+					string ShellExecutable = BuildHostPlatform.Current.Shell.FullName;
 					string ShellParametersBegin = (BuildHostPlatform.Current.ShellType == ShellType.Sh) ? "-c '" : "/c ";
 					string ShellParametersEnd = (BuildHostPlatform.Current.ShellType == ShellType.Sh) ? "'" : "";
 					switch (AntVerbosity.ToLower())
@@ -3907,7 +3907,7 @@ namespace UnrealBuildTool
 						Directory.CreateDirectory(Path.GetDirectoryName(DestApkName));
 
 						// Use gradle to build the .apk file
-						string ShellExecutable = BuildHostPlatform.Current.Shell;
+						string ShellExecutable = BuildHostPlatform.Current.Shell.FullName;
 						string ShellParametersBegin = (BuildHostPlatform.Current.ShellType == ShellType.Sh) ? "-c '" : "/c ";
 						string ShellParametersEnd = (BuildHostPlatform.Current.ShellType == ShellType.Sh) ? "'" : "";
 						RunCommandLineProgramWithExceptionAndFiltering(UE4BuildGradlePath, ShellExecutable, ShellParametersBegin + "\"" + GradleScriptPath + "\" " + GradleOptions + ShellParametersEnd, "Making .apk with Gradle...");
diff --git a/Engine/Source/Programs/UnrealBuildTool/Platform/HTML5/HTML5SDKInfo.cs b/Engine/Source/Programs/UnrealBuildTool/Platform/HTML5/HTML5SDKInfo.cs
--- a/Engine/Source/Programs/UnrealBuildTool/Platform/HTML5/HTML5SDKInfo.cs
+++ b/Engine/Source/Programs/UnrealBuildTool/Platform/HTML5/HTML5SDKInfo.cs
@@ -1,4 +1,4 @@
-﻿// Copyright 1998-2019 Epic Games, Inc. All Rights Reserved.
+// Copyright 1998-2019 Epic Games, Inc. All Rights Reserved.
 using System;
 using System.IO;
 using System.Text;
@@ -424,9 +424,9 @@ namespace UnrealBuildTool
 		///
 		/// </summary>
 		/// <returns></returns>
-		public static string Python()
+		public static FileReference Python()
 		{
-			return PYTHON;
+			return new FileReference(PYTHON);
 		}
 
 		/// <summary>
diff --git a/Engine/Source/Programs/UnrealBuildTool/Platform/HTML5/HTML5ToolChain.cs b/Engine/Source/Programs/UnrealBuildTool/Platform/HTML5/HTML5ToolChain.cs
--- a/Engine/Source/Programs/UnrealBuildTool/Platform/HTML5/HTML5ToolChain.cs
+++ b/Engine/Source/Programs/UnrealBuildTool/Platform/HTML5/HTML5ToolChain.cs
@@ -411,7 +411,7 @@ namespace UnrealBuildTool
 					CompileAction.ProducedItems.Add(DependencyListFile);
 				}
 
-				CompileAction.WorkingDirectory = UnrealBuildTool.EngineSourceDirectory.FullName;
+				CompileAction.WorkingDirectory = UnrealBuildTool.EngineSourceDirectory;
 				CompileAction.CommandPath = HTML5SDKInfo.Python();
 
 				CompileAction.CommandArguments = HTML5SDKInfo.EmscriptenCompiler() + " " + Arguments + FileArguments + CompileEnvironment.AdditionalArguments;
@@ -459,7 +459,7 @@ namespace UnrealBuildTool
 			List<string> ReponseLines = new List<string>();
 
 			LinkAction.bCanExecuteRemotely = false;
-			LinkAction.WorkingDirectory = UnrealBuildTool.EngineSourceDirectory.FullName;
+			LinkAction.WorkingDirectory = UnrealBuildTool.EngineSourceDirectory;
 			LinkAction.CommandPath = HTML5SDKInfo.Python();
 			LinkAction.CommandArguments = HTML5SDKInfo.EmscriptenCompiler();
 //			bool bIsBuildingLibrary = LinkEnvironment.bIsBuildingLibrary || bBuildImportLibraryOnly;
diff --git a/Engine/Source/Programs/UnrealBuildTool/Platform/IOS/IOSToolChain.cs b/Engine/Source/Programs/UnrealBuildTool/Platform/IOS/IOSToolChain.cs
--- a/Engine/Source/Programs/UnrealBuildTool/Platform/IOS/IOSToolChain.cs
+++ b/Engine/Source/Programs/UnrealBuildTool/Platform/IOS/IOSToolChain.cs
@@ -697,7 +697,7 @@ namespace UnrealBuildTool
 
 				// RPC utility parameters are in terms of the Mac side
 				CompileAction.WorkingDirectory = GetMacDevSrcRoot();
-				CompileAction.CommandPath = CompilerPath;
+				CompileAction.CommandPath = new FileReference(CompilerPath);
 				CompileAction.CommandArguments = AllArgs; // Arguments + FileArguments + CompileEnvironment.AdditionalArguments;
 				CompileAction.CommandDescription = "Compile";
 				CompileAction.StatusDescription = string.Format("{0}", Path.GetFileName(SourceFile.AbsolutePath));
@@ -824,7 +824,7 @@ namespace UnrealBuildTool
 
 			LinkAction.StatusDescription = string.Format("{0}", OutputFile.AbsolutePath);
 
-			LinkAction.CommandPath = "sh";
+			LinkAction.CommandPath = BuildHostPlatform.Current.Shell;
 			if(LinkEnvironment.Configuration == CppConfiguration.Shipping && Path.GetExtension(OutputFile.AbsolutePath) != ".a")
 			{
 				// When building a shipping package, symbols are stripped from the exe as the last build step. This is a problem
@@ -928,7 +928,7 @@ namespace UnrealBuildTool
             Action GenDebugAction = new Action(ActionType.GenerateDebugInfo);
 
 			GenDebugAction.WorkingDirectory = GetMacDevSrcRoot();
-			GenDebugAction.CommandPath = "sh";
+			GenDebugAction.CommandPath = BuildHostPlatform.Current.Shell;
 			string DsymutilPath = GetDsymutilPath();
 			if(ProjectSettings.bGeneratedSYMBundle)
 			{
@@ -976,10 +976,10 @@ namespace UnrealBuildTool
 
             // Make the compile action
             Action GenDebugAction = new Action(ActionType.GenerateDebugInfo);
-            GenDebugAction.WorkingDirectory = GetMacDevEngineRoot() + "/Binaries/Mac/";
+            GenDebugAction.WorkingDirectory = DirectoryReference.Combine(UnrealBuildTool.EngineDirectory, "Binaries", "Mac");
 
-            GenDebugAction.CommandPath = "sh";
-            GenDebugAction.CommandArguments = string.Format("-c 'rm -rf \"{1}\"; dwarfdump --uuid \"{3}\" | cut -d\" \" -f2; chmod 777 ./DsymExporter; ./DsymExporter -UUID=$(dwarfdump --uuid \"{3}\" | cut -d\" \" -f2) \"{0}\" \"{2}\"'",
+            GenDebugAction.CommandPath = BuildHostPlatform.Current.Shell;
+			GenDebugAction.CommandArguments = string.Format("-c 'rm -rf \"{1}\"; dwarfdump --uuid \"{3}\" | cut -d\" \" -f2; chmod 777 ./DsymExporter; ./DsymExporter -UUID=$(dwarfdump --uuid \"{3}\" | cut -d\" \" -f2) \"{0}\" \"{2}\"'",
                     DWARFFile.AbsolutePath,
                     OutputFile.AbsolutePath,
                     Path.GetDirectoryName(OutputFile.AbsolutePath),
@@ -1129,9 +1129,9 @@ namespace UnrealBuildTool
 				FileItem ExtractScriptFileItem = FileItem.CreateIntermediateTextFile(new FileReference(Framework.OutputDirectory.FullName + ".sh"), ExtractScript.ToString());
 
 				Action UnzipAction = new Action(ActionType.BuildProject);
-				UnzipAction.CommandPath = "/bin/sh";
+				UnzipAction.CommandPath = new FileReference("/bin/sh");
 				UnzipAction.CommandArguments = Utils.EscapeShellArgument(ExtractScriptFileItem.AbsolutePath);
-				UnzipAction.WorkingDirectory = UnrealBuildTool.EngineDirectory.FullName;
+				UnzipAction.WorkingDirectory = UnrealBuildTool.EngineDirectory;
 				UnzipAction.PrerequisiteItems.Add(InputFile);
 				UnzipAction.PrerequisiteItems.Add(ExtractScriptFileItem);
 				UnzipAction.ProducedItems.Add(Framework.ExtractedTokenFile);
@@ -1314,7 +1314,7 @@ namespace UnrealBuildTool
 
 				Action StripAction = new Action(ActionType.CreateAppBundle);
 				StripAction.WorkingDirectory = GetMacDevSrcRoot();
-				StripAction.CommandPath = "sh";
+				StripAction.CommandPath = BuildHostPlatform.Current.Shell;
 				StripAction.CommandArguments = String.Format("-c '\"{0}strip\" \"{1}\" && touch \"{2}\"'", Settings.Value.ToolchainDir, Executable.Location, StripCompleteFile);
 				StripAction.PrerequisiteItems.Add(Executable);
 				StripAction.PrerequisiteItems.AddRange(OutputFiles);
@@ -1338,7 +1338,7 @@ namespace UnrealBuildTool
 				// Make the compile action
 				Action CompileAssetAction = new Action(ActionType.CreateAppBundle);
 				CompileAssetAction.WorkingDirectory = GetMacDevSrcRoot();
-				CompileAssetAction.CommandPath = "/usr/bin/xcrun";
+				CompileAssetAction.CommandPath = new FileReference("/usr/bin/xcrun");
 				CompileAssetAction.CommandArguments = GetAssetCatalogArgs(CppPlatform, ResourcesDir.FullName, Path.GetDirectoryName(AssetCatalogFile.AbsolutePath));
 				CompileAssetAction.PrerequisiteItems.Add(Executable);
 				CompileAssetAction.ProducedItems.Add(AssetCatalogFile);
@@ -1378,7 +1378,7 @@ namespace UnrealBuildTool
 				string PostBuildSyncArguments = String.Format("-Input=\"{0}\" -XmlConfigCache=\"{1}\"", PostBuildSyncFile, XmlConfig.CacheFile);
 
 				Action PostBuildSyncAction = Action.CreateRecursiveAction<IOSPostBuildSyncMode>(ActionType.CreateAppBundle, PostBuildSyncArguments);
-				PostBuildSyncAction.WorkingDirectory = UnrealBuildTool.EngineSourceDirectory.FullName;
+				PostBuildSyncAction.WorkingDirectory = UnrealBuildTool.EngineSourceDirectory;
 				PostBuildSyncAction.PrerequisiteItems.Add(Executable);
 				PostBuildSyncAction.PrerequisiteItems.AddRange(OutputFiles);
 				PostBuildSyncAction.ProducedItems.Add(FileItem.GetItemByFileReference(GetStagedExecutablePath(Executable.Location, Target.Name)));
diff --git a/Engine/Source/Programs/UnrealBuildTool/Platform/Linux/LinuxToolChain.cs b/Engine/Source/Programs/UnrealBuildTool/Platform/Linux/LinuxToolChain.cs
--- a/Engine/Source/Programs/UnrealBuildTool/Platform/Linux/LinuxToolChain.cs
+++ b/Engine/Source/Programs/UnrealBuildTool/Platform/Linux/LinuxToolChain.cs
@@ -1238,14 +1238,14 @@ namespace UnrealBuildTool
 					CompileAction.ProducedItems.Add(DependencyListFile);
 				}
 
-				CompileAction.WorkingDirectory = UnrealBuildTool.EngineSourceDirectory.FullName;
+				CompileAction.WorkingDirectory = UnrealBuildTool.EngineSourceDirectory;
 				if (!UsingClang())
 				{
-					CompileAction.CommandPath = GCCPath;
+					CompileAction.CommandPath = new FileReference(GCCPath);
 				}
 				else
 				{
-					CompileAction.CommandPath = ClangPath;
+					CompileAction.CommandPath = new FileReference(ClangPath);
 				}
 
 				string AllArguments = (Arguments + FileArguments + CompileEnvironment.AdditionalArguments);
@@ -1285,7 +1285,7 @@ namespace UnrealBuildTool
 		{
 			// Create an archive action
 			Action ArchiveAction = new Action(ActionType.Link);
-			ArchiveAction.WorkingDirectory = UnrealBuildTool.EngineSourceDirectory.FullName;
+			ArchiveAction.WorkingDirectory = UnrealBuildTool.EngineSourceDirectory;
 			ArchiveAction.CommandPath = BuildHostPlatform.Current.Shell;
 
 			if (BuildHostPlatform.Current.ShellType == ShellType.Sh)
@@ -1363,14 +1363,14 @@ namespace UnrealBuildTool
 				Log.TraceVerbose("Adding postlink step");
 
 				bool bUseCmdExe = BuildHostPlatform.Current.ShellType == ShellType.Cmd;
-				string ShellBinary = BuildHostPlatform.Current.Shell;
+				FileReference ShellBinary = BuildHostPlatform.Current.Shell;
 				string ExecuteSwitch = bUseCmdExe ? " /C" : ""; // avoid -c so scripts don't need +x
 				string ScriptName = bUseCmdExe ? "FixDependencies.bat" : "FixDependencies.sh";
 
 				FileItem FixDepsScript = FileItem.GetItemByFileReference(FileReference.Combine(LinkEnvironment.LocalShadowDirectory, ScriptName));
 
 				Action PostLinkAction = new Action(ActionType.Link);
-				PostLinkAction.WorkingDirectory = UnrealBuildTool.EngineSourceDirectory.FullName;
+				PostLinkAction.WorkingDirectory = UnrealBuildTool.EngineSourceDirectory;
 				PostLinkAction.CommandPath = ShellBinary;
 				PostLinkAction.StatusDescription = string.Format("{0}", Path.GetFileName(Executable.AbsolutePath));
 				PostLinkAction.CommandDescription = "FixDeps";
@@ -1434,7 +1434,7 @@ namespace UnrealBuildTool
 
 			// Create an action that invokes the linker.
 			Action LinkAction = new Action(ActionType.Link);
-			LinkAction.WorkingDirectory = UnrealBuildTool.EngineSourceDirectory.FullName;
+			LinkAction.WorkingDirectory = UnrealBuildTool.EngineSourceDirectory;
 
 			string LinkCommandString;
 			if (String.IsNullOrEmpty(ClangPath))
@@ -1654,8 +1654,8 @@ namespace UnrealBuildTool
 			LinkCommandString = LinkCommandString.Replace("\\\\", "/");
 			LinkCommandString = LinkCommandString.Replace("\\", "/");
 
-			bool bUseCmdExe = BuildHostPlatform.Current.Platform == UnrealTargetPlatform.Win64 || BuildHostPlatform.Current.Platform == UnrealTargetPlatform.Win32;
-			string ShellBinary = bUseCmdExe ? "cmd.exe" : "/bin/sh";
+			bool bUseCmdExe = BuildHostPlatform.Current.ShellType == ShellType.Cmd;
+			FileReference ShellBinary = BuildHostPlatform.Current.Shell;
 			string ExecuteSwitch = bUseCmdExe ? " /C" : ""; // avoid -c so scripts don't need +x
 
 			// Linux has issues with scripts and parameter expansion from curely brakets
diff --git a/Engine/Source/Programs/UnrealBuildTool/Platform/Mac/MacToolChain.cs b/Engine/Source/Programs/UnrealBuildTool/Platform/Mac/MacToolChain.cs
--- a/Engine/Source/Programs/UnrealBuildTool/Platform/Mac/MacToolChain.cs
+++ b/Engine/Source/Programs/UnrealBuildTool/Platform/Mac/MacToolChain.cs
@@ -520,7 +520,7 @@ namespace UnrealBuildTool
 				}
 
 				CompileAction.WorkingDirectory = GetMacDevSrcRoot();
-				CompileAction.CommandPath = CompilerPath;
+				CompileAction.CommandPath = new FileReference(CompilerPath);
 				CompileAction.CommandArguments = AllArgs;
 				CompileAction.CommandDescription = "Compile";
 				CompileAction.StatusDescription = Path.GetFileName(SourceFile.AbsolutePath);
@@ -640,7 +640,7 @@ namespace UnrealBuildTool
 			Action LinkAction = new Action(ActionType.Link);
 
 			LinkAction.WorkingDirectory = GetMacDevSrcRoot();
-			LinkAction.CommandPath = "/bin/sh";
+			LinkAction.CommandPath = BuildHostPlatform.Current.Shell;
 			LinkAction.CommandDescription = "Link";
 
 			string EngineAPIVersion = LoadEngineAPIVersion();
@@ -1041,8 +1041,8 @@ namespace UnrealBuildTool
 		FileItem FixDylibDependencies(LinkEnvironment LinkEnvironment, FileItem Executable, List<Action> Actions)
 		{
 			Action LinkAction = new Action(ActionType.Link);
-			LinkAction.WorkingDirectory = UnrealBuildTool.EngineSourceDirectory.FullName;
-			LinkAction.CommandPath = "/bin/sh";
+			LinkAction.WorkingDirectory = UnrealBuildTool.EngineSourceDirectory;
+			LinkAction.CommandPath = BuildHostPlatform.Current.Shell;
 			LinkAction.CommandDescription = "";
 
 			// Call the FixDylibDependencies.sh script which will link the dylibs and the main executable, this time proper ones, as it's called
@@ -1108,7 +1108,7 @@ namespace UnrealBuildTool
 			// Make the compile action
 			Action GenDebugAction = new Action(ActionType.GenerateDebugInfo);
 			GenDebugAction.WorkingDirectory = GetMacDevSrcRoot();
-			GenDebugAction.CommandPath = "sh";
+			GenDebugAction.CommandPath = BuildHostPlatform.Current.Shell;
 
 			// Deletes ay existing file on the building machine. Also, waits 30 seconds, if needed, for the input file to be created in an attempt to work around
 			// a problem where dsymutil would exit with an error saying the input file did not exist.
@@ -1147,7 +1147,7 @@ namespace UnrealBuildTool
 			// Make the compile action
 			Action FinalizeAppBundleAction = new Action(ActionType.CreateAppBundle);
 			FinalizeAppBundleAction.WorkingDirectory = GetMacDevSrcRoot(); // Path.GetFullPath(".");
-			FinalizeAppBundleAction.CommandPath = "/bin/sh";
+			FinalizeAppBundleAction.CommandPath = BuildHostPlatform.Current.Shell;
 			FinalizeAppBundleAction.CommandDescription = "";
 
 			// make path to the script
@@ -1167,7 +1167,7 @@ namespace UnrealBuildTool
 		{
 			Action CopyAction = new Action(ActionType.CreateAppBundle);
 			CopyAction.WorkingDirectory = GetMacDevSrcRoot(); // Path.GetFullPath(".");
-			CopyAction.CommandPath = "/bin/sh";
+			CopyAction.CommandPath = BuildHostPlatform.Current.Shell;
 			CopyAction.CommandDescription = "";
 
 			string BundlePath = BundleDirectory.FullName;
diff --git a/Engine/Source/Programs/UnrealBuildTool/Platform/Windows/PVSToolChain.cs b/Engine/Source/Programs/UnrealBuildTool/Platform/Windows/PVSToolChain.cs
--- a/Engine/Source/Programs/UnrealBuildTool/Platform/Windows/PVSToolChain.cs
+++ b/Engine/Source/Programs/UnrealBuildTool/Platform/Windows/PVSToolChain.cs
@@ -245,8 +245,8 @@ namespace UnrealBuildTool
 				FileItem PreprocessedFileItem = FileItem.GetItemByFileReference(PreprocessedFileLocation);
 
 				Action PreprocessAction = new Action(ActionType.Compile);
-				PreprocessAction.CommandPath = EnvVars.CompilerPath.FullName;
-				PreprocessAction.WorkingDirectory = UnrealBuildTool.EngineSourceDirectory.FullName;
+				PreprocessAction.CommandPath = EnvVars.CompilerPath;
+				PreprocessAction.WorkingDirectory = UnrealBuildTool.EngineSourceDirectory;
 				PreprocessAction.CommandArguments = " @\"" + ResponseFileItem.AbsolutePath + "\"";
 				PreprocessAction.PrerequisiteItems.AddRange(CompileEnvironment.ForceIncludeFiles);
 				PreprocessAction.PrerequisiteItems.Add(SourceFile);
@@ -305,8 +305,8 @@ namespace UnrealBuildTool
 				Action AnalyzeAction = new Action(ActionType.Compile);
 				AnalyzeAction.CommandDescription = "Analyzing";
 				AnalyzeAction.StatusDescription = BaseFileName;
-				AnalyzeAction.WorkingDirectory = UnrealBuildTool.EngineSourceDirectory.FullName;
-				AnalyzeAction.CommandPath = AnalyzerFile.FullName;
+				AnalyzeAction.WorkingDirectory = UnrealBuildTool.EngineSourceDirectory;
+				AnalyzeAction.CommandPath = AnalyzerFile;
 				AnalyzeAction.CommandArguments = String.Format("--cl-params \"{0}\" --source-file \"{1}\" --output-file \"{2}\" --cfg \"{3}\" --analysis-mode 4", PreprocessAction.CommandArguments, SourceFile.AbsolutePath, OutputFileLocation, ConfigFileItem.AbsolutePath);
 				if (LicenseFile != null)
 				{
@@ -348,7 +348,7 @@ namespace UnrealBuildTool
 			Action AnalyzeAction = new Action(ActionType.Compile);
 			AnalyzeAction.CommandPath = UnrealBuildTool.GetUBTPath();
 			AnalyzeAction.CommandArguments = String.Format("-Mode=PVSGather -Input=\"{0}\" -Output=\"{1}\"", InputFileListItem.Location, OutputFile);
-			AnalyzeAction.WorkingDirectory = UnrealBuildTool.EngineSourceDirectory.FullName;
+			AnalyzeAction.WorkingDirectory = UnrealBuildTool.EngineSourceDirectory;
 			AnalyzeAction.PrerequisiteItems.Add(InputFileListItem);
 			AnalyzeAction.PrerequisiteItems.AddRange(Makefile.OutputItems);
 			AnalyzeAction.ProducedItems.Add(FileItem.GetItemByFileReference(OutputFile));
diff --git a/Engine/Source/Programs/UnrealBuildTool/Platform/Windows/VCToolChain.cs b/Engine/Source/Programs/UnrealBuildTool/Platform/Windows/VCToolChain.cs
--- a/Engine/Source/Programs/UnrealBuildTool/Platform/Windows/VCToolChain.cs
+++ b/Engine/Source/Programs/UnrealBuildTool/Platform/Windows/VCToolChain.cs
@@ -1051,8 +1051,8 @@ namespace UnrealBuildTool
 					AppendCLArguments_CPP(CompileEnvironment, FileArguments);
 				}
 
-				CompileAction.WorkingDirectory = UnrealBuildTool.EngineSourceDirectory.FullName;
-				CompileAction.CommandPath = EnvVars.CompilerPath.FullName;
+				CompileAction.WorkingDirectory = UnrealBuildTool.EngineSourceDirectory;
+				CompileAction.CommandPath = EnvVars.CompilerPath;
 				CompileAction.PrerequisiteItems.AddRange(CompileEnvironment.ForceIncludeFiles);
 
 				string[] AdditionalArguments = String.IsNullOrEmpty(CompileEnvironment.AdditionalArguments)? new string[0] : new string[] { CompileEnvironment.AdditionalArguments };
@@ -1076,7 +1076,7 @@ namespace UnrealBuildTool
 				{
 					CompileAction.DependencyListFile = FileItem.GetItemByFileReference(FileReference.Combine(OutputDir, SourceFile.Location.GetFileName() + ".txt"));
 					CompileAction.CommandArguments = String.Format("{0} -- {1} {2} /showIncludes", Utils.MakePathSafeToUseWithCommandLine(CompileAction.DependencyListFile.Location), Utils.MakePathSafeToUseWithCommandLine(CompileAction.CommandPath), CompileAction.CommandArguments);
-					CompileAction.CommandPath = FileReference.Combine(UnrealBuildTool.EngineDirectory, "Build", "Windows", "cl-filter", "cl-filter.exe").FullName;
+					CompileAction.CommandPath = FileReference.Combine(UnrealBuildTool.EngineDirectory, "Build", "Windows", "cl-filter", "cl-filter.exe");
 					CompileAction.ProducedItems.Add(CompileAction.DependencyListFile);
 				}
 
@@ -1119,8 +1119,8 @@ namespace UnrealBuildTool
 			{
 				Action CompileAction = new Action(ActionType.Compile);
 				CompileAction.CommandDescription = "Resource";
-				CompileAction.WorkingDirectory = UnrealBuildTool.EngineSourceDirectory.FullName;
-				CompileAction.CommandPath = EnvVars.ResourceCompilerPath.FullName;
+				CompileAction.WorkingDirectory = UnrealBuildTool.EngineSourceDirectory;
+				CompileAction.CommandPath = EnvVars.ResourceCompilerPath;
 				CompileAction.StatusDescription = Path.GetFileName(RCFile.AbsolutePath);
 				CompileAction.PrerequisiteItems.AddRange(CompileEnvironment.ForceIncludeFiles);
 
@@ -1412,8 +1412,8 @@ namespace UnrealBuildTool
 			// Create an action that invokes the linker.
 			Action LinkAction = new Action(ActionType.Link);
 			LinkAction.CommandDescription = "Link";
-			LinkAction.WorkingDirectory = UnrealBuildTool.EngineSourceDirectory.FullName;
-			LinkAction.CommandPath = bIsBuildingLibraryOrImportLibrary ? EnvVars.LibraryManagerPath.FullName : EnvVars.LinkerPath.FullName;
+			LinkAction.WorkingDirectory = UnrealBuildTool.EngineSourceDirectory;
+			LinkAction.CommandPath = bIsBuildingLibraryOrImportLibrary ? EnvVars.LibraryManagerPath : EnvVars.LinkerPath;
 			LinkAction.CommandArguments = String.Format("@\"{0}\"", ResponseFileName);
 			LinkAction.ProducedItems.AddRange(ProducedItems);
 			LinkAction.PrerequisiteItems.AddRange(PrerequisiteItems);
diff --git a/Engine/Source/Programs/UnrealBuildTool/ProjectFiles/CodeLite/CodeLiteProject.cs b/Engine/Source/Programs/UnrealBuildTool/ProjectFiles/CodeLite/CodeLiteProject.cs
--- a/Engine/Source/Programs/UnrealBuildTool/ProjectFiles/CodeLite/CodeLiteProject.cs
+++ b/Engine/Source/Programs/UnrealBuildTool/ProjectFiles/CodeLite/CodeLiteProject.cs
@@ -404,7 +404,7 @@ namespace UnrealBuildTool
 						// Add the working directory for the custom build commands.
 						//
 						XElement CustomBuildWorkingDirectory = new XElement("WorkingDirectory");
-						XText CustuomBuildWorkingDirectory = new XText(Path.GetDirectoryName(UnrealBuildTool.GetUBTPath()));
+						XText CustuomBuildWorkingDirectory = new XText(UnrealBuildTool.GetUBTPath().Directory.FullName);
 						CustomBuildWorkingDirectory.Add(CustuomBuildWorkingDirectory);
 						CodeLiteConfigurationCustomBuild.Add(CustomBuildWorkingDirectory);
 
@@ -420,7 +420,7 @@ namespace UnrealBuildTool
 						XElement CustomBuildCommand = new XElement("BuildCommand");
 						CodeLiteConfigurationCustomBuild.Add(CustomBuildCommand);
 
-						string BuildTarget = Path.GetFileName(UnrealBuildTool.GetUBTPath()) + " " + TargetName + " " + ProjectPlatformName + " " + CurConf.ToString();
+						string BuildTarget = UnrealBuildTool.GetUBTPath().GetFileName() + " " + TargetName + " " + ProjectPlatformName + " " + CurConf.ToString();
 						if( (BuildHostPlatform.Current.Platform != UnrealTargetPlatform.Win64) &&
 							(BuildHostPlatform.Current.Platform != UnrealTargetPlatform.Win32))
 						{
diff --git a/Engine/Source/Programs/UnrealBuildTool/System/Action.cs b/Engine/Source/Programs/UnrealBuildTool/System/Action.cs
--- a/Engine/Source/Programs/UnrealBuildTool/System/Action.cs
+++ b/Engine/Source/Programs/UnrealBuildTool/System/Action.cs
@@ -73,7 +73,7 @@ namespace UnrealBuildTool
 		/// <summary>
 		/// Directory from which to execute the program to create produced items
 		/// </summary>
-		public string WorkingDirectory = null;
+		public DirectoryReference WorkingDirectory = null;
 
 		/// <summary>
 		/// True if we should log extra information when we run a program to create produced items
@@ -83,7 +83,7 @@ namespace UnrealBuildTool
 		/// <summary>
 		/// The command to run to create produced items
 		/// </summary>
-		public string CommandPath = null;
+		public FileReference CommandPath = null;
 
 		/// <summary>
 		/// Command-line parameters to pass to the program
@@ -171,9 +171,9 @@ namespace UnrealBuildTool
 		public Action(BinaryArchiveReader Reader)
 		{
 			ActionType = (ActionType)Reader.ReadByte();
-			WorkingDirectory = Reader.ReadString();
+			WorkingDirectory = Reader.ReadDirectoryReference();
 			bPrintDebugInfo = Reader.ReadBool();
-			CommandPath = Reader.ReadString();
+			CommandPath = Reader.ReadFileReference();
 			CommandArguments = Reader.ReadString();
 			CommandDescription = Reader.ReadString();
 			StatusDescription = Reader.ReadString();
@@ -195,9 +195,9 @@ namespace UnrealBuildTool
 		public void Write(BinaryArchiveWriter Writer)
 		{
 			Writer.WriteByte((byte)ActionType);
-			Writer.WriteString(WorkingDirectory);
+			Writer.WriteDirectoryReference(WorkingDirectory);
 			Writer.WriteBool(bPrintDebugInfo);
-			Writer.WriteString(CommandPath);
+			Writer.WriteFileReference(CommandPath);
 			Writer.WriteString(CommandArguments);
 			Writer.WriteString(CommandDescription);
 			Writer.WriteString(StatusDescription);
@@ -263,12 +263,12 @@ namespace UnrealBuildTool
 			}
 			if(WorkingDirectory != Other.WorkingDirectory)
 			{
-				LogConflict("working directory is different", WorkingDirectory, Other.WorkingDirectory);
+				LogConflict("working directory is different", WorkingDirectory.FullName, Other.WorkingDirectory.FullName);
 				bResult = false;
 			}
 			if(CommandPath != Other.CommandPath)
 			{
-				LogConflict("command path is different", CommandPath, Other.CommandPath);
+				LogConflict("command path is different", CommandPath.FullName, Other.CommandPath.FullName);
 				bResult = false;
 			}
 			if(CommandArguments != Other.CommandArguments)
diff --git a/Engine/Source/Programs/UnrealBuildTool/System/ActionGraph.cs b/Engine/Source/Programs/UnrealBuildTool/System/ActionGraph.cs
--- a/Engine/Source/Programs/UnrealBuildTool/System/ActionGraph.cs
+++ b/Engine/Source/Programs/UnrealBuildTool/System/ActionGraph.cs
@@ -367,7 +367,7 @@ namespace UnrealBuildTool
 				{
 					// Check if the command-line of the action previously used to produce the item is outdated.
 					string OldProducingCommandLine = "";
-					string NewProducingCommandLine = RootAction.CommandPath + " " + RootAction.CommandArguments;
+					string NewProducingCommandLine = RootAction.CommandPath.FullName + " " + RootAction.CommandArguments;
 					if (!ActionHistory.TryGetProducingCommandLine(ProducedItem, out OldProducingCommandLine)
 					|| !String.Equals(OldProducingCommandLine, NewProducingCommandLine, StringComparison.InvariantCultureIgnoreCase))
 					{
diff --git a/Engine/Source/Programs/UnrealBuildTool/System/BuildHostPlatform.cs b/Engine/Source/Programs/UnrealBuildTool/System/BuildHostPlatform.cs
--- a/Engine/Source/Programs/UnrealBuildTool/System/BuildHostPlatform.cs
+++ b/Engine/Source/Programs/UnrealBuildTool/System/BuildHostPlatform.cs
@@ -88,7 +88,7 @@ namespace UnrealBuildTool
 		/// <summary>
 		/// Gets the path to the shell for this platform
 		/// </summary>
-		abstract public string Shell { get; }
+		abstract public FileReference Shell { get; }
 
 		/// <summary>
 		/// The type of shell returned by the Shell parameter
@@ -246,9 +246,9 @@ namespace UnrealBuildTool
 			get { return UnrealTargetPlatform.Win64; }
 		}
 
-		public override string Shell
+		public override FileReference Shell
 		{
-			get { return Environment.GetEnvironmentVariable("COMSPEC"); }
+			get { return new FileReference(Environment.GetEnvironmentVariable("COMSPEC")); }
 		}
 
 		public override ShellType ShellType
@@ -269,9 +269,9 @@ namespace UnrealBuildTool
 			get { return UnrealTargetPlatform.Mac; }
 		}
 
-		public override string Shell
+		public override FileReference Shell
 		{
-			get { return "/bin/sh"; }
+			get { return new FileReference("/bin/sh"); }
 		}
 
 		public override ShellType ShellType
@@ -393,9 +393,9 @@ namespace UnrealBuildTool
 			get { return UnrealTargetPlatform.Linux; }
 		}
 
-		public override string Shell
+		public override FileReference Shell
 		{
-			get { return "/bin/sh"; }
+			get { return new FileReference("/bin/sh"); }
 		}
 
 		public override ShellType ShellType
diff --git a/Engine/Source/Programs/UnrealBuildTool/System/DynamicCompilation.cs b/Engine/Source/Programs/UnrealBuildTool/System/DynamicCompilation.cs
--- a/Engine/Source/Programs/UnrealBuildTool/System/DynamicCompilation.cs
+++ b/Engine/Source/Programs/UnrealBuildTool/System/DynamicCompilation.cs
@@ -47,7 +47,7 @@ namespace UnrealBuildTool
 			// recently than the dynamically-compiled assembly, then we'll always recompile it.  This is
 			// because Unreal Build Tool's code may have changed in such a way that invalidate these
 			// previously-compiled assembly files.
-			FileItem ExecutableItem = FileItem.GetItemByPath(UnrealBuildTool.GetUBTPath());
+			FileItem ExecutableItem = FileItem.GetItemByFileReference(UnrealBuildTool.GetUBTPath());
 			if (ExecutableItem.LastWriteTimeUtc > OutputAssemblyInfo.LastWriteTimeUtc)
 			{
 				Log.TraceLog("Compiling {0}: {1} is newer", OutputAssemblyPath, ExecutableItem.Name);
diff --git a/Engine/Source/Programs/UnrealBuildTool/System/Utils.cs b/Engine/Source/Programs/UnrealBuildTool/System/Utils.cs
--- a/Engine/Source/Programs/UnrealBuildTool/System/Utils.cs
+++ b/Engine/Source/Programs/UnrealBuildTool/System/Utils.cs
@@ -1057,7 +1057,7 @@ namespace UnrealBuildTool
 			foreach(FileReference ScriptFile in ScriptFiles)
 			{
 				ProcessStartInfo StartInfo = new ProcessStartInfo();
-				StartInfo.FileName = BuildHostPlatform.Current.Shell;
+				StartInfo.FileName = BuildHostPlatform.Current.Shell.FullName;
 
 				if(BuildHostPlatform.Current.ShellType == ShellType.Cmd)
 				{
diff --git a/Engine/Source/Programs/UnrealBuildTool/ToolChain/AppleToolChain.cs b/Engine/Source/Programs/UnrealBuildTool/ToolChain/AppleToolChain.cs
--- a/Engine/Source/Programs/UnrealBuildTool/ToolChain/AppleToolChain.cs
+++ b/Engine/Source/Programs/UnrealBuildTool/ToolChain/AppleToolChain.cs
@@ -140,16 +140,11 @@ namespace UnrealBuildTool
 			ProjectFile = InProjectFile;
 		}
 
-		protected string GetMacDevSrcRoot()
+		protected DirectoryReference GetMacDevSrcRoot()
 		{
-			return UnrealBuildTool.EngineSourceDirectory.FullName;
+			return UnrealBuildTool.EngineSourceDirectory;
 		}
 
-        protected string GetMacDevEngineRoot()
-        {
-            return UnrealBuildTool.EngineDirectory.FullName;
-        }
-
 		protected void StripSymbolsWithXcode(FileReference SourceFile, FileReference TargetFile, string ToolchainDir)
 		{
 			if (SourceFile != TargetFile)
diff --git a/Engine/Source/Programs/UnrealBuildTool/ToolChain/RemoteMac.cs b/Engine/Source/Programs/UnrealBuildTool/ToolChain/RemoteMac.cs
--- a/Engine/Source/Programs/UnrealBuildTool/ToolChain/RemoteMac.cs
+++ b/Engine/Source/Programs/UnrealBuildTool/ToolChain/RemoteMac.cs
@@ -209,7 +209,7 @@ namespace UnrealBuildTool
 				CommandLine.AppendFormat(" \"{0}\"", UnrealBuildTool.EngineDirectory);
 				CommandLine.Append("\"");
 
-				using(Process ChildProcess = Process.Start(BuildHostPlatform.Current.Shell, CommandLine.ToString()))
+				using(Process ChildProcess = Process.Start(BuildHostPlatform.Current.Shell.FullName, CommandLine.ToString()))
 				{
 					ChildProcess.WaitForExit();
 				}
diff --git a/Engine/Source/Programs/UnrealBuildTool/UnrealBuildTool.cs b/Engine/Source/Programs/UnrealBuildTool/UnrealBuildTool.cs
--- a/Engine/Source/Programs/UnrealBuildTool/UnrealBuildTool.cs
+++ b/Engine/Source/Programs/UnrealBuildTool/UnrealBuildTool.cs
@@ -195,9 +195,9 @@ namespace UnrealBuildTool
 		/// Gets the absolute path to the UBT assembly.
 		/// </summary>
 		/// <returns>A string containing the path to the UBT assembly.</returns>
-		static public string GetUBTPath()
+		static public FileReference GetUBTPath()
 		{
-			string UnrealBuildToolPath = Assembly.GetExecutingAssembly().GetOriginalLocation();
+			FileReference UnrealBuildToolPath = new FileReference(Assembly.GetExecutingAssembly().GetOriginalLocation());
 			return UnrealBuildToolPath;
 		}
 
