UBT: Allow querying for the presence of a FileItem or DirectoryItem without actually creating it.

#rnx

diff --git a/Engine/Source/Programs/UnrealBuildTool/System/DirectoryItem.cs b/Engine/Source/Programs/UnrealBuildTool/System/DirectoryItem.cs
--- a/Engine/Source/Programs/UnrealBuildTool/System/DirectoryItem.cs
+++ b/Engine/Source/Programs/UnrealBuildTool/System/DirectoryItem.cs
@@ -184,6 +184,29 @@ namespace UnrealBuildTool
 			return Directories;
 		}
 
+		/// <summary>
+		/// Attempts to get a sub-directory by name
+		/// </summary>
+		/// <param name="Name">Name of the directory</param>
+		/// <param name="OutDirectory">If successful receives the matching directory item with this name</param>
+		/// <returns>True if the file exists, false otherwise</returns>
+		public bool TryGetDirectory(string Name, out DirectoryItem OutDirectory)
+		{
+			CacheFiles();
+
+			foreach(DirectoryItem Directory in Directories)
+			{
+				if(String.Equals(Directory.Name, Name, DirectoryReference.Comparison))
+				{
+					OutDirectory = Directory;
+					return true;
+				}
+			}
+
+			OutDirectory = null;
+			return false;
+		}
+
 		/// <summary>
 		/// Caches the files in this directory
 		/// </summary>
@@ -213,6 +236,30 @@ namespace UnrealBuildTool
 			return Files;
 		}
 
+		/// <summary>
+		/// Attempts to get a file from this directory by name. Unlike creating a file item and checking whether it exists, this does not
+		/// cause any I/O to check whether it exists, nor does it create a permanent FileItem object.
+		/// </summary>
+		/// <param name="Name">Name of the file</param>
+		/// <param name="OutFile">If successful receives the matching file item with this name</param>
+		/// <returns>True if the file exists, false otherwise</returns>
+		public bool TryGetFile(string Name, out FileItem OutFile)
+		{
+			CacheFiles();
+
+			foreach(FileItem File in Files)
+			{
+				if(String.Equals(File.Name, Name, FileReference.Comparison))
+				{
+					OutFile = File;
+					return true;
+				}
+			}
+
+			OutFile = null;
+			return false;
+		}
+
 		/// <summary>
 		/// Formats this object as a string for debugging
 		/// </summary>
