UBT: Remove environment variable expansion from executors. Individual toolchains already do this if necessary, and the current implementation is not universally supported.

#rnx

diff --git a/Engine/Source/Programs/UnrealBuildTool/Executors/LocalExecutor.cs b/Engine/Source/Programs/UnrealBuildTool/Executors/LocalExecutor.cs
--- a/Engine/Source/Programs/UnrealBuildTool/Executors/LocalExecutor.cs
+++ b/Engine/Source/Programs/UnrealBuildTool/Executors/LocalExecutor.cs
@@ -36,24 +36,6 @@ namespace UnrealBuildTool
 		int JobNumber;
 		int TotalJobs;
 
-		/// <summary>
-		/// Regex that matches environment variables in $(Variable) format.
-		/// </summary>
-		static Regex EnvironmentVariableRegex = new Regex("\\$\\(([\\d\\w]+)\\)");
-
-		/// <summary>
-		/// Replaces the environment variables references in a string with their values.
-		/// </summary>
-		public static string ExpandEnvironmentVariables(string Text)
-		{
-			foreach (Match EnvironmentVariableMatch in EnvironmentVariableRegex.Matches(Text))
-			{
-				string VariableValue = Environment.GetEnvironmentVariable(EnvironmentVariableMatch.Groups[1].Value);
-				Text = Text.Replace(EnvironmentVariableMatch.Value, VariableValue);
-			}
-			return Text;
-		}
-
 		/// <summary>
 		/// Constructor, takes the action to process
 		/// </summary>
@@ -92,23 +74,20 @@ namespace UnrealBuildTool
 			// Create the action's process.
 			ProcessStartInfo ActionStartInfo = new ProcessStartInfo();
 			ActionStartInfo.WorkingDirectory = Action.WorkingDirectory.FullName;
-
-			string ExpandedCommandPath = ExpandEnvironmentVariables(Action.CommandPath.FullName);
-			ActionStartInfo.FileName = ExpandedCommandPath;
-			ActionStartInfo.Arguments = ExpandEnvironmentVariables(Action.CommandArguments);
-
+			ActionStartInfo.FileName = Action.CommandPath.FullName;
+			ActionStartInfo.Arguments = Action.CommandArguments;
 			ActionStartInfo.UseShellExecute = false;
 			ActionStartInfo.RedirectStandardInput = false;
 			ActionStartInfo.RedirectStandardOutput = false;
 			ActionStartInfo.RedirectStandardError = false;
 
 			// Log command-line used to execute task if debug info printing is enabled.
-			Log.TraceVerbose("Executing: {0} {1}", ExpandedCommandPath, ActionStartInfo.Arguments);
+			Log.TraceVerbose("Executing: {0} {1}", ActionStartInfo.FileName, ActionStartInfo.Arguments);
 
 			// Log summary if wanted.
 			if (Action.bShouldOutputStatusDescription)
 			{
-				string CommandDescription = Action.CommandDescription != null ? Action.CommandDescription : Path.GetFileName(ExpandedCommandPath);
+				string CommandDescription = Action.CommandDescription != null ? Action.CommandDescription : Path.GetFileName(ActionStartInfo.FileName);
 				if (string.IsNullOrEmpty(CommandDescription))
 				{
 					Log.TraceInformation(Action.StatusDescription);
diff --git a/Engine/Source/Programs/UnrealBuildTool/Executors/SNDBS.cs b/Engine/Source/Programs/UnrealBuildTool/Executors/SNDBS.cs
--- a/Engine/Source/Programs/UnrealBuildTool/Executors/SNDBS.cs
+++ b/Engine/Source/Programs/UnrealBuildTool/Executors/SNDBS.cs
@@ -236,7 +236,7 @@ namespace UnrealBuildTool
 						{
 							// Add to script for execution by SN-DBS
 							string NewCommandArguments = "\"" + Action.CommandPath + "\"" + " " + Action.CommandArguments;
-							ScriptFile.WriteLine(ActionThread.ExpandEnvironmentVariables(NewCommandArguments));
+							ScriptFile.WriteLine(NewCommandArguments);
 							InActionThreadDictionary.Add(Action, DummyActionThread);
 							Action.StartTime = Action.EndTime = DateTimeOffset.Now;
 							Log.TraceInformation("[{0}/{1}] {2} {3}", JobNumber, InActions.Count, Action.CommandDescription, Action.StatusDescription);
diff --git a/Engine/Source/Programs/UnrealBuildTool/Platform/Android/AndroidToolChain.cs b/Engine/Source/Programs/UnrealBuildTool/Platform/Android/AndroidToolChain.cs
--- a/Engine/Source/Programs/UnrealBuildTool/Platform/Android/AndroidToolChain.cs
+++ b/Engine/Source/Programs/UnrealBuildTool/Platform/Android/AndroidToolChain.cs
@@ -1433,7 +1433,7 @@ namespace UnrealBuildTool
 							AllArguments = AllArguments.Replace("-fvisibility=hidden -fvisibility-inlines-hidden", "");
 						}
 
-						AllArguments = ActionThread.ExpandEnvironmentVariables(AllArguments);
+						AllArguments = Utils.ExpandVariables(AllArguments);
 						AllArguments = AllArguments.Replace("\\", "/");
 
 						// Remove shadow warning for this file if requested
@@ -1622,7 +1622,7 @@ namespace UnrealBuildTool
 						foreach (DirectoryReference LibraryPath in LinkEnvironment.LibraryPaths)
 						{
 							// LinkerPaths could be relative or absolute
-							string AbsoluteLibraryPath = ActionThread.ExpandEnvironmentVariables(LibraryPath.FullName);
+							string AbsoluteLibraryPath = Utils.ExpandVariables(LibraryPath.FullName);
 							if (IsDirectoryForArch(AbsoluteLibraryPath, Arch))
 							{
 								// environment variables aren't expanded when using the $( style
diff --git a/Engine/Source/Programs/UnrealBuildTool/Platform/Lumin/UEBuildLumin.cs b/Engine/Source/Programs/UnrealBuildTool/Platform/Lumin/UEBuildLumin.cs
--- a/Engine/Source/Programs/UnrealBuildTool/Platform/Lumin/UEBuildLumin.cs
+++ b/Engine/Source/Programs/UnrealBuildTool/Platform/Lumin/UEBuildLumin.cs
@@ -165,7 +165,7 @@ namespace UnrealBuildTool
 			CompileEnvironment.Definitions.Add("USE_ANDROID_OPENGL=0");
 			CompileEnvironment.Definitions.Add("WITH_OGGVORBIS=1");
 
-			DirectoryReference MLSDKDir = new DirectoryReference("$(MLSDK)", DirectoryReference.Sanitize.None);
+			DirectoryReference MLSDKDir = new DirectoryReference(Environment.GetEnvironmentVariable("MLSDK"));
 			CompileEnvironment.SystemIncludePaths.Add(DirectoryReference.Combine(MLSDKDir, "lumin/stl/gnu-libstdc++/include"));
 			CompileEnvironment.SystemIncludePaths.Add(DirectoryReference.Combine(MLSDKDir, "lumin/stl/gnu-libstdc++/include/aarch64-linux-android"));
 			CompileEnvironment.SystemIncludePaths.Add(DirectoryReference.Combine(MLSDKDir, "include"));
diff --git a/Engine/Source/Programs/UnrealBuildTool/Platform/Windows/PVSToolChain.cs b/Engine/Source/Programs/UnrealBuildTool/Platform/Windows/PVSToolChain.cs
--- a/Engine/Source/Programs/UnrealBuildTool/Platform/Windows/PVSToolChain.cs
+++ b/Engine/Source/Programs/UnrealBuildTool/Platform/Windows/PVSToolChain.cs
@@ -188,7 +188,7 @@ namespace UnrealBuildTool
 
 		static string GetFullIncludePath(string IncludePath)
 		{
-			return Path.GetFullPath(ActionThread.ExpandEnvironmentVariables(IncludePath));
+			return Path.GetFullPath(Utils.ExpandVariables(IncludePath));
 		}
 
 		public override CPPOutput CompileCPPFiles(CppCompileEnvironment CompileEnvironment, List<FileItem> InputFiles, DirectoryReference OutputDir, string ModuleName, List<Action> Actions)
diff --git a/Engine/Source/Programs/UnrealBuildTool/Platform/Windows/VCToolChain.cs b/Engine/Source/Programs/UnrealBuildTool/Platform/Windows/VCToolChain.cs
--- a/Engine/Source/Programs/UnrealBuildTool/Platform/Windows/VCToolChain.cs
+++ b/Engine/Source/Programs/UnrealBuildTool/Platform/Windows/VCToolChain.cs
@@ -1063,7 +1063,7 @@ namespace UnrealBuildTool
 				{
 					FileItem TargetFile = CompileAction.ProducedItems[0];
 					FileReference ResponseFileName = new FileReference(TargetFile.AbsolutePath + ".response");
-					FileItem ResponseFileItem = FileItem.CreateIntermediateTextFile(ResponseFileName, SharedArguments.Concat(FileArguments).Concat(AdditionalArguments).Select(x => ActionThread.ExpandEnvironmentVariables(x)));
+					FileItem ResponseFileItem = FileItem.CreateIntermediateTextFile(ResponseFileName, SharedArguments.Concat(FileArguments).Concat(AdditionalArguments).Select(x => Utils.ExpandVariables(x)));
 					CompileAction.CommandArguments = " @\"" + ResponseFileName + "\"";
 					CompileAction.PrerequisiteItems.Add(ResponseFileItem);
 				}
