Fix Lumin CIS errors. Tool paths are quoted if they are passed on the command line, rather than globally.

#rnx

diff --git a/Engine/Source/Programs/UnrealBuildTool/Platform/Android/AndroidToolChain.cs b/Engine/Source/Programs/UnrealBuildTool/Platform/Android/AndroidToolChain.cs
--- a/Engine/Source/Programs/UnrealBuildTool/Platform/Android/AndroidToolChain.cs
+++ b/Engine/Source/Programs/UnrealBuildTool/Platform/Android/AndroidToolChain.cs
@@ -1455,13 +1455,13 @@ namespace UnrealBuildTool
 						if(bExecuteCompilerThroughShell)
 						{
 							CompileAction.CommandPath = BuildHostPlatform.Current.Shell;
-							if (BuildHostPlatform.Current.ShellType == ShellType.Sh)
+							if (BuildHostPlatform.Current.ShellType == ShellType.Cmd)
 							{
-								CompileAction.CommandArguments = String.Format("/c \"{0} {1}\"", ClangPath, ResponseArgument);
+								CompileAction.CommandArguments = String.Format("/c \"{0} {1}\"", Utils.MakePathSafeToUseWithCommandLine(ClangPath), ResponseArgument);
 							}
 							else
 							{
-								CompileAction.CommandArguments = String.Format("-c \'{0} {1}\'", ClangPath, ResponseArgument);
+								CompileAction.CommandArguments = String.Format("-c \'{0} {1}\'", Utils.MakePathSafeToUseWithCommandLine(ClangPath), ResponseArgument);
 							}
 							CompileAction.CommandDescription = "Compile";
 						}
diff --git a/Engine/Source/Programs/UnrealBuildTool/Platform/Lumin/LuminToolChain.cs b/Engine/Source/Programs/UnrealBuildTool/Platform/Lumin/LuminToolChain.cs
--- a/Engine/Source/Programs/UnrealBuildTool/Platform/Lumin/LuminToolChain.cs
+++ b/Engine/Source/Programs/UnrealBuildTool/Platform/Lumin/LuminToolChain.cs
@@ -109,7 +109,7 @@ namespace UnrealBuildTool
 						throw new BuildException("Unexpected output from mabu in --print-tools: '{0}'", Line);
 					}
 
-					ToolsDict.Add(Split[0].Trim(), "\"" + Split[1].Trim() + "\"");
+					ToolsDict.Add(Split[0].Trim(), Split[1].Trim());
 				}
 			}
 
