UBT: Fix exception using TryGetDirectoryItem class due to contents not being cached correctly.

#rnx

diff --git a/Engine/Source/Programs/UnrealBuildTool/System/DirectoryItem.cs b/Engine/Source/Programs/UnrealBuildTool/System/DirectoryItem.cs
--- a/Engine/Source/Programs/UnrealBuildTool/System/DirectoryItem.cs
+++ b/Engine/Source/Programs/UnrealBuildTool/System/DirectoryItem.cs
@@ -87,6 +87,16 @@ namespace UnrealBuildTool
 			return DirectoryItem.GetItemByDirectoryReference(DirectoryReference.Combine(BaseDirectory.Location, Fragments));
 		}
 
+		/// <summary>
+		/// Finds or creates a directory item from its location
+		/// </summary>
+		/// <param name="Location">Path to the directory</param>
+		/// <returns>The directory item for this location</returns>
+		public static DirectoryItem GetItemByPath(string Location)
+		{
+			return GetItemByDirectoryReference(new DirectoryReference(Location));
+		}
+
 		/// <summary>
 		/// Finds or creates a directory item from its location
 		/// </summary>
@@ -192,7 +202,7 @@ namespace UnrealBuildTool
 		/// <returns>True if the file exists, false otherwise</returns>
 		public bool TryGetDirectory(string Name, out DirectoryItem OutDirectory)
 		{
-			CacheFiles();
+			CacheDirectories();
 
 			foreach(DirectoryItem Directory in Directories)
 			{
