UBT: Add methods to FileItem and DirectoryItem to get the parent DirectoryItem, and add support for resolving virtual directory names (ie. "." and "..") when querying for a directory by name

#rnx

diff --git a/Engine/Source/Programs/UnrealBuildTool/System/DirectoryItem.cs b/Engine/Source/Programs/UnrealBuildTool/System/DirectoryItem.cs
--- a/Engine/Source/Programs/UnrealBuildTool/System/DirectoryItem.cs
+++ b/Engine/Source/Programs/UnrealBuildTool/System/DirectoryItem.cs
@@ -76,6 +76,21 @@ namespace UnrealBuildTool
 			get { return Info.LastWriteTimeUtc; }
 		}
 
+		/// <summary>
+		/// Gets the parent directory item
+		/// </summary>
+		public DirectoryItem GetParentDirectoryItem()
+		{
+			if(Info.Parent == null)
+			{
+				return null;
+			}
+			else
+			{
+				return GetItemByDirectoryInfo(Info.Parent);
+			}
+		}
+
 		/// <summary>
 		/// Gets a new directory item by combining the existing directory item with the given path fragments
 		/// </summary>
@@ -202,6 +217,18 @@ namespace UnrealBuildTool
 		/// <returns>True if the file exists, false otherwise</returns>
 		public bool TryGetDirectory(string Name, out DirectoryItem OutDirectory)
 		{
+			if(Name.Equals(".", StringComparison.Ordinal))
+			{
+				OutDirectory = this;
+				return true;
+			}
+
+			if(Name.Equals("..", StringComparison.Ordinal))
+			{
+				OutDirectory = GetParentDirectoryItem();
+				return OutDirectory != null;
+			}
+
 			CacheDirectories();
 
 			foreach(DirectoryItem Directory in Directories)
diff --git a/Engine/Source/Programs/UnrealBuildTool/System/FileItem.cs b/Engine/Source/Programs/UnrealBuildTool/System/FileItem.cs
--- a/Engine/Source/Programs/UnrealBuildTool/System/FileItem.cs
+++ b/Engine/Source/Programs/UnrealBuildTool/System/FileItem.cs
@@ -94,6 +94,15 @@ namespace UnrealBuildTool
 			return Location.HasExtension(Extension);
 		}
 
+		/// <summary>
+		/// Gets the directory containing this file
+		/// </summary>
+		/// <returns>DirectoryItem for the directory containing this file</returns>
+		public DirectoryItem GetDirectoryItem()
+		{
+			return DirectoryItem.GetItemByDirectoryInfo(Info.Directory);
+		}
+
 		/// <summary>
 		/// Gets a FileItem corresponding to the given path
 		/// </summary>
