UBT: Use a dictionary to store contents of DirectoryItem instances to improve performance.

#rnx

diff --git a/Engine/Source/Programs/UnrealBuildTool/System/DirectoryItem.cs b/Engine/Source/Programs/UnrealBuildTool/System/DirectoryItem.cs
--- a/Engine/Source/Programs/UnrealBuildTool/System/DirectoryItem.cs
+++ b/Engine/Source/Programs/UnrealBuildTool/System/DirectoryItem.cs
@@ -29,12 +29,12 @@ namespace UnrealBuildTool
 		/// <summary>
 		/// Cached map of name to subdirectory item
 		/// </summary>
-		List<DirectoryItem> Directories;
+		Dictionary<string, DirectoryItem> Directories;
 
 		/// <summary>
 		/// Cached map of name to file
 		/// </summary>
-		List<FileItem> Files;
+		Dictionary<string, FileItem> Files;
 
 		/// <summary>
 		/// Global map of location to item
@@ -176,7 +176,7 @@ namespace UnrealBuildTool
 		{
 			if(Directories == null)
 			{
-				List<DirectoryItem> NewDirectories = new List<DirectoryItem>();
+				Dictionary<string, DirectoryItem> NewDirectories = new Dictionary<string, DirectoryItem>(DirectoryReference.Comparer);
 				if(Info.Exists)
 				{
 					foreach(DirectoryInfo SubDirectoryInfo in Info.EnumerateDirectories())
@@ -191,7 +191,7 @@ namespace UnrealBuildTool
 						}
 						else
 						{
-							NewDirectories.Add(DirectoryItem.GetItemByDirectoryInfo(SubDirectoryInfo));
+							NewDirectories.Add(SubDirectoryInfo.Name, DirectoryItem.GetItemByDirectoryInfo(SubDirectoryInfo));
 						}
 					}
 				}
@@ -206,7 +206,7 @@ namespace UnrealBuildTool
 		public IEnumerable<DirectoryItem> EnumerateDirectories()
 		{
 			CacheDirectories();
-			return Directories;
+			return Directories.Values;
 		}
 
 		/// <summary>
@@ -230,18 +230,7 @@ namespace UnrealBuildTool
 			}
 
 			CacheDirectories();
-
-			foreach(DirectoryItem Directory in Directories)
-			{
-				if(String.Equals(Directory.Name, Name, DirectoryReference.Comparison))
-				{
-					OutDirectory = Directory;
-					return true;
-				}
-			}
-
-			OutDirectory = null;
-			return false;
+			return Directories.TryGetValue(Name, out OutDirectory);
 		}
 
 		/// <summary>
@@ -251,12 +240,12 @@ namespace UnrealBuildTool
 		{
 			if(Files == null)
 			{
-				List<FileItem> NewFiles = new List<FileItem>();
+				Dictionary<string, FileItem> NewFiles = new Dictionary<string, FileItem>(FileReference.Comparer);
 				if(Info.Exists)
 				{
 					foreach(FileInfo FileInfo in Info.EnumerateFiles())
 					{
-						NewFiles.Add(FileItem.GetItemByFileInfo(FileInfo));
+						NewFiles.Add(FileInfo.Name, FileItem.GetItemByFileInfo(FileInfo));
 					}
 				}
 				Files = NewFiles;
@@ -270,7 +259,7 @@ namespace UnrealBuildTool
 		public IEnumerable<FileItem> EnumerateFiles()
 		{
 			CacheFiles();
-			return Files;
+			return Files.Values;
 		}
 
 		/// <summary>
@@ -283,18 +272,7 @@ namespace UnrealBuildTool
 		public bool TryGetFile(string Name, out FileItem OutFile)
 		{
 			CacheFiles();
-
-			foreach(FileItem File in Files)
-			{
-				if(String.Equals(File.Name, Name, FileReference.Comparison))
-				{
-					OutFile = File;
-					return true;
-				}
-			}
-
-			OutFile = null;
-			return false;
+			return Files.TryGetValue(Name, out OutFile);
 		}
 
 		/// <summary>
