UBT: Fix import library changes being treated as outdated prerequisites in the action graph. Causes larger builds than necessary.

#rnx

diff --git a/Engine/Source/Programs/UnrealBuildTool/System/ActionGraph.cs b/Engine/Source/Programs/UnrealBuildTool/System/ActionGraph.cs
--- a/Engine/Source/Programs/UnrealBuildTool/System/ActionGraph.cs
+++ b/Engine/Source/Programs/UnrealBuildTool/System/ActionGraph.cs
@@ -499,7 +499,7 @@ namespace UnrealBuildTool
 		{
 			if(PrerequisiteAction.bProducesImportLibrary)
 			{
-				return PrerequisiteAction.ProducedItems.Any(x => !x.Location.HasExtension(".lib") && RootAction.PrerequisiteItems.Contains(x));
+				return PrerequisiteAction.ProducedItems.All(x => x.Location.HasExtension(".lib") || !RootAction.PrerequisiteItems.Contains(x));
 			}
 			else
 			{
