UBT: Fix another edge case for import library dependencies triggering new dependent items to be rebuilt.

#rnx

diff --git a/Engine/Source/Programs/UnrealBuildTool/System/ActionGraph.cs b/Engine/Source/Programs/UnrealBuildTool/System/ActionGraph.cs
--- a/Engine/Source/Programs/UnrealBuildTool/System/ActionGraph.cs
+++ b/Engine/Source/Programs/UnrealBuildTool/System/ActionGraph.cs
@@ -445,9 +445,13 @@ namespace UnrealBuildTool
 							bool bPrerequisiteItemIsNewerThanLastExecution = TimeDifference.TotalSeconds > 1;
 							if (bPrerequisiteItemIsNewerThanLastExecution)
 							{
-								Log.TraceLog("{0}: Prerequisite {1} is newer than the last execution of the action: {2} vs {3}", RootAction.StatusDescription, Path.GetFileName(PrerequisiteItem.AbsolutePath), PrerequisiteItem.LastWriteTimeUtc.ToLocalTime(), LastExecutionTimeUtc.LocalDateTime);
-								bIsOutdated = true;
-								break;
+								// Need to check for import libraries here too
+								if(!bIgnoreOutdatedImportLibraries || !IsImportLibraryDependency(RootAction, PrerequisiteItem))
+								{
+									Log.TraceLog("{0}: Prerequisite {1} is newer than the last execution of the action: {2} vs {3}", RootAction.StatusDescription, Path.GetFileName(PrerequisiteItem.AbsolutePath), PrerequisiteItem.LastWriteTimeUtc.ToLocalTime(), LastExecutionTimeUtc.LocalDateTime);
+									bIsOutdated = true;
+									break;
+								}
 							}
 						}
 					}
@@ -507,6 +511,27 @@ namespace UnrealBuildTool
 			}
 		}
 
+		/// <summary>
+		/// Determines if the dependency on a between two actions is only for an import library
+		/// </summary>
+		/// <param name="RootAction">The action to check</param>
+		/// <param name="PrerequisiteItem">The dependency that is out of date</param>
+		/// <returns>True if the only dependency between two actions is for an import library</returns>
+		static bool IsImportLibraryDependency(Action RootAction, FileItem PrerequisiteItem)
+		{
+			if(PrerequisiteItem.Location.HasExtension(".lib"))
+			{
+				foreach(Action PrerequisiteAction in RootAction.PrerequisiteActions)
+				{
+					if(PrerequisiteAction.bProducesImportLibrary && PrerequisiteAction.ProducedItems.Contains(PrerequisiteItem))
+					{
+						return true;
+					}
+				}
+			}
+			return false;
+		}
+
 		/// <summary>
 		/// Builds a dictionary containing the actions from AllActions that are outdated by calling
 		/// IsActionOutdated.
