UBT: Add an explicit version number to the WriteMetadata build mode, so that it only needs to be re-run when the version number increases (rather than whenever UBT is rebuilt).

#rnx

diff --git a/Engine/Source/Programs/UnrealBuildTool/Configuration/UEBuildTarget.cs b/Engine/Source/Programs/UnrealBuildTool/Configuration/UEBuildTarget.cs
--- a/Engine/Source/Programs/UnrealBuildTool/Configuration/UEBuildTarget.cs
+++ b/Engine/Source/Programs/UnrealBuildTool/Configuration/UEBuildTarget.cs
@@ -1409,6 +1409,7 @@ namespace UnrealBuildTool
 
 				StringBuilder WriteMetadataArguments = new StringBuilder();
 				WriteMetadataArguments.AppendFormat("-Input={0}", Utils.MakePathSafeToUseWithCommandLine(MetadataTargetFile));
+				WriteMetadataArguments.AppendFormat(" -Version={0}", WriteMetadataMode.CurrentVersionNumber);
 				if(Rules.bNoManifestChanges)
 				{
 					WriteMetadataArguments.Append(" -NoManifestChanges");
@@ -1419,7 +1420,6 @@ namespace UnrealBuildTool
 				WriteMetadataAction.StatusDescription = ReceiptFileName.GetFileName();
 				WriteMetadataAction.bCanExecuteRemotely = false;
 				WriteMetadataAction.PrerequisiteItems.AddRange(Makefile.OutputItems);
-				WriteMetadataAction.PrerequisiteItems.Add(FileItem.GetItemByPath(Assembly.GetExecutingAssembly().Location));
 				WriteMetadataAction.ProducedItems.Add(FileItem.GetItemByFileReference(ReceiptFileName));
 				Makefile.Actions.Add(WriteMetadataAction);
 
diff --git a/Engine/Source/Programs/UnrealBuildTool/Modes/WriteMetadataMode.cs b/Engine/Source/Programs/UnrealBuildTool/Modes/WriteMetadataMode.cs
--- a/Engine/Source/Programs/UnrealBuildTool/Modes/WriteMetadataMode.cs
+++ b/Engine/Source/Programs/UnrealBuildTool/Modes/WriteMetadataMode.cs
@@ -66,6 +66,12 @@ namespace UnrealBuildTool
 	[ToolMode("WriteMetadata", ToolModeOptions.None)]
 	class WriteMetadataMode : ToolMode
 	{
+		/// <summary>
+		/// Version number for output files. This is not used directly, but can be appended to command-line invocations of the tool to ensure that actions to generate metadata are updated if the output format changes. 
+		/// The action graph is regenerated whenever UBT is rebuilt, so this should always match.
+		/// </summary>
+		public const int CurrentVersionNumber = 1;
+
 		/// <summary>
 		/// Execute the command
 		/// </summary>
@@ -76,8 +82,15 @@ namespace UnrealBuildTool
 			// Read the target info
 			WriteMetadataTargetInfo TargetInfo = BinaryFormatterUtils.Load<WriteMetadataTargetInfo>(Arguments.GetFileReference("-Input="));
 			bool bNoManifestChanges = Arguments.HasOption("-NoManifestChanges");
+			int VersionNumber = Arguments.GetInteger("-Version=");
 			Arguments.CheckAllArgumentsUsed();
 
+			// Make sure the version number is correct
+			if(VersionNumber != CurrentVersionNumber)
+			{
+				throw new BuildException("Version number to WriteMetadataMode is incorrect (expected {0}, got {1})", CurrentVersionNumber, VersionNumber);
+			}
+
 			// Check if we need to set a build id
 			TargetReceipt Receipt = TargetInfo.Receipt;
 			if(Receipt.Version.BuildId == null)
