UBT: Use UTC everywhere for timestamp checks when determining whether UHT needs to run. Inconsistency at the moment is causing it to run all the time.

#rnx

diff --git a/Engine/Source/Programs/UnrealBuildTool/System/ExternalExecution.cs b/Engine/Source/Programs/UnrealBuildTool/System/ExternalExecution.cs
--- a/Engine/Source/Programs/UnrealBuildTool/System/ExternalExecution.cs
+++ b/Engine/Source/Programs/UnrealBuildTool/System/ExternalExecution.cs
@@ -711,7 +711,7 @@ namespace UnrealBuildTool
 		/// Gets the latest write time of any of the UnrealHeaderTool binaries (including DLLs and Plugins) or DateTime.MaxValue if UnrealHeaderTool does not exist
 		/// </summary>
 		/// <returns>Latest timestamp of UHT binaries or DateTime.MaxValue if UnrealHeaderTool is out of date and needs to be rebuilt.</returns>
-		static bool GetHeaderToolTimestamp(FileReference ReceiptPath, out DateTime Timestamp)
+		static bool GetHeaderToolTimestampUtc(FileReference ReceiptPath, out DateTime Timestamp)
 		{
 			using (ScopedTimer TimestampTimer = new ScopedTimer("GetHeaderToolTimestamp"))
 			{
@@ -751,7 +751,7 @@ namespace UnrealBuildTool
 		/// Gets the timestamp of CoreUObject.gen.cpp file.
 		/// </summary>
 		/// <returns>Last write time of CoreUObject.gen.cpp or DateTime.MaxValue if it doesn't exist.</returns>
-		private static DateTime GetCoreGeneratedTimestamp(string ModuleName, string ModuleGeneratedCodeDirectory)
+		private static DateTime GetCoreGeneratedTimestampUtc(string ModuleName, string ModuleGeneratedCodeDirectory)
 		{
 			// In Installed Builds, we don't check the timestamps on engine headers.  Default to a very old date.
 			if (UnrealBuildTool.IsEngineInstalled())
@@ -763,7 +763,7 @@ namespace UnrealBuildTool
 			FileInfo CoreGeneratedFileInfo = new FileInfo(Path.Combine(ModuleGeneratedCodeDirectory, ModuleName + ".init.gen.cpp"));
 			if (CoreGeneratedFileInfo.Exists)
 			{
-				return CoreGeneratedFileInfo.LastWriteTime;
+				return CoreGeneratedFileInfo.LastWriteTimeUtc;
 			}
 
 			// Doesn't exist, so use a 'newer that everything' date to force rebuild headers.
@@ -775,26 +775,26 @@ namespace UnrealBuildTool
 		/// </summary>
 		/// <param name="BuildConfiguration">Build configuration</param>
 		/// <param name="UObjectModules">Modules that we generate headers for</param>
-		/// <param name="HeaderToolTimestamp">Timestamp for UHT</param>
+		/// <param name="HeaderToolTimestampUtc">Timestamp for UHT</param>
 		/// <param name="bIsGatheringBuild"></param>
 		/// <param name="bIsAssemblingBuild"></param>
 		/// <returns>True if the code files are out of date</returns>
-		private static bool AreGeneratedCodeFilesOutOfDate(BuildConfiguration BuildConfiguration, List<UHTModuleInfo> UObjectModules, DateTime HeaderToolTimestamp, bool bIsGatheringBuild, bool bIsAssemblingBuild)
+		private static bool AreGeneratedCodeFilesOutOfDate(BuildConfiguration BuildConfiguration, List<UHTModuleInfo> UObjectModules, DateTime HeaderToolTimestampUtc, bool bIsGatheringBuild, bool bIsAssemblingBuild)
 		{
 			// Get CoreUObject.init.gen.cpp timestamp.  If the source files are older than the CoreUObject generated code, we'll
 			// need to regenerate code for the module
-			DateTime? CoreGeneratedTimestamp = null;
+			DateTime? CoreGeneratedTimestampUtc = null;
 			{
 				// Find the CoreUObject module
 				foreach (UHTModuleInfo Module in UObjectModules)
 				{
 					if (Module.ModuleName.Equals("CoreUObject", StringComparison.InvariantCultureIgnoreCase))
 					{
-						CoreGeneratedTimestamp = GetCoreGeneratedTimestamp(Module.ModuleName, Path.GetDirectoryName(Module.GeneratedCPPFilenameBase));
+						CoreGeneratedTimestampUtc = GetCoreGeneratedTimestampUtc(Module.ModuleName, Path.GetDirectoryName(Module.GeneratedCPPFilenameBase));
 						break;
 					}
 				}
-				if (CoreGeneratedTimestamp == null)
+				if (CoreGeneratedTimestampUtc == null)
 				{
 					throw new BuildException("Could not find CoreUObject in list of all UObjectModules");
 				}
@@ -833,8 +833,8 @@ namespace UnrealBuildTool
 				}
 
 				// Make sure the last UHT run completed after UnrealHeaderTool.exe was compiled last, and after the CoreUObject headers were touched last.
-				DateTime SavedTimestamp = SavedTimestampFileInfo.LastWriteTime;
-				if (HeaderToolTimestamp > SavedTimestamp || CoreGeneratedTimestamp > SavedTimestamp)
+				DateTime SavedTimestampUtc = SavedTimestampFileInfo.LastWriteTimeUtc;
+				if (HeaderToolTimestampUtc > SavedTimestampUtc || CoreGeneratedTimestampUtc > SavedTimestampUtc)
 				{
 					// Generated code is older than UnrealHeaderTool.exe or CoreUObject headers.  Out of date!
 					Log.TraceVerbose("UnrealHeaderTool needs to run because UnrealHeaderTool.exe or CoreUObject headers are newer than SavedTimestamp for module {0}", Module.ModuleName);
@@ -843,7 +843,7 @@ namespace UnrealBuildTool
 
 				// Has the .build.cs file changed since we last generated headers successfully?
 				FileInfo ModuleRulesFile = new FileInfo(Module.ModuleRulesFile.FullName);
-				if (!ModuleRulesFile.Exists || ModuleRulesFile.LastWriteTime > SavedTimestamp)
+				if (!ModuleRulesFile.Exists || ModuleRulesFile.LastWriteTimeUtc > SavedTimestampUtc)
 				{
 					Log.TraceVerbose("UnrealHeaderTool needs to run because SavedTimestamp is older than the rules file ({0}) for module {1}", Module.ModuleRulesFile, Module.ModuleName);
 					return true;
@@ -875,10 +875,10 @@ namespace UnrealBuildTool
 
 				foreach (FileItem HeaderFile in AllUObjectHeaders)
 				{
-					DateTime HeaderFileTimestamp = HeaderFile.LastWriteTimeUtc;
+					DateTime HeaderFileTimestampUtc = HeaderFile.LastWriteTimeUtc;
 
 					// Has the source header changed since we last generated headers successfully?
-					if (HeaderFileTimestamp > SavedTimestamp)
+					if (HeaderFileTimestampUtc > SavedTimestampUtc)
 					{
 						Log.TraceVerbose("UnrealHeaderTool needs to run because SavedTimestamp is older than HeaderFileTimestamp ({0}) for module {1}", HeaderFile.AbsolutePath, Module.ModuleName);
 						return true;
@@ -898,8 +898,8 @@ namespace UnrealBuildTool
 						// files were actually UObject headers, but because we don't know all of the files we processed
 						// in the previous run, we need to assume our generated code is out of date if the directory timestamp
 						// is newer.
-						DateTime HeaderDirectoryTimestamp = new DirectoryInfo(Path.GetDirectoryName(HeaderFile.AbsolutePath)).LastWriteTime;
-						if (HeaderDirectoryTimestamp > SavedTimestamp)
+						DateTime HeaderDirectoryTimestampUtc = new DirectoryInfo(Path.GetDirectoryName(HeaderFile.AbsolutePath)).LastWriteTimeUtc;
+						if (HeaderDirectoryTimestampUtc > SavedTimestampUtc)
 						{
 							Log.TraceVerbose("UnrealHeaderTool needs to run because the directory containing an existing header ({0}) has changed, and headers may have been added to or deleted from module {1}", HeaderFile.AbsolutePath, Module.ModuleName);
 							return true;
@@ -1064,11 +1064,11 @@ namespace UnrealBuildTool
 				FileReference HeaderToolReceipt = GetHeaderToolReceiptFile(ProjectFile, UHTConfig, bHasProjectScriptPlugin);
 
 				// check if UHT is out of date
-				DateTime HeaderToolTimestamp = DateTime.MaxValue;
-				bool bHaveHeaderTool = !bIsBuildingUHT && GetHeaderToolTimestamp(HeaderToolReceipt, out HeaderToolTimestamp);
+				DateTime HeaderToolTimestampUtc = DateTime.MaxValue;
+				bool bHaveHeaderTool = !bIsBuildingUHT && GetHeaderToolTimestampUtc(HeaderToolReceipt, out HeaderToolTimestampUtc);
 
 				// ensure the headers are up to date
-				bool bUHTNeedsToRun = (BuildConfiguration.bForceHeaderGeneration || !bHaveHeaderTool || AreGeneratedCodeFilesOutOfDate(BuildConfiguration, UObjectModules, HeaderToolTimestamp, bIsGatheringBuild, bIsAssemblingBuild));
+				bool bUHTNeedsToRun = (BuildConfiguration.bForceHeaderGeneration || !bHaveHeaderTool || AreGeneratedCodeFilesOutOfDate(BuildConfiguration, UObjectModules, HeaderToolTimestampUtc, bIsGatheringBuild, bIsAssemblingBuild));
 
 				// Get the file containing dependencies for the generated code
 				FileReference ExternalDependenciesFile = ModuleInfoFileName.ChangeExtension(".deps");
