UBT: Improve logging for when UHT has to run.

#rnx

diff --git a/Engine/Source/Programs/UnrealBuildTool/System/ExternalExecution.cs b/Engine/Source/Programs/UnrealBuildTool/System/ExternalExecution.cs
--- a/Engine/Source/Programs/UnrealBuildTool/System/ExternalExecution.cs
+++ b/Engine/Source/Programs/UnrealBuildTool/System/ExternalExecution.cs
@@ -815,7 +815,7 @@ namespace UnrealBuildTool
 				if (!TestDirectory.Exists)
 				{
 					// Generated code directory is missing entirely!
-					Log.TraceVerbose("UnrealHeaderTool needs to run because no generated code directory was found for module {0}", Module.ModuleName);
+					Log.TraceLog("UnrealHeaderTool needs to run because no generated code directory was found for module {0}", Module.ModuleName);
 					return true;
 				}
 
@@ -828,16 +828,22 @@ namespace UnrealBuildTool
 				if (!SavedTimestampFileInfo.Exists)
 				{
 					// Timestamp file was missing (possibly deleted/cleaned), so headers are out of date
-					Log.TraceVerbose("UnrealHeaderTool needs to run because UHT Timestamp file did not exist for module {0}", Module.ModuleName);
+					Log.TraceLog("UnrealHeaderTool needs to run because UHT Timestamp file did not exist for module {0}", Module.ModuleName);
 					return true;
 				}
 
 				// Make sure the last UHT run completed after UnrealHeaderTool.exe was compiled last, and after the CoreUObject headers were touched last.
 				DateTime SavedTimestampUtc = SavedTimestampFileInfo.LastWriteTimeUtc;
-				if (HeaderToolTimestampUtc > SavedTimestampUtc || CoreGeneratedTimestampUtc > SavedTimestampUtc)
+				if (HeaderToolTimestampUtc > SavedTimestampUtc)
 				{
-					// Generated code is older than UnrealHeaderTool.exe or CoreUObject headers.  Out of date!
-					Log.TraceVerbose("UnrealHeaderTool needs to run because UnrealHeaderTool.exe or CoreUObject headers are newer than SavedTimestamp for module {0}", Module.ModuleName);
+					// Generated code is older than UnrealHeaderTool.exe.  Out of date!
+					Log.TraceLog("UnrealHeaderTool needs to run because UnrealHeaderTool timestamp ({0}) is later than timestamp for module {1} ({2})", HeaderToolTimestampUtc.ToLocalTime(), Module.ModuleName, SavedTimestampUtc.ToLocalTime());
+					return true;
+				}
+				if (CoreGeneratedTimestampUtc > SavedTimestampUtc)
+				{
+					// Generated code is older than CoreUObject headers.  Out of date!
+					Log.TraceLog("UnrealHeaderTool needs to run because CoreUObject timestamp ({0}) is newer than timestamp for module {1} ({2})", CoreGeneratedTimestampUtc.Value.ToLocalTime(), Module.ModuleName, SavedTimestampUtc.ToLocalTime());
 					return true;
 				}
 
@@ -845,7 +851,7 @@ namespace UnrealBuildTool
 				FileInfo ModuleRulesFile = new FileInfo(Module.ModuleRulesFile.FullName);
 				if (!ModuleRulesFile.Exists || ModuleRulesFile.LastWriteTimeUtc > SavedTimestampUtc)
 				{
-					Log.TraceVerbose("UnrealHeaderTool needs to run because SavedTimestamp is older than the rules file ({0}) for module {1}", Module.ModuleRulesFile, Module.ModuleName);
+					Log.TraceLog("UnrealHeaderTool needs to run because SavedTimestamp is older than the rules file ({0}) for module {1}", Module.ModuleRulesFile, Module.ModuleName);
 					return true;
 				}
 
@@ -860,14 +866,14 @@ namespace UnrealBuildTool
 					string[] UObjectFilesFromPreviousRun = File.ReadAllLines(TimestampFile);
 					if (AllUObjectHeaders.Count != UObjectFilesFromPreviousRun.Length)
 					{
-						Log.TraceVerbose("UnrealHeaderTool needs to run because there are a different number of UObject source files in module {0}", Module.ModuleName);
+						Log.TraceLog("UnrealHeaderTool needs to run because there are a different number of UObject source files in module {0}", Module.ModuleName);
 						return true;
 					}
 					for (int FileIndex = 0; FileIndex < AllUObjectHeaders.Count; ++FileIndex)
 					{
 						if (!UObjectFilesFromPreviousRun[FileIndex].Equals(AllUObjectHeaders[FileIndex].AbsolutePath, StringComparison.InvariantCultureIgnoreCase))
 						{
-							Log.TraceVerbose("UnrealHeaderTool needs to run because the set of UObject source files in module {0} has changed", Module.ModuleName);
+							Log.TraceLog("UnrealHeaderTool needs to run because the set of UObject source files in module {0} has changed", Module.ModuleName);
 							return true;
 						}
 					}
@@ -880,7 +886,7 @@ namespace UnrealBuildTool
 					// Has the source header changed since we last generated headers successfully?
 					if (HeaderFileTimestampUtc > SavedTimestampUtc)
 					{
-						Log.TraceVerbose("UnrealHeaderTool needs to run because SavedTimestamp is older than HeaderFileTimestamp ({0}) for module {1}", HeaderFile.AbsolutePath, Module.ModuleName);
+						Log.TraceLog("UnrealHeaderTool needs to run because SavedTimestamp is older than HeaderFileTimestamp ({0}) for module {1}", HeaderFile.AbsolutePath, Module.ModuleName);
 						return true;
 					}
 
@@ -901,7 +907,7 @@ namespace UnrealBuildTool
 						DateTime HeaderDirectoryTimestampUtc = new DirectoryInfo(Path.GetDirectoryName(HeaderFile.AbsolutePath)).LastWriteTimeUtc;
 						if (HeaderDirectoryTimestampUtc > SavedTimestampUtc)
 						{
-							Log.TraceVerbose("UnrealHeaderTool needs to run because the directory containing an existing header ({0}) has changed, and headers may have been added to or deleted from module {1}", HeaderFile.AbsolutePath, Module.ModuleName);
+							Log.TraceLog("UnrealHeaderTool needs to run because the directory containing an existing header ({0}) has changed, and headers may have been added to or deleted from module {1}", HeaderFile.AbsolutePath, Module.ModuleName);
 							return true;
 						}
 					}
@@ -1068,7 +1074,19 @@ namespace UnrealBuildTool
 				bool bHaveHeaderTool = !bIsBuildingUHT && GetHeaderToolTimestampUtc(HeaderToolReceipt, out HeaderToolTimestampUtc);
 
 				// ensure the headers are up to date
-				bool bUHTNeedsToRun = (BuildConfiguration.bForceHeaderGeneration || !bHaveHeaderTool || AreGeneratedCodeFilesOutOfDate(BuildConfiguration, UObjectModules, HeaderToolTimestampUtc, bIsGatheringBuild, bIsAssemblingBuild));
+				bool bUHTNeedsToRun = false;
+				if(!bHaveHeaderTool)
+				{
+					bUHTNeedsToRun = true;
+				}
+				else if(BuildConfiguration.bForceHeaderGeneration)
+				{
+					bUHTNeedsToRun = true;
+				}
+				else if(AreGeneratedCodeFilesOutOfDate(BuildConfiguration, UObjectModules, HeaderToolTimestampUtc, bIsGatheringBuild, bIsAssemblingBuild))
+				{
+					bUHTNeedsToRun = true;
+				}
 
 				// Get the file containing dependencies for the generated code
 				FileReference ExternalDependenciesFile = ModuleInfoFileName.ChangeExtension(".deps");
