UBT: Minor improvements for DirectoryItem performance; showed up in profile.

#rnx

diff --git a/Engine/Source/Programs/UnrealBuildTool/System/DirectoryItem.cs b/Engine/Source/Programs/UnrealBuildTool/System/DirectoryItem.cs
--- a/Engine/Source/Programs/UnrealBuildTool/System/DirectoryItem.cs
+++ b/Engine/Source/Programs/UnrealBuildTool/System/DirectoryItem.cs
@@ -217,16 +217,18 @@ namespace UnrealBuildTool
 		/// <returns>True if the file exists, false otherwise</returns>
 		public bool TryGetDirectory(string Name, out DirectoryItem OutDirectory)
 		{
-			if(Name.Equals(".", StringComparison.Ordinal))
+			if(Name.Length > 0 && Name[0] == '.')
 			{
-				OutDirectory = this;
-				return true;
-			}
-
-			if(Name.Equals("..", StringComparison.Ordinal))
-			{
-				OutDirectory = GetParentDirectoryItem();
-				return OutDirectory != null;
+				if(Name.Length == 1)
+				{
+					OutDirectory = this;
+					return true;
+				}
+				else if(Name.Length == 2 && Name[1] == '.')
+				{
+					OutDirectory = GetParentDirectoryItem();
+					return OutDirectory != null;
+				}
 			}
 
 			CacheDirectories();
