UBT: Rename WindowsPlatform.DefaultVersion to WindowsPlatform.DefaultWindowsSdkVersion for clarity.

#rnx

diff --git a/Engine/Source/Programs/UnrealBuildTool/Platform/Windows/UEBuildWindows.cs b/Engine/Source/Programs/UnrealBuildTool/Platform/Windows/UEBuildWindows.cs
--- a/Engine/Source/Programs/UnrealBuildTool/Platform/Windows/UEBuildWindows.cs
+++ b/Engine/Source/Programs/UnrealBuildTool/Platform/Windows/UEBuildWindows.cs
@@ -411,7 +411,7 @@ namespace UnrealBuildTool
 		/// <summary>
 		/// The default Windows SDK version to be used, if installed.
 		/// </summary>
-		static readonly VersionNumber DefaultVersion = VersionNumber.Parse("10.0.16299.0");
+		static readonly VersionNumber DefaultWindowsSdkVersion = VersionNumber.Parse("10.0.16299.0");
 
 		/// <summary>
 		/// Cache of Visual Studio installation directories
@@ -1317,9 +1317,9 @@ namespace UnrealBuildTool
 			}
 			else
 			{
-				if(CachedWindowsSdkDirs.ContainsKey(DefaultVersion))
+				if(CachedWindowsSdkDirs.ContainsKey(DefaultWindowsSdkVersion))
 				{
-					WindowsSdkVersion = DefaultVersion;
+					WindowsSdkVersion = DefaultWindowsSdkVersion;
 				}
 				else if(CachedWindowsSdkDirs.Count > 0)
 				{
