Fix UBT compile error after merge.

#rnx

diff --git a/Engine/Source/Programs/UnrealBuildTool/Platform/Mac/MacToolChain.cs b/Engine/Source/Programs/UnrealBuildTool/Platform/Mac/MacToolChain.cs
--- a/Engine/Source/Programs/UnrealBuildTool/Platform/Mac/MacToolChain.cs
+++ b/Engine/Source/Programs/UnrealBuildTool/Platform/Mac/MacToolChain.cs
@@ -426,7 +426,7 @@ namespace UnrealBuildTool
 			List<string> FrameworksSearchPaths = new List<string>();
 			foreach (UEBuildFramework Framework in CompileEnvironment.AdditionalFrameworks)
 			{
-				string FrameworkPath = Path.GetDirectoryName(Path.GetFullPath(Framework.FrameworkName));
+				string FrameworkPath = Path.GetDirectoryName(Path.GetFullPath(Framework.Name));
 				if (!FrameworksSearchPaths.Contains(FrameworkPath))
 				{
 					Arguments.Append(" -F \"");
