UBT: Fix the -allmodules option so that it works with any target, not just engine targets.

#rnx

diff --git a/Engine/Source/Programs/UnrealBuildTool/Configuration/UEBuildTarget.cs b/Engine/Source/Programs/UnrealBuildTool/Configuration/UEBuildTarget.cs
--- a/Engine/Source/Programs/UnrealBuildTool/Configuration/UEBuildTarget.cs
+++ b/Engine/Source/Programs/UnrealBuildTool/Configuration/UEBuildTarget.cs
@@ -2078,130 +2078,106 @@ namespace UnrealBuildTool
 				}
 			}
 
-			// If we're compiling a base engine target, create binaries for all the engine modules that are compatible with it.
-			if (ProjectFile == null && TargetType != TargetType.Program)
-			{
-				// Find all the known module names in this assembly
-				List<string> ModuleNames = new List<string>();
-				RulesAssembly.GetAllModuleNames(ModuleNames);
-
-				// Find all the platform folders to exclude from the list of valid modules
-				List<string> ExcludeFolders = new List<string>();
-				foreach (UnrealTargetPlatform TargetPlatform in Enum.GetValues(typeof(UnrealTargetPlatform)))
-				{
-					if (TargetPlatform != Platform)
-					{
-						ExcludeFolders.Add(TargetPlatform.ToString());
-					}
-				}
+			// Find all the known module names in this assembly
+			List<string> ModuleNames = new List<string>();
+			RulesAssembly.GetAllModuleNames(ModuleNames);
 
-				// Also exclude all the platform groups that this platform is not a part of
-				List<UnrealPlatformGroup> IncludePlatformGroups = new List<UnrealPlatformGroup>(UEBuildPlatform.GetPlatformGroups(Platform));
-				foreach (UnrealPlatformGroup PlatformGroup in Enum.GetValues(typeof(UnrealPlatformGroup)))
-				{
-					if (!IncludePlatformGroups.Contains(PlatformGroup))
-					{
-						ExcludeFolders.Add(PlatformGroup.ToString());
-					}
-				}
-
-				// Create an array of folders to exclude
-				string[] ExcludeFoldersArray = ExcludeFolders.ToArray();
+			// Find all the platform folders to exclude from the list of valid modules
+			ReadOnlyHashSet<string> ExcludeFolders = UEBuildPlatform.GetBuildPlatform(Platform).GetExcludedFolderNames();
 
-				// Find all the directories containing engine modules that may be compatible with this target
-				List<DirectoryReference> Directories = new List<DirectoryReference>();
-				if (TargetType == TargetType.Editor)
-				{
-					Directories.Add(UnrealBuildTool.EngineSourceEditorDirectory);
-				}
-				Directories.Add(UnrealBuildTool.EngineSourceRuntimeDirectory);
+			// Find all the directories containing engine modules that may be compatible with this target
+			List<DirectoryReference> Directories = new List<DirectoryReference>();
+			if (TargetType == TargetType.Editor)
+			{
+				Directories.Add(UnrealBuildTool.EngineSourceEditorDirectory);
+			}
+			Directories.Add(UnrealBuildTool.EngineSourceRuntimeDirectory);
 
-				// Also allow anything in the developer directory in non-shipping configurations (though we blacklist by default unless the PrecompileForTargets
-				// setting indicates that it's actually useful at runtime).
-				bool bAllowDeveloperModules = false;
-				if(Configuration != UnrealTargetConfiguration.Shipping)
-				{
-					Directories.Add(UnrealBuildTool.EngineSourceDeveloperDirectory);
-					Directories.Add(DirectoryReference.Combine(UnrealBuildTool.EnterpriseSourceDirectory, "Developer"));
-					bAllowDeveloperModules = true;
-				}
+			// Also allow anything in the developer directory in non-shipping configurations (though we blacklist by default unless the PrecompileForTargets
+			// setting indicates that it's actually useful at runtime).
+			bool bAllowDeveloperModules = false;
+			if(Configuration != UnrealTargetConfiguration.Shipping)
+			{
+				Directories.Add(UnrealBuildTool.EngineSourceDeveloperDirectory);
+				Directories.Add(DirectoryReference.Combine(UnrealBuildTool.EnterpriseSourceDirectory, "Developer"));
+				bAllowDeveloperModules = true;
+			}
 
-				// Find all the modules that are not part of the standard set
-				HashSet<string> FilteredModuleNames = new HashSet<string>();
-				foreach (string ModuleName in ModuleNames)
+			// Find all the modules that are not part of the standard set
+			HashSet<string> FilteredModuleNames = new HashSet<string>();
+			foreach (string ModuleName in ModuleNames)
+			{
+				FileReference ModuleFileName = RulesAssembly.GetModuleFileName(ModuleName);
+				foreach(DirectoryReference BaseDir in Directories)
 				{
-					FileReference ModuleFileName = RulesAssembly.GetModuleFileName(ModuleName);
-					foreach(DirectoryReference BaseDir in Directories)
+					if(ModuleFileName.IsUnderDirectory(BaseDir))
 					{
-						if(ModuleFileName.IsUnderDirectory(BaseDir))
-						{
-							Type RulesType = RulesAssembly.GetModuleRulesType(ModuleName);
+						Type RulesType = RulesAssembly.GetModuleRulesType(ModuleName);
 
-							SupportedPlatformsAttribute SupportedPlatforms = RulesType.GetCustomAttribute<SupportedPlatformsAttribute>();
-							if(SupportedPlatforms != null)
+						SupportedPlatformsAttribute SupportedPlatforms = RulesType.GetCustomAttribute<SupportedPlatformsAttribute>();
+						if(SupportedPlatforms != null)
+						{
+							if(SupportedPlatforms.Platforms.Contains(Platform))
 							{
-								if(SupportedPlatforms.Platforms.Contains(Platform))
-								{
-									FilteredModuleNames.Add(ModuleName);
-								}
+								FilteredModuleNames.Add(ModuleName);
 							}
-							else
+						}
+						else
+						{
+							if(!ModuleFileName.ContainsAnyNames(ExcludeFolders, BaseDir))
 							{
-								if(!ModuleFileName.ContainsAnyNames(ExcludeFoldersArray, BaseDir))
-								{
-									FilteredModuleNames.Add(ModuleName);
-								}
+								FilteredModuleNames.Add(ModuleName);
 							}
 						}
 					}
 				}
+			}
 
-				// Add all the plugin modules that need to be compiled
-				List<PluginInfo> Plugins = RulesAssembly.EnumeratePlugins().ToList();
-				foreach(PluginInfo Plugin in Plugins)
+			// Add all the plugin modules that need to be compiled
+			List<PluginInfo> Plugins = RulesAssembly.EnumeratePlugins().ToList();
+			foreach(PluginInfo Plugin in Plugins)
+			{
+				if(Rules.bIncludePluginsForTargetPlatforms || Plugin.Descriptor.SupportsTargetPlatform(Platform))
 				{
-					if(Rules.bIncludePluginsForTargetPlatforms || Plugin.Descriptor.SupportsTargetPlatform(Platform))
+					if(Plugin.Descriptor.Modules != null && (!Plugin.Descriptor.bRequiresBuildPlatform || !Plugin.File.ContainsAnyNames(ExcludeFolders, UnrealBuildTool.EngineDirectory)))
 					{
-						if(Plugin.Descriptor.Modules != null && (!Plugin.Descriptor.bRequiresBuildPlatform || !Plugin.File.ContainsAnyNames(ExcludeFoldersArray, UnrealBuildTool.EngineDirectory)))
+						foreach (ModuleDescriptor ModuleDescriptor in Plugin.Descriptor.Modules)
 						{
-							foreach (ModuleDescriptor ModuleDescriptor in Plugin.Descriptor.Modules)
+							if (ModuleDescriptor.IsCompiledInConfiguration(Platform, Configuration, TargetName, TargetType, bAllowDeveloperModules && Rules.bBuildDeveloperTools, Rules.bBuildEditor, Rules.bBuildRequiresCookedData))
 							{
-								if (ModuleDescriptor.IsCompiledInConfiguration(Platform, Configuration, TargetName, TargetType, bAllowDeveloperModules && Rules.bBuildDeveloperTools, Rules.bBuildEditor, Rules.bBuildRequiresCookedData))
+								FileReference ModuleFileName = RulesAssembly.GetModuleFileName(ModuleDescriptor.Name);
+								if(ModuleFileName == null)
 								{
-									FileReference ModuleFileName = RulesAssembly.GetModuleFileName(ModuleDescriptor.Name);
-									if(ModuleFileName == null)
-									{
-										throw new BuildException("Unable to find module '{0}' referenced by {1}", ModuleDescriptor.Name, Plugin.File);
-									}
-									if(!ModuleFileName.ContainsAnyNames(ExcludeFoldersArray, Plugin.Directory))
-									{
-										FilteredModuleNames.Add(ModuleDescriptor.Name);
-									}
+									throw new BuildException("Unable to find module '{0}' referenced by {1}", ModuleDescriptor.Name, Plugin.File);
+								}
+								if(!ModuleFileName.ContainsAnyNames(ExcludeFolders, Plugin.Directory))
+								{
+									FilteredModuleNames.Add(ModuleDescriptor.Name);
 								}
 							}
 						}
 					}
 				}
+			}
 
-				// Create rules for each remaining module, and check that it's set to be compiled
-				foreach(string FilteredModuleName in FilteredModuleNames)
+			// Create rules for each remaining module, and check that it's set to be compiled
+			foreach(string FilteredModuleName in FilteredModuleNames)
+			{
+				// Try to create the rules object, but catch any exceptions if it fails. Some modules (eg. SQLite) may determine that they are unavailable in the constructor.
+				ModuleRules ModuleRules;
+				try
 				{
-					// Try to create the rules object, but catch any exceptions if it fails. Some modules (eg. SQLite) may determine that they are unavailable in the constructor.
-					ModuleRules ModuleRules;
-					try
-					{
-						ModuleRules = RulesAssembly.CreateModuleRules(FilteredModuleName, this.Rules, "precompile option");
-					}
-					catch (BuildException)
-					{
-						ModuleRules = null;
-					}
+					ModuleRules = RulesAssembly.CreateModuleRules(FilteredModuleName, this.Rules, "all modules option");
+				}
+				catch (BuildException)
+				{
+					ModuleRules = null;
+				}
 
-					// Figure out if it can be precompiled
-					if (ModuleRules != null && ModuleRules.IsValidForTarget(ModuleRules.File))
-					{
-						ValidModuleNames.Add(FilteredModuleName);
-					}
+				// Figure out if it can be precompiled
+				if (ModuleRules != null && ModuleRules.IsValidForTarget(ModuleRules.File))
+				{
+					ValidModuleNames.Add(FilteredModuleName);
 				}
 			}
 
