Prevent programs from precompiling modules by default.

#rnx

diff --git a/Engine/Source/Programs/UnrealBuildTool/Configuration/ModuleRules.cs b/Engine/Source/Programs/UnrealBuildTool/Configuration/ModuleRules.cs
--- a/Engine/Source/Programs/UnrealBuildTool/Configuration/ModuleRules.cs
+++ b/Engine/Source/Programs/UnrealBuildTool/Configuration/ModuleRules.cs
@@ -956,13 +956,13 @@ namespace UnrealBuildTool
 					case ModuleRules.PrecompileTargetsType.None:
 						return false;
 					case ModuleRules.PrecompileTargetsType.Default:
-						return !RulesFile.IsUnderDirectory(UnrealBuildTool.EngineSourceDeveloperDirectory) || Target.Type == TargetType.Editor;
+						return (Target.Type != TargetType.Program && (!RulesFile.IsUnderDirectory(UnrealBuildTool.EngineSourceDeveloperDirectory) || Target.Type == TargetType.Editor));
 					case ModuleRules.PrecompileTargetsType.Game:
 						return (Target.Type == TargetType.Client || Target.Type == TargetType.Server || Target.Type == TargetType.Game);
 					case ModuleRules.PrecompileTargetsType.Editor:
 						return (Target.Type == TargetType.Editor);
 					case ModuleRules.PrecompileTargetsType.Any:
-						return true;
+						return (Target.Type != TargetType.Program);
 				}
 			}
 			return false;
