UBT: Update default Clang version to 7.10.

#rnx

diff --git a/Engine/Source/Programs/UnrealBuildTool/Platform/Windows/UEBuildWindows.cs b/Engine/Source/Programs/UnrealBuildTool/Platform/Windows/UEBuildWindows.cs
--- a/Engine/Source/Programs/UnrealBuildTool/Platform/Windows/UEBuildWindows.cs
+++ b/Engine/Source/Programs/UnrealBuildTool/Platform/Windows/UEBuildWindows.cs
@@ -401,7 +401,7 @@ namespace UnrealBuildTool
 		/// <summary>
 		/// The default compiler version to be used, if installed. 
 		/// </summary>
-		static readonly VersionNumber DefaultClangToolChainVersion = VersionNumber.Parse("6.0.0");
+		static readonly VersionNumber DefaultClangToolChainVersion = VersionNumber.Parse("7.1.0");
 
 		/// <summary>
 		/// The default compiler version to be used, if installed. 
