UBT: Changes to the way the -allmodules option works with program targets.

* Engine modules are not built for program targets by default.
* Plugin modules that are whitelisted for the given program will be built by default, unless disabled via the PrecompileForTargets setting in the ModuleRules class.
* The same logic is used to determine plugin suitability for a target as the UEBuildTarget.AddPlugin() method.

#rnx

diff --git a/Engine/Source/Programs/UnrealBuildTool/Configuration/ModuleRules.cs b/Engine/Source/Programs/UnrealBuildTool/Configuration/ModuleRules.cs
--- a/Engine/Source/Programs/UnrealBuildTool/Configuration/ModuleRules.cs
+++ b/Engine/Source/Programs/UnrealBuildTool/Configuration/ModuleRules.cs
@@ -956,13 +956,13 @@ namespace UnrealBuildTool
 					case ModuleRules.PrecompileTargetsType.None:
 						return false;
 					case ModuleRules.PrecompileTargetsType.Default:
-						return (Target.Type != TargetType.Program && (!RulesFile.IsUnderDirectory(UnrealBuildTool.EngineSourceDeveloperDirectory) || Target.Type == TargetType.Editor));
+						return (Target.Type == TargetType.Editor || !RulesFile.IsUnderDirectory(UnrealBuildTool.EngineSourceDeveloperDirectory) || Plugin != null);
 					case ModuleRules.PrecompileTargetsType.Game:
 						return (Target.Type == TargetType.Client || Target.Type == TargetType.Server || Target.Type == TargetType.Game);
 					case ModuleRules.PrecompileTargetsType.Editor:
 						return (Target.Type == TargetType.Editor);
 					case ModuleRules.PrecompileTargetsType.Any:
-						return (Target.Type != TargetType.Program);
+						return true;
 				}
 			}
 			return false;
diff --git a/Engine/Source/Programs/UnrealBuildTool/Configuration/UEBuildTarget.cs b/Engine/Source/Programs/UnrealBuildTool/Configuration/UEBuildTarget.cs
--- a/Engine/Source/Programs/UnrealBuildTool/Configuration/UEBuildTarget.cs
+++ b/Engine/Source/Programs/UnrealBuildTool/Configuration/UEBuildTarget.cs
@@ -2078,55 +2078,62 @@ namespace UnrealBuildTool
 				}
 			}
 
-			// Find all the known module names in this assembly
-			List<string> ModuleNames = new List<string>();
-			RulesAssembly.GetAllModuleNames(ModuleNames);
+			// Whether to allow developer modules
+			bool bAllowDeveloperModules = (Configuration != UnrealTargetConfiguration.Shipping);
 
 			// Find all the platform folders to exclude from the list of valid modules
 			ReadOnlyHashSet<string> ExcludeFolders = UEBuildPlatform.GetBuildPlatform(Platform).GetExcludedFolderNames();
 
-			// Find all the directories containing engine modules that may be compatible with this target
-			List<DirectoryReference> Directories = new List<DirectoryReference>();
-			if (TargetType == TargetType.Editor)
-			{
-				Directories.Add(UnrealBuildTool.EngineSourceEditorDirectory);
-			}
-			Directories.Add(UnrealBuildTool.EngineSourceRuntimeDirectory);
+			// Set of module names to build
+			HashSet<string> FilteredModuleNames = new HashSet<string>();
 
-			// Also allow anything in the developer directory in non-shipping configurations (though we blacklist by default unless the PrecompileForTargets
-			// setting indicates that it's actually useful at runtime).
-			bool bAllowDeveloperModules = false;
-			if(Configuration != UnrealTargetConfiguration.Shipping)
+			// Only add engine modules for non-program targets. Programs only compile whitelisted modules through plugins.
+			if(TargetType != TargetType.Program)
 			{
-				Directories.Add(UnrealBuildTool.EngineSourceDeveloperDirectory);
-				Directories.Add(DirectoryReference.Combine(UnrealBuildTool.EnterpriseSourceDirectory, "Developer"));
-				bAllowDeveloperModules = true;
-			}
+				// Find all the known module names in this assembly
+				List<string> ModuleNames = new List<string>();
+				RulesAssembly.GetAllModuleNames(ModuleNames);
 
-			// Find all the modules that are not part of the standard set
-			HashSet<string> FilteredModuleNames = new HashSet<string>();
-			foreach (string ModuleName in ModuleNames)
-			{
-				FileReference ModuleFileName = RulesAssembly.GetModuleFileName(ModuleName);
-				foreach(DirectoryReference BaseDir in Directories)
+				// Find all the directories containing engine modules that may be compatible with this target
+				List<DirectoryReference> Directories = new List<DirectoryReference>();
+				if (TargetType == TargetType.Editor)
 				{
-					if(ModuleFileName.IsUnderDirectory(BaseDir))
-					{
-						Type RulesType = RulesAssembly.GetModuleRulesType(ModuleName);
+					Directories.Add(UnrealBuildTool.EngineSourceEditorDirectory);
+				}
+				Directories.Add(UnrealBuildTool.EngineSourceRuntimeDirectory);
+
+				// Also allow anything in the developer directory in non-shipping configurations (though we blacklist by default unless the PrecompileForTargets
+				// setting indicates that it's actually useful at runtime).
+				if(bAllowDeveloperModules)
+				{
+					Directories.Add(UnrealBuildTool.EngineSourceDeveloperDirectory);
+					Directories.Add(DirectoryReference.Combine(UnrealBuildTool.EnterpriseSourceDirectory, "Developer"));
+				}
 
-						SupportedPlatformsAttribute SupportedPlatforms = RulesType.GetCustomAttribute<SupportedPlatformsAttribute>();
-						if(SupportedPlatforms != null)
+				// Find all the modules that are not part of the standard set
+				foreach (string ModuleName in ModuleNames)
+				{
+					FileReference ModuleFileName = RulesAssembly.GetModuleFileName(ModuleName);
+					foreach(DirectoryReference BaseDir in Directories)
+					{
+						if(ModuleFileName.IsUnderDirectory(BaseDir))
 						{
-							if(SupportedPlatforms.Platforms.Contains(Platform))
+							Type RulesType = RulesAssembly.GetModuleRulesType(ModuleName);
+
+							SupportedPlatformsAttribute SupportedPlatforms = RulesType.GetCustomAttribute<SupportedPlatformsAttribute>();
+							if(SupportedPlatforms != null)
 							{
-								FilteredModuleNames.Add(ModuleName);
+								if(SupportedPlatforms.Platforms.Contains(Platform))
+								{
+									FilteredModuleNames.Add(ModuleName);
+								}
 							}
-						}
-						else
-						{
-							if(!ModuleFileName.ContainsAnyNames(ExcludeFolders, BaseDir))
+							else
 							{
-								FilteredModuleNames.Add(ModuleName);
+								if(!ModuleFileName.ContainsAnyNames(ExcludeFolders, BaseDir))
+								{
+									FilteredModuleNames.Add(ModuleName);
+								}
 							}
 						}
 					}
@@ -2137,24 +2144,43 @@ namespace UnrealBuildTool
 			List<PluginInfo> Plugins = RulesAssembly.EnumeratePlugins().ToList();
 			foreach(PluginInfo Plugin in Plugins)
 			{
-				if(Rules.bIncludePluginsForTargetPlatforms || Plugin.Descriptor.SupportsTargetPlatform(Platform))
+				// Ignore plugins without any modules
+				if(Plugin.Descriptor.Modules == null)
+				{
+					continue;
+				}
+
+				// Disable any plugin which does not support the target platform. The editor should update such references in the .uproject file on load.
+				if (!Rules.bIncludePluginsForTargetPlatforms && !Plugin.Descriptor.SupportsTargetPlatform(Platform))
+				{
+					continue;
+				}
+
+				// Disable any plugin that requires the build platform
+				if(Plugin.Descriptor.bRequiresBuildPlatform && ShouldExcludePlugin(Plugin, ExcludeFolders))
+				{
+					continue;
+				}
+
+				// Disable any plugins that aren't compatible with this program
+				if (TargetType == TargetType.Program && (Plugin.Descriptor.SupportedPrograms == null || !Plugin.Descriptor.SupportedPrograms.Contains(AppName)))
 				{
-					if(Plugin.Descriptor.Modules != null && (!Plugin.Descriptor.bRequiresBuildPlatform || !Plugin.File.ContainsAnyNames(ExcludeFolders, UnrealBuildTool.EngineDirectory)))
+					continue;
+				}
+
+				// Add all the modules
+				foreach (ModuleDescriptor ModuleDescriptor in Plugin.Descriptor.Modules)
+				{
+					if (ModuleDescriptor.IsCompiledInConfiguration(Platform, Configuration, TargetName, TargetType, bAllowDeveloperModules && Rules.bBuildDeveloperTools, Rules.bBuildEditor, Rules.bBuildRequiresCookedData))
 					{
-						foreach (ModuleDescriptor ModuleDescriptor in Plugin.Descriptor.Modules)
+						FileReference ModuleFileName = RulesAssembly.GetModuleFileName(ModuleDescriptor.Name);
+						if(ModuleFileName == null)
 						{
-							if (ModuleDescriptor.IsCompiledInConfiguration(Platform, Configuration, TargetName, TargetType, bAllowDeveloperModules && Rules.bBuildDeveloperTools, Rules.bBuildEditor, Rules.bBuildRequiresCookedData))
-							{
-								FileReference ModuleFileName = RulesAssembly.GetModuleFileName(ModuleDescriptor.Name);
-								if(ModuleFileName == null)
-								{
-									throw new BuildException("Unable to find module '{0}' referenced by {1}", ModuleDescriptor.Name, Plugin.File);
-								}
-								if(!ModuleFileName.ContainsAnyNames(ExcludeFolders, Plugin.Directory))
-								{
-									FilteredModuleNames.Add(ModuleDescriptor.Name);
-								}
-							}
+							throw new BuildException("Unable to find module '{0}' referenced by {1}", ModuleDescriptor.Name, Plugin.File);
+						}
+						if(!ModuleFileName.ContainsAnyNames(ExcludeFolders, Plugin.Directory))
+						{
+							FilteredModuleNames.Add(ModuleDescriptor.Name);
 						}
 					}
 				}
@@ -2667,7 +2693,7 @@ namespace UnrealBuildTool
 		/// <summary>
 		/// Checks whether a plugin path contains a platform directory fragment
 		/// </summary>
-		private bool ShouldExcludePlugin(PluginInfo Plugin, string[] ExcludeFolders)
+		private bool ShouldExcludePlugin(PluginInfo Plugin, IEnumerable<string> ExcludeFolders)
 		{
 			if (Plugin.LoadedFrom == PluginLoadedFrom.Engine)
 			{
