UBT: Save command-line arguments into the target makefile, so we can invalidate it if they change.

diff --git a/Engine/Source/Programs/UnrealBuildTool/Modes/BuildMode.cs b/Engine/Source/Programs/UnrealBuildTool/Modes/BuildMode.cs
--- a/Engine/Source/Programs/UnrealBuildTool/Modes/BuildMode.cs
+++ b/Engine/Source/Programs/UnrealBuildTool/Modes/BuildMode.cs
@@ -280,7 +280,7 @@ namespace UnrealBuildTool
 				using(Timeline.ScopeEvent("TargetMakefile.Load()"))
 				{
 					string ReasonNotLoaded;
-					Makefile = TargetMakefile.Load(MakefileLocation, TargetDescriptor.ProjectFile, TargetDescriptor.Platform, WorkingSet, out ReasonNotLoaded);
+					Makefile = TargetMakefile.Load(MakefileLocation, TargetDescriptor.ProjectFile, TargetDescriptor.Platform, TargetDescriptor.AdditionalArguments.GetRawArray(), WorkingSet, out ReasonNotLoaded);
 					if (Makefile == null)
 					{
 						Log.TraceInformation("Creating makefile for {0} ({1})", TargetDescriptor.Name, ReasonNotLoaded);
@@ -306,6 +306,9 @@ namespace UnrealBuildTool
 					Makefile = Target.Build(BuildConfiguration, WorkingSet, bIsAssemblingBuild);
 				}
 
+				// Save the additional command line arguments
+				Makefile.AdditionalArguments = TargetDescriptor.AdditionalArguments.GetRawArray();
+
 				// Save the environment variables
 				foreach (System.Collections.DictionaryEntry EnvironmentVariable in Environment.GetEnvironmentVariables())
 				{
diff --git a/Engine/Source/Programs/UnrealBuildTool/System/TargetMakefile.cs b/Engine/Source/Programs/UnrealBuildTool/System/TargetMakefile.cs
--- a/Engine/Source/Programs/UnrealBuildTool/System/TargetMakefile.cs
+++ b/Engine/Source/Programs/UnrealBuildTool/System/TargetMakefile.cs
@@ -47,6 +47,11 @@ namespace UnrealBuildTool
 		/// </summary>
 		public bool bHasProjectScriptPlugin;
 
+		/// <summary>
+		/// The array of command-line arguments. The makefile will be invalidated whenever these change.
+		/// </summary>
+		public string[] AdditionalArguments;
+
 		/// <summary>
 		/// Scripts which should be run before building anything
 		/// </summary>
@@ -146,6 +151,7 @@ namespace UnrealBuildTool
 			TargetType = (TargetType)Reader.ReadInt();
 			bDeployAfterCompile = Reader.ReadBool();
 			bHasProjectScriptPlugin = Reader.ReadBool();
+			AdditionalArguments = Reader.ReadArray(() => Reader.ReadString());
 			PreBuildScripts = Reader.ReadArray(() => Reader.ReadFileReference());
 			Actions = Reader.ReadList(() => new Action(Reader));
 			EnvironmentVariables = Reader.ReadList(() => Tuple.Create(Reader.ReadString(), Reader.ReadString()));
@@ -171,6 +177,7 @@ namespace UnrealBuildTool
 			Writer.WriteInt((int)TargetType);
 			Writer.WriteBool(bDeployAfterCompile);
 			Writer.WriteBool(bHasProjectScriptPlugin);
+			Writer.WriteArray(AdditionalArguments, Item => Writer.WriteString(Item));
 			Writer.WriteArray(PreBuildScripts, Item => Writer.WriteFileReference(Item));
 			Writer.WriteList(Actions, Action => Action.Write(Writer));
 			Writer.WriteList(EnvironmentVariables, x => { Writer.WriteString(x.Item1); Writer.WriteString(x.Item2); });
@@ -205,10 +212,11 @@ namespace UnrealBuildTool
 		/// <param name="MakefilePath">Path to the makefile to load</param>
 		/// <param name="ProjectFile">Path to the project file</param>
 		/// <param name="Platform">Platform for this makefile</param>
+		/// <param name="Arguments">Command line arguments for this target</param>
 		/// <param name="WorkingSet">Interface to query which source files are in the working set</param>
 		/// <param name="ReasonNotLoaded">If the function returns null, this string will contain the reason why</param>
 		/// <returns>The loaded makefile, or null if it failed for some reason.  On failure, the 'ReasonNotLoaded' variable will contain information about why</returns>
-		public static TargetMakefile Load(FileReference MakefilePath, FileReference ProjectFile, UnrealTargetPlatform Platform, ISourceFileWorkingSet WorkingSet, out string ReasonNotLoaded)
+		public static TargetMakefile Load(FileReference MakefilePath, FileReference ProjectFile, UnrealTargetPlatform Platform, string[] Arguments, ISourceFileWorkingSet WorkingSet, out string ReasonNotLoaded)
 		{
 			// Check the directory timestamp on the project files directory.  If the user has generated project files more
 			// recently than the UBTMakefile, then we need to consider the file to be out of date
@@ -327,6 +335,13 @@ namespace UnrealBuildTool
 
 			using(Timeline.ScopeEvent("Checking makefile validity"))
 			{
+				// Check if the arguments are different
+				if(!Enumerable.SequenceEqual(LoadedUBTMakefile.AdditionalArguments, Arguments))
+				{
+					ReasonNotLoaded = "command line arguments changed";
+					return null;
+				}
+
 				// Check if ini files are newer. Ini files contain build settings too.
 				DirectoryReference ProjectDirectory = DirectoryReference.FromFile(ProjectFile);
 				foreach (ConfigHierarchyType IniType in (ConfigHierarchyType[])Enum.GetValues(typeof(ConfigHierarchyType)))
