UBT: Renaming variables, updating comments.

#rnx

diff --git a/Engine/Source/Programs/UnrealBuildTool/System/TargetMakefile.cs b/Engine/Source/Programs/UnrealBuildTool/System/TargetMakefile.cs
--- a/Engine/Source/Programs/UnrealBuildTool/System/TargetMakefile.cs
+++ b/Engine/Source/Programs/UnrealBuildTool/System/TargetMakefile.cs
@@ -13,7 +13,7 @@ using Tools.DotNETCommon;
 namespace UnrealBuildTool
 {
 	/// <summary>
-	/// A special Makefile that UBT is able to create in "-gather" mode, then load in "-assemble" mode to accelerate iterative compiling and linking
+	/// Cached list of actions that need to be executed to build a target, along with the information needed to determine whether they are valid.
 	/// </summary>
 	class TargetMakefile
 	{
@@ -207,7 +207,7 @@ namespace UnrealBuildTool
 		}
 
 		/// <summary>
-		/// Loads a UBTMakefile from disk
+		/// Loads a makefile  from disk
 		/// </summary>
 		/// <param name="MakefilePath">Path to the makefile to load</param>
 		/// <param name="ProjectFile">Path to the project file</param>
@@ -218,19 +218,18 @@ namespace UnrealBuildTool
 		/// <returns>The loaded makefile, or null if it failed for some reason.  On failure, the 'ReasonNotLoaded' variable will contain information about why</returns>
 		public static TargetMakefile Load(FileReference MakefilePath, FileReference ProjectFile, UnrealTargetPlatform Platform, string[] Arguments, ISourceFileWorkingSet WorkingSet, out string ReasonNotLoaded)
 		{
-			// Check the directory timestamp on the project files directory.  If the user has generated project files more
-			// recently than the UBTMakefile, then we need to consider the file to be out of date
-			FileInfo UBTMakefileInfo = new FileInfo(MakefilePath.FullName);
-			if (!UBTMakefileInfo.Exists)
+			// Check the directory timestamp on the project files directory.  If the user has generated project files more recently than the makefile, then we need to consider the file to be out of date
+			FileInfo MakefileInfo = new FileInfo(MakefilePath.FullName);
+			if (!MakefileInfo.Exists)
 			{
-				// UBTMakefile doesn't even exist, so we won't bother loading it
+				// Makefile doesn't even exist, so we won't bother loading it
 				ReasonNotLoaded = "no existing makefile";
 				return null;
 			}
 
 			// Check the build version
 			FileInfo BuildVersionFileInfo = new FileInfo(BuildVersion.GetDefaultFileName().FullName);
-			if (BuildVersionFileInfo.Exists && UBTMakefileInfo.LastWriteTime.CompareTo(BuildVersionFileInfo.LastWriteTime) < 0)
+			if (BuildVersionFileInfo.Exists && MakefileInfo.LastWriteTime.CompareTo(BuildVersionFileInfo.LastWriteTime) < 0)
 			{
 				Log.TraceLog("Existing makefile is older than Build.version, ignoring it");
 				ReasonNotLoaded = "Build.version is newer";
@@ -247,9 +246,9 @@ namespace UnrealBuildTool
 				if (DirectoryReference.Exists(ProjectFileGenerator.IntermediateProjectFilesPath))
 				{
 					DateTime EngineProjectFilesLastUpdateTime = new FileInfo(ProjectFileGenerator.ProjectTimestampFile).LastWriteTime;
-					if (UBTMakefileInfo.LastWriteTime.CompareTo(EngineProjectFilesLastUpdateTime) < 0)
+					if (MakefileInfo.LastWriteTime.CompareTo(EngineProjectFilesLastUpdateTime) < 0)
 					{
-						// Engine project files are newer than UBTMakefile
+						// Engine project files are newer than makefile
 						Log.TraceLog("Existing makefile is older than generated engine project files, ignoring it");
 						ReasonNotLoaded = "project files are newer";
 						return null;
@@ -262,9 +261,9 @@ namespace UnrealBuildTool
 			{
 				string ProjectFilename = ProjectFile.FullName;
 				FileInfo ProjectFileInfo = new FileInfo(ProjectFilename);
-				if (!ProjectFileInfo.Exists || UBTMakefileInfo.LastWriteTime.CompareTo(ProjectFileInfo.LastWriteTime) < 0)
+				if (!ProjectFileInfo.Exists || MakefileInfo.LastWriteTime.CompareTo(ProjectFileInfo.LastWriteTime) < 0)
 				{
-					// .uproject file is newer than UBTMakefile
+					// .uproject file is newer than makefile
 					Log.TraceLog("Makefile is older than .uproject file, ignoring it");
 					ReasonNotLoaded = ".uproject file is newer";
 					return null;
@@ -275,9 +274,9 @@ namespace UnrealBuildTool
 				if (Directory.Exists(GameIntermediateProjectFilesPath))
 				{
 					DateTime GameProjectFilesLastUpdateTime = new DirectoryInfo(GameIntermediateProjectFilesPath).LastWriteTime;
-					if (UBTMakefileInfo.LastWriteTime.CompareTo(GameProjectFilesLastUpdateTime) < 0)
+					if (MakefileInfo.LastWriteTime.CompareTo(GameProjectFilesLastUpdateTime) < 0)
 					{
-						// Game project files are newer than UBTMakefile
+						// Game project files are newer than makefile
 						Log.TraceLog("Makefile is older than generated game project files, ignoring it");
 						ReasonNotLoaded = "game project files are newer";
 						return null;
@@ -285,11 +284,11 @@ namespace UnrealBuildTool
 				}
 			}
 
-			// Check to see if UnrealBuildTool.exe was compiled more recently than the UBTMakefile
+			// Check to see if UnrealBuildTool.exe was compiled more recently than the makefile
 			DateTime UnrealBuildToolTimestamp = new FileInfo(Assembly.GetExecutingAssembly().Location).LastWriteTime;
-			if (UBTMakefileInfo.LastWriteTime.CompareTo(UnrealBuildToolTimestamp) < 0)
+			if (MakefileInfo.LastWriteTime.CompareTo(UnrealBuildToolTimestamp) < 0)
 			{
-				// UnrealBuildTool.exe was compiled more recently than the UBTMakefile
+				// UnrealBuildTool.exe was compiled more recently than the makefile
 				Log.TraceLog("Makefile is older than UnrealBuildTool.exe, ignoring it");
 				ReasonNotLoaded = "UnrealBuildTool.exe is newer";
 				return null;
@@ -300,7 +299,7 @@ namespace UnrealBuildTool
 			foreach (XmlConfig.InputFile InputFile in InputFiles)
 			{
 				FileInfo InputFileInfo = new FileInfo(InputFile.Location.FullName);
-				if (InputFileInfo.LastWriteTime > UBTMakefileInfo.LastWriteTime)
+				if (InputFileInfo.LastWriteTime > MakefileInfo.LastWriteTime)
 				{
 					Log.TraceLog("Makefile is older than BuildConfiguration.xml, ignoring it");
 					ReasonNotLoaded = "BuildConfiguration.xml is newer";
@@ -308,12 +307,12 @@ namespace UnrealBuildTool
 				}
 			}
 
-			TargetMakefile LoadedUBTMakefile;
+			TargetMakefile Makefile;
 			using(Timeline.ScopeEvent("Loading makefile"))
 			{
 				try
 				{
-					using(BinaryArchiveReader Reader = new BinaryArchiveReader(new FileReference(UBTMakefileInfo)))
+					using(BinaryArchiveReader Reader = new BinaryArchiveReader(new FileReference(MakefileInfo)))
 					{
 						int Version = Reader.ReadInt();
 						if(Version != CurrentVersion)
@@ -321,7 +320,7 @@ namespace UnrealBuildTool
 							ReasonNotLoaded = "makefile version does not match";
 							return null;
 						}
-						LoadedUBTMakefile = new TargetMakefile(Reader);
+						Makefile = new TargetMakefile(Reader);
 					}
 				}
 				catch (Exception Ex)
@@ -336,7 +335,7 @@ namespace UnrealBuildTool
 			using(Timeline.ScopeEvent("Checking makefile validity"))
 			{
 				// Check if the arguments are different
-				if(!Enumerable.SequenceEqual(LoadedUBTMakefile.AdditionalArguments, Arguments))
+				if(!Enumerable.SequenceEqual(Makefile.AdditionalArguments, Arguments))
 				{
 					ReasonNotLoaded = "command line arguments changed";
 					return null;
@@ -349,19 +348,19 @@ namespace UnrealBuildTool
 					foreach (FileReference IniFilename in ConfigHierarchy.EnumerateConfigFileLocations(IniType, ProjectDirectory, Platform))
 					{
 						FileInfo IniFileInfo = new FileInfo(IniFilename.FullName);
-						if (UBTMakefileInfo.LastWriteTime.CompareTo(IniFileInfo.LastWriteTime) < 0)
+						if (MakefileInfo.LastWriteTime.CompareTo(IniFileInfo.LastWriteTime) < 0)
 						{
-							// Ini files are newer than UBTMakefile
-							ReasonNotLoaded = "ini files are newer than UBTMakefile";
+							// Ini files are newer than makefile
+							ReasonNotLoaded = "ini files are newer than makefile";
 							return null;
 						}
 					}
 				}
 
-				foreach(KeyValuePair<DirectoryItem, FileItem[]> Pair in LoadedUBTMakefile.DirectoryToSourceFiles)
+				foreach(KeyValuePair<DirectoryItem, FileItem[]> Pair in Makefile.DirectoryToSourceFiles)
 				{
 					DirectoryItem InputDirectory = Pair.Key;
-					if(!InputDirectory.Exists || InputDirectory.LastWriteTimeUtc > UBTMakefileInfo.LastWriteTimeUtc)
+					if(!InputDirectory.Exists || InputDirectory.LastWriteTimeUtc > MakefileInfo.LastWriteTimeUtc)
 					{
 						FileItem[] SourceFiles = UEBuildModuleCPP.GetSourceFiles(InputDirectory);
 						if(SourceFiles.Length < Pair.Value.Length)
@@ -382,7 +381,7 @@ namespace UnrealBuildTool
 					}
 				}
 
-				foreach(FileItem AdditionalDependency in LoadedUBTMakefile.AdditionalDependencies)
+				foreach(FileItem AdditionalDependency in Makefile.AdditionalDependencies)
 				{
 					if (!AdditionalDependency.Exists)
 					{
@@ -390,7 +389,7 @@ namespace UnrealBuildTool
 						ReasonNotLoaded = string.Format("{0} deleted", AdditionalDependency.Location.GetFileName());
 						return null;
 					}
-					if(AdditionalDependency.LastWriteTimeUtc > UBTMakefileInfo.LastWriteTimeUtc)
+					if(AdditionalDependency.LastWriteTimeUtc > MakefileInfo.LastWriteTimeUtc)
 					{
 						Log.TraceLog("{0} has been modified since makefile was built.", AdditionalDependency.Location);
 						ReasonNotLoaded = string.Format("{0} modified", AdditionalDependency.Location.GetFileName());
@@ -404,11 +403,11 @@ namespace UnrealBuildTool
 
 				// Get all H files in processed modules newer than the makefile itself
 				HashSet<FileItem> HFilesNewerThanMakefile = new HashSet<FileItem>();
-				foreach(UHTModuleHeaderInfo ModuleHeaderInfo in LoadedUBTMakefile.UObjectModuleHeaders)
+				foreach(UHTModuleHeaderInfo ModuleHeaderInfo in Makefile.UObjectModuleHeaders)
 				{
 					foreach(FileItem HeaderFile in ModuleHeaderInfo.SourceFolder.EnumerateFiles())
 					{
-						if(HeaderFile.HasExtension(".h") && HeaderFile.LastWriteTimeUtc > UBTMakefileInfo.LastWriteTimeUtc)
+						if(HeaderFile.HasExtension(".h") && HeaderFile.LastWriteTimeUtc > MakefileInfo.LastWriteTimeUtc)
 						{
 							HFilesNewerThanMakefile.Add(HeaderFile);
 						}
@@ -416,7 +415,7 @@ namespace UnrealBuildTool
 				}
 
 				// Get all H files in all modules processed in the last makefile build
-				HashSet<FileItem> AllUHTHeaders = new HashSet<FileItem>(LoadedUBTMakefile.UObjectModuleHeaders.SelectMany(x => x.HeaderFiles));
+				HashSet<FileItem> AllUHTHeaders = new HashSet<FileItem>(Makefile.UObjectModuleHeaders.SelectMany(x => x.HeaderFiles));
 
 				// Check whether any headers have been deleted. If they have, we need to regenerate the makefile since the module might now be empty. If we don't,
 				// and the file has been moved to a different module, we may include stale generated headers.
@@ -440,7 +439,7 @@ namespace UnrealBuildTool
 					bool bWasProcessed = AllUHTHeaders.Contains(HeaderFile);
 					if (bContainsUHTData != bWasProcessed)
 					{
-						Log.TraceLog("{0} {1} contain UHT types and now {2} , ignoring it ({3})", HeaderFile, bWasProcessed ? "used to" : "didn't", bWasProcessed ? "doesn't" : "does", UBTMakefileInfo.FullName);
+						Log.TraceLog("{0} {1} contain UHT types and now {2} , ignoring it ({3})", HeaderFile, bWasProcessed ? "used to" : "didn't", bWasProcessed ? "doesn't" : "does", MakefileInfo.FullName);
 						ReasonNotLoaded = string.Format("new files with reflected types");
 						return null;
 					}
@@ -453,22 +452,22 @@ namespace UnrealBuildTool
 				// iteration times.)
 
 				// Check if any source files in the working set no longer belong in it
-				foreach (FileItem SourceFile in LoadedUBTMakefile.WorkingSet)
+				foreach (FileItem SourceFile in Makefile.WorkingSet)
 				{
-					if (!WorkingSet.Contains(SourceFile.Location) && SourceFile.LastWriteTimeUtc > UBTMakefileInfo.LastWriteTimeUtc)
+					if (!WorkingSet.Contains(SourceFile.Location) && SourceFile.LastWriteTimeUtc > MakefileInfo.LastWriteTimeUtc)
 					{
-						Log.TraceLog("{0} was part of source working set and now is not; invalidating makefile ({1})", SourceFile.AbsolutePath, UBTMakefileInfo.FullName);
+						Log.TraceLog("{0} was part of source working set and now is not; invalidating makefile ({1})", SourceFile.AbsolutePath, MakefileInfo.FullName);
 						ReasonNotLoaded = string.Format("working set of source files changed");
 						return null;
 					}
 				}
 
 				// Check if any source files that are eligible for being in the working set have been modified
-				foreach (FileItem SourceFile in LoadedUBTMakefile.CandidatesForWorkingSet)
+				foreach (FileItem SourceFile in Makefile.CandidatesForWorkingSet)
 				{
-					if (WorkingSet.Contains(SourceFile.Location) && SourceFile.LastWriteTimeUtc > UBTMakefileInfo.LastWriteTimeUtc)
+					if (WorkingSet.Contains(SourceFile.Location) && SourceFile.LastWriteTimeUtc > MakefileInfo.LastWriteTimeUtc)
 					{
-						Log.TraceLog("{0} was part of source working set and now is not; invalidating makefile ({1})", SourceFile.AbsolutePath, UBTMakefileInfo.FullName);
+						Log.TraceLog("{0} was part of source working set and now is not; invalidating makefile ({1})", SourceFile.AbsolutePath, MakefileInfo.FullName);
 						ReasonNotLoaded = string.Format("working set of source files changed");
 						return null;
 					}
@@ -476,7 +475,7 @@ namespace UnrealBuildTool
 			}
 
 			ReasonNotLoaded = null;
-			return LoadedUBTMakefile;
+			return Makefile;
 		}
 
 		/// <summary>
@@ -490,7 +489,7 @@ namespace UnrealBuildTool
 		public static FileReference GetLocation(FileReference ProjectFile, string TargetName, UnrealTargetPlatform Platform, UnrealTargetConfiguration Configuration)
 		{
 			DirectoryReference BaseDirectory = DirectoryReference.FromFile(ProjectFile) ?? UnrealBuildTool.EngineDirectory;
-			return FileReference.Combine(BaseDirectory, "Intermediate", "Build", Platform.ToString(), TargetName, Configuration.ToString(), "Makefile.ubt");
+			return FileReference.Combine(BaseDirectory, "Intermediate", "Build", Platform.ToString(), TargetName, Configuration.ToString(), "Makefile.bin");
 		}
 	}
 }
