UBT: Allow disabling platform project generation for specific platforms via the BuildConfiguration.xml file, to allow building platforms that don't have Visual Studio integration installed.

Example usage:

<ProjectFileGenerator><DisablePlatformProjectGenerators><Item>Win64</Item><DisablePlatformProjectGenerators><ProjectFileGenerator>

#rnx

diff --git a/Engine/Source/Programs/UnrealBuildTool/Modes/GenerateProjectFilesMode.cs b/Engine/Source/Programs/UnrealBuildTool/Modes/GenerateProjectFilesMode.cs
--- a/Engine/Source/Programs/UnrealBuildTool/Modes/GenerateProjectFilesMode.cs
+++ b/Engine/Source/Programs/UnrealBuildTool/Modes/GenerateProjectFilesMode.cs
@@ -38,6 +38,12 @@ namespace UnrealBuildTool
 		[CommandLine("-CLion", Value = nameof(ProjectFileFormat.CLion))]
 		HashSet<ProjectFileFormat> ProjectFileFormats = new HashSet<ProjectFileFormat>();
 
+		/// <summary>
+		/// Platforms to disable native project file generators for. Platforms with native project file generators typically require IDE extensions to be installed.
+		/// </summary>
+		[XmlConfigFile(Category = "ProjectFileGenerator")]
+		string[] DisablePlatformProjectGenerators = null;
+
 		/// <summary>
 		/// Execute the tool mode
 		/// </summary>
@@ -48,6 +54,9 @@ namespace UnrealBuildTool
 			// Apply any command line arguments to this class
 			Arguments.ApplyTo(this);
 
+			// Apply the XML config to this class
+			XmlConfig.ApplyTo(this);
+
 			// Parse rocket-specific arguments.
 			FileReference ProjectFile;
 			TryParseProjectFileArgument(Arguments, out ProjectFile);
@@ -84,8 +93,11 @@ namespace UnrealBuildTool
 					PlatformProjectGenerator Generator = (PlatformProjectGenerator)Activator.CreateInstance(CheckType, Arguments);
 					foreach(UnrealTargetPlatform Platform in Generator.GetPlatforms())
 					{
-						Log.TraceVerbose("Registering project generator {0} for {1}", CheckType, Platform);
-						PlatformProjectGenerators.RegisterPlatformProjectGenerator(Platform, Generator);
+						if(DisablePlatformProjectGenerators != null && !DisablePlatformProjectGenerators.Any(x => x.Equals(Platform.ToString(), StringComparison.OrdinalIgnoreCase)))
+						{
+							Log.TraceVerbose("Registering project generator {0} for {1}", CheckType, Platform);
+							PlatformProjectGenerators.RegisterPlatformProjectGenerator(Platform, Generator);
+						}
 					}
 				}
 			}
diff --git a/Engine/Source/Programs/UnrealBuildTool/ProjectFiles/VisualStudio/VCProject.cs b/Engine/Source/Programs/UnrealBuildTool/ProjectFiles/VisualStudio/VCProject.cs
--- a/Engine/Source/Programs/UnrealBuildTool/ProjectFiles/VisualStudio/VCProject.cs
+++ b/Engine/Source/Programs/UnrealBuildTool/ProjectFiles/VisualStudio/VCProject.cs
@@ -201,12 +201,6 @@ namespace UnrealBuildTool
 				}
 			}
 
-			PlatformProjectGenerator PlatformProjectGenerator = PlatformProjectGenerators.GetPlatformProjectGenerator(Platform, true);
-			if (PlatformProjectGenerator == null)
-			{
-				return false;
-			}
-
 			UEBuildPlatform BuildPlatform = UEBuildPlatform.GetBuildPlatform(Platform, true);
 			if (BuildPlatform == null)
 			{
@@ -385,7 +379,7 @@ namespace UnrealBuildTool
 			string ProjectConfigurationName = "Invalid";
 
 			// Get the default platform. If there were not valid platforms for this project, just use one that will always be available in VS.
-			string ProjectPlatformName = InvalidConfigPlatformNames[0];
+			string ProjectPlatformName = InvalidConfigPlatformNames.First();
 
 			// Whether the configuration should be built automatically as part of the solution
 			bool bBuildByDefault = false;
@@ -544,7 +538,7 @@ namespace UnrealBuildTool
 			}
 		}
 
-		List<string> InvalidConfigPlatformNames;
+		HashSet<string> InvalidConfigPlatformNames;
 		List<ProjectConfigAndTargetCombination> ProjectConfigAndTargetCombinations;
 
 		private void BuildProjectConfigAndTargetCombinations(List<UnrealTargetPlatform> InPlatforms, List<UnrealTargetConfiguration> InConfigurations, PlatformProjectGeneratorCollection PlatformProjectGenerators)
@@ -604,14 +598,13 @@ namespace UnrealBuildTool
 				}
 
 				// Create a list of platforms for the "invalid" configuration. We always require at least one of these.
-				InvalidConfigPlatformNames = new List<string>();
 				if(ProjectConfigAndTargetCombinations.Count == 0)
 				{
-					InvalidConfigPlatformNames.Add(DefaultPlatformName);
+					InvalidConfigPlatformNames = new HashSet<string>{ DefaultPlatformName };
 				}
 				else
 				{
-					InvalidConfigPlatformNames.AddRange(ProjectConfigAndTargetCombinations.Select(x => x.ProjectPlatformName));
+					InvalidConfigPlatformNames = new HashSet<string>(ProjectConfigAndTargetCombinations.Select(x => x.ProjectPlatformName));
 				}
 			}
 		}
@@ -1223,10 +1216,6 @@ namespace UnrealBuildTool
 		private void WritePostDefaultPropsConfiguration(UnrealTargetPlatform TargetPlatform, UnrealTargetConfiguration TargetConfiguration, string ProjectPlatformName, string ProjectConfigurationName, PlatformProjectGeneratorCollection PlatformProjectGenerators, StringBuilder VCProjectFileContent)
 		{
 			PlatformProjectGenerator ProjGenerator = PlatformProjectGenerators.GetPlatformProjectGenerator(TargetPlatform, true);
-			if (((ProjGenerator == null) && (TargetPlatform != UnrealTargetPlatform.Unknown)))
-			{
-				return;
-			}
 
 			string ProjectConfigurationAndPlatformName = ProjectConfigurationName + "|" + ProjectPlatformName;
 			string ConditionString = "Condition=\"'$(Configuration)|$(Platform)'=='" + ProjectConfigurationAndPlatformName + "'\"";
@@ -1255,10 +1244,6 @@ namespace UnrealBuildTool
 			UnrealTargetConfiguration Configuration = Combination.Configuration;
 
 			PlatformProjectGenerator ProjGenerator = PlatformProjectGenerators.GetPlatformProjectGenerator(Platform, true);
-			if (((ProjGenerator == null) && (Platform != UnrealTargetPlatform.Unknown)))
-			{
-				return;
-			}
 
 			string UProjectPath = "";
 			if (IsForeignProject)
@@ -1434,8 +1419,10 @@ namespace UnrealBuildTool
 
 					VCProjectFileContent.AppendLine("  </PropertyGroup>");
 
-					string LayoutDirString = (ProjGenerator != null) ? ProjGenerator.GetVisualStudioLayoutDirSection(Platform, Configuration, ConditionString, Combination.ProjectTarget.TargetRules.Type, Combination.ProjectTarget.TargetFilePath, ProjectFilePath, NMakePath, ProjectFileFormat) : "";
-					VCProjectFileContent.Append(LayoutDirString);
+					if(ProjGenerator != null)
+					{
+						VCProjectFileContent.Append(ProjGenerator.GetVisualStudioLayoutDirSection(Platform, Configuration, ConditionString, Combination.ProjectTarget.TargetRules.Type, Combination.ProjectTarget.TargetFilePath, ProjectFilePath, NMakePath, ProjectFileFormat));
+					}
 				}
 
 				if (VCUserFileContent != null && Combination.ProjectTarget != null)
@@ -1465,8 +1452,10 @@ namespace UnrealBuildTool
 						VCUserFileContent.AppendLine("  </PropertyGroup>");
 					}
 
-					string PlatformUserFileStrings = (ProjGenerator != null) ? ProjGenerator.GetVisualStudioUserFileStrings(Platform, Configuration, ConditionString, TargetRulesObject, Combination.ProjectTarget.TargetFilePath, ProjectFilePath) : "";
-					VCUserFileContent.Append(PlatformUserFileStrings);
+					if(ProjGenerator != null)
+					{
+						VCUserFileContent.Append(ProjGenerator.GetVisualStudioUserFileStrings(Platform, Configuration, ConditionString, TargetRulesObject, Combination.ProjectTarget.TargetFilePath, ProjectFilePath));
+					}
 				}
 			}
 		}
