UBT: Allow setting the C++ standard version for the compiler on a per-target or per-module basis. Currently only implemented for MSVC.

Example:

    CppStandard = CppStandardVersion.Cpp17;

Or:

    CppStandard = CppStandardVersion.Latest;

diff --git a/Engine/Source/Programs/UnrealBuildTool/Configuration/ModuleRules.cs b/Engine/Source/Programs/UnrealBuildTool/Configuration/ModuleRules.cs
--- a/Engine/Source/Programs/UnrealBuildTool/Configuration/ModuleRules.cs
+++ b/Engine/Source/Programs/UnrealBuildTool/Configuration/ModuleRules.cs
@@ -693,6 +693,11 @@ namespace UnrealBuildTool
 		/// </summary>
 		public bool? bLegacyPublicIncludePaths;
 
+		/// <summary>
+		/// Which stanard to use for compiling this module
+		/// </summary>
+		public CppStandardVersion CppStandard = CppStandardVersion.Default;
+
 		/// <summary>
 		/// The current engine directory
 		/// </summary>
diff --git a/Engine/Source/Programs/UnrealBuildTool/Configuration/TargetRules.cs b/Engine/Source/Programs/UnrealBuildTool/Configuration/TargetRules.cs
--- a/Engine/Source/Programs/UnrealBuildTool/Configuration/TargetRules.cs
+++ b/Engine/Source/Programs/UnrealBuildTool/Configuration/TargetRules.cs
@@ -1025,6 +1025,13 @@ namespace UnrealBuildTool
 		/// </summary>
 		public bool bLegacyPublicIncludePaths = true;
 
+		/// <summary>
+		/// Which C++ stanard to use for compiling this target
+		/// </summary>
+		[RequiresUniqueBuildEnvironment]
+		[XmlConfigFile(Category = "BuildConfiguration")]
+		public CppStandardVersion CppStandard = CppStandardVersion.Default;
+
 		/// <summary>
 		/// Do not allow manifest changes when building this target. Used to cause earlier errors when building multiple targets with a shared build environment.
 		/// </summary>
@@ -2161,6 +2168,11 @@ namespace UnrealBuildTool
 			get { return Inner.bLegacyPublicIncludePaths; }
 		}
 
+		public CppStandardVersion CppStandard
+		{
+			get { return Inner.CppStandard; }
+		}
+
 		internal bool bNoManifestChanges
 		{
 			get { return Inner.bNoManifestChanges; }
diff --git a/Engine/Source/Programs/UnrealBuildTool/Configuration/UEBuildModuleCPP.cs b/Engine/Source/Programs/UnrealBuildTool/Configuration/UEBuildModuleCPP.cs
--- a/Engine/Source/Programs/UnrealBuildTool/Configuration/UEBuildModuleCPP.cs
+++ b/Engine/Source/Programs/UnrealBuildTool/Configuration/UEBuildModuleCPP.cs
@@ -974,6 +974,12 @@ namespace UnrealBuildTool
 			Result.bEnableShadowVariableWarnings = Rules.bEnableShadowVariableWarnings;
 			Result.bEnableUndefinedIdentifierWarnings = Rules.bEnableUndefinedIdentifierWarnings;
 
+			// If the module overrides the C++ language version, override it on the compile environment
+			if(Rules.CppStandard != CppStandardVersion.Default)
+			{
+				Result.CppStandard = Rules.CppStandard;
+			}
+
 			// Set the macro used to check whether monolithic headers can be used
 			if (Rules.bTreatAsEngineModule && (!Rules.bEnforceIWYU || !Target.bEnforceIWYU))
 			{
diff --git a/Engine/Source/Programs/UnrealBuildTool/Configuration/UEBuildTarget.cs b/Engine/Source/Programs/UnrealBuildTool/Configuration/UEBuildTarget.cs
--- a/Engine/Source/Programs/UnrealBuildTool/Configuration/UEBuildTarget.cs
+++ b/Engine/Source/Programs/UnrealBuildTool/Configuration/UEBuildTarget.cs
@@ -2808,6 +2808,7 @@ namespace UnrealBuildTool
 			GlobalCompileEnvironment.bUseRTTI = Rules.bForceEnableRTTI;
 			GlobalCompileEnvironment.bUseInlining = Rules.bUseInlining;
 			GlobalCompileEnvironment.bHideSymbolsByDefault = Rules.bHideSymbolsByDefault;
+			GlobalCompileEnvironment.CppStandard = Rules.CppStandard;
 			GlobalCompileEnvironment.AdditionalArguments = Rules.AdditionalCompilerArguments;
 
 			GlobalLinkEnvironment.bIsBuildingConsoleApplication = Rules.bIsBuildingConsoleApplication;
diff --git a/Engine/Source/Programs/UnrealBuildTool/Platform/Windows/VCToolChain.cs b/Engine/Source/Programs/UnrealBuildTool/Platform/Windows/VCToolChain.cs
--- a/Engine/Source/Programs/UnrealBuildTool/Platform/Windows/VCToolChain.cs
+++ b/Engine/Source/Programs/UnrealBuildTool/Platform/Windows/VCToolChain.cs
@@ -543,6 +543,21 @@ namespace UnrealBuildTool
 			{
 				Arguments.Add("/Qstd=c++14");
 			}
+			else
+			{
+				if(CompileEnvironment.CppStandard >= CppStandardVersion.Latest)
+				{
+					Arguments.Add("/std:c++latest");
+				}
+				else if(CompileEnvironment.CppStandard >= CppStandardVersion.Cpp17)
+				{
+					Arguments.Add("/std:c++17");
+				}
+				else if(CompileEnvironment.CppStandard >= CppStandardVersion.Cpp14)
+				{
+					Arguments.Add("/std:c++14");
+				}
+			}
 
 			if (Target.WindowsPlatform.Compiler == WindowsCompiler.Clang)
 			{
diff --git a/Engine/Source/Programs/UnrealBuildTool/System/CppCompileEnvironment.cs b/Engine/Source/Programs/UnrealBuildTool/System/CppCompileEnvironment.cs
--- a/Engine/Source/Programs/UnrealBuildTool/System/CppCompileEnvironment.cs
+++ b/Engine/Source/Programs/UnrealBuildTool/System/CppCompileEnvironment.cs
@@ -41,6 +41,32 @@ namespace UnrealBuildTool
 		Shipping
 	}
 
+	/// <summary>
+	/// Specifies which language standard to use. This enum should be kept in order, so that toolchains can check whether the requested setting is >= values that they support.
+	/// </summary>
+	public enum CppStandardVersion
+	{
+		/// <summary>
+		/// Use the default standard version
+		/// </summary>
+		Default,
+
+		/// <summary>
+		/// Supports C++14
+		/// </summary>
+		Cpp14,
+
+		/// <summary>
+		/// Supports C++17
+		/// </summary>
+		Cpp17,
+
+		/// <summary>
+		/// Latest standard supported by the compiler
+		/// </summary>
+		Latest,
+	}
+
 	/// <summary>
 	/// The optimization level that may be compilation targets for C# files.
 	/// </summary>
@@ -337,6 +363,11 @@ namespace UnrealBuildTool
 		/// </summary>
 		public bool bHideSymbolsByDefault;
 
+		/// <summary>
+		/// Which C++ standard to support. May not be compatible with all platforms.
+		/// </summary>
+		public CppStandardVersion CppStandard = CppStandardVersion.Default;
+
 		/// <summary>
 		/// Default constructor.
 		/// </summary>
@@ -408,6 +439,7 @@ namespace UnrealBuildTool
 			PrecompiledHeaderFile = Other.PrecompiledHeaderFile;
 			bHackHeaderGenerator = Other.bHackHeaderGenerator;
 			bHideSymbolsByDefault = Other.bHideSymbolsByDefault;
+			CppStandard = Other.CppStandard;
 		}
 	}
 }
