UBT: Remove the project dependency on ShaderCompileWorker being compiled for any project targets. Instead, SCW is built through the same UBT invocation that builds the editor, allowing files to be built in parallel. This also removes the dependency on ShaderCompileWorker being built for non-editor targets (since Visual Studio does not support per-configuration project dependencies).

diff --git a/Engine/Build/BatchFiles/Clean.bat b/Engine/Build/BatchFiles/Clean.bat
--- a/Engine/Build/BatchFiles/Clean.bat
+++ b/Engine/Build/BatchFiles/Clean.bat
@@ -11,7 +11,6 @@ rem ##     %2 is the platform name
 rem ##     %3 is the configuration name
 
 setlocal
-echo Cleaning %1 Binaries...
 
 rem ## First, make sure the batch file exists in the folder we expect it to.  This is necessary in order to
 rem ## verify that our relative path to the /Engine/Source directory is correct
diff --git a/Engine/Source/Programs/UnrealBuildTool/Configuration/UEBuildTarget.cs b/Engine/Source/Programs/UnrealBuildTool/Configuration/UEBuildTarget.cs
--- a/Engine/Source/Programs/UnrealBuildTool/Configuration/UEBuildTarget.cs
+++ b/Engine/Source/Programs/UnrealBuildTool/Configuration/UEBuildTarget.cs
@@ -1097,8 +1097,6 @@ namespace UnrealBuildTool
 			LinkEnvironment GlobalLinkEnvironment = new LinkEnvironment(GlobalCompileEnvironment.Platform, GlobalCompileEnvironment.Configuration, GlobalCompileEnvironment.Architecture);
 
 			UEToolChain TargetToolChain = CreateToolchain(CppPlatform);
-			TargetToolChain.PrintVersionInfo();
-			
 			SetupGlobalEnvironment(TargetToolChain, GlobalCompileEnvironment, GlobalLinkEnvironment);
 
 			// Save off the original list of binaries. We'll use this to figure out which PCHs to create later, to avoid switching PCHs when compiling single modules.
@@ -1197,11 +1195,10 @@ namespace UnrealBuildTool
 			}
 
 			// Create the makefile
-			TargetMakefile Makefile = new TargetMakefile(ReceiptFileName, ProjectIntermediateDirectory, TargetType, bDeployAfterCompile, bHasProjectScriptPlugin);
+			TargetMakefile Makefile = new TargetMakefile(TargetToolChain.GetVersionInfo(), ReceiptFileName, ProjectIntermediateDirectory, TargetType, bDeployAfterCompile, bHasProjectScriptPlugin);
 
 			// Execute the pre-build steps
 			Makefile.PreBuildScripts = CreatePreBuildScripts();
-			Utils.ExecuteCustomBuildSteps(Makefile.PreBuildScripts);
 
 			// If we're compiling monolithic, make sure the executable knows about all referenced modules
 			if (ShouldCompileMonolithic())
diff --git a/Engine/Source/Programs/UnrealBuildTool/Modes/BuildMode.cs b/Engine/Source/Programs/UnrealBuildTool/Modes/BuildMode.cs
--- a/Engine/Source/Programs/UnrealBuildTool/Modes/BuildMode.cs
+++ b/Engine/Source/Programs/UnrealBuildTool/Modes/BuildMode.cs
@@ -12,6 +12,33 @@ using Tools.DotNETCommon;
 
 namespace UnrealBuildTool
 {
+	/// <summary>
+	/// Options controlling how a target is built
+	/// </summary>
+	[Flags]
+	enum BuildOptions
+	{
+		/// <summary>
+		/// Default options
+		/// </summary>
+		None = 0,
+
+		/// <summary>
+		/// Don't output any messages unless we're going to build something
+		/// </summary>
+		Quiet = 1,
+
+		/// <summary>
+		/// Don't build anything, just do target setup and terminate
+		/// </summary>
+		SkipBuild = 2,
+
+		/// <summary>
+		/// Just output a list of XGE actions; don't build anything
+		/// </summary>
+		XGEExport = 4,
+	}
+
 	/// <summary>
 	/// Builds a target
 	/// </summary>
@@ -31,6 +58,18 @@ namespace UnrealBuildTool
 		[CommandLine("-IgnoreJunk")]
 		public bool bIgnoreJunk = false;
 
+		/// <summary>
+		/// Skip building; just do setup and terminate.
+		/// </summary>
+		[CommandLine("-SkipBuild")]
+		public bool bSkipBuild = false;
+
+		/// <summary>
+		/// Whether we should just export the XGE XML and pretend it succeeded
+		/// </summary>
+		[CommandLine("-XGEExport")]
+		public bool bXGEExport = false;
+
 		/// <summary>
 		/// Main entry point
 		/// </summary>
@@ -78,12 +117,6 @@ namespace UnrealBuildTool
 			XmlConfig.ApplyTo(BuildConfiguration);
 			Arguments.ApplyTo(BuildConfiguration);
 
-			// Then let the command lines override any configs necessary.
-			if(BuildConfiguration.bXGEExport)
-			{
-				BuildConfiguration.bAllowXGE = true;
-			}
-
 			// Parse the remote INI setting
 			string RemoteIniPath;
 			Arguments.TryGetValue("-RemoteIni=", out RemoteIniPath);
@@ -145,10 +178,21 @@ namespace UnrealBuildTool
 						}
 					}
 
+					// Get all the build options
+					BuildOptions Options = BuildOptions.None;
+					if(bSkipBuild)
+					{
+						Options |= BuildOptions.SkipBuild;
+					}
+					if(bXGEExport)
+					{
+						Options |= BuildOptions.XGEExport;
+					}
+
 					// Create the working set provider
 					using (ISourceFileWorkingSet WorkingSet = SourceFileWorkingSet.Create(UnrealBuildTool.RootDirectory, ProjectDirs))
 					{
-						Build(TargetDescriptors, BuildConfiguration, WorkingSet);
+						Build(TargetDescriptors, BuildConfiguration, WorkingSet, Options);
 					}
 				}
 			}
@@ -170,8 +214,9 @@ namespace UnrealBuildTool
 		/// <param name="TargetDescriptors">Target descriptors</param>
 		/// <param name="BuildConfiguration">Current build configuration</param>
 		/// <param name="WorkingSet">The source file working set</param>
+		/// <param name="Options">Additional options for the build</param>
 		/// <returns>Result from the compilation</returns>
-		public static void Build(List<TargetDescriptor> TargetDescriptors, BuildConfiguration BuildConfiguration, ISourceFileWorkingSet WorkingSet)
+		public static void Build(List<TargetDescriptor> TargetDescriptors, BuildConfiguration BuildConfiguration, ISourceFileWorkingSet WorkingSet, BuildOptions Options)
 		{
 			// Create a makefile for each target
 			TargetMakefile[] Makefiles = new TargetMakefile[TargetDescriptors.Count];
@@ -181,11 +226,17 @@ namespace UnrealBuildTool
 			}
 
 			// Execute the build
-			if(!BuildConfiguration.bSkipBuild)
+			if((Options & BuildOptions.SkipBuild) == 0)
 			{
+				// Execute all the pre-build steps
+				foreach(TargetMakefile Makefile in Makefiles)
+				{
+					Utils.ExecuteCustomBuildSteps(Makefile.PreBuildScripts);
+				}
+
 				// Make sure that none of the actions conflict with any other (producing output files differently, etc...)
 				ActionGraph.CheckForConflicts(Makefiles.SelectMany(x => x.Actions));
-				
+
 				// Find all the actions to be executed
 				HashSet<Action>[] ActionsToExecute = new HashSet<Action>[TargetDescriptors.Count];
 				for(int TargetIdx = 0; TargetIdx < TargetDescriptors.Count; TargetIdx++)
@@ -227,7 +278,7 @@ namespace UnrealBuildTool
 				ActionGraph.CreateDirectoriesForProducedItems(MergedActionsToExecute);
 
 				// Execute the actions
-				if (BuildConfiguration.bXGEExport)
+				if ((Options & BuildOptions.XGEExport) != 0)
 				{
 					// Just export to an XML file
 					using(Timeline.ScopeEvent("XGE.ExportActions()"))
@@ -238,9 +289,26 @@ namespace UnrealBuildTool
 				else
 				{
 					// Execute the actions
-					using(Timeline.ScopeEvent("ActionGraph.ExecuteActions()"))
+					if(MergedActionsToExecute.Count == 0)
 					{
-						ActionGraph.ExecuteActions(BuildConfiguration, MergedActionsToExecute);
+						if((Options & BuildOptions.Quiet) == 0)
+						{
+							Log.TraceInformation((TargetDescriptors.Count == 1)? "Target is up to date" : "Targets are up to date");
+						}
+					}
+					else
+					{
+						if((Options & BuildOptions.Quiet) != 0)
+						{
+							Log.TraceInformation("Building {0}...", StringUtils.FormatList(TargetDescriptors.Select(x => x.Name).Distinct()));
+						}
+
+						OutputToolchainInfo(TargetDescriptors, Makefiles);
+
+						using(Timeline.ScopeEvent("ActionGraph.ExecuteActions()"))
+						{
+							ActionGraph.ExecuteActions(BuildConfiguration, MergedActionsToExecute);
+						}
 					}
 
 					// Run the deployment steps
@@ -257,6 +325,30 @@ namespace UnrealBuildTool
 			}
 		}
 
+		/// <summary>
+		/// Outputs the toolchain used to build each target
+		/// </summary>
+		/// <param name="TargetDescriptors">List of targets being built</param>
+		/// <param name="Makefiles">Matching array of makefiles for each target</param>
+		static void OutputToolchainInfo(List<TargetDescriptor> TargetDescriptors, TargetMakefile[] Makefiles)
+		{
+			List<string> UniqueStrings = new List<string>(Makefiles.Select(x => x.ToolchainInfo).Where(x => x != null).Distinct());
+			if(UniqueStrings.Count == 1)
+			{
+				Log.TraceInformation("{0}", UniqueStrings[0]);
+			}
+			else
+			{
+				for(int Idx = 0; Idx < TargetDescriptors.Count; Idx++)
+				{
+					if(Makefiles[Idx].ToolchainInfo != null)
+					{
+						Log.TraceInformation("For {0}: {1}", TargetDescriptors[Idx], Makefiles[Idx].ToolchainInfo);
+					}
+				}
+			}
+		}
+
 		/// <summary>
 		/// Creates the makefile for a target. If an existing, valid makefile already exists on disk, loads that instead.
 		/// </summary>
@@ -332,12 +424,6 @@ namespace UnrealBuildTool
 					Environment.SetEnvironmentVariable(EnvironmentVariable.Item1, EnvironmentVariable.Item2);
 				}
 
-				// Execute all the pre-build steps
-				if (!BuildConfiguration.bXGEExport)
-				{
-					Utils.ExecuteCustomBuildSteps(Makefile.PreBuildScripts);
-				}
-
 				// If the target needs UHT to be run, we'll go ahead and do that now
 				if (Makefile.UObjectModules.Count > 0)
 				{
diff --git a/Engine/Source/Programs/UnrealBuildTool/Modes/CleanMode.cs b/Engine/Source/Programs/UnrealBuildTool/Modes/CleanMode.cs
--- a/Engine/Source/Programs/UnrealBuildTool/Modes/CleanMode.cs
+++ b/Engine/Source/Programs/UnrealBuildTool/Modes/CleanMode.cs
@@ -79,6 +79,9 @@ namespace UnrealBuildTool
 				}
 			}
 
+			// Output the list of targets that we're cleaning
+			Log.TraceInformation("Cleaning {0} binaries...", StringUtils.FormatList(TargetDescriptors.Select(x => x.Name).Distinct()));
+
 			// Loop through all the targets, and clean them all
 			HashSet<FileReference> FilesToDelete = new HashSet<FileReference>();
 			HashSet<DirectoryReference> DirectoriesToDelete = new HashSet<DirectoryReference>();
diff --git a/Engine/Source/Programs/UnrealBuildTool/Platform/Windows/VCToolChain.cs b/Engine/Source/Programs/UnrealBuildTool/Platform/Windows/VCToolChain.cs
--- a/Engine/Source/Programs/UnrealBuildTool/Platform/Windows/VCToolChain.cs
+++ b/Engine/Source/Programs/UnrealBuildTool/Platform/Windows/VCToolChain.cs
@@ -30,6 +30,11 @@ namespace UnrealBuildTool
 			this.Target = Target;
 			this.EnvVars = Target.WindowsPlatform.Environment;
 
+			Log.TraceLog("Compiler: {0}", EnvVars.CompilerPath);
+			Log.TraceLog("Linker: {0}", EnvVars.LinkerPath);
+			Log.TraceLog("Library Manager: {0}", EnvVars.LibraryManagerPath);
+			Log.TraceLog("Resource Compiler: {0}", EnvVars.ResourceCompilerPath);
+
 			if (Target.WindowsPlatform.ObjSrcMapFile != null)
 			{
 				try
@@ -42,20 +47,20 @@ namespace UnrealBuildTool
 			}
 		}
 
-		public override void PrintVersionInfo()
+		/// <summary>
+		/// Returns the version info for the toolchain. This will be output before building.
+		/// </summary>
+		/// <returns>String describing the current toolchain</returns>
+		public override string GetVersionInfo()
 		{
 			if(EnvVars.Compiler == EnvVars.ToolChain)
 			{
-				Log.TraceInformation("Using {0} {1} toolchain ({2}) and Windows {3} SDK ({4}).", WindowsPlatform.GetCompilerName(EnvVars.Compiler), EnvVars.ToolChainVersion, EnvVars.ToolChainDir, EnvVars.WindowsSdkVersion, EnvVars.WindowsSdkDir);
+				return String.Format("Using {0} {1} toolchain ({2}) and Windows {3} SDK ({4}).", WindowsPlatform.GetCompilerName(EnvVars.Compiler), EnvVars.ToolChainVersion, EnvVars.ToolChainDir, EnvVars.WindowsSdkVersion, EnvVars.WindowsSdkDir);
 			}
 			else
 			{
-				Log.TraceInformation("Using {0} {1} compiler ({2}) with {3} {4} runtime ({5}) and Windows {6} SDK ({7}).", WindowsPlatform.GetCompilerName(EnvVars.Compiler), EnvVars.CompilerVersion, EnvVars.CompilerDir, WindowsPlatform.GetCompilerName(EnvVars.ToolChain), EnvVars.ToolChainVersion, EnvVars.ToolChainDir, EnvVars.WindowsSdkVersion, EnvVars.WindowsSdkDir);
+				return String.Format("Using {0} {1} compiler ({2}) with {3} {4} runtime ({5}) and Windows {6} SDK ({7}).", WindowsPlatform.GetCompilerName(EnvVars.Compiler), EnvVars.CompilerVersion, EnvVars.CompilerDir, WindowsPlatform.GetCompilerName(EnvVars.ToolChain), EnvVars.ToolChainVersion, EnvVars.ToolChainDir, EnvVars.WindowsSdkVersion, EnvVars.WindowsSdkDir);
 			}
-			Log.TraceLog("Compiler: {0}", EnvVars.CompilerPath);
-			Log.TraceLog("Linker: {0}", EnvVars.LinkerPath);
-			Log.TraceLog("Library Manager: {0}", EnvVars.LibraryManagerPath);
-			Log.TraceLog("Resource Compiler: {0}", EnvVars.ResourceCompilerPath);
 		}
 
 		static void AddDefinition(List<string> Arguments, string Definition)
@@ -72,7 +77,6 @@ namespace UnrealBuildTool
 			}
 		}
 
-
 		static void AddDefinition(List<string> Arguments, string Variable, string Value)
 		{
 			// If the value has a space in it and isn't wrapped in quotes, do that now
diff --git a/Engine/Source/Programs/UnrealBuildTool/ProjectFiles/ProjectFileGenerator.cs b/Engine/Source/Programs/UnrealBuildTool/ProjectFiles/ProjectFileGenerator.cs
--- a/Engine/Source/Programs/UnrealBuildTool/ProjectFiles/ProjectFileGenerator.cs
+++ b/Engine/Source/Programs/UnrealBuildTool/ProjectFiles/ProjectFileGenerator.cs
@@ -2256,17 +2256,11 @@ namespace UnrealBuildTool
 							CreateRulesDelegate = (Platform, Configuration) => RulesAssembly.CreateTargetRules(TargetName, Platform, Configuration, "", CheckProjectFile, null)
                         };
 
-					if (TargetName == "ShaderCompileWorker")		// @todo projectfiles: Ideally, the target rules file should set this
-					{
-						ProjectTarget.ForceDevelopmentConfiguration = true;
-					}
-
 					ProjectFile.ProjectTargets.Add(ProjectTarget);
 
 					// Make sure the *.Target.cs file is in the project.
 					ProjectFile.AddFileToProject(TargetFilePath, BaseFolder);
 
-					// We special case ShaderCompileWorker.  It needs to always be compiled in Development mode.
 					Log.TraceVerbose("Generating target {0} for {1}", TargetRulesObject.Type.ToString(), ProjectFilePath);
 				}
 			}
diff --git a/Engine/Source/Programs/UnrealBuildTool/ProjectFiles/VisualStudio/VCProject.cs b/Engine/Source/Programs/UnrealBuildTool/ProjectFiles/VisualStudio/VCProject.cs
--- a/Engine/Source/Programs/UnrealBuildTool/ProjectFiles/VisualStudio/VCProject.cs
+++ b/Engine/Source/Programs/UnrealBuildTool/ProjectFiles/VisualStudio/VCProject.cs
@@ -281,6 +281,7 @@ namespace UnrealBuildTool
 		bool bUseFastPDB;
 		bool bUsePerFileIntellisense;
 		bool bUsePrecompiled;
+		bool bEditorDependsOnShaderCompileWorker;
 		string BuildToolOverride;
 
 		/// This is the platform name that Visual Studio is always guaranteed to support.  We'll use this as
@@ -303,8 +304,9 @@ namespace UnrealBuildTool
 		/// <param name="bUseFastPDB">If true, adds the -FastPDB argument to build command lines</param>
 		/// <param name="bUsePerFileIntellisense">If true, generates per-file intellisense data</param>
 		/// <param name="bUsePrecompiled">Whether to add the -UsePrecompiled argumemnt when building targets</param>
+		/// <param name="bEditorDependsOnShaderCompileWorker">Whether editor targets should also build ShaderCompileWorker</param>
 		/// <param name="BuildToolOverride">Optional arguments to pass to UBT when building</param>
-		public VCProjectFile(FileReference InFilePath, FileReference InOnlyGameProject, VCProjectFileFormat InProjectFileFormat, bool bUseFastPDB, bool bUsePerFileIntellisense, bool bUsePrecompiled, string BuildToolOverride)
+		public VCProjectFile(FileReference InFilePath, FileReference InOnlyGameProject, VCProjectFileFormat InProjectFileFormat, bool bUseFastPDB, bool bUsePerFileIntellisense, bool bUsePrecompiled, bool bEditorDependsOnShaderCompileWorker, string BuildToolOverride)
 			: base(InFilePath)
 		{
 			OnlyGameProject = InOnlyGameProject;
@@ -312,6 +314,7 @@ namespace UnrealBuildTool
 			this.bUseFastPDB = bUseFastPDB;
 			this.bUsePerFileIntellisense = bUsePerFileIntellisense;
 			this.bUsePrecompiled = bUsePrecompiled;
+			this.bEditorDependsOnShaderCompileWorker = bEditorDependsOnShaderCompileWorker;
 			this.BuildToolOverride = BuildToolOverride;
 		}
 
@@ -1376,39 +1379,48 @@ namespace UnrealBuildTool
 					//	..\..\Build\BatchFiles\Build.bat <TARGETNAME> <PLATFORM> <CONFIGURATION>
 					//	ie ..\..\Build\BatchFiles\Build.bat BlankProgram Win64 Debug
 
-					string BuildArguments = " " + TargetName + " " + UBTPlatformName + " " + UBTConfigurationName;
+					StringBuilder BuildArguments = new StringBuilder();
+
+					BuildArguments.AppendFormat("{0} {1} {2}", TargetName, UBTPlatformName, UBTConfigurationName);
+					if (IsForeignProject)
+					{
+						BuildArguments.AppendFormat(" -Project=\"{0}\"", UProjectPath);
+					}
+
 					if (bUsePrecompiled)
 					{
-						BuildArguments += " -UsePrecompiled";
+						BuildArguments.Append(" -UsePrecompiled");
 					}
-					if (IsForeignProject)
+					else if(TargetRulesObject.Type == TargetType.Editor && bEditorDependsOnShaderCompileWorker)
 					{
-						BuildArguments += " " + UProjectPath;
+						BuildArguments.Replace("\"", "\\\"");
+						BuildArguments.Insert(0, "-Target=\"ShaderCompileWorker Win64 Development\" -Target=\"");
+						BuildArguments.Append("\"");
 					}
 
 					// Always wait for the mutex between UBT invocations, so that building the whole solution doesn't fail.
-					BuildArguments += " -WaitMutex";
+					BuildArguments.Append(" -WaitMutex");
 
 					// Always include a flag to format log messages for MSBuild
-					BuildArguments += " -FromMsBuild";
+					BuildArguments.Append(" -FromMsBuild");
 
 					if (bUseFastPDB)
 					{
 						// Pass Fast PDB option to make use of Visual Studio's /DEBUG:FASTLINK option
-						BuildArguments += " -FastPDB";
+						BuildArguments.Append(" -FastPDB");
 					}
 
 					DirectoryReference BatchFilesDirectory = DirectoryReference.Combine(UnrealBuildTool.EngineDirectory, "Build", "BatchFiles");
 
 					if(BuildToolOverride != null)
 					{
-						BuildArguments += " " + BuildToolOverride;
+						BuildArguments.AppendFormat(" {0}", BuildToolOverride);
 					}
 
 					// NMake Build command line
-					VCProjectFileContent.AppendLine("    <NMakeBuildCommandLine>{0}{1}</NMakeBuildCommandLine>", EscapePath(NormalizeProjectPath(FileReference.Combine(BatchFilesDirectory, "Build.bat"))), BuildArguments.ToString());
-					VCProjectFileContent.AppendLine("    <NMakeReBuildCommandLine>{0}{1}</NMakeReBuildCommandLine>", EscapePath(NormalizeProjectPath(FileReference.Combine(BatchFilesDirectory, "Rebuild.bat"))), BuildArguments.ToString());
-					VCProjectFileContent.AppendLine("    <NMakeCleanCommandLine>{0}{1}</NMakeCleanCommandLine>", EscapePath(NormalizeProjectPath(FileReference.Combine(BatchFilesDirectory, "Clean.bat"))), BuildArguments.ToString());
+					VCProjectFileContent.AppendLine("    <NMakeBuildCommandLine>{0} {1}</NMakeBuildCommandLine>", EscapePath(NormalizeProjectPath(FileReference.Combine(BatchFilesDirectory, "Build.bat"))), BuildArguments.ToString());
+					VCProjectFileContent.AppendLine("    <NMakeReBuildCommandLine>{0} {1}</NMakeReBuildCommandLine>", EscapePath(NormalizeProjectPath(FileReference.Combine(BatchFilesDirectory, "Rebuild.bat"))), BuildArguments.ToString());
+					VCProjectFileContent.AppendLine("    <NMakeCleanCommandLine>{0} {1}</NMakeCleanCommandLine>", EscapePath(NormalizeProjectPath(FileReference.Combine(BatchFilesDirectory, "Clean.bat"))), BuildArguments.ToString());
 					VCProjectFileContent.AppendLine("    <NMakeOutput>{0}</NMakeOutput>", NormalizeProjectPath(NMakePath.FullName));
 
 					if (TargetRulesObject.Type == TargetType.Game || TargetRulesObject.Type == TargetType.Client || TargetRulesObject.Type == TargetType.Server)
diff --git a/Engine/Source/Programs/UnrealBuildTool/ProjectFiles/VisualStudio/VCProjectFileGenerator.cs b/Engine/Source/Programs/UnrealBuildTool/ProjectFiles/VisualStudio/VCProjectFileGenerator.cs
--- a/Engine/Source/Programs/UnrealBuildTool/ProjectFiles/VisualStudio/VCProjectFileGenerator.cs
+++ b/Engine/Source/Programs/UnrealBuildTool/ProjectFiles/VisualStudio/VCProjectFileGenerator.cs
@@ -179,7 +179,7 @@ namespace UnrealBuildTool
 		/// <returns>The newly allocated project file object</returns>
 		protected override ProjectFile AllocateProjectFile(FileReference InitFilePath)
 		{
-			return new VCProjectFile(InitFilePath, OnlyGameProject, ProjectFileFormat, bAddFastPDBToProjects, bUsePerFileIntellisense, bUsePrecompiled, BuildToolOverride);
+			return new VCProjectFile(InitFilePath, OnlyGameProject, ProjectFileFormat, bAddFastPDBToProjects, bUsePerFileIntellisense, bUsePrecompiled, bEditorDependsOnShaderCompileWorker, BuildToolOverride);
 		}
 
 
@@ -461,34 +461,6 @@ namespace UnrealBuildTool
 				throw new BuildException("Unexpected ProjectFileFormat");
 			}
 
-			// Find the projects for ShaderCompileWorker and UnrealLightmass
-			ProjectFile ShaderCompileWorkerProject = null;
-			ProjectFile UnrealLightmassProject = null;
-			foreach (ProjectFile Project in AllProjectFiles)
-			{
-				if (Project.ProjectTargets.Count == 1)
-				{
-					FileReference TargetFilePath = Project.ProjectTargets[0].TargetFilePath;
-					if (TargetFilePath != null)
-					{
-						string TargetFileName = TargetFilePath.GetFileNameWithoutAnyExtensions();
-						if (TargetFileName.Equals("ShaderCompileWorker", StringComparison.InvariantCultureIgnoreCase))
-						{
-							ShaderCompileWorkerProject = Project;
-						}
-						else if (TargetFileName.Equals("UnrealLightmass", StringComparison.InvariantCultureIgnoreCase))
-						{
-							UnrealLightmassProject = Project;
-						}
-					}
-					if (ShaderCompileWorkerProject != null
-						&& UnrealLightmassProject != null)
-					{
-						break;
-					}
-				}
-			}
-
 			// Solution folders, files and project entries
 			{
 				// This the GUID that Visual Studio uses to identify a solution folder
@@ -555,10 +527,6 @@ namespace UnrealBuildTool
 							Dependencies.Add(UBTProject);
 							Dependencies.AddRange(UBTProject.DependsOnProjects);
 						}
-						if (bEditorDependsOnShaderCompileWorker && !bUsePrecompiled && CurProject.IsGeneratedProject && ShaderCompileWorkerProject != null && CurProject.ProjectTargets.Any(x => x.TargetRules != null && x.TargetRules.Type == TargetType.Editor))
-						{
-							Dependencies.Add(ShaderCompileWorkerProject);
-						}
 						Dependencies.AddRange(CurProject.DependsOnProjects);
 
 						if (Dependencies.Count > 0)
diff --git a/Engine/Source/Programs/UnrealBuildTool/System/ExternalExecution.cs b/Engine/Source/Programs/UnrealBuildTool/System/ExternalExecution.cs
--- a/Engine/Source/Programs/UnrealBuildTool/System/ExternalExecution.cs
+++ b/Engine/Source/Programs/UnrealBuildTool/System/ExternalExecution.cs
@@ -1114,7 +1114,6 @@ namespace UnrealBuildTool
 					{
 						// If it is out of date or not there it will be built.
 						// If it is there and up to date, it will add 0.8 seconds to the build time.
-						Log.TraceInformation("Building UnrealHeaderTool...");
 
 						// Which desktop platform do we need to compile UHT for?
 						UnrealTargetPlatform Platform = BuildHostPlatform.Current.Platform;
@@ -1146,12 +1145,7 @@ namespace UnrealBuildTool
 
 						using(Timeline.ScopeEvent("Buildng UnrealHeaderTool"))
 						{
-							bool bPrevXGEExport = BuildConfiguration.bXGEExport;
-							BuildConfiguration.bXGEExport = false;
-
-							BuildMode.Build(TargetDescriptors, BuildConfiguration, WorkingSet);
-
-							BuildConfiguration.bXGEExport = bPrevXGEExport;
+							BuildMode.Build(TargetDescriptors, BuildConfiguration, WorkingSet, BuildOptions.Quiet);
 						}
 					}
 
diff --git a/Engine/Source/Programs/UnrealBuildTool/System/TargetMakefile.cs b/Engine/Source/Programs/UnrealBuildTool/System/TargetMakefile.cs
--- a/Engine/Source/Programs/UnrealBuildTool/System/TargetMakefile.cs
+++ b/Engine/Source/Programs/UnrealBuildTool/System/TargetMakefile.cs
@@ -22,6 +22,11 @@ namespace UnrealBuildTool
 		/// </summary>
 		public const int CurrentVersion = 9;
 
+		/// <summary>
+		/// Information about the toolchain used to build. This string will be output before building.
+		/// </summary>
+		public string ToolchainInfo;
+
 		/// <summary>
 		/// Path to the receipt file for this target
 		/// </summary>
@@ -116,13 +121,15 @@ namespace UnrealBuildTool
 		/// <summary>
 		/// Constructor
 		/// </summary>
+		/// <param name="ToolchainInfo">String describing the toolchain used to build. This will be output before executing actions.</param>
 		/// <param name="ReceiptFile">Path to the receipt file</param>
 		/// <param name="ProjectIntermediateDirectory">Path to the project intermediate directory</param>
 		/// <param name="TargetType">The type of target</param>
 		/// <param name="bDeployAfterCompile">Whether to deploy the target after compiling</param>
 		/// <param name="bHasProjectScriptPlugin">Whether the target has a project script plugin</param>
-		public TargetMakefile(FileReference ReceiptFile, DirectoryReference ProjectIntermediateDirectory, TargetType TargetType, bool bDeployAfterCompile, bool bHasProjectScriptPlugin)
+		public TargetMakefile(string ToolchainInfo, FileReference ReceiptFile, DirectoryReference ProjectIntermediateDirectory, TargetType TargetType, bool bDeployAfterCompile, bool bHasProjectScriptPlugin)
 		{
+			this.ToolchainInfo = ToolchainInfo;
 			this.ReceiptFile = ReceiptFile;
 			this.ProjectIntermediateDirectory = ProjectIntermediateDirectory;
 			this.TargetType = TargetType;
@@ -146,6 +153,7 @@ namespace UnrealBuildTool
 		/// <param name="Reader">The archive to read from</param>
 		public TargetMakefile(BinaryArchiveReader Reader)
 		{
+			ToolchainInfo = Reader.ReadString();
 			ReceiptFile = Reader.ReadFileReference();
 			ProjectIntermediateDirectory = Reader.ReadDirectoryReference();
 			TargetType = (TargetType)Reader.ReadInt();
@@ -172,6 +180,7 @@ namespace UnrealBuildTool
 		/// <param name="Writer">The archive to write to</param>
 		public void Write(BinaryArchiveWriter Writer)
 		{
+			Writer.WriteString(ToolchainInfo);
 			Writer.WriteFileReference(ReceiptFile);
 			Writer.WriteDirectoryReference(ProjectIntermediateDirectory);
 			Writer.WriteInt((int)TargetType);
diff --git a/Engine/Source/Programs/UnrealBuildTool/ToolChain/UEToolChain.cs b/Engine/Source/Programs/UnrealBuildTool/ToolChain/UEToolChain.cs
--- a/Engine/Source/Programs/UnrealBuildTool/ToolChain/UEToolChain.cs
+++ b/Engine/Source/Programs/UnrealBuildTool/ToolChain/UEToolChain.cs
@@ -22,8 +22,9 @@ namespace UnrealBuildTool
 			CppPlatform = InCppPlatform;
 		}
 
-		public virtual void PrintVersionInfo()
+		public virtual string GetVersionInfo()
 		{
+			return null;
 		}
 
 		public abstract CPPOutput CompileCPPFiles(CppCompileEnvironment CompileEnvironment, List<FileItem> InputFiles, DirectoryReference OutputDir, string ModuleName, List<Action> Actions);
