Missing files.

#rnx

diff --git a/Engine/Source/Programs/UnrealBuildTool/Configuration/BuildConfiguration.cs b/Engine/Source/Programs/UnrealBuildTool/Configuration/BuildConfiguration.cs
--- a/Engine/Source/Programs/UnrealBuildTool/Configuration/BuildConfiguration.cs
+++ b/Engine/Source/Programs/UnrealBuildTool/Configuration/BuildConfiguration.cs
@@ -48,24 +48,12 @@ namespace UnrealBuildTool
 		[CommandLine("-NoXGE", Value = "false")]
 		public bool bAllowXGE = true;
 
-		/// <summary>
-		/// Whether we should just export the XGE XML and pretend it succeeded
-		/// </summary>
-		[CommandLine("-XGEExport")]
-		public bool bXGEExport = false;
-
 		/// <summary>
 		/// Whether SN-DBS may be used.
 		/// </summary>
 		[XmlConfigFile]
 		public bool bAllowSNDBS = true;
 
-		/// <summary>
-		/// Skip building; just do setup and terminate.
-		/// </summary>
-		[CommandLine("-SkipBuild")]
-		public bool bSkipBuild = false;
-
 		/// <summary>
 		/// Enables support for very fast iterative builds by caching target data.  Turning this on causes Unreal Build Tool to emit
 		/// 'UBT Makefiles' for targets when they're built the first time.  Subsequent builds will load these Makefiles and begin
diff --git a/Engine/Source/Programs/UnrealBuildTool/Executors/XGE.cs b/Engine/Source/Programs/UnrealBuildTool/Executors/XGE.cs
--- a/Engine/Source/Programs/UnrealBuildTool/Executors/XGE.cs
+++ b/Engine/Source/Programs/UnrealBuildTool/Executors/XGE.cs
@@ -341,7 +341,7 @@ namespace UnrealBuildTool
 				}
 				if(Action.GroupNames.Count > 0)
 				{
-					ToolElement.SetAttribute("GroupPrefix", String.Format("** For {0}", String.Join(" + ", Action.GroupNames)));
+					ToolElement.SetAttribute("GroupPrefix", String.Format("** For {0} **", String.Join(" + ", Action.GroupNames)));
 				}
 
 				ToolElement.SetAttribute("Params", Action.CommandArguments);
