UBT: Fix exception on case-sensitive file systems when a folder contains multiple directories with differing case.

#rnx

diff --git a/Engine/Source/Programs/UnrealBuildTool/System/DirectoryItem.cs b/Engine/Source/Programs/UnrealBuildTool/System/DirectoryItem.cs
--- a/Engine/Source/Programs/UnrealBuildTool/System/DirectoryItem.cs
+++ b/Engine/Source/Programs/UnrealBuildTool/System/DirectoryItem.cs
@@ -191,7 +191,7 @@ namespace UnrealBuildTool
 						}
 						else
 						{
-							NewDirectories.Add(SubDirectoryInfo.Name, DirectoryItem.GetItemByDirectoryInfo(SubDirectoryInfo));
+							NewDirectories[SubDirectoryInfo.Name] = DirectoryItem.GetItemByDirectoryInfo(SubDirectoryInfo);
 						}
 					}
 				}
@@ -247,7 +247,7 @@ namespace UnrealBuildTool
 				{
 					foreach(FileInfo FileInfo in Info.EnumerateFiles())
 					{
-						NewFiles.Add(FileInfo.Name, FileItem.GetItemByFileInfo(FileInfo));
+						NewFiles[FileInfo.Name] = FileItem.GetItemByFileInfo(FileInfo);
 					}
 				}
 				Files = NewFiles;
