UBT: Invalidate the makefile if any plugin is added or removed. It may be enabled by default.

diff --git a/Engine/Source/Programs/UnrealBuildTool/Configuration/UEBuildTarget.cs b/Engine/Source/Programs/UnrealBuildTool/Configuration/UEBuildTarget.cs
--- a/Engine/Source/Programs/UnrealBuildTool/Configuration/UEBuildTarget.cs
+++ b/Engine/Source/Programs/UnrealBuildTool/Configuration/UEBuildTarget.cs
@@ -1485,6 +1485,13 @@ namespace UnrealBuildTool
 				CheckForEULAViolation();
 			}
 
+			// Add all the plugins to be tracked
+			foreach(FileReference PluginFile in global::UnrealBuildTool.Plugins.EnumeratePlugins(ProjectFile))
+			{
+				FileItem PluginFileItem = FileItem.GetItemByFileReference(PluginFile);
+				Makefile.PluginFiles.Add(PluginFileItem);
+			}
+
 			// Add all the input files to the predicate store
 			Makefile.AdditionalDependencies.Add(FileItem.GetItemByFileReference(TargetRulesFile));
 			foreach(UEBuildModule Module in Modules.Values)
@@ -1496,6 +1503,7 @@ namespace UnrealBuildTool
 					Makefile.AdditionalDependencies.Add(FileItem.GetItemByFileReference(Location));
 				}
 			}
+			Makefile.AdditionalDependencies.UnionWith(Makefile.PluginFiles);
 
 			// Clean any stale modules which exist in multiple output directories. This can lead to the wrong DLL being loaded on Windows.
 			CleanStaleModules();
diff --git a/Engine/Source/Programs/UnrealBuildTool/System/TargetMakefile.cs b/Engine/Source/Programs/UnrealBuildTool/System/TargetMakefile.cs
--- a/Engine/Source/Programs/UnrealBuildTool/System/TargetMakefile.cs
+++ b/Engine/Source/Programs/UnrealBuildTool/System/TargetMakefile.cs
@@ -113,6 +113,11 @@ namespace UnrealBuildTool
 		/// </summary>
 		public List<UHTModuleHeaderInfo> UObjectModuleHeaders = new List<UHTModuleHeaderInfo>();
 
+		/// <summary>
+		/// List of all plugin names. The makefile will be considered invalid if any of these changes, or new plugins are added.
+		/// </summary>
+		public HashSet<FileItem> PluginFiles;
+
 		/// <summary>
 		/// Additional files which are required 
 		/// </summary>
@@ -144,6 +149,7 @@ namespace UnrealBuildTool
 			this.CandidatesForWorkingSet = new HashSet<FileItem>();
 			this.UObjectModules = new List<UHTModuleInfo>();
 			this.UObjectModuleHeaders = new List<UHTModuleHeaderInfo>();
+			this.PluginFiles = new HashSet<FileItem>();
 			this.AdditionalDependencies = new HashSet<FileItem>();
 		}
 
@@ -171,6 +177,7 @@ namespace UnrealBuildTool
 			CandidatesForWorkingSet = Reader.ReadHashSet(() => Reader.ReadFileItem());
 			UObjectModules = Reader.ReadList(() => new UHTModuleInfo(Reader));
 			UObjectModuleHeaders = Reader.ReadList(() => new UHTModuleHeaderInfo(Reader));
+			PluginFiles = Reader.ReadHashSet(() => Reader.ReadFileItem());
 			AdditionalDependencies = Reader.ReadHashSet(() => Reader.ReadFileItem());
 		}
 
@@ -198,6 +205,7 @@ namespace UnrealBuildTool
 			Writer.WriteHashSet(CandidatesForWorkingSet, x => Writer.WriteFileItem(x));
 			Writer.WriteList(UObjectModules, e => e.Write(Writer));
 			Writer.WriteList(UObjectModuleHeaders, x => x.Write(Writer));
+			Writer.WriteHashSet(PluginFiles, x => Writer.WriteFileItem(x));
 			Writer.WriteHashSet(AdditionalDependencies, x => Writer.WriteFileItem(x));
 		}
 
@@ -390,6 +398,7 @@ namespace UnrealBuildTool
 					}
 				}
 
+				// Check if any of the additional dependencies has changed
 				foreach(FileItem AdditionalDependency in Makefile.AdditionalDependencies)
 				{
 					if (!AdditionalDependency.Exists)
@@ -406,6 +415,18 @@ namespace UnrealBuildTool
 					}
 				}
 
+				// Check that no new plugins have been added
+				foreach(FileReference PluginFile in Plugins.EnumeratePlugins(ProjectFile))
+				{
+					FileItem PluginFileItem = FileItem.GetItemByFileReference(PluginFile);
+					if(!Makefile.PluginFiles.Contains(PluginFileItem))
+					{
+						Log.TraceLog("{0} has been added", PluginFile.GetFileName());
+						ReasonNotLoaded = string.Format("{0} has been added", PluginFile.GetFileName());
+						return null;
+					}
+				}
+
 				// We do a check to see if any modules' headers have changed which have
 				// acquired or lost UHT types.  If so, which should be rare,
 				// we'll just invalidate the entire makefile and force it to be rebuilt.
