UBT: Don't assume the default architecture for a platform is the empty string. It's not on Linux, which causes warnings when generating project files.

#rnx

diff --git a/Engine/Source/Programs/UnrealBuildTool/ProjectFiles/ProjectFileGenerator.cs b/Engine/Source/Programs/UnrealBuildTool/ProjectFiles/ProjectFileGenerator.cs
--- a/Engine/Source/Programs/UnrealBuildTool/ProjectFiles/ProjectFileGenerator.cs
+++ b/Engine/Source/Programs/UnrealBuildTool/ProjectFiles/ProjectFileGenerator.cs
@@ -1702,7 +1702,8 @@ namespace UnrealBuildTool
 
 						try
 						{
-							const string DefaultArchitecture = "";
+							// Get the architecture from the target platform
+							string DefaultArchitecture = UEBuildPlatform.GetBuildPlatform(BuildHostPlatform.Current.Platform).GetDefaultArchitecture(CurTarget.UnrealProjectFilePath);
 
 							// Create the target descriptor
 							TargetDescriptor TargetDesc = new TargetDescriptor(CurTarget.UnrealProjectFilePath, CurTarget.Name, BuildHostPlatform.Current.Platform, UnrealTargetConfiguration.Development, DefaultArchitecture, new CommandLineArguments(NewArguments.ToArray()));
