UBT: Fix support for running DebugGame editor builds with a unique build environment. The output directory cannot be used to infer which configuration binaries are built in.

#rnx

diff --git a/Engine/Source/Programs/UnrealBuildTool/Configuration/UEBuildTarget.cs b/Engine/Source/Programs/UnrealBuildTool/Configuration/UEBuildTarget.cs
--- a/Engine/Source/Programs/UnrealBuildTool/Configuration/UEBuildTarget.cs
+++ b/Engine/Source/Programs/UnrealBuildTool/Configuration/UEBuildTarget.cs
@@ -926,8 +926,8 @@ namespace UnrealBuildTool
 					if(Binary.Type == UEBuildBinaryType.DynamicLinkLibrary)
 					{
 						DirectoryReference DirectoryName = Binary.OutputFilePath.Directory;
-						bool bIsGameDirectory = !DirectoryName.IsUnderDirectory(UnrealBuildTool.EngineDirectory);
-						FileReference ManifestFileName = FileReference.Combine(DirectoryName, ModuleManifest.GetStandardFileName(AppName, Platform, Configuration, Architecture, bIsGameDirectory));
+						bool bIsGameBinary = RulesAssembly.IsGameModule(Binary.PrimaryModule.Name);
+						FileReference ManifestFileName = FileReference.Combine(DirectoryName, ModuleManifest.GetStandardFileName(AppName, Platform, Configuration, Architecture, bIsGameBinary));
 
 						ModuleManifest Manifest;
 						if (!FileNameToModuleManifest.TryGetValue(ManifestFileName, out Manifest))
