PR #5252: Prevent to use engine intermediate folder for generated linking files if engine is installed (Contributed by [REDACTED_PERSON])

#rnx

diff --git a/Engine/Source/Programs/UnrealBuildTool/Configuration/UEBuildTarget.cs b/Engine/Source/Programs/UnrealBuildTool/Configuration/UEBuildTarget.cs
--- a/Engine/Source/Programs/UnrealBuildTool/Configuration/UEBuildTarget.cs
+++ b/Engine/Source/Programs/UnrealBuildTool/Configuration/UEBuildTarget.cs
@@ -2864,16 +2864,13 @@ namespace UnrealBuildTool
 			// units too, so they can't be shared between different targets.  They are effectively project-specific engine intermediates.
 			if (UnrealBuildTool.IsEngineInstalled() || (ProjectFile != null && ShouldCompileMonolithic()))
 			{
-				if (ShouldCompileMonolithic())
+				if (ProjectFile != null)
 				{
-					if (ProjectFile != null)
-					{
-						LinkIntermediateDirectory = DirectoryReference.Combine(ProjectFile.Directory, PlatformIntermediateFolder, AppName, Configuration.ToString());
-					}
-					else if (ForeignPlugin != null)
-					{
-						LinkIntermediateDirectory = DirectoryReference.Combine(ForeignPlugin.Directory, PlatformIntermediateFolder, AppName, Configuration.ToString());
-					}
+					LinkIntermediateDirectory = DirectoryReference.Combine(ProjectFile.Directory, PlatformIntermediateFolder, AppName, Configuration.ToString());
+				}
+				else if (ForeignPlugin != null)
+				{
+					LinkIntermediateDirectory = DirectoryReference.Combine(ForeignPlugin.Directory, PlatformIntermediateFolder, AppName, Configuration.ToString());
 				}
 			}
 
