UBT: Perform scan for reflection markup in parallel. Saves ~1.3s on editor builds without makefiles.

#rnx

diff --git a/Engine/Source/Programs/UnrealBuildTool/System/ExternalExecution.cs b/Engine/Source/Programs/UnrealBuildTool/System/ExternalExecution.cs
--- a/Engine/Source/Programs/UnrealBuildTool/System/ExternalExecution.cs
+++ b/Engine/Source/Programs/UnrealBuildTool/System/ExternalExecution.cs
@@ -207,15 +207,15 @@ namespace UnrealBuildTool
 		/// </summary>
 		public bool bIsReadOnly;
 
-		public UHTModuleInfo(string ModuleName, FileReference ModuleRulesFile, DirectoryReference ModuleDirectory, UHTModuleType ModuleType, List<FileItem> PublicUObjectClassesHeaders, List<FileItem> PublicUObjectHeaders, List<FileItem> PrivateUObjectHeaders, EGeneratedCodeVersion GeneratedCodeVersion, bool bIsReadOnly)
+		public UHTModuleInfo(string ModuleName, FileReference ModuleRulesFile, DirectoryReference ModuleDirectory, UHTModuleType ModuleType, EGeneratedCodeVersion GeneratedCodeVersion, bool bIsReadOnly)
 		{
 			this.ModuleName = ModuleName;
 			this.ModuleRulesFile = ModuleRulesFile;
 			this.ModuleDirectory = ModuleDirectory;
 			this.ModuleType = ModuleType.ToString();
-			this.PublicUObjectClassesHeaders = PublicUObjectClassesHeaders;
-			this.PublicUObjectHeaders = PublicUObjectHeaders;
-			this.PrivateUObjectHeaders = PrivateUObjectHeaders;
+			this.PublicUObjectClassesHeaders = new List<FileItem>();
+			this.PublicUObjectHeaders = new List<FileItem>();
+			this.PrivateUObjectHeaders = new List<FileItem>();
 			this.GeneratedCodeVersion = GeneratedCodeVersion;
 			this.bIsReadOnly = bIsReadOnly;
 		}
@@ -363,47 +363,6 @@ namespace UnrealBuildTool
 	/// </summary>
 	class ExternalExecution
 	{
-		/// <summary>
-		/// Generates a UHTModuleInfo for a particular named module under a directory.
-		/// </summary>
-		/// <returns></returns>
-		public static UHTModuleInfo CreateUHTModuleInfo(List<FileItem> HeaderFiles, string ModuleName, FileReference ModuleRulesFile, DirectoryReference ModuleDirectory, UHTModuleType ModuleType, EGeneratedCodeVersion GeneratedCodeVersion, bool bIsReadOnly, SourceFileMetadataCache MetadataCache)
-		{
-			DirectoryReference ClassesFolder = DirectoryReference.Combine(ModuleDirectory, "Classes");
-			DirectoryReference PublicFolder = DirectoryReference.Combine(ModuleDirectory, "Public");
-
-			List<FileItem> PublicClassesUObjectHeaders = new List<FileItem>();
-			List<FileItem> PublicUObjectHeaders = new List<FileItem>();
-			List<FileItem> PrivateUObjectHeaders = new List<FileItem>();
-
-			foreach (FileItem HeaderFile in HeaderFiles)
-			{
-				// Check to see if we know anything about this file.  If we have up-to-date cached information about whether it has
-				// UObjects or not, we can skip doing a test here.
-				if (MetadataCache.ContainsReflectionMarkup(HeaderFile))
-				{
-					if (HeaderFile.Location.IsUnderDirectory(ClassesFolder))
-					{
-						PublicClassesUObjectHeaders.Add(HeaderFile);
-					}
-					else if (HeaderFile.Location.IsUnderDirectory(PublicFolder))
-					{
-						PublicUObjectHeaders.Add(HeaderFile);
-					}
-					else
-					{
-						PrivateUObjectHeaders.Add(HeaderFile);
-					}
-				}
-			}
-
-			return new UHTModuleInfo(ModuleName, ModuleRulesFile, ModuleDirectory, ModuleType, PublicClassesUObjectHeaders, PublicUObjectHeaders, PrivateUObjectHeaders, GeneratedCodeVersion, bIsReadOnly);
-		}
-
-		static ExternalExecution()
-		{
-		}
-
 		static UHTModuleType GetEngineModuleTypeFromDescriptor(ModuleDescriptor Module)
 		{
             UHTModuleType? Type = UHTModuleTypeExtensions.EngineModuleTypeFromHostType(Module.Type);
@@ -635,15 +594,27 @@ namespace UnrealBuildTool
 			List<UEBuildModuleCPP> ModulesSortedByType = ModulesToGenerateHeadersFor.OrderBy(c => ModuleToType[c]).ToList();
 			StableTopologicalSort(ModulesSortedByType);
 
-			ReadOnlyHashSet<string> ExcludedFolders = UEBuildPlatform.GetBuildPlatform(Platform, true).GetExcludedFolderNames();
-			foreach (UEBuildModuleCPP Module in ModulesSortedByType)
+			// Create the info for each module in parallel
+			UHTModuleInfo[] ModuleInfoArray = new UHTModuleInfo[ModulesSortedByType.Count];
+			using(ThreadPoolWorkQueue Queue = new ThreadPoolWorkQueue())
 			{
-				DirectoryItem ModuleDirectoryItem = DirectoryItem.GetItemByDirectoryReference(Module.ModuleDirectory);
+				ReadOnlyHashSet<string> ExcludedFolders = UEBuildPlatform.GetBuildPlatform(Platform, true).GetExcludedFolderNames();
+				for(int Idx = 0; Idx < ModulesSortedByType.Count; Idx++)
+				{
+					UEBuildModuleCPP Module = ModulesSortedByType[Idx];
+
+					UHTModuleInfo Info = new UHTModuleInfo(Module.Name, Module.RulesFile, Module.ModuleDirectory, ModuleToType[Module], GeneratedCodeVersion, Module.Rules.bUsePrecompiled);
+					ModuleInfoArray[Idx] = Info;
 
-				List<FileItem> HeaderFiles = new List<FileItem>();
-				FindHeaders(ModuleDirectoryItem, ExcludedFolders, HeaderFiles);
+					Queue.Enqueue(() => SetupUObjectModule(Info, ExcludedFolders, MetadataCache, Queue));
+				}
+			}
 
-				UHTModuleInfo Info = ExternalExecution.CreateUHTModuleInfo(HeaderFiles, Module.Name, Module.RulesFile, Module.ModuleDirectory, ModuleToType[Module], GeneratedCodeVersion, Module.Rules.bUsePrecompiled, MetadataCache);
+			// Filter out all the modules with reflection data
+			for(int Idx = 0; Idx < ModulesSortedByType.Count; Idx++)
+			{
+				UEBuildModuleCPP Module = ModulesSortedByType[Idx];
+				UHTModuleInfo Info = ModuleInfoArray[Idx];
 				if (Info.PublicUObjectClassesHeaders.Count > 0 || Info.PrivateUObjectHeaders.Count > 0 || Info.PublicUObjectHeaders.Count > 0)
 				{
 					// Set a flag indicating that we need to add the generated headers directory
@@ -659,13 +630,13 @@ namespace UnrealBuildTool
 
 					UObjectModules.Add(Info);
 
+					DirectoryItem ModuleDirectoryItem = DirectoryItem.GetItemByDirectoryReference(Module.ModuleDirectory);
+
 					List<FileItem> ReflectedHeaderFiles = new List<FileItem>();
 					ReflectedHeaderFiles.AddRange(Info.PublicUObjectHeaders);
 					ReflectedHeaderFiles.AddRange(Info.PublicUObjectClassesHeaders);
 					ReflectedHeaderFiles.AddRange(Info.PrivateUObjectHeaders);
 					UObjectModuleHeaders.Add(new UHTModuleHeaderInfo(ModuleDirectoryItem, ReflectedHeaderFiles));
-
-					Log.TraceVerbose("Detected UObject module: " + Info.ModuleName);
 				}
 				else
 				{
@@ -674,7 +645,6 @@ namespace UnrealBuildTool
 					{
 						if (DirectoryReference.Exists(Module.GeneratedCodeDirectory))
 						{
-							Log.TraceVerbose("Deleting stale generated code directory: " + Module.GeneratedCodeDirectory.ToString());
 							Directory.Delete(Module.GeneratedCodeDirectory.FullName, true);
 						}
 					}
@@ -682,6 +652,43 @@ namespace UnrealBuildTool
 			}
 		}
 
+		static void SetupUObjectModule(UHTModuleInfo ModuleInfo, ReadOnlyHashSet<string> ExcludedFolders, SourceFileMetadataCache MetadataCache, ThreadPoolWorkQueue Queue)
+		{
+			DirectoryItem ModuleDirectoryItem = DirectoryItem.GetItemByDirectoryReference(ModuleInfo.ModuleDirectory);
+
+			List<FileItem> HeaderFiles = new List<FileItem>();
+			FindHeaders(ModuleDirectoryItem, ExcludedFolders, HeaderFiles);
+
+			foreach (FileItem HeaderFile in HeaderFiles)
+			{
+				Queue.Enqueue(() => SetupUObjectModuleHeader(ModuleInfo, HeaderFile, MetadataCache));
+			}
+		}
+
+		static void SetupUObjectModuleHeader(UHTModuleInfo ModuleInfo, FileItem HeaderFile, SourceFileMetadataCache MetadataCache)
+		{
+			// Check to see if we know anything about this file.  If we have up-to-date cached information about whether it has
+			// UObjects or not, we can skip doing a test here.
+			if (MetadataCache.ContainsReflectionMarkup(HeaderFile))
+			{
+				lock(ModuleInfo)
+				{
+					if (HeaderFile.Location.IsUnderDirectory(DirectoryReference.Combine(ModuleInfo.ModuleDirectory, "Classes")))
+					{
+						ModuleInfo.PublicUObjectClassesHeaders.Add(HeaderFile);
+					}
+					else if (HeaderFile.Location.IsUnderDirectory(DirectoryReference.Combine(ModuleInfo.ModuleDirectory, "Public")))
+					{
+						ModuleInfo.PublicUObjectHeaders.Add(HeaderFile);
+					}
+					else
+					{
+						ModuleInfo.PrivateUObjectHeaders.Add(HeaderFile);
+					}
+				}
+			}
+		}
+
 		/// <summary>
 		/// Gets the path to the receipt for UHT
 		/// </summary>
