UBT: Only reset cached info for generated code directories after running UHT. Resetting cached info for all files causes slowdown for dependency checking.

#rnx

diff --git a/Engine/Source/Programs/UnrealBuildTool/Configuration/UEBuildModuleCPP.cs b/Engine/Source/Programs/UnrealBuildTool/Configuration/UEBuildModuleCPP.cs
--- a/Engine/Source/Programs/UnrealBuildTool/Configuration/UEBuildModuleCPP.cs
+++ b/Engine/Source/Programs/UnrealBuildTool/Configuration/UEBuildModuleCPP.cs
@@ -830,7 +830,7 @@ namespace UnrealBuildTool
 			if (IncludedFile.LastWriteTimeUtc > WrapperFile.LastWriteTimeUtc)
 			{
 				File.SetLastWriteTimeUtc(WrapperFile.AbsolutePath, DateTime.UtcNow);
-				WrapperFile.ResetFileInfo();
+				WrapperFile.ResetCachedInfo();
 			}
 			return WrapperFile;
 		}
diff --git a/Engine/Source/Programs/UnrealBuildTool/System/ActionGraph.cs b/Engine/Source/Programs/UnrealBuildTool/System/ActionGraph.cs
--- a/Engine/Source/Programs/UnrealBuildTool/System/ActionGraph.cs
+++ b/Engine/Source/Programs/UnrealBuildTool/System/ActionGraph.cs
@@ -150,6 +150,15 @@ namespace UnrealBuildTool
 				}
 				Log.TraceInformation("Total time in {0} executor: {1:0.00} seconds", Executor.Name, Timer.Elapsed.TotalSeconds);
 
+				// Reset the file info for all the produced items
+				foreach (Action BuildAction in ActionsToExecute)
+				{
+					foreach(FileItem ProducedItem in BuildAction.ProducedItems)
+					{
+						ProducedItem.ResetCachedInfo();
+					}
+				}
+
 				// Verify the link outputs were created (seems to happen with Win64 compiles)
 				foreach (Action BuildAction in ActionsToExecute)
 				{
@@ -157,7 +166,6 @@ namespace UnrealBuildTool
 					{
 						foreach (FileItem Item in BuildAction.ProducedItems)
 						{
-							Item.ResetFileInfo();
 							if(!Item.Exists)
 							{
 								throw new BuildException("Failed to produce item: {0}", Item.AbsolutePath);
diff --git a/Engine/Source/Programs/UnrealBuildTool/System/DirectoryItem.cs b/Engine/Source/Programs/UnrealBuildTool/System/DirectoryItem.cs
--- a/Engine/Source/Programs/UnrealBuildTool/System/DirectoryItem.cs
+++ b/Engine/Source/Programs/UnrealBuildTool/System/DirectoryItem.cs
@@ -163,10 +163,27 @@ namespace UnrealBuildTool
 		/// <summary>
 		/// Reset the contents of the directory and allow them to be fetched again
 		/// </summary>
-		public void Refresh()
+		public void ResetCachedInfo()
 		{
-			Directories = null;
-			Files = null;
+			Dictionary<string, DirectoryItem> PrevDirectories = Directories;
+			if(PrevDirectories != null)
+			{
+				foreach(DirectoryItem SubDirectory in PrevDirectories.Values)
+				{
+					SubDirectory.ResetCachedInfo();
+				}
+				Directories = null;
+			}
+
+			Dictionary<string, FileItem> PrevFiles = Files;
+			if(PrevFiles != null)
+			{
+				foreach(FileItem File in PrevFiles.Values)
+				{
+					File.ResetCachedInfo();
+				}
+				Files = null;
+			}
 		}
 
 		/// <summary>
diff --git a/Engine/Source/Programs/UnrealBuildTool/System/DirectoryLookupCache.cs b/Engine/Source/Programs/UnrealBuildTool/System/DirectoryLookupCache.cs
--- a/Engine/Source/Programs/UnrealBuildTool/System/DirectoryLookupCache.cs
+++ b/Engine/Source/Programs/UnrealBuildTool/System/DirectoryLookupCache.cs
@@ -32,7 +32,7 @@ namespace UnrealBuildTool
 
         static public void InvalidateCachedDirectory(DirectoryReference Directory)
         {
-			DirectoryItem.GetItemByDirectoryReference(Directory).Refresh();
+			DirectoryItem.GetItemByDirectoryReference(Directory).ResetCachedInfo();
         }
 	}
 }
diff --git a/Engine/Source/Programs/UnrealBuildTool/System/ExternalExecution.cs b/Engine/Source/Programs/UnrealBuildTool/System/ExternalExecution.cs
--- a/Engine/Source/Programs/UnrealBuildTool/System/ExternalExecution.cs
+++ b/Engine/Source/Programs/UnrealBuildTool/System/ExternalExecution.cs
@@ -192,6 +192,11 @@ namespace UnrealBuildTool
 		/// </summary>
 		public List<FileItem> PrivateUObjectHeaders;
 
+		/// <summary>
+		/// Directory containing generated code
+		/// </summary>
+		public DirectoryItem GeneratedCodeDirectory;
+
 		/// <summary>
 		/// Base (i.e. extensionless) path+filename of the .gen files
 		/// </summary>
@@ -207,7 +212,7 @@ namespace UnrealBuildTool
 		/// </summary>
 		public bool bIsReadOnly;
 
-		public UHTModuleInfo(string ModuleName, FileReference ModuleRulesFile, DirectoryReference ModuleDirectory, UHTModuleType ModuleType, EGeneratedCodeVersion GeneratedCodeVersion, bool bIsReadOnly)
+		public UHTModuleInfo(string ModuleName, FileReference ModuleRulesFile, DirectoryReference ModuleDirectory, UHTModuleType ModuleType, DirectoryItem GeneratedCodeDirectory, EGeneratedCodeVersion GeneratedCodeVersion, bool bIsReadOnly)
 		{
 			this.ModuleName = ModuleName;
 			this.ModuleRulesFile = ModuleRulesFile;
@@ -216,6 +221,7 @@ namespace UnrealBuildTool
 			this.PublicUObjectClassesHeaders = new List<FileItem>();
 			this.PublicUObjectHeaders = new List<FileItem>();
 			this.PrivateUObjectHeaders = new List<FileItem>();
+			this.GeneratedCodeDirectory = GeneratedCodeDirectory;
 			this.GeneratedCodeVersion = GeneratedCodeVersion;
 			this.bIsReadOnly = bIsReadOnly;
 		}
@@ -230,6 +236,7 @@ namespace UnrealBuildTool
 			PublicUObjectHeaders = Reader.ReadList(() => Reader.ReadFileItem());
 			PrivateUObjectHeaders = Reader.ReadList(() => Reader.ReadFileItem());
 			GeneratedCPPFilenameBase = Reader.ReadString();
+			GeneratedCodeDirectory = Reader.ReadDirectoryItem();
 			GeneratedCodeVersion = (EGeneratedCodeVersion)Reader.ReadInt();
 			bIsReadOnly = Reader.ReadBool();
 		}
@@ -244,6 +251,7 @@ namespace UnrealBuildTool
 			Writer.WriteList(PublicUObjectHeaders, Item => Writer.WriteFileItem(Item));
 			Writer.WriteList(PrivateUObjectHeaders, Item => Writer.WriteFileItem(Item));
 			Writer.WriteString(GeneratedCPPFilenameBase);
+			Writer.WriteDirectoryItem(GeneratedCodeDirectory);
 			Writer.WriteInt((int)GeneratedCodeVersion);
 			Writer.WriteBool(bIsReadOnly);
 		}
@@ -603,7 +611,7 @@ namespace UnrealBuildTool
 				{
 					UEBuildModuleCPP Module = ModulesSortedByType[Idx];
 
-					UHTModuleInfo Info = new UHTModuleInfo(Module.Name, Module.RulesFile, Module.ModuleDirectory, ModuleToType[Module], GeneratedCodeVersion, Module.Rules.bUsePrecompiled);
+					UHTModuleInfo Info = new UHTModuleInfo(Module.Name, Module.RulesFile, Module.ModuleDirectory, ModuleToType[Module], DirectoryItem.GetItemByDirectoryReference(Module.GeneratedCodeDirectory), GeneratedCodeVersion, Module.Rules.bUsePrecompiled);
 					ModuleInfoArray[Idx] = Info;
 
 					Queue.Enqueue(() => SetupUObjectModule(Info, ExcludedFolders, MetadataCache, Queue));
@@ -1218,10 +1226,10 @@ namespace UnrealBuildTool
 
 					// Now that UHT has successfully finished generating code, we need to update all cached FileItems in case their last write time has changed.
 					// Otherwise UBT might not detect changes UHT made.
-					DateTime StartTime = DateTime.UtcNow;
-					FileItem.ResetInfos();
-					double ResetDuration = (DateTime.UtcNow - StartTime).TotalSeconds;
-					Log.TraceVerbose("FileItem.ResetInfos() duration: {0}s", ResetDuration);
+					using(Timeline.ScopeEvent("ExternalExecution.ResetCachedHeaderInfo()"))
+					{
+						ResetCachedHeaderInfo(UObjectModules);
+					}
 				}
 				else
 				{
@@ -1240,5 +1248,13 @@ namespace UnrealBuildTool
 				Log.WriteLine(LogEventType.Console, "@progress pop");
 			}
 		}
+
+		static void ResetCachedHeaderInfo(List<UHTModuleInfo> UObjectModules)
+		{
+			foreach(UHTModuleInfo ModuleInfo in UObjectModules)
+			{
+				ModuleInfo.GeneratedCodeDirectory.ResetCachedInfo();
+			}
+		}
 	}
 }
diff --git a/Engine/Source/Programs/UnrealBuildTool/System/FileItem.cs b/Engine/Source/Programs/UnrealBuildTool/System/FileItem.cs
--- a/Engine/Source/Programs/UnrealBuildTool/System/FileItem.cs
+++ b/Engine/Source/Programs/UnrealBuildTool/System/FileItem.cs
@@ -211,7 +211,7 @@ namespace UnrealBuildTool
 
 			// Reset the file info, in case it already knows about the old file
 			FileItem Item = GetItemByFileReference(Location);
-			Item.ResetFileInfo();
+			Item.ResetCachedInfo();
 			return Item;
 		}
 
@@ -277,22 +277,11 @@ namespace UnrealBuildTool
 		/// <summary>
 		/// Resets the cached file info
 		/// </summary>
-		public void ResetFileInfo()
+		public void ResetCachedInfo()
 		{
 			Info = Location.ToFileInfo();
 		}
 
-		/// <summary>
-		/// Reset file information on all cached FileItems.
-		/// </summary>
-		public static void ResetInfos()
-		{
-			foreach (KeyValuePair<FileReference, FileItem> Item in UniqueSourceFileMap)
-			{
-				Item.Value.ResetFileInfo();
-			}
-		}
-
 		/// <summary>
 		/// Return the path to this FileItem to debugging
 		/// </summary>
