UBT: Add a timer around updating generated code timestamps.

#rnx

diff --git a/Engine/Source/Programs/UnrealBuildTool/System/ExternalExecution.cs b/Engine/Source/Programs/UnrealBuildTool/System/ExternalExecution.cs
--- a/Engine/Source/Programs/UnrealBuildTool/System/ExternalExecution.cs
+++ b/Engine/Source/Programs/UnrealBuildTool/System/ExternalExecution.cs
@@ -1239,7 +1239,10 @@ namespace UnrealBuildTool
 				Progress.Write(2, 3);
 
 				// touch the directories
-				UpdateDirectoryTimestamps(UObjectModules);
+				using(Timeline.ScopeEvent("ExternalExecution.UpdateDirectoryTimestamps()"))
+				{
+					UpdateDirectoryTimestamps(UObjectModules);
+				}
 
 				Progress.Write(3, 3);
 			}
