UBT: Fix project file generator including all projects in the workspace in the solution, even if a specific project file is specified.

#rnx

diff --git a/Engine/Source/Programs/UnrealBuildTool/ProjectFiles/ProjectFileGenerator.cs b/Engine/Source/Programs/UnrealBuildTool/ProjectFiles/ProjectFileGenerator.cs
--- a/Engine/Source/Programs/UnrealBuildTool/ProjectFiles/ProjectFileGenerator.cs
+++ b/Engine/Source/Programs/UnrealBuildTool/ProjectFiles/ProjectFileGenerator.cs
@@ -468,7 +468,7 @@ namespace UnrealBuildTool
 		public List<FileReference> FindGameProjects()
 		{
 			List<FileReference> ProjectFiles = new List<FileReference>();
-			if(bGeneratingGameProjectFiles)
+			if(OnlyGameProject != null)
 			{
 				ProjectFiles.Add(OnlyGameProject);
 			}
