PR #5429: Clang 7.0.1 support (Contributed by [REDACTED_PERSON])

diff --git a/Engine/Source/Programs/UnrealBuildTool/Platform/Linux/LinuxToolChain.cs b/Engine/Source/Programs/UnrealBuildTool/Platform/Linux/LinuxToolChain.cs
--- a/Engine/Source/Programs/UnrealBuildTool/Platform/Linux/LinuxToolChain.cs
+++ b/Engine/Source/Programs/UnrealBuildTool/Platform/Linux/LinuxToolChain.cs
@@ -171,10 +171,10 @@ namespace UnrealBuildTool
 				throw new BuildException("Unable to build: no compatible clang version found. Please run Setup.sh");
 			}
 			// prevent unknown clangs since the build is likely to fail on too old or too new compilers
-			else if ((CompilerVersionMajor * 10 + CompilerVersionMinor) > 60 || (CompilerVersionMajor * 10 + CompilerVersionMinor) < 38)
+			else if ((CompilerVersionMajor * 10 + CompilerVersionMinor) > 70 || (CompilerVersionMajor * 10 + CompilerVersionMinor) < 38)
 			{
 				throw new BuildException(
-					string.Format("This version of the Unreal Engine can only be compiled with clang 6.0, 5.0, 4.0, 3.9, 3.8. clang {0} may not build it - please use a different version.",
+					string.Format("This version of the Unreal Engine can only be compiled with clang 7.0, 6.0, 5.0, 4.0, 3.9, 3.8. clang {0} may not build it - please use a different version.",
 						CompilerVersionString)
 					);
 			}
@@ -921,7 +921,7 @@ namespace UnrealBuildTool
 			Result += " -Wl,--build-id";
 			if (bSuppressPIE && !LinkEnvironment.bIsBuildingDLL)
 			{
-				Result += " -Wl,-nopie";
+				Result += " -Wl,--no-pie";
 			}
 
 			// whether we actually can do that is checked in CanUseLTO() earlier
