UBT: Remove code to add the -Monolithic option when generating intellisense data. This only exists to prevent intellisense errors because the DLLEXPORT and DLLIMPORT macros are mismatched for the project containing all modules; we can achieve this by explicitly modifying the macros instead.

#rnx

diff --git a/Engine/Source/Programs/UnrealBuildTool/ProjectFiles/Project.cs b/Engine/Source/Programs/UnrealBuildTool/ProjectFiles/Project.cs
--- a/Engine/Source/Programs/UnrealBuildTool/ProjectFiles/Project.cs
+++ b/Engine/Source/Programs/UnrealBuildTool/ProjectFiles/Project.cs
@@ -291,9 +291,10 @@ namespace UnrealBuildTool
 		/// <param name="NewPreprocessorDefinitions">List of preprocessor definitons to add</param>
 		public void AddIntelliSensePreprocessorDefinitions(List<string> NewPreprocessorDefinitions)
 		{
-			foreach (string CurDef in NewPreprocessorDefinitions)
+			foreach (string NewPreprocessorDefinition in NewPreprocessorDefinitions)
 			{
 				// Don't add definitions and value combinations that have already been added for this project
+				string CurDef = NewPreprocessorDefinition;
 				if (KnownIntelliSensePreprocessorDefinitions.Add(CurDef))
 				{
 					// Go ahead and check to see if the definition already exists, but the value is different
@@ -301,6 +302,14 @@ namespace UnrealBuildTool
 
 					string Def, Value;
 					SplitDefinitionAndValue(CurDef, out Def, out Value);
+
+					// Ignore any API macros being import/export; we'll assume they're valid across the whole project
+					if(Def.EndsWith("_API", StringComparison.Ordinal))
+					{
+						CurDef = Def + "=";
+						Value = "";
+					}
+
 					for (int DefineIndex = 0; DefineIndex < IntelliSensePreprocessorDefinitions.Count; ++DefineIndex)
 					{
 						string ExistingDef, ExistingValue;
diff --git a/Engine/Source/Programs/UnrealBuildTool/ProjectFiles/ProjectFileGenerator.cs b/Engine/Source/Programs/UnrealBuildTool/ProjectFiles/ProjectFileGenerator.cs
--- a/Engine/Source/Programs/UnrealBuildTool/ProjectFiles/ProjectFileGenerator.cs
+++ b/Engine/Source/Programs/UnrealBuildTool/ProjectFiles/ProjectFileGenerator.cs
@@ -1693,11 +1693,6 @@ namespace UnrealBuildTool
 						{
 							NewArguments.Add("-precompile");
 						}
-						if (CurTarget.TargetRules.Type == TargetType.Editor && CurTarget.TargetRules.BuildEnvironment != TargetBuildEnvironment.Unique)
-						{
-							NewArguments.Add("-SharedBuildEnvironment"); // Force the shared build environment, even though we have -Monolithic
-						}
-						NewArguments.Add("-Monolithic"); // Don't want any _API macros to be DLLIMPORT; this will cause Intellisense errors in the file that they are defined.
 						NewArguments.AddRange(Arguments);
 
 						try
