Revert change to -nopie argument in LinuxToolChain from CL 4685944. Doesn't seem to work on all distros.

#rnx

diff --git a/Engine/Source/Programs/UnrealBuildTool/Platform/Linux/LinuxToolChain.cs b/Engine/Source/Programs/UnrealBuildTool/Platform/Linux/LinuxToolChain.cs
--- a/Engine/Source/Programs/UnrealBuildTool/Platform/Linux/LinuxToolChain.cs
+++ b/Engine/Source/Programs/UnrealBuildTool/Platform/Linux/LinuxToolChain.cs
@@ -921,7 +921,7 @@ namespace UnrealBuildTool
 			Result += " -Wl,--build-id";
 			if (bSuppressPIE && !LinkEnvironment.bIsBuildingDLL)
 			{
-				Result += " -Wl,--no-pie";
+				Result += " -Wl,-nopie";
 			}
 
 			// whether we actually can do that is checked in CanUseLTO() earlier
