UBT: Fix exception when building remotely if a log path is specified on the command line (eg. when being launched via UAT)

#rnx

diff --git a/Engine/Source/Programs/DotNETCommon/DotNETUtilities/Log.cs b/Engine/Source/Programs/DotNETCommon/DotNETUtilities/Log.cs
--- a/Engine/Source/Programs/DotNETCommon/DotNETUtilities/Log.cs
+++ b/Engine/Source/Programs/DotNETCommon/DotNETUtilities/Log.cs
@@ -171,6 +171,14 @@ namespace Tools.DotNETCommon
 			get; set;
 		}
 
+		/// <summary>
+		/// Path to the log file being written to. May be null.
+		/// </summary>
+		public static FileReference OutputFile
+		{
+			get; private set;
+		}
+
 		/// <summary>
 		/// Whether console output is redirected. This prevents writing status updates that rely on moving the cursor.
 		/// </summary>
@@ -235,6 +243,7 @@ namespace Tools.DotNETCommon
 		{
 			try
 			{
+				Log.OutputFile = OutputFile;
 				DirectoryReference.CreateDirectory(OutputFile.Directory);
 				TextWriterTraceListener LogTraceListener = new TextWriterTraceListener(new StreamWriter(OutputFile.FullName), Name);
 				Trace.Listeners.Add(LogTraceListener);
diff --git a/Engine/Source/Programs/UnrealBuildTool/Modes/BuildMode.cs b/Engine/Source/Programs/UnrealBuildTool/Modes/BuildMode.cs
--- a/Engine/Source/Programs/UnrealBuildTool/Modes/BuildMode.cs
+++ b/Engine/Source/Programs/UnrealBuildTool/Modes/BuildMode.cs
@@ -98,10 +98,9 @@ namespace UnrealBuildTool
 			XmlConfig.ApplyTo(this);
 
 			// Create the log file, and flush the startup listener to it
-			FileReference LogFile = null;
 			if(!Arguments.HasOption("-NoLog") && !Log.HasFileWriter())
 			{
-				LogFile = new FileReference(BaseLogFileName);
+				FileReference LogFile = new FileReference(BaseLogFileName);
 				foreach(string LogSuffix in Arguments.GetValues("-LogSuffix="))
 				{
 					LogFile = LogFile.ChangeExtension(null) + "_" + LogSuffix + LogFile.GetExtension();
@@ -147,14 +146,11 @@ namespace UnrealBuildTool
 					TargetDescriptor TargetDesc = TargetDescriptors[Idx];
 					if(RemoteMac.HandlesTargetPlatform(TargetDesc.Platform))
 					{
-						FileReference RemoteLogFile = null;
-						if(LogFile != null)
-						{
-							RemoteLogFile = FileReference.Combine(LogFile.Directory, LogFile.GetFileNameWithoutExtension() + "_Remote.txt");
-						}
+						FileReference BaseLogFile = Log.OutputFile ?? new FileReference(BaseLogFileName);
+						FileReference RemoteLogFile = FileReference.Combine(BaseLogFile.Directory, BaseLogFile.GetFileNameWithoutExtension() + "_Remote.txt");
 
 						RemoteMac RemoteMac = new RemoteMac(TargetDesc.ProjectFile);
-						if(!RemoteMac.Build(TargetDesc, LogFile))
+						if(!RemoteMac.Build(TargetDesc, RemoteLogFile))
 						{
 							return (int)CompilationResult.Unknown;
 						}
diff --git a/Engine/Source/Programs/UnrealBuildTool/ToolChain/RemoteMac.cs b/Engine/Source/Programs/UnrealBuildTool/ToolChain/RemoteMac.cs
--- a/Engine/Source/Programs/UnrealBuildTool/ToolChain/RemoteMac.cs
+++ b/Engine/Source/Programs/UnrealBuildTool/ToolChain/RemoteMac.cs
@@ -380,6 +380,7 @@ namespace UnrealBuildTool
 			// Get the directory for working files
 			DirectoryReference TempDir = CreateTempDirectory(TargetDesc);
 
+			// Map the path containing the remote log file
 			bool bLogIsMapped = false;
 			foreach (RemoteMapping Mapping in Mappings)
 			{
