Fix hot reload failures due to hot-reload module names not being set up correctly in the makefile.

diff --git a/Engine/Source/Programs/UnrealBuildTool/Configuration/UEBuildTarget.cs b/Engine/Source/Programs/UnrealBuildTool/Configuration/UEBuildTarget.cs
--- a/Engine/Source/Programs/UnrealBuildTool/Configuration/UEBuildTarget.cs
+++ b/Engine/Source/Programs/UnrealBuildTool/Configuration/UEBuildTarget.cs
@@ -1197,9 +1197,12 @@ namespace UnrealBuildTool
 			// Create the makefile
 			TargetMakefile Makefile = new TargetMakefile(TargetToolChain.GetVersionInfo(), ReceiptFileName, ProjectIntermediateDirectory, TargetType, bDeployAfterCompile, bHasProjectScriptPlugin);
 
-			// Execute the pre-build steps
+			// Setup the pre-build steps
 			Makefile.PreBuildScripts = CreatePreBuildScripts();
 
+			// Setup the hot reload module list
+			Makefile.HotReloadModuleNames = GetHotReloadModuleNames();
+
 			// If we're compiling monolithic, make sure the executable knows about all referenced modules
 			if (ShouldCompileMonolithic())
 			{
diff --git a/Engine/Source/Programs/UnrealBuildTool/Modes/BuildMode.cs b/Engine/Source/Programs/UnrealBuildTool/Modes/BuildMode.cs
--- a/Engine/Source/Programs/UnrealBuildTool/Modes/BuildMode.cs
+++ b/Engine/Source/Programs/UnrealBuildTool/Modes/BuildMode.cs
@@ -449,13 +449,13 @@ namespace UnrealBuildTool
 			HotReloadMode HotReloadMode = TargetDescriptor.HotReloadMode;
 			if(HotReloadMode == HotReloadMode.Default)
 			{
-				if (BuildConfiguration.bAllowHotReloadFromIDE && HotReload.ShouldDoHotReloadFromIDE(BuildConfiguration, TargetDescriptor))
+				if (TargetDescriptor.HotReloadModuleNameToSuffix.Count > 0 && TargetDescriptor.ForeignPlugin == null)
 				{
-					HotReloadMode = HotReloadMode.FromIDE;
+					HotReloadMode = HotReloadMode.FromEditor;
 				}
-				else if (TargetDescriptor.HotReloadModuleNameToSuffix.Count > 0 && TargetDescriptor.ForeignPlugin == null)
+				else if (BuildConfiguration.bAllowHotReloadFromIDE && HotReload.ShouldDoHotReloadFromIDE(BuildConfiguration, TargetDescriptor))
 				{
-					HotReloadMode = HotReloadMode.FromEditor;
+					HotReloadMode = HotReloadMode.FromIDE;
 				}
 				else
 				{
