Fix null reference exception when packaging for TVOS.

#rnx

diff --git a/Engine/Source/Programs/UnrealBuildTool/Platform/IOS/IOSToolChain.cs b/Engine/Source/Programs/UnrealBuildTool/Platform/IOS/IOSToolChain.cs
--- a/Engine/Source/Programs/UnrealBuildTool/Platform/IOS/IOSToolChain.cs
+++ b/Engine/Source/Programs/UnrealBuildTool/Platform/IOS/IOSToolChain.cs
@@ -76,14 +76,14 @@ namespace UnrealBuildTool
 		protected IOSProjectSettings ProjectSettings;
 
 		public IOSToolChain(ReadOnlyTargetRules Target, IOSProjectSettings InProjectSettings)
-			: this(CppPlatform.IOS, (Target == null)? null : Target.ProjectFile, InProjectSettings, () => new IOSToolChainSettings())
+			: this(CppPlatform.IOS, Target, InProjectSettings, () => new IOSToolChainSettings())
 		{
-			this.Target = Target;
 		}
 
-		protected IOSToolChain(CppPlatform TargetPlatform, FileReference InProjectFile, IOSProjectSettings InProjectSettings, Func<IOSToolChainSettings> InCreateSettings)
-			: base(TargetPlatform, InProjectFile)
+		protected IOSToolChain(CppPlatform TargetPlatform, ReadOnlyTargetRules Target, IOSProjectSettings InProjectSettings, Func<IOSToolChainSettings> InCreateSettings)
+			: base(TargetPlatform, (Target == null)? null : Target.ProjectFile)
 		{
+			this.Target = Target;
 			ProjectSettings = InProjectSettings;
 			Settings = new Lazy<IOSToolChainSettings>(InCreateSettings);
 		}
diff --git a/Engine/Source/Programs/UnrealBuildTool/Platform/TVOS/TVOSToolChain.cs b/Engine/Source/Programs/UnrealBuildTool/Platform/TVOS/TVOSToolChain.cs
--- a/Engine/Source/Programs/UnrealBuildTool/Platform/TVOS/TVOSToolChain.cs
+++ b/Engine/Source/Programs/UnrealBuildTool/Platform/TVOS/TVOSToolChain.cs
@@ -25,8 +25,8 @@ namespace UnrealBuildTool
 
 	class TVOSToolChain : IOSToolChain
 	{
-		public TVOSToolChain(FileReference InProjectFile, TVOSProjectSettings InProjectSettings)
-			: base(CppPlatform.TVOS, InProjectFile, InProjectSettings, () => new TVOSToolChainSettings())
+		public TVOSToolChain(ReadOnlyTargetRules InTarget, TVOSProjectSettings InProjectSettings)
+			: base(CppPlatform.TVOS, InTarget, InProjectSettings, () => new TVOSToolChainSettings())
 		{
 		}
 
diff --git a/Engine/Source/Programs/UnrealBuildTool/Platform/TVOS/UEBuildTVOS.cs b/Engine/Source/Programs/UnrealBuildTool/Platform/TVOS/UEBuildTVOS.cs
--- a/Engine/Source/Programs/UnrealBuildTool/Platform/TVOS/UEBuildTVOS.cs
+++ b/Engine/Source/Programs/UnrealBuildTool/Platform/TVOS/UEBuildTVOS.cs
@@ -127,7 +127,7 @@ namespace UnrealBuildTool
 		public override UEToolChain CreateToolChain(CppPlatform CppPlatform, ReadOnlyTargetRules Target)
 		{
 			TVOSProjectSettings ProjectSettings = ((TVOSPlatform)UEBuildPlatform.GetBuildPlatform(UnrealTargetPlatform.TVOS)).ReadProjectSettings(Target.ProjectFile);
-			return new TVOSToolChain(Target.ProjectFile, ProjectSettings);
+			return new TVOSToolChain(Target, ProjectSettings);
 		}
 
 		public override void Deploy(TargetReceipt Receipt)
