Fix lumin deployment assuming that the target name matches the project name, and is not a client. Can just use the receipt directly instead.

#rnx

diff --git a/Engine/Source/Programs/AutomationTool/Lumin/LuminPlatform.Automation.cs b/Engine/Source/Programs/AutomationTool/Lumin/LuminPlatform.Automation.cs
--- a/Engine/Source/Programs/AutomationTool/Lumin/LuminPlatform.Automation.cs
+++ b/Engine/Source/Programs/AutomationTool/Lumin/LuminPlatform.Automation.cs
@@ -553,7 +553,7 @@ public class LuminPlatform : Platform
 			return;
 		}
 
-		Deploy.InitUPL(Params.ShortProjectName, SC.ProjectRoot, SC.StageTargets[0].Receipt.Configuration);
+		Deploy.InitUPL(SC.StageTargets[0].Receipt);
 
 		string MpkName = GetFinalMpkName(Params, SC);
 		string PackageName = GetPackageName(Params);
diff --git a/Engine/Source/Programs/UnrealBuildTool/Platform/Lumin/LuminExports.cs b/Engine/Source/Programs/UnrealBuildTool/Platform/Lumin/LuminExports.cs
--- a/Engine/Source/Programs/UnrealBuildTool/Platform/Lumin/LuminExports.cs
+++ b/Engine/Source/Programs/UnrealBuildTool/Platform/Lumin/LuminExports.cs
@@ -25,10 +25,8 @@ namespace UnrealBuildTool
 		/// <summary>
 		/// 
 		/// </summary>
-		/// <param name="ProjectName"></param>
-		/// <param name="ProjectDirectory"></param>
-		/// <param name="Configuration"></param>
-		void InitUPL(string ProjectName, DirectoryReference ProjectDirectory, UnrealTargetConfiguration Configuration);
+		/// <param name="Receipt"></param>
+		void InitUPL(TargetReceipt Receipt);
 
 		/// <summary>
 		/// 
diff --git a/Engine/Source/Programs/UnrealBuildTool/Platform/Lumin/UEDeployLumin.cs b/Engine/Source/Programs/UnrealBuildTool/Platform/Lumin/UEDeployLumin.cs
--- a/Engine/Source/Programs/UnrealBuildTool/Platform/Lumin/UEDeployLumin.cs
+++ b/Engine/Source/Programs/UnrealBuildTool/Platform/Lumin/UEDeployLumin.cs
@@ -262,27 +262,21 @@ namespace UnrealBuildTool
 			return PluginExtras;
 		}
 
-		public void InitUPL(string ProjectName, DirectoryReference ProjectDirectory, UnrealTargetConfiguration Configuration)
+		public void InitUPL(TargetReceipt Receipt)
 		{
+			DirectoryReference ProjectDirectory = Receipt.ProjectDir ?? UnrealBuildTool.EngineDirectory;
+
 			string UE4BuildPath = Path.Combine(ProjectDirectory.FullName, "Intermediate/Lumin/Mabu");
 			string RelativeEnginePath = UnrealBuildTool.EngineDirectory.MakeRelativeTo(DirectoryReference.GetCurrentDirectory());
 			string RelativeProjectPath = ProjectDirectory.MakeRelativeTo(DirectoryReference.GetCurrentDirectory());//.MakeRelativeTo(ProjectDirectory);
 
-			UnrealTargetPlatform Platform = UnrealTargetPlatform.Lumin;
-			string ConfigurationString = Configuration.ToString();
+			string ConfigurationString = Receipt.Configuration.ToString();
 
 			string Architecture = "arm64-v8a";
 			List<string> MLSDKArches = new List<string>();
 			MLSDKArches.Add(Architecture);
 
-			// get the receipt
-			FileReference ReceiptFilename = TargetReceipt.GetDefaultPath(ProjectDirectory, ProjectName, Platform, Configuration, "");
-			if (!File.Exists(ReceiptFilename.ToString()))
-			{
-				ReceiptFilename = TargetReceipt.GetDefaultPath(UnrealBuildTool.EngineDirectory, "UE4Game", Platform, Configuration, "");
-			}
-			Log.TraceInformation("Receipt Filename: {0}", ReceiptFilename);
-			SetLuminPluginData(MLSDKArches, CollectPluginDataPaths(TargetReceipt.Read(ReceiptFilename)));
+			SetLuminPluginData(MLSDKArches, CollectPluginDataPaths(Receipt));
 
 			//gather all of the xml
 			UPL.Init(MLSDKArches, true, RelativeEnginePath, UE4BuildPath, RelativeProjectPath, ConfigurationString);
