UBT: Allow pre-build steps to do things that may invalidate the makefile. After loading a makefile, we validate that anything that can affect the pre-build scripts and execute them, then check whether the makefile is still valid for the working set of source files. If not, we discard the makefile and rebuild it (without executing the pre-build steps).

diff --git a/Engine/Source/Programs/UnrealBuildTool/Configuration/UEBuildTarget.cs b/Engine/Source/Programs/UnrealBuildTool/Configuration/UEBuildTarget.cs
--- a/Engine/Source/Programs/UnrealBuildTool/Configuration/UEBuildTarget.cs
+++ b/Engine/Source/Programs/UnrealBuildTool/Configuration/UEBuildTarget.cs
@@ -1197,9 +1197,6 @@ namespace UnrealBuildTool
 			// Create the makefile
 			TargetMakefile Makefile = new TargetMakefile(TargetToolChain.GetVersionInfo(), ReceiptFileName, ProjectIntermediateDirectory, TargetType, bDeployAfterCompile, bHasProjectScriptPlugin);
 
-			// Setup the pre-build steps
-			Makefile.PreBuildScripts = CreatePreBuildScripts();
-
 			// Setup the hot reload module list
 			Makefile.HotReloadModuleNames = GetHotReloadModuleNames();
 
@@ -1834,7 +1831,7 @@ namespace UnrealBuildTool
 		/// <summary>
 		/// Creates scripts for executing the pre-build scripts
 		/// </summary>
-		private FileReference[] CreatePreBuildScripts()
+		public FileReference[] CreatePreBuildScripts()
 		{
 			// Find all the pre-build steps
 			List<Tuple<string[], UEBuildPlugin>> PreBuildCommandBatches = new List<Tuple<string[], UEBuildPlugin>>();
diff --git a/Engine/Source/Programs/UnrealBuildTool/Modes/BuildMode.cs b/Engine/Source/Programs/UnrealBuildTool/Modes/BuildMode.cs
--- a/Engine/Source/Programs/UnrealBuildTool/Modes/BuildMode.cs
+++ b/Engine/Source/Programs/UnrealBuildTool/Modes/BuildMode.cs
@@ -224,12 +224,6 @@ namespace UnrealBuildTool
 			// Execute the build
 			if((Options & BuildOptions.SkipBuild) == 0)
 			{
-				// Execute all the pre-build steps
-				foreach(TargetMakefile Makefile in Makefiles)
-				{
-					Utils.ExecuteCustomBuildSteps(Makefile.PreBuildScripts);
-				}
-
 				// Make sure that none of the actions conflict with any other (producing output files differently, etc...)
 				ActionGraph.CheckForConflicts(Makefiles.SelectMany(x => x.Actions));
 
@@ -368,7 +362,7 @@ namespace UnrealBuildTool
 				using(Timeline.ScopeEvent("TargetMakefile.Load()"))
 				{
 					string ReasonNotLoaded;
-					Makefile = TargetMakefile.Load(MakefileLocation, TargetDescriptor.ProjectFile, TargetDescriptor.Platform, TargetDescriptor.AdditionalArguments.GetRawArray(), WorkingSet, out ReasonNotLoaded);
+					Makefile = TargetMakefile.Load(MakefileLocation, TargetDescriptor.ProjectFile, TargetDescriptor.Platform, TargetDescriptor.AdditionalArguments.GetRawArray(), out ReasonNotLoaded);
 					if (Makefile == null)
 					{
 						Log.TraceInformation("Creating makefile for {0} ({1})", TargetDescriptor.Name, ReasonNotLoaded);
@@ -376,6 +370,29 @@ namespace UnrealBuildTool
 				}
 			}
 
+			// If we have a makefile, execute the pre-build steps and check it's still valid
+			bool bHasRunPreBuildScripts = false;
+			if(Makefile != null)
+			{
+				// Execute the scripts. We have to invalidate all cached file info after doing so, because we don't know what may have changed.
+				if(Makefile.PreBuildScripts.Length > 0)
+				{
+					Utils.ExecuteCustomBuildSteps(Makefile.PreBuildScripts);
+					DirectoryItem.ResetAllCachedInfo_SLOW();
+				}
+
+				// Don't run the pre-build steps again, even if we invalidate the makefile.
+				bHasRunPreBuildScripts = true;
+
+				// Check that the makefile is still valid
+				string Reason;
+				if(!TargetMakefile.IsValidForSourceFiles(Makefile, TargetDescriptor.ProjectFile, WorkingSet, out Reason))
+				{
+					Log.TraceInformation("Invalidating makefile for {0} ({1})", TargetDescriptor.Name, Reason);
+					Makefile = null;
+				}
+			}
+
 			// If we couldn't load a makefile, create a new one
 			if(Makefile == null)
 			{
@@ -386,14 +403,26 @@ namespace UnrealBuildTool
 					Target = UEBuildTarget.Create(TargetDescriptor, BuildConfiguration.bSkipRulesCompile, BuildConfiguration.bUsePrecompiled);
 				}
 
-				// Build the target
-				const bool bIsAssemblingBuild = true;
+				// Create the pre-build scripts
+				FileReference[] PreBuildScripts = Target.CreatePreBuildScripts();
+
+				// Execute the pre-build scripts
+				if(!bHasRunPreBuildScripts)
+				{
+					Utils.ExecuteCustomBuildSteps(PreBuildScripts);
+					bHasRunPreBuildScripts = true;
+				}
 
+				// Build the target
 				using(Timeline.ScopeEvent("UEBuildTarget.Build()"))
 				{
+					const bool bIsAssemblingBuild = true;
 					Makefile = Target.Build(BuildConfiguration, WorkingSet, bIsAssemblingBuild);
 				}
 
+				// Save the pre-build scripts onto the makefile
+				Makefile.PreBuildScripts = PreBuildScripts;
+
 				// Save the additional command line arguments
 				Makefile.AdditionalArguments = TargetDescriptor.AdditionalArguments.GetRawArray();
 
diff --git a/Engine/Source/Programs/UnrealBuildTool/System/DirectoryItem.cs b/Engine/Source/Programs/UnrealBuildTool/System/DirectoryItem.cs
--- a/Engine/Source/Programs/UnrealBuildTool/System/DirectoryItem.cs
+++ b/Engine/Source/Programs/UnrealBuildTool/System/DirectoryItem.cs
@@ -165,6 +165,8 @@ namespace UnrealBuildTool
 		/// </summary>
 		public void ResetCachedInfo()
 		{
+			Info = new DirectoryInfo(Info.FullName);
+
 			Dictionary<string, DirectoryItem> PrevDirectories = Directories;
 			if(PrevDirectories != null)
 			{
@@ -186,6 +188,20 @@ namespace UnrealBuildTool
 			}
 		}
 
+		/// <summary>
+		/// Resets all cached directory info. Significantly reduces performance; do not use unless strictly necessary.
+		/// </summary>
+		public static void ResetAllCachedInfo_SLOW()
+		{
+			foreach(DirectoryItem Item in LocationToItem.Values)
+			{
+				Item.Info = new DirectoryInfo(Item.Info.FullName);
+				Item.Directories = null;
+				Item.Files = null;
+			}
+			FileItem.ResetAllCachedInfo_SLOW();
+		}
+
 		/// <summary>
 		/// Caches the subdirectories of this directories
 		/// </summary>
diff --git a/Engine/Source/Programs/UnrealBuildTool/System/FileItem.cs b/Engine/Source/Programs/UnrealBuildTool/System/FileItem.cs
--- a/Engine/Source/Programs/UnrealBuildTool/System/FileItem.cs
+++ b/Engine/Source/Programs/UnrealBuildTool/System/FileItem.cs
@@ -282,6 +282,17 @@ namespace UnrealBuildTool
 			Info = Location.ToFileInfo();
 		}
 
+		/// <summary>
+		/// Resets all cached file info. Significantly reduces performance; do not use unless strictly necessary.
+		/// </summary>
+		public static void ResetAllCachedInfo_SLOW()
+		{
+			foreach(FileItem Item in UniqueSourceFileMap.Values)
+			{
+				Item.ResetCachedInfo();
+			}
+		}
+
 		/// <summary>
 		/// Return the path to this FileItem to debugging
 		/// </summary>
diff --git a/Engine/Source/Programs/UnrealBuildTool/System/TargetMakefile.cs b/Engine/Source/Programs/UnrealBuildTool/System/TargetMakefile.cs
--- a/Engine/Source/Programs/UnrealBuildTool/System/TargetMakefile.cs
+++ b/Engine/Source/Programs/UnrealBuildTool/System/TargetMakefile.cs
@@ -20,7 +20,12 @@ namespace UnrealBuildTool
 		/// <summary>
 		/// The version number to write
 		/// </summary>
-		public const int CurrentVersion = 9;
+		public const int CurrentVersion = 10;
+
+		/// <summary>
+		/// The time at which the makefile was created
+		/// </summary>
+		public DateTime CreateTimeUtc;
 
 		/// <summary>
 		/// Information about the toolchain used to build. This string will be output before building.
@@ -134,6 +139,7 @@ namespace UnrealBuildTool
 		/// <param name="bHasProjectScriptPlugin">Whether the target has a project script plugin</param>
 		public TargetMakefile(string ToolchainInfo, FileReference ReceiptFile, DirectoryReference ProjectIntermediateDirectory, TargetType TargetType, bool bDeployAfterCompile, bool bHasProjectScriptPlugin)
 		{
+			this.CreateTimeUtc = DateTime.UtcNow;
 			this.ToolchainInfo = ToolchainInfo;
 			this.ReceiptFile = ReceiptFile;
 			this.ProjectIntermediateDirectory = ProjectIntermediateDirectory;
@@ -159,6 +165,7 @@ namespace UnrealBuildTool
 		/// <param name="Reader">The archive to read from</param>
 		public TargetMakefile(BinaryArchiveReader Reader)
 		{
+			CreateTimeUtc = new DateTime(Reader.ReadLong(), DateTimeKind.Utc);
 			ToolchainInfo = Reader.ReadString();
 			ReceiptFile = Reader.ReadFileReference();
 			ProjectIntermediateDirectory = Reader.ReadDirectoryReference();
@@ -187,6 +194,7 @@ namespace UnrealBuildTool
 		/// <param name="Writer">The archive to write to</param>
 		public void Write(BinaryArchiveWriter Writer)
 		{
+			Writer.WriteLong(CreateTimeUtc.Ticks);
 			Writer.WriteString(ToolchainInfo);
 			Writer.WriteFileReference(ReceiptFile);
 			Writer.WriteDirectoryReference(ProjectIntermediateDirectory);
@@ -230,106 +238,108 @@ namespace UnrealBuildTool
 		/// <param name="ProjectFile">Path to the project file</param>
 		/// <param name="Platform">Platform for this makefile</param>
 		/// <param name="Arguments">Command line arguments for this target</param>
-		/// <param name="WorkingSet">Interface to query which source files are in the working set</param>
 		/// <param name="ReasonNotLoaded">If the function returns null, this string will contain the reason why</param>
 		/// <returns>The loaded makefile, or null if it failed for some reason.  On failure, the 'ReasonNotLoaded' variable will contain information about why</returns>
-		public static TargetMakefile Load(FileReference MakefilePath, FileReference ProjectFile, UnrealTargetPlatform Platform, string[] Arguments, ISourceFileWorkingSet WorkingSet, out string ReasonNotLoaded)
+		public static TargetMakefile Load(FileReference MakefilePath, FileReference ProjectFile, UnrealTargetPlatform Platform, string[] Arguments, out string ReasonNotLoaded)
 		{
-			// Check the directory timestamp on the project files directory.  If the user has generated project files more recently than the makefile, then we need to consider the file to be out of date
-			FileInfo MakefileInfo = new FileInfo(MakefilePath.FullName);
-			if (!MakefileInfo.Exists)
+			using(Timeline.ScopeEvent("Checking dependent timestamps"))
 			{
-				// Makefile doesn't even exist, so we won't bother loading it
-				ReasonNotLoaded = "no existing makefile";
-				return null;
-			}
+				// Check the directory timestamp on the project files directory.  If the user has generated project files more recently than the makefile, then we need to consider the file to be out of date
+				FileInfo MakefileInfo = new FileInfo(MakefilePath.FullName);
+				if (!MakefileInfo.Exists)
+				{
+					// Makefile doesn't even exist, so we won't bother loading it
+					ReasonNotLoaded = "no existing makefile";
+					return null;
+				}
 
-			// Check the build version
-			FileInfo BuildVersionFileInfo = new FileInfo(BuildVersion.GetDefaultFileName().FullName);
-			if (BuildVersionFileInfo.Exists && MakefileInfo.LastWriteTime.CompareTo(BuildVersionFileInfo.LastWriteTime) < 0)
-			{
-				Log.TraceLog("Existing makefile is older than Build.version, ignoring it");
-				ReasonNotLoaded = "Build.version is newer";
-				return null;
-			}
+				// Check the build version
+				FileInfo BuildVersionFileInfo = new FileInfo(BuildVersion.GetDefaultFileName().FullName);
+				if (BuildVersionFileInfo.Exists && MakefileInfo.LastWriteTime.CompareTo(BuildVersionFileInfo.LastWriteTime) < 0)
+				{
+					Log.TraceLog("Existing makefile is older than Build.version, ignoring it");
+					ReasonNotLoaded = "Build.version is newer";
+					return null;
+				}
 
-			// @todo ubtmake: This will only work if the directory timestamp actually changes with every single GPF.  Force delete existing files before creating new ones?  Eh... really we probably just want to delete + create a file in that folder
-			//			-> UPDATE: Seems to work OK right now though on Windows platform, maybe due to GUID changes
-			// @todo ubtmake: Some platforms may not save any files into this folder.  We should delete + generate a "touch" file to force the directory timestamp to be updated (or just check the timestamp file itself.  We could put it ANYWHERE, actually)
+				// @todo ubtmake: This will only work if the directory timestamp actually changes with every single GPF.  Force delete existing files before creating new ones?  Eh... really we probably just want to delete + create a file in that folder
+				//			-> UPDATE: Seems to work OK right now though on Windows platform, maybe due to GUID changes
+				// @todo ubtmake: Some platforms may not save any files into this folder.  We should delete + generate a "touch" file to force the directory timestamp to be updated (or just check the timestamp file itself.  We could put it ANYWHERE, actually)
 
-			// Installed Build doesn't need to check engine projects for outdatedness
-			if (!UnrealBuildTool.IsEngineInstalled())
-			{
-				if (DirectoryReference.Exists(ProjectFileGenerator.IntermediateProjectFilesPath))
+				// Installed Build doesn't need to check engine projects for outdatedness
+				if (!UnrealBuildTool.IsEngineInstalled())
+				{
+					if (DirectoryReference.Exists(ProjectFileGenerator.IntermediateProjectFilesPath))
+					{
+						DateTime EngineProjectFilesLastUpdateTime = new FileInfo(ProjectFileGenerator.ProjectTimestampFile).LastWriteTime;
+						if (MakefileInfo.LastWriteTime.CompareTo(EngineProjectFilesLastUpdateTime) < 0)
+						{
+							// Engine project files are newer than makefile
+							Log.TraceLog("Existing makefile is older than generated engine project files, ignoring it");
+							ReasonNotLoaded = "project files are newer";
+							return null;
+						}
+					}
+				}
+
+				// Check the game project directory too
+				if (ProjectFile != null)
 				{
-					DateTime EngineProjectFilesLastUpdateTime = new FileInfo(ProjectFileGenerator.ProjectTimestampFile).LastWriteTime;
-					if (MakefileInfo.LastWriteTime.CompareTo(EngineProjectFilesLastUpdateTime) < 0)
+					string ProjectFilename = ProjectFile.FullName;
+					FileInfo ProjectFileInfo = new FileInfo(ProjectFilename);
+					if (!ProjectFileInfo.Exists || MakefileInfo.LastWriteTime.CompareTo(ProjectFileInfo.LastWriteTime) < 0)
 					{
-						// Engine project files are newer than makefile
-						Log.TraceLog("Existing makefile is older than generated engine project files, ignoring it");
-						ReasonNotLoaded = "project files are newer";
+						// .uproject file is newer than makefile
+						Log.TraceLog("Makefile is older than .uproject file, ignoring it");
+						ReasonNotLoaded = ".uproject file is newer";
 						return null;
 					}
+
+					DirectoryReference MasterProjectRelativePath = ProjectFile.Directory;
+					string GameIntermediateProjectFilesPath = Path.Combine(MasterProjectRelativePath.FullName, "Intermediate", "ProjectFiles");
+					if (Directory.Exists(GameIntermediateProjectFilesPath))
+					{
+						DateTime GameProjectFilesLastUpdateTime = new DirectoryInfo(GameIntermediateProjectFilesPath).LastWriteTime;
+						if (MakefileInfo.LastWriteTime.CompareTo(GameProjectFilesLastUpdateTime) < 0)
+						{
+							// Game project files are newer than makefile
+							Log.TraceLog("Makefile is older than generated game project files, ignoring it");
+							ReasonNotLoaded = "game project files are newer";
+							return null;
+						}
+					}
 				}
-			}
 
-			// Check the game project directory too
-			if (ProjectFile != null)
-			{
-				string ProjectFilename = ProjectFile.FullName;
-				FileInfo ProjectFileInfo = new FileInfo(ProjectFilename);
-				if (!ProjectFileInfo.Exists || MakefileInfo.LastWriteTime.CompareTo(ProjectFileInfo.LastWriteTime) < 0)
+				// Check to see if UnrealBuildTool.exe was compiled more recently than the makefile
+				DateTime UnrealBuildToolTimestamp = new FileInfo(Assembly.GetExecutingAssembly().Location).LastWriteTime;
+				if (MakefileInfo.LastWriteTime.CompareTo(UnrealBuildToolTimestamp) < 0)
 				{
-					// .uproject file is newer than makefile
-					Log.TraceLog("Makefile is older than .uproject file, ignoring it");
-					ReasonNotLoaded = ".uproject file is newer";
+					// UnrealBuildTool.exe was compiled more recently than the makefile
+					Log.TraceLog("Makefile is older than UnrealBuildTool.exe, ignoring it");
+					ReasonNotLoaded = "UnrealBuildTool.exe is newer";
 					return null;
 				}
 
-				DirectoryReference MasterProjectRelativePath = ProjectFile.Directory;
-				string GameIntermediateProjectFilesPath = Path.Combine(MasterProjectRelativePath.FullName, "Intermediate", "ProjectFiles");
-				if (Directory.Exists(GameIntermediateProjectFilesPath))
+				// Check to see if any BuildConfiguration files have changed since the last build
+				List<XmlConfig.InputFile> InputFiles = XmlConfig.FindInputFiles();
+				foreach (XmlConfig.InputFile InputFile in InputFiles)
 				{
-					DateTime GameProjectFilesLastUpdateTime = new DirectoryInfo(GameIntermediateProjectFilesPath).LastWriteTime;
-					if (MakefileInfo.LastWriteTime.CompareTo(GameProjectFilesLastUpdateTime) < 0)
+					FileInfo InputFileInfo = new FileInfo(InputFile.Location.FullName);
+					if (InputFileInfo.LastWriteTime > MakefileInfo.LastWriteTime)
 					{
-						// Game project files are newer than makefile
-						Log.TraceLog("Makefile is older than generated game project files, ignoring it");
-						ReasonNotLoaded = "game project files are newer";
+						Log.TraceLog("Makefile is older than BuildConfiguration.xml, ignoring it");
+						ReasonNotLoaded = "BuildConfiguration.xml is newer";
 						return null;
 					}
 				}
 			}
 
-			// Check to see if UnrealBuildTool.exe was compiled more recently than the makefile
-			DateTime UnrealBuildToolTimestamp = new FileInfo(Assembly.GetExecutingAssembly().Location).LastWriteTime;
-			if (MakefileInfo.LastWriteTime.CompareTo(UnrealBuildToolTimestamp) < 0)
-			{
-				// UnrealBuildTool.exe was compiled more recently than the makefile
-				Log.TraceLog("Makefile is older than UnrealBuildTool.exe, ignoring it");
-				ReasonNotLoaded = "UnrealBuildTool.exe is newer";
-				return null;
-			}
-
-			// Check to see if any BuildConfiguration files have changed since the last build
-			List<XmlConfig.InputFile> InputFiles = XmlConfig.FindInputFiles();
-			foreach (XmlConfig.InputFile InputFile in InputFiles)
-			{
-				FileInfo InputFileInfo = new FileInfo(InputFile.Location.FullName);
-				if (InputFileInfo.LastWriteTime > MakefileInfo.LastWriteTime)
-				{
-					Log.TraceLog("Makefile is older than BuildConfiguration.xml, ignoring it");
-					ReasonNotLoaded = "BuildConfiguration.xml is newer";
-					return null;
-				}
-			}
-
 			TargetMakefile Makefile;
 			using(Timeline.ScopeEvent("Loading makefile"))
 			{
 				try
 				{
-					using(BinaryArchiveReader Reader = new BinaryArchiveReader(new FileReference(MakefileInfo)))
+					using(BinaryArchiveReader Reader = new BinaryArchiveReader(MakefilePath))
 					{
 						int Version = Reader.ReadInt();
 						if(Version != CurrentVersion)
@@ -365,7 +375,7 @@ namespace UnrealBuildTool
 					foreach (FileReference IniFilename in ConfigHierarchy.EnumerateConfigFileLocations(IniType, ProjectDirectory, Platform))
 					{
 						FileInfo IniFileInfo = new FileInfo(IniFilename.FullName);
-						if (MakefileInfo.LastWriteTime.CompareTo(IniFileInfo.LastWriteTime) < 0)
+						if(IniFileInfo.LastWriteTimeUtc > Makefile.CreateTimeUtc)
 						{
 							// Ini files are newer than makefile
 							ReasonNotLoaded = "ini files are newer than makefile";
@@ -373,27 +383,46 @@ namespace UnrealBuildTool
 						}
 					}
 				}
+			}
 
+			// The makefile is ok
+			ReasonNotLoaded = null;
+			return Makefile;
+		}
+
+		/// <summary>
+		/// Checks if the makefile is valid for the current set of source files. This is done separately to the Load() method to allow pre-build steps to modify source files.
+		/// </summary>
+		/// <param name="Makefile">The makefile that has been loaded</param>
+		/// <param name="ProjectFile">Path to the project file</param>
+		/// <param name="WorkingSet">The current working set of source files</param>
+		/// <param name="ReasonNotLoaded">If the makefile is not valid, is set to a message describing why</param>
+		/// <returns>True if the makefile is valid, false otherwise</returns>
+		public static bool IsValidForSourceFiles(TargetMakefile Makefile, FileReference ProjectFile, ISourceFileWorkingSet WorkingSet, out string ReasonNotLoaded)
+		{
+			using(Timeline.ScopeEvent("TargetMakefile.IsValidForSourceFiles()"))
+			{
+				// Check if any source files have been added or removed
 				foreach(KeyValuePair<DirectoryItem, FileItem[]> Pair in Makefile.DirectoryToSourceFiles)
 				{
 					DirectoryItem InputDirectory = Pair.Key;
-					if(!InputDirectory.Exists || InputDirectory.LastWriteTimeUtc > MakefileInfo.LastWriteTimeUtc)
+					if(!InputDirectory.Exists || InputDirectory.LastWriteTimeUtc > Makefile.CreateTimeUtc)
 					{
 						FileItem[] SourceFiles = UEBuildModuleCPP.GetSourceFiles(InputDirectory);
 						if(SourceFiles.Length < Pair.Value.Length)
 						{
 							ReasonNotLoaded = "source file removed";
-							return null;
+							return false;
 						}
 						else if(SourceFiles.Length > Pair.Value.Length)
 						{
 							ReasonNotLoaded = "source file added";
-							return null;
+							return false;
 						}
 						else if(SourceFiles.Intersect(Pair.Value).Count() != SourceFiles.Length)
 						{
 							ReasonNotLoaded = "source file modified";
-							return null;
+							return false;
 						}
 					}
 				}
@@ -405,13 +434,13 @@ namespace UnrealBuildTool
 					{
 						Log.TraceLog("{0} has been deleted since makefile was built.", AdditionalDependency.Location);
 						ReasonNotLoaded = string.Format("{0} deleted", AdditionalDependency.Location.GetFileName());
-						return null;
+						return false;
 					}
-					if(AdditionalDependency.LastWriteTimeUtc > MakefileInfo.LastWriteTimeUtc)
+					if(AdditionalDependency.LastWriteTimeUtc > Makefile.CreateTimeUtc)
 					{
 						Log.TraceLog("{0} has been modified since makefile was built.", AdditionalDependency.Location);
 						ReasonNotLoaded = string.Format("{0} modified", AdditionalDependency.Location.GetFileName());
-						return null;
+						return false;
 					}
 				}
 
@@ -423,7 +452,7 @@ namespace UnrealBuildTool
 					{
 						Log.TraceLog("{0} has been added", PluginFile.GetFileName());
 						ReasonNotLoaded = string.Format("{0} has been added", PluginFile.GetFileName());
-						return null;
+						return false;
 					}
 				}
 
@@ -437,7 +466,7 @@ namespace UnrealBuildTool
 				{
 					foreach(FileItem HeaderFile in ModuleHeaderInfo.SourceFolder.EnumerateFiles())
 					{
-						if(HeaderFile.HasExtension(".h") && HeaderFile.LastWriteTimeUtc > MakefileInfo.LastWriteTimeUtc)
+						if(HeaderFile.HasExtension(".h") && HeaderFile.LastWriteTimeUtc > Makefile.CreateTimeUtc)
 						{
 							HFilesNewerThanMakefile.Add(HeaderFile);
 						}
@@ -455,7 +484,7 @@ namespace UnrealBuildTool
 					{
 						Log.TraceLog("File processed by UHT was deleted ({0}); invalidating makefile", HeaderFile);
 						ReasonNotLoaded = string.Format("UHT file was deleted");
-						return null;
+						return false;
 					}
 				}
 
@@ -469,9 +498,9 @@ namespace UnrealBuildTool
 					bool bWasProcessed = AllUHTHeaders.Contains(HeaderFile);
 					if (bContainsUHTData != bWasProcessed)
 					{
-						Log.TraceLog("{0} {1} contain UHT types and now {2} , ignoring it ({3})", HeaderFile, bWasProcessed ? "used to" : "didn't", bWasProcessed ? "doesn't" : "does", MakefileInfo.FullName);
+						Log.TraceLog("{0} {1} contain UHT types and now {2} , ignoring it", HeaderFile, bWasProcessed ? "used to" : "didn't", bWasProcessed ? "doesn't" : "does");
 						ReasonNotLoaded = string.Format("new files with reflected types");
-						return null;
+						return false;
 					}
 				}
 
@@ -484,28 +513,28 @@ namespace UnrealBuildTool
 				// Check if any source files in the working set no longer belong in it
 				foreach (FileItem SourceFile in Makefile.WorkingSet)
 				{
-					if (!WorkingSet.Contains(SourceFile.Location) && SourceFile.LastWriteTimeUtc > MakefileInfo.LastWriteTimeUtc)
+					if (!WorkingSet.Contains(SourceFile.Location) && SourceFile.LastWriteTimeUtc > Makefile.CreateTimeUtc)
 					{
-						Log.TraceLog("{0} was part of source working set and now is not; invalidating makefile ({1})", SourceFile.AbsolutePath, MakefileInfo.FullName);
+						Log.TraceLog("{0} was part of source working set and now is not; invalidating makefile", SourceFile.AbsolutePath);
 						ReasonNotLoaded = string.Format("working set of source files changed");
-						return null;
+						return false;
 					}
 				}
 
 				// Check if any source files that are eligible for being in the working set have been modified
 				foreach (FileItem SourceFile in Makefile.CandidatesForWorkingSet)
 				{
-					if (WorkingSet.Contains(SourceFile.Location) && SourceFile.LastWriteTimeUtc > MakefileInfo.LastWriteTimeUtc)
+					if (WorkingSet.Contains(SourceFile.Location) && SourceFile.LastWriteTimeUtc > Makefile.CreateTimeUtc)
 					{
-						Log.TraceLog("{0} was part of source working set and now is not; invalidating makefile ({1})", SourceFile.AbsolutePath, MakefileInfo.FullName);
+						Log.TraceLog("{0} was part of source working set and now is not", SourceFile.AbsolutePath);
 						ReasonNotLoaded = string.Format("working set of source files changed");
-						return null;
+						return false;
 					}
 				}
 			}
 
 			ReasonNotLoaded = null;
-			return Makefile;
+			return true;
 		}
 
 		/// <summary>
