Fix *.user files not being generated when the DisablePlatformProjectGenerators config entry is empty.

diff --git a/Engine/Source/Programs/UnrealBuildTool/Modes/GenerateProjectFilesMode.cs b/Engine/Source/Programs/UnrealBuildTool/Modes/GenerateProjectFilesMode.cs
--- a/Engine/Source/Programs/UnrealBuildTool/Modes/GenerateProjectFilesMode.cs
+++ b/Engine/Source/Programs/UnrealBuildTool/Modes/GenerateProjectFilesMode.cs
@@ -93,7 +93,7 @@ namespace UnrealBuildTool
 					PlatformProjectGenerator Generator = (PlatformProjectGenerator)Activator.CreateInstance(CheckType, Arguments);
 					foreach(UnrealTargetPlatform Platform in Generator.GetPlatforms())
 					{
-						if(DisablePlatformProjectGenerators != null && !DisablePlatformProjectGenerators.Any(x => x.Equals(Platform.ToString(), StringComparison.OrdinalIgnoreCase)))
+						if(DisablePlatformProjectGenerators == null || !DisablePlatformProjectGenerators.Any(x => x.Equals(Platform.ToString(), StringComparison.OrdinalIgnoreCase)))
 						{
 							Log.TraceVerbose("Registering project generator {0} for {1}", CheckType, Platform);
 							PlatformProjectGenerators.RegisterPlatformProjectGenerator(Platform, Generator);
