Acquire a global mutex when writing out target metadata. Some targets (eg. Development and DebugGame editors) can share output files, and cannot be updated at the same time.

#rnx

diff --git a/Engine/Source/Programs/UnrealBuildTool/Modes/WriteMetadataMode.cs b/Engine/Source/Programs/UnrealBuildTool/Modes/WriteMetadataMode.cs
--- a/Engine/Source/Programs/UnrealBuildTool/Modes/WriteMetadataMode.cs
+++ b/Engine/Source/Programs/UnrealBuildTool/Modes/WriteMetadataMode.cs
@@ -78,6 +78,22 @@ namespace UnrealBuildTool
 		/// <param name="Arguments">Command line arguments</param>
 		/// <returns>Exit code</returns>
 		public override int Execute(CommandLineArguments Arguments)
+		{
+			// Acquire a different mutex to the regular UBT instance, since this mode will be called as part of a build. We need the mutex to ensure that building two modular configurations 
+			// in parallel don't clash over writing shared *.modules files (eg. DebugGame and Development editors).
+			string MutexName = SingleInstanceMutex.GetUniqueMutexForPath("UnrealBuildTool_WriteMetadata", UnrealBuildTool.RootDirectory.FullName);
+			using(new SingleInstanceMutex(MutexName, true))
+			{
+				return ExecuteInternal(Arguments);
+			}
+		}
+
+		/// <summary>
+		/// Execute the command, having obtained the appropriate mutex
+		/// </summary>
+		/// <param name="Arguments">Command line arguments</param>
+		/// <returns>Exit code</returns>
+		private int ExecuteInternal(CommandLineArguments Arguments)
 		{
 			// Read the target info
 			WriteMetadataTargetInfo TargetInfo = BinaryFormatterUtils.Load<WriteMetadataTargetInfo>(Arguments.GetFileReference("-Input="));
diff --git a/Engine/Source/Programs/UnrealBuildTool/System/SingleInstanceMutex.cs b/Engine/Source/Programs/UnrealBuildTool/System/SingleInstanceMutex.cs
--- a/Engine/Source/Programs/UnrealBuildTool/System/SingleInstanceMutex.cs
+++ b/Engine/Source/Programs/UnrealBuildTool/System/SingleInstanceMutex.cs
@@ -11,22 +11,6 @@ using Tools.DotNETCommon;
 
 namespace UnrealBuildTool
 {
-	/// <summary>
-	/// The type of mutex to aquire
-	/// </summary>
-	enum SingleInstanceMutexType
-	{
-		/// <summary>
-		/// Prevent any instance running in the system
-		/// </summary>
-		Global,
-
-		/// <summary>
-		/// Prevents any instance running in the current branch
-		/// </summary>
-		PerBranch,
-	}
-
 	/// <summary>
 	/// System-wide mutex allowing only one instance of the program to run at a time
 	/// </summary>
@@ -40,21 +24,10 @@ namespace UnrealBuildTool
 		/// <summary>
 		/// Constructor. Attempts to acquire the global mutex
 		/// </summary>
-		/// <param name="Type">The type of mutex to acquire</param>
+		/// <param name="MutexName">Name of the mutex to acquire</param>
 		/// <param name="bWaitMutex"></param>
-		public SingleInstanceMutex(SingleInstanceMutexType Type, bool bWaitMutex)
+		public SingleInstanceMutex(string MutexName, bool bWaitMutex)
 		{
-			// Get the mutex name
-			string MutexName;
-			if (Type == SingleInstanceMutexType.Global)
-			{
-				MutexName = "Global\\UnrealBuildTool_Mutex_AutoSDKS";
-			}
-			else
-			{
-				MutexName = "Global\\UnrealBuildTool_Mutex_" + Assembly.GetExecutingAssembly().CodeBase.GetHashCode().ToString();
-			}
-
 			// Try to create the mutex, with it initially locked
 			bool bCreatedMutex;
 			GlobalMutex = new Mutex(true, MutexName, out bCreatedMutex);
@@ -72,33 +45,21 @@ namespace UnrealBuildTool
 					{
 					}
 				}
-				else if (Type == SingleInstanceMutexType.Global)
-				{
-					throw new BuildException("A conflicting instance of UnrealBuildTool is already running.");
-				}
 				else
 				{
-					throw new BuildException("A conflicting instance of UnrealBuildTool is already running at {0}.", Assembly.GetExecutingAssembly().CodeBase);
+					throw new BuildException("A conflicting instance of UnrealBuildTool is already running.");
 				}
 			}
 		}
 
 		/// <summary>
-		/// Attempt to aquire the single instance mutex, using settings determined from the command line
+		/// Gets the name of a mutex unique for the given path
 		/// </summary>
-		/// <param name="Type">The type of mutex to aquire</param>
-		/// <param name="Arguments">Command line arguments</param>
-		/// <returns>The single instance mutex, or null if the -NoMutex argument is specified on the command line</returns>
-		public static SingleInstanceMutex Acquire(SingleInstanceMutexType Type, CommandLineArguments Arguments)
+		/// <param name="Name">Base name of the mutex</param>
+		/// <param name="UniquePath">Path to identify a unique mutex</param>
+		public static string GetUniqueMutexForPath(string Name, string UniquePath)
 		{
-			if(Arguments.HasOption("-NoMutex"))
-			{
-				return null;
-			}
-			else
-			{
-				return new SingleInstanceMutex(Type, Arguments.HasOption("-WaitMutex"));
-			}
+			return String.Format("Global\\{0}_{1}", Name, UniquePath.GetHashCode());
 		}
 
 		/// <summary>
diff --git a/Engine/Source/Programs/UnrealBuildTool/UnrealBuildTool.cs b/Engine/Source/Programs/UnrealBuildTool/UnrealBuildTool.cs
--- a/Engine/Source/Programs/UnrealBuildTool/UnrealBuildTool.cs
+++ b/Engine/Source/Programs/UnrealBuildTool/UnrealBuildTool.cs
@@ -265,6 +265,18 @@ namespace UnrealBuildTool
 			[CommandLine(Prefix = "-Progress")]
 			public bool bWriteProgressMarkup = false;
 
+			/// <summary>
+			/// Whether to ignore the mutex
+			/// </summary>
+			[CommandLine(Prefix = "-NoMutex")]
+			public bool bNoMutex = false;
+
+			/// <summary>
+			/// Whether to wait for the mutex rather than aborting immediately
+			/// </summary>
+			[CommandLine(Prefix = "-WaitMutex")]
+			public bool bWaitMutex = false;
+
 			/// <summary>
 			/// The mode to execute
 			/// </summary>
@@ -400,11 +412,12 @@ namespace UnrealBuildTool
 				}
 
 				// Acquire a lock for this branch
-				if((ModeOptions & ToolModeOptions.SingleInstance) != 0)
+				if((ModeOptions & ToolModeOptions.SingleInstance) != 0 && !Options.bNoMutex)
 				{
 					using(Timeline.ScopeEvent("SingleInstanceMutex.Acquire()"))
 					{
-						Mutex = SingleInstanceMutex.Acquire(SingleInstanceMutexType.PerBranch, Arguments);
+						string MutexName = SingleInstanceMutex.GetUniqueMutexForPath("UnrealBuildTool_Mutex", Assembly.GetExecutingAssembly().CodeBase);
+						Mutex = new SingleInstanceMutex(MutexName, Options.bWaitMutex);
 					}
 				}
 
