Fixed a problem that with UBT not generating debug symbols for Mac UE4Editor-Cmd

diff --git a/Engine/Source/Programs/UnrealBuildTool/Configuration/UEBuildBinary.cs b/Engine/Source/Programs/UnrealBuildTool/Configuration/UEBuildBinary.cs
--- a/Engine/Source/Programs/UnrealBuildTool/Configuration/UEBuildBinary.cs
+++ b/Engine/Source/Programs/UnrealBuildTool/Configuration/UEBuildBinary.cs
@@ -211,7 +211,13 @@ namespace UnrealBuildTool
 				ConsoleAppLinkEvironment.OutputFilePaths = ConsoleAppLinkEvironment.OutputFilePaths.Select(Path => GetAdditionalConsoleAppPath(Path)).ToList();
 
 				// Link the console app executable
-				OutputFiles.AddRange(ToolChain.LinkAllFiles(ConsoleAppLinkEvironment, false, Makefile.Actions));
+				FileItem[] ConsoleAppOutputFiles = ToolChain.LinkAllFiles(ConsoleAppLinkEvironment, false, Makefile.Actions);
+				OutputFiles.AddRange(ConsoleAppOutputFiles);
+
+				foreach (FileItem Executable in ConsoleAppOutputFiles)
+				{
+					OutputFiles.AddRange(ToolChain.PostBuild(Executable, ConsoleAppLinkEvironment, Makefile.Actions));
+				}
 			}
 
 			foreach (FileItem Executable in Executables)
diff --git a/Engine/Source/Programs/UnrealBuildTool/Platform/Mac/MacToolChain.cs b/Engine/Source/Programs/UnrealBuildTool/Platform/Mac/MacToolChain.cs
--- a/Engine/Source/Programs/UnrealBuildTool/Platform/Mac/MacToolChain.cs
+++ b/Engine/Source/Programs/UnrealBuildTool/Platform/Mac/MacToolChain.cs
@@ -1340,7 +1340,7 @@ namespace UnrealBuildTool
 				}
 			}
 
-			if (BinaryLinkEnvironment.bIsBuildingDLL)
+			if (BinaryLinkEnvironment.bIsBuildingDLL || (BinaryLinkEnvironment.bIsBuildingConsoleApplication && Executable.Name.Contains("/UE4Editor") && Executable.Name.EndsWith("-Cmd")))
 			{
 				return OutputFiles;
 			}
