UBT: Add a "hybrid" executor, which uses the XGE executor for large numbers of parallel actions, and the Parallel executor for any leaf actions that will not saturate the local machine. For platforms that take a long time to link, this frees up XGE to be used for compiling other branches.

To enable, set <BuildConfiguration><bEnableHybridExecutor>true</bEnableHybridExecutor></BuildConfiguration> in the BuildConfiguration.xml file.

diff --git a/Engine/Source/Programs/UnrealBuildTool/Configuration/BuildConfiguration.cs b/Engine/Source/Programs/UnrealBuildTool/Configuration/BuildConfiguration.cs
--- a/Engine/Source/Programs/UnrealBuildTool/Configuration/BuildConfiguration.cs
+++ b/Engine/Source/Programs/UnrealBuildTool/Configuration/BuildConfiguration.cs
@@ -41,6 +41,12 @@ namespace UnrealBuildTool
 		[XmlConfigFile]
 		public bool bLogDetailedActionStats = false;
 
+		/// <summary>
+		/// Whether the hybrid executor will be used (a remote executor and local executor)
+		/// </summary>
+		[XmlConfigFile]
+		public bool bAllowHybridExecutor = false;
+
 		/// <summary>
 		/// Whether XGE may be used.
 		/// </summary>
diff --git a/Engine/Source/Programs/UnrealBuildTool/Executors/HybridExecutor.cs b/Engine/Source/Programs/UnrealBuildTool/Executors/HybridExecutor.cs
new file mode 100644
--- /dev/null
+++ b/Engine/Source/Programs/UnrealBuildTool/Executors/HybridExecutor.cs
@@ -0,0 +1,151 @@
+using System;
+using System.Collections.Generic;
+using System.Linq;
+using System.Text;
+using System.Threading.Tasks;
+
+namespace UnrealBuildTool
+{
+	/// <summary>
+	/// Executor which distributes large sets of parallel actions through a remote executor, and any leaf serial actions through a local executor. Currently uses
+	/// XGE for remote actions and ParallelExecutor for local actions.
+	/// </summary>
+	class HybridExecutor : ActionExecutor
+	{
+		/// <summary>
+		/// Maximum number of actions to execute locally
+		/// </summary>
+		[XmlConfigFile(Category = "HybridExecutor")]
+		int MaxLocalActions;
+
+		/// <summary>
+		/// Executor to use for remote actions
+		/// </summary>
+		ActionExecutor RemoteExecutor;
+
+		/// <summary>
+		/// Executor to use for local actions
+		/// </summary>
+		ActionExecutor LocalExecutor;
+
+		/// <summary>
+		/// Constructor
+		/// </summary>
+		public HybridExecutor()
+		{
+			this.LocalExecutor = new ParallelExecutor();
+			this.RemoteExecutor = new XGE();
+
+			XmlConfig.ApplyTo(this);
+
+			if(MaxLocalActions == 0)
+			{
+				MaxLocalActions = Utils.GetPhysicalProcessorCount();
+			}
+		}
+
+		/// <summary>
+		/// Tests whether this executor can be used
+		/// </summary>
+		/// <returns>True if the executor may be used</returns>
+		public static bool IsAvailable()
+		{
+			return XGE.IsAvailable() && ParallelExecutor.IsAvailable();
+		}
+
+		/// <summary>
+		/// Name of this executor for diagnostic output
+		/// </summary>
+		public override string Name
+		{
+			get { return String.Format("hybrid ({0}+{1})", LocalExecutor.Name, RemoteExecutor.Name); }
+		}
+
+		/// <summary>
+		/// Execute the given actions
+		/// </summary>
+		/// <param name="ActionsToExecute">Actions to be executed</param>
+		/// <param name="bLogDetailedActionStats">Whether to log detailed action stats</param>
+		/// <returns>True if the build succeeded, false otherwise</returns>
+		public override bool ExecuteActions(List<Action> ActionsToExecute, bool bLogDetailedActionStats)
+		{
+			// Find the number of dependants for each action
+			Dictionary<Action, int> ActionToNumDependents = ActionsToExecute.ToDictionary(x => x, x => 0);
+			foreach(Action Action in ActionsToExecute)
+			{
+				foreach(Action PrerequisiteAction in Action.PrerequisiteActions)
+				{
+					ActionToNumDependents[PrerequisiteAction]++;
+				}
+			}
+
+			// Build up a set of leaf actions in several iterations, ensuring that the number of leaf actions in each 
+			HashSet<Action> LeafActions = new HashSet<Action>();
+			for(;;)
+			{
+				// Find all the leaf actions in the graph
+				List<Action> NewLeafActions = new List<Action>();
+				foreach(Action Action in ActionsToExecute)
+				{
+					if(ActionToNumDependents[Action] == 0 && !LeafActions.Contains(Action))
+					{
+						NewLeafActions.Add(Action);
+					}
+				}
+
+				// Exit once we can't prune any more layers from the tree
+				if(NewLeafActions.Count == 0 || NewLeafActions.Count >= MaxLocalActions)
+				{
+					break;
+				}
+
+				// Add these actions to the set of leaf actions
+				LeafActions.UnionWith(NewLeafActions);
+
+				// Decrement the dependent counts for any of their prerequisites, so we can try and remove those from the tree in another iteration
+				foreach(Action NewLeafAction in NewLeafActions)
+				{
+					foreach(Action PrerequisiteAction in NewLeafAction.PrerequisiteActions)
+					{
+						ActionToNumDependents[PrerequisiteAction]--;
+					}
+				}
+			}
+
+			// Split the list of actions into those which should be executed locally and remotely
+			List<Action> LocalActionsToExecute = new List<Action>(LeafActions.Count);
+			List<Action> RemoteActionsToExecute = new List<Action>(ActionsToExecute.Count - LeafActions.Count);
+			foreach(Action ActionToExecute in ActionsToExecute)
+			{
+				if(LeafActions.Contains(ActionToExecute))
+				{
+					LocalActionsToExecute.Add(ActionToExecute);
+				}
+				else
+				{
+					RemoteActionsToExecute.Add(ActionToExecute);
+				}
+			}
+
+			// Execute the remote actions
+			if(RemoteActionsToExecute.Count > 0)
+			{
+				if (!RemoteExecutor.ExecuteActions(RemoteActionsToExecute, bLogDetailedActionStats))
+				{
+					return false;
+				}
+			}
+
+			// Pass all the local actions through to the parallel executor
+			if(LocalActionsToExecute.Count > 0)
+			{
+				if(!LocalExecutor.ExecuteActions(LocalActionsToExecute, bLogDetailedActionStats))
+				{
+					return false;
+				}
+			}
+
+			return true;
+		}
+	}
+}
diff --git a/Engine/Source/Programs/UnrealBuildTool/System/ActionGraph.cs b/Engine/Source/Programs/UnrealBuildTool/System/ActionGraph.cs
--- a/Engine/Source/Programs/UnrealBuildTool/System/ActionGraph.cs
+++ b/Engine/Source/Programs/UnrealBuildTool/System/ActionGraph.cs
@@ -121,7 +121,11 @@ namespace UnrealBuildTool
 			{
 				// Figure out which executor to use
 				ActionExecutor Executor;
-				if (BuildConfiguration.bAllowXGE && XGE.IsAvailable())
+				if (BuildConfiguration.bAllowHybridExecutor && HybridExecutor.IsAvailable())
+				{
+					Executor = new HybridExecutor();
+				}
+				else if (BuildConfiguration.bAllowXGE && XGE.IsAvailable())
 				{
 					Executor = new XGE();
 				}
