Missing project file.

#rnx

diff --git a/Engine/Source/Programs/UnrealBuildTool/UnrealBuildTool.csproj b/Engine/Source/Programs/UnrealBuildTool/UnrealBuildTool.csproj
--- a/Engine/Source/Programs/UnrealBuildTool/UnrealBuildTool.csproj
+++ b/Engine/Source/Programs/UnrealBuildTool/UnrealBuildTool.csproj
@@ -143,6 +143,7 @@
     </When>
   </Choose>
   <ItemGroup>
+    <Compile Include="Executors\HybridExecutor.cs" />
     <Compile Include="Modes\BuildMode.cs" />
     <Compile Include="Modes\CleanMode.cs" />
     <Compile Include="Modes\DeployMode.cs" />
