Speculative fix for Mac FixDylibDependencies script being executed too early

diff --git a/Engine/Source/Programs/UnrealBuildTool/Platform/Mac/MacToolChain.cs b/Engine/Source/Programs/UnrealBuildTool/Platform/Mac/MacToolChain.cs
--- a/Engine/Source/Programs/UnrealBuildTool/Platform/Mac/MacToolChain.cs
+++ b/Engine/Source/Programs/UnrealBuildTool/Platform/Mac/MacToolChain.cs
@@ -1072,8 +1072,6 @@ namespace UnrealBuildTool
 				LinkAction.PrerequisiteItems.Add(Dependency);
 			}
 
-			BundleDependencies.Clear();
-
 			LinkAction.StatusDescription = string.Format("Fixing dylib dependencies for {0}", Path.GetFileName(Executable.AbsolutePath));
 			LinkAction.bCanExecuteRemotely = false;
 
@@ -1340,7 +1338,7 @@ namespace UnrealBuildTool
 				}
 			}
 
-			if (BinaryLinkEnvironment.bIsBuildingDLL || (BinaryLinkEnvironment.bIsBuildingConsoleApplication && Executable.Name.Contains("/UE4Editor") && Executable.Name.EndsWith("-Cmd")))
+			if (BinaryLinkEnvironment.bIsBuildingDLL || (BinaryLinkEnvironment.bIsBuildingConsoleApplication && Executable.Name.Contains("UE4Editor") && Executable.Name.EndsWith("-Cmd")))
 			{
 				return OutputFiles;
 			}
