Fix various issues building code projects from an installed engine build on Windows.

#rnx

diff --git a/Engine/Build/InstalledEngineFilters.xml b/Engine/Build/InstalledEngineFilters.xml
--- a/Engine/Build/InstalledEngineFilters.xml
+++ b/Engine/Build/InstalledEngineFilters.xml
@@ -200,6 +200,9 @@
 			Engine/Build/BatchFiles/RunUAT.bat
 			Engine/Build/BatchFiles/GetMSBuildPath.bat
 
+			<!-- Build utilities -->
+			Engine/Build/Windows/cl-filter/cl-filter.exe
+			
 			<!-- Binaries which are explicitly staged by UAT rather than listing runtime dependencies -->
 			Engine/Binaries/ThirdParty/OpenSSL/...
 
diff --git a/Engine/Source/Programs/UnrealBuildTool/ProjectFiles/VisualStudio/VCProject.cs b/Engine/Source/Programs/UnrealBuildTool/ProjectFiles/VisualStudio/VCProject.cs
--- a/Engine/Source/Programs/UnrealBuildTool/ProjectFiles/VisualStudio/VCProject.cs
+++ b/Engine/Source/Programs/UnrealBuildTool/ProjectFiles/VisualStudio/VCProject.cs
@@ -1391,7 +1391,7 @@ namespace UnrealBuildTool
 					{
 						BuildArguments.Append(" -UsePrecompiled");
 					}
-					else if(TargetRulesObject.Type == TargetType.Editor && bEditorDependsOnShaderCompileWorker)
+					else if(TargetRulesObject.Type == TargetType.Editor && bEditorDependsOnShaderCompileWorker && !UnrealBuildTool.IsEngineInstalled())
 					{
 						BuildArguments.Replace("\"", "\\\"");
 						BuildArguments.Insert(0, "-Target=\"ShaderCompileWorker Win64 Development\" -Target=\"");
diff --git a/Engine/Source/Programs/UnrealBuildTool/System/RulesCompiler.cs b/Engine/Source/Programs/UnrealBuildTool/System/RulesCompiler.cs
--- a/Engine/Source/Programs/UnrealBuildTool/System/RulesCompiler.cs
+++ b/Engine/Source/Programs/UnrealBuildTool/System/RulesCompiler.cs
@@ -542,14 +542,10 @@ namespace UnrealBuildTool
 			{
 				RulesAssembly = CreateEngineRulesAssembly(bUsePrecompiled, bSkipRulesCompile);
 
-				if (RulesAssembly.GetTargetFileName(TargetName) == null)
+				if (RulesAssembly.GetTargetFileName(TargetName) == null && DirectoryReference.Exists(UnrealBuildTool.EnterpriseDirectory))
 				{
 					// Target isn't part of the engine assembly, try the enterprise assembly
-					RulesAssembly EnterpriseRulesAssembly = CreateEnterpriseRulesAssembly(bUsePrecompiled, bSkipRulesCompile);
-					if (EnterpriseRulesAssembly != null)
-					{
-						RulesAssembly = EnterpriseRulesAssembly;
-					}
+					RulesAssembly = CreateEnterpriseRulesAssembly(bUsePrecompiled, bSkipRulesCompile);
 				}
 			}
 			if (ForeignPlugin != null)
