Hacky fix for single file compile; when building from Visual Studio, ignore the ShaderCompileWorker target.

#rnx

diff --git a/Engine/Source/Programs/UnrealBuildTool/Modes/BuildMode.cs b/Engine/Source/Programs/UnrealBuildTool/Modes/BuildMode.cs
--- a/Engine/Source/Programs/UnrealBuildTool/Modes/BuildMode.cs
+++ b/Engine/Source/Programs/UnrealBuildTool/Modes/BuildMode.cs
@@ -140,6 +140,12 @@ namespace UnrealBuildTool
 					TargetDescriptors = TargetDescriptor.ParseCommandLine(Arguments, BuildConfiguration.bUsePrecompiled, BuildConfiguration.bSkipRulesCompile);
 				}
 
+				// Hack for single file compile; don't build the ShaderCompileWorker target that's added to the command line for generated project files
+				if(TargetDescriptors.Count >= 2)
+				{
+					TargetDescriptors.RemoveAll(x => x.Name == "ShaderCompileWorker" && x.SingleFileToCompile != null);
+				}
+
 				// Handle remote builds
 				for(int Idx = 0; Idx < TargetDescriptors.Count; Idx++)
 				{
