Fix compile errors in UBT after merge

#rnx

diff --git a/Engine/Source/Programs/UnrealBuildTool/Platform/Android/UEBuildAndroid.cs b/Engine/Source/Programs/UnrealBuildTool/Platform/Android/UEBuildAndroid.cs
--- a/Engine/Source/Programs/UnrealBuildTool/Platform/Android/UEBuildAndroid.cs
+++ b/Engine/Source/Programs/UnrealBuildTool/Platform/Android/UEBuildAndroid.cs
@@ -95,9 +95,6 @@ namespace UnrealBuildTool
 			ValidateTarget(Target);
 
 			Target.bDeployAfterCompile = true;
-
-			Target.bCompileSimplygon = false;
-			Target.bCompileSimplygonSSF = false;
 		}
 
 		public override void ValidateTarget(TargetRules Target)
diff --git a/Engine/Source/Programs/UnrealBuildTool/Platform/HTML5/UEBuildHTML5.cs b/Engine/Source/Programs/UnrealBuildTool/Platform/HTML5/UEBuildHTML5.cs
--- a/Engine/Source/Programs/UnrealBuildTool/Platform/HTML5/UEBuildHTML5.cs
+++ b/Engine/Source/Programs/UnrealBuildTool/Platform/HTML5/UEBuildHTML5.cs
@@ -42,9 +42,6 @@ namespace UnrealBuildTool
 		public override void ResetTarget(TargetRules Target)
 		{
 			ValidateTarget(Target);
-
-			Target.bCompileSimplygon = false;
-			Target.bCompileSimplygonSSF = false;
 		}
 
 		public override void ValidateTarget(TargetRules Target)
diff --git a/Engine/Source/Programs/UnrealBuildTool/Platform/IOS/UEBuildIOS.cs b/Engine/Source/Programs/UnrealBuildTool/Platform/IOS/UEBuildIOS.cs
--- a/Engine/Source/Programs/UnrealBuildTool/Platform/IOS/UEBuildIOS.cs
+++ b/Engine/Source/Programs/UnrealBuildTool/Platform/IOS/UEBuildIOS.cs
@@ -644,8 +644,6 @@ namespace UnrealBuildTool
 
 			Target.bCompileAPEX = false;
             Target.bCompileNvCloth = false;
-            Target.bCompileSimplygon = false;
-            Target.bCompileSimplygonSSF = false;
 
 			Target.bDeployAfterCompile = true;
 		}
diff --git a/Engine/Source/Programs/UnrealBuildTool/Platform/Linux/UEBuildLinux.cs b/Engine/Source/Programs/UnrealBuildTool/Platform/Linux/UEBuildLinux.cs
--- a/Engine/Source/Programs/UnrealBuildTool/Platform/Linux/UEBuildLinux.cs
+++ b/Engine/Source/Programs/UnrealBuildTool/Platform/Linux/UEBuildLinux.cs
@@ -229,9 +229,6 @@ namespace UnrealBuildTool
 		public override void ResetTarget(TargetRules Target)
 		{
 			ValidateTarget(Target);
-
-			Target.bCompileSimplygon = false;
-			Target.bCompileSimplygonSSF = false;
 		}
 
 		public override void ValidateTarget(TargetRules Target)
