UBT: Remove the bIsUsingPCH flag from each action. This isn't used for anything useful any more.

#rnx

diff --git a/Engine/Source/Programs/UnrealBuildTool/Executors/Distcc.cs b/Engine/Source/Programs/UnrealBuildTool/Executors/Distcc.cs
--- a/Engine/Source/Programs/UnrealBuildTool/Executors/Distcc.cs
+++ b/Engine/Source/Programs/UnrealBuildTool/Executors/Distcc.cs
@@ -358,12 +358,11 @@ namespace UnrealBuildTool
 
 					Log.WriteLineIf(bLogDetailedActionStats,
 						LogEventType.Console,
-						"^{0}^{1:0.00}^{2}^{3}^{4}",
+						"^{0}^{1:0.00}^{2}^{3}",
 						Action.ActionType.ToString(),
 						ThreadSeconds,
 						Action.CommandPath.GetFileName(),
-						Action.StatusDescription,
-						Action.bIsUsingPCH);
+						Action.StatusDescription);
 
 					// Keep track of total thread seconds spent on tasks.
 					TotalThreadSeconds += ThreadSeconds;
diff --git a/Engine/Source/Programs/UnrealBuildTool/Executors/LocalExecutor.cs b/Engine/Source/Programs/UnrealBuildTool/Executors/LocalExecutor.cs
--- a/Engine/Source/Programs/UnrealBuildTool/Executors/LocalExecutor.cs
+++ b/Engine/Source/Programs/UnrealBuildTool/Executors/LocalExecutor.cs
@@ -434,12 +434,11 @@ namespace UnrealBuildTool
 
 				Log.WriteLineIf(bLogDetailedActionStats,
 					LogEventType.Console,
-					"^{0}^{1:0.00}^{2}^{3}^{4}",
+					"^{0}^{1:0.00}^{2}^{3}",
 					Action.ActionType.ToString(),
 					ThreadSeconds,
 					Action.CommandPath.GetFileName(),
-					  Action.StatusDescription,
-					Action.bIsUsingPCH);
+					Action.StatusDescription);
 
 				// Update statistics
 				switch (Action.ActionType)
diff --git a/Engine/Source/Programs/UnrealBuildTool/Executors/SNDBS.cs b/Engine/Source/Programs/UnrealBuildTool/Executors/SNDBS.cs
--- a/Engine/Source/Programs/UnrealBuildTool/Executors/SNDBS.cs
+++ b/Engine/Source/Programs/UnrealBuildTool/Executors/SNDBS.cs
@@ -406,12 +406,11 @@ namespace UnrealBuildTool
 
 					Log.WriteLineIf(bLogDetailedActionStats,
 						LogEventType.Console,
-						"^{0}^{1:0.00}^{2}^{3}^{4}",
+						"^{0}^{1:0.00}^{2}^{3}",
 						Action.ActionType.ToString(),
 						ThreadSeconds,
 						Action.CommandPath.GetFileName(),
-						Action.StatusDescription,
-						Action.bIsUsingPCH);
+						Action.StatusDescription);
 
 					// Keep track of total thread seconds spent on tasks.
 					TotalThreadSeconds += ThreadSeconds;
diff --git a/Engine/Source/Programs/UnrealBuildTool/Platform/Android/AndroidToolChain.cs b/Engine/Source/Programs/UnrealBuildTool/Platform/Android/AndroidToolChain.cs
--- a/Engine/Source/Programs/UnrealBuildTool/Platform/Android/AndroidToolChain.cs
+++ b/Engine/Source/Programs/UnrealBuildTool/Platform/Android/AndroidToolChain.cs
@@ -1384,7 +1384,6 @@ namespace UnrealBuildTool
 						{
 							if (CompileEnvironment.PrecompiledHeaderAction == PrecompiledHeaderAction.Include)
 							{
-								CompileAction.bIsUsingPCH = true;
 								FileItem ArchPrecompiledHeaderFile = FileItem.GetItemByPath(InlineArchName(BasePCHName, Arch, GPUArchitecture) + PCHExtension);
 								CompileAction.PrerequisiteItems.Add(ArchPrecompiledHeaderFile);
 							}
diff --git a/Engine/Source/Programs/UnrealBuildTool/Platform/IOS/IOSToolChain.cs b/Engine/Source/Programs/UnrealBuildTool/Platform/IOS/IOSToolChain.cs
--- a/Engine/Source/Programs/UnrealBuildTool/Platform/IOS/IOSToolChain.cs
+++ b/Engine/Source/Programs/UnrealBuildTool/Platform/IOS/IOSToolChain.cs
@@ -660,7 +660,6 @@ namespace UnrealBuildTool
 				{
 					if (CompileEnvironment.PrecompiledHeaderAction == PrecompiledHeaderAction.Include)
 					{
-						CompileAction.bIsUsingPCH = true;
 						CompileAction.PrerequisiteItems.Add(CompileEnvironment.PrecompiledHeaderFile);
 					}
 					// Add the object file to the produced item list.
diff --git a/Engine/Source/Programs/UnrealBuildTool/Platform/Linux/LinuxToolChain.cs b/Engine/Source/Programs/UnrealBuildTool/Platform/Linux/LinuxToolChain.cs
--- a/Engine/Source/Programs/UnrealBuildTool/Platform/Linux/LinuxToolChain.cs
+++ b/Engine/Source/Programs/UnrealBuildTool/Platform/Linux/LinuxToolChain.cs
@@ -1217,7 +1217,6 @@ namespace UnrealBuildTool
 				{
 					if (CompileEnvironment.PrecompiledHeaderAction == PrecompiledHeaderAction.Include)
 					{
-						CompileAction.bIsUsingPCH = true;
 						CompileAction.PrerequisiteItems.Add(CompileEnvironment.PrecompiledHeaderFile);
 					}
 
diff --git a/Engine/Source/Programs/UnrealBuildTool/Platform/Mac/MacToolChain.cs b/Engine/Source/Programs/UnrealBuildTool/Platform/Mac/MacToolChain.cs
--- a/Engine/Source/Programs/UnrealBuildTool/Platform/Mac/MacToolChain.cs
+++ b/Engine/Source/Programs/UnrealBuildTool/Platform/Mac/MacToolChain.cs
@@ -496,7 +496,6 @@ namespace UnrealBuildTool
 				{
 					if (CompileEnvironment.PrecompiledHeaderAction == PrecompiledHeaderAction.Include)
 					{
-						CompileAction.bIsUsingPCH = true;
 						CompileAction.PrerequisiteItems.Add(CompileEnvironment.PrecompiledHeaderFile);
 					}
 					// Add the object file to the produced item list.
diff --git a/Engine/Source/Programs/UnrealBuildTool/Platform/Windows/VCToolChain.cs b/Engine/Source/Programs/UnrealBuildTool/Platform/Windows/VCToolChain.cs
--- a/Engine/Source/Programs/UnrealBuildTool/Platform/Windows/VCToolChain.cs
+++ b/Engine/Source/Programs/UnrealBuildTool/Platform/Windows/VCToolChain.cs
@@ -979,7 +979,6 @@ namespace UnrealBuildTool
 				{
 					if (CompileEnvironment.PrecompiledHeaderAction == PrecompiledHeaderAction.Include)
 					{
-						CompileAction.bIsUsingPCH = true;
 						CompileAction.PrerequisiteItems.Add(CompileEnvironment.PrecompiledHeaderFile);
 
 						if (Target.WindowsPlatform.Compiler == WindowsCompiler.Clang)
diff --git a/Engine/Source/Programs/UnrealBuildTool/System/Action.cs b/Engine/Source/Programs/UnrealBuildTool/System/Action.cs
--- a/Engine/Source/Programs/UnrealBuildTool/System/Action.cs
+++ b/Engine/Source/Programs/UnrealBuildTool/System/Action.cs
@@ -120,11 +120,6 @@ namespace UnrealBuildTool
 		/// </summary>
 		public bool bIsGCCCompiler = false;
 
-		/// <summary>
-		/// Whether the action is using a pre-compiled header to speed it up.
-		/// </summary>
-		public bool bIsUsingPCH = false;
-
 		/// <summary>
 		/// Whether we should log this action, whether executed locally or remotely.  This is useful for actions that take time
 		/// but invoke tools without any console output.
@@ -180,7 +175,6 @@ namespace UnrealBuildTool
 			bCanExecuteRemotely = Reader.ReadBool();
 			bCanExecuteRemotelyWithSNDBS = Reader.ReadBool();
 			bIsGCCCompiler = Reader.ReadBool();
-			bIsUsingPCH = Reader.ReadBool();
 			bShouldOutputStatusDescription = Reader.ReadBool();
 			bProducesImportLibrary = Reader.ReadBool();
 			PrerequisiteItems = Reader.ReadList(() => Reader.ReadFileItem());
@@ -204,7 +198,6 @@ namespace UnrealBuildTool
 			Writer.WriteBool(bCanExecuteRemotely);
 			Writer.WriteBool(bCanExecuteRemotelyWithSNDBS);
 			Writer.WriteBool(bIsGCCCompiler);
-			Writer.WriteBool(bIsUsingPCH);
 			Writer.WriteBool(bShouldOutputStatusDescription);
 			Writer.WriteBool(bProducesImportLibrary);
 			Writer.WriteList(PrerequisiteItems, Item => Writer.WriteFileItem(Item));
