Fixup and remove OnlineSubsystemLive from the list of engine modules that don't support IWYU.

#rnx

diff --git a/Engine/Source/Programs/UnrealBuildTool/Configuration/UEBuildModuleCPP.cs b/Engine/Source/Programs/UnrealBuildTool/Configuration/UEBuildModuleCPP.cs
--- a/Engine/Source/Programs/UnrealBuildTool/Configuration/UEBuildModuleCPP.cs
+++ b/Engine/Source/Programs/UnrealBuildTool/Configuration/UEBuildModuleCPP.cs
@@ -921,7 +921,6 @@ namespace UnrealBuildTool
 				case "PS4RHI":
                 case "Gnmx":
 				case "OnlineSubsystemIOS":
-				case "OnlineSubsystemLive":
 					return true;
 			}
 			return false;
