Add support for building additional platforms with SN-DBS.

#rnx

diff --git a/Engine/Source/Programs/UnrealBuildTool/Executors/SNDBS.cs b/Engine/Source/Programs/UnrealBuildTool/Executors/SNDBS.cs
--- a/Engine/Source/Programs/UnrealBuildTool/Executors/SNDBS.cs
+++ b/Engine/Source/Programs/UnrealBuildTool/Executors/SNDBS.cs
@@ -11,6 +11,7 @@ using System.Reflection;
 using System.Threading;
 using System.ServiceProcess;
 using Tools.DotNETCommon;
+using System.Text;
 
 namespace UnrealBuildTool
 {
@@ -234,8 +235,32 @@ namespace UnrealBuildTool
 						}
 						else
 						{
+							// Create a dummy force-included file which references PCH files, so that SN-DBS knows they are dependencies.
+							string AdditionalStubIncludes = "";
+							if (Action.CommandPath.GetFileName().Equals("cl.exe", StringComparison.OrdinalIgnoreCase))
+							{
+								string ResponseFile = Action.CommandArguments.Replace("\"", "").Replace("@", "").Trim();
+								StringBuilder WrapperContents = new StringBuilder();
+								using (StringWriter Writer = new StringWriter(WrapperContents))
+								{
+									Writer.WriteLine("// PCH dependencies for {0}", ResponseFile);
+									Writer.WriteLine("#if 0");
+									foreach (FileItem Preqrequisite in Action.PrerequisiteItems)
+									{
+										if (Preqrequisite.AbsolutePath.EndsWith(".pch"))
+										{
+											Writer.WriteLine("#include \"{0}\"", Preqrequisite.AbsolutePath.Replace(".pch", ".obj"));
+										}
+									}
+									Writer.WriteLine("#endif");
+								}
+
+								FileItem DummyResponseFileDependency = FileItem.CreateIntermediateTextFile(new FileReference(ResponseFile + ".dummy.h"), WrapperContents.ToString());
+								AdditionalStubIncludes = string.Format("/FI\"{0}\"", DummyResponseFileDependency);
+							}
+
 							// Add to script for execution by SN-DBS
-							string NewCommandArguments = "\"" + Action.CommandPath + "\"" + " " + Action.CommandArguments;
+							string NewCommandArguments = "\"" + Action.CommandPath + "\"" + " " + AdditionalStubIncludes + " " + Action.CommandArguments;
 							ScriptFile.WriteLine(NewCommandArguments);
 							InActionThreadDictionary.Add(Action, DummyActionThread);
 							Action.StartTime = Action.EndTime = DateTimeOffset.Now;
@@ -257,7 +282,8 @@ namespace UnrealBuildTool
                 // Create the process
                 string SCERoot = Environment.GetEnvironmentVariable("SCE_ROOT_DIR");
                 string SNDBSExecutable = Path.Combine(SCERoot, "Common/SN-DBS/bin/dbsbuild.exe");
-                ProcessStartInfo PSI = new ProcessStartInfo(SNDBSExecutable, String.Format("-q -p UE4 -s \"{0}\"", FileReference.Combine(UnrealBuildTool.EngineDirectory, "Intermediate", "Build", "sndbs.bat").FullName));
+				DirectoryReference TemplatesDir = DirectoryReference.Combine(UnrealBuildTool.EngineDirectory, "Programs", "UnrealBuildTool", "SndbsTemplates");
+                ProcessStartInfo PSI = new ProcessStartInfo(SNDBSExecutable, String.Format("-q -p UE4 -s \"{0}\" -templates \"{1}\"", FileReference.Combine(UnrealBuildTool.EngineDirectory, "Intermediate", "Build", "sndbs.bat").FullName, TemplatesDir.FullName));
 				PSI.RedirectStandardOutput = true;
 				PSI.RedirectStandardError = true;
 				PSI.UseShellExecute = false;
