UBT: Fix LIB files not being generated to the same directory as the DLL when bShouldCompileAsDLL = true in the target.

diff --git a/Engine/Source/Programs/UnrealBuildTool/Platform/Windows/VCToolChain.cs b/Engine/Source/Programs/UnrealBuildTool/Platform/Windows/VCToolChain.cs
--- a/Engine/Source/Programs/UnrealBuildTool/Platform/Windows/VCToolChain.cs
+++ b/Engine/Source/Programs/UnrealBuildTool/Platform/Windows/VCToolChain.cs
@@ -1383,12 +1383,18 @@ namespace UnrealBuildTool
 			// file is not needed for our builds, but there is no way to prevent MSVC from generating it when
 			// linking targets that have exports.  We don't want this to clobber our LIB file and invalidate the
 			// existing timstamp, so instead we simply emit it with a different name
-			FileReference ImportLibraryFilePath = FileReference.Combine(LinkEnvironment.IntermediateDirectory,
-														 LinkEnvironment.OutputFilePath.GetFileNameWithoutExtension() + ".lib");
-
+			FileReference ImportLibraryFilePath;
 			if (LinkEnvironment.bIsCrossReferenced && !bBuildImportLibraryOnly)
 			{
-				ImportLibraryFilePath = ImportLibraryFilePath.ChangeExtension(".suppressed" + ImportLibraryFilePath.GetExtension());
+				ImportLibraryFilePath = FileReference.Combine(LinkEnvironment.IntermediateDirectory, LinkEnvironment.OutputFilePath.GetFileNameWithoutExtension() + ".suppressed.lib");
+			}
+			else if(Target.bShouldCompileAsDLL)
+			{
+				ImportLibraryFilePath = FileReference.Combine(LinkEnvironment.OutputDirectory, LinkEnvironment.OutputFilePath.GetFileNameWithoutExtension() + ".lib");
+			}
+			else
+			{
+				ImportLibraryFilePath = FileReference.Combine(LinkEnvironment.IntermediateDirectory, LinkEnvironment.OutputFilePath.GetFileNameWithoutExtension() + ".lib");
 			}
 
 			FileItem OutputFile;
@@ -1627,7 +1633,14 @@ namespace UnrealBuildTool
 		{
 			if (Binary.Type == UEBuildBinaryType.DynamicLinkLibrary)
 			{
-				BuildProducts.Add(FileReference.Combine(Binary.IntermediateDirectory, Binary.OutputFilePath.GetFileNameWithoutExtension() + ".lib"), BuildProductType.BuildResource);
+				if(Target.bShouldCompileAsDLL)
+				{
+					BuildProducts.Add(FileReference.Combine(Binary.OutputDir, Binary.OutputFilePath.GetFileNameWithoutExtension() + ".lib"), BuildProductType.BuildResource);
+				}
+				else
+				{
+					BuildProducts.Add(FileReference.Combine(Binary.IntermediateDirectory, Binary.OutputFilePath.GetFileNameWithoutExtension() + ".lib"), BuildProductType.BuildResource);
+				}
 			}
 			if(Binary.Type == UEBuildBinaryType.Executable && Target.bCreateMapFile)
 			{
