Fix "launch on" always rebuilding the executable, due to order of enumerated plugins being non-deterministic.

diff --git a/Engine/Source/Programs/UnrealBuildTool/System/Plugins.cs b/Engine/Source/Programs/UnrealBuildTool/System/Plugins.cs
--- a/Engine/Source/Programs/UnrealBuildTool/System/Plugins.cs
+++ b/Engine/Source/Programs/UnrealBuildTool/System/Plugins.cs
@@ -336,6 +336,9 @@ namespace UnrealBuildTool
 					}
 				}
 
+				// Sort the filenames to ensure that the plugin order is deterministic; otherwise response files will change with each build.
+				FileNames = FileNames.OrderBy(x => x.FullName, StringComparer.OrdinalIgnoreCase).ToList();
+
 				PluginFileCache.Add(ParentDirectory, FileNames);
 			}
 			return FileNames;
