UBT: Add a new target setting (bAdaptiveUnityDisablesPCHForProject) which allows overriding the "bAdaptiveUnityDisablesPCH" setting for project-specific code.

diff --git a/Engine/Source/Programs/UnrealBuildTool/Configuration/TargetRules.cs b/Engine/Source/Programs/UnrealBuildTool/Configuration/TargetRules.cs
--- a/Engine/Source/Programs/UnrealBuildTool/Configuration/TargetRules.cs
+++ b/Engine/Source/Programs/UnrealBuildTool/Configuration/TargetRules.cs
@@ -663,6 +663,21 @@ namespace UnrealBuildTool
 		[XmlConfigFile(Category = "BuildConfiguration")]
 		public bool bAdaptiveUnityDisablesPCH = true;
 
+		/// <summary>
+		/// Backing storage for bAdaptiveUnityDisablesProjectPCH.
+		/// </summary>
+		[XmlConfigFile(Category = "BuildConfiguration")]
+		bool? bAdaptiveUnityDisablesProjectPCHForProjectPrivate;
+
+		/// <summary>
+		/// Whether to disable force-included PCHs for project source files in the adaptive non-unity working set. Defaults to bAdaptiveUnityDisablesPCH;
+		/// </summary>
+		public bool bAdaptiveUnityDisablesPCHForProject
+		{
+			get { return bAdaptiveUnityDisablesProjectPCHForProjectPrivate ?? bAdaptiveUnityDisablesPCH; }
+			set { bAdaptiveUnityDisablesProjectPCHForProjectPrivate = value; }
+		}
+
 		/// <summary>
 		/// Creates a dedicated PCH for each source file in the working set, allowing faster iteration on cpp-only changes.
 		/// </summary>
@@ -1896,6 +1911,11 @@ namespace UnrealBuildTool
 			get { return Inner.bAdaptiveUnityDisablesPCH; }
 		}
 
+		public bool bAdaptiveUnityDisablesPCHForProject
+		{
+			get { return Inner.bAdaptiveUnityDisablesPCHForProject; }
+		}
+
 		public bool bAdaptiveUnityCreatesDedicatedPCH
 		{
 			get { return Inner.bAdaptiveUnityCreatesDedicatedPCH; }
diff --git a/Engine/Source/Programs/UnrealBuildTool/Configuration/UEBuildModuleCPP.cs b/Engine/Source/Programs/UnrealBuildTool/Configuration/UEBuildModuleCPP.cs
--- a/Engine/Source/Programs/UnrealBuildTool/Configuration/UEBuildModuleCPP.cs
+++ b/Engine/Source/Programs/UnrealBuildTool/Configuration/UEBuildModuleCPP.cs
@@ -643,7 +643,18 @@ namespace UnrealBuildTool
 			List<FileItem> NormalFiles = new List<FileItem>();
 			List<FileItem> AdaptiveFiles = new List<FileItem>();
 
-			bool bAdaptiveUnityDisablesPCH = (Target.bAdaptiveUnityDisablesPCH && Rules.PCHUsage == ModuleRules.PCHUsageMode.UseExplicitOrSharedPCHs);
+			bool bAdaptiveUnityDisablesPCH = false;
+			if(Rules.PCHUsage == ModuleRules.PCHUsageMode.UseExplicitOrSharedPCHs)
+			{
+				if(Rules.bTreatAsEngineModule || Rules.PrivatePCHHeaderFile == null)
+				{
+					bAdaptiveUnityDisablesPCH = Target.bAdaptiveUnityDisablesPCH;
+				}
+				else
+				{
+					bAdaptiveUnityDisablesPCH = Target.bAdaptiveUnityDisablesPCHForProject;
+				}
+			}
 
 			if ((Target.bAdaptiveUnityDisablesOptimizations || bAdaptiveUnityDisablesPCH || Target.bAdaptiveUnityCreatesDedicatedPCH) && !Target.bStressTestUnity)
 			{
