UBT: Add support for writing a header containing the public compile environment for modules in the primary binary of a target (including all *_API macros in 'export' mode). To use, set ExportPublicHeader = "Foo.h" from the target.cs file.

diff --git a/Engine/Source/Programs/UnrealBuildTool/Configuration/TargetRules.cs b/Engine/Source/Programs/UnrealBuildTool/Configuration/TargetRules.cs
--- a/Engine/Source/Programs/UnrealBuildTool/Configuration/TargetRules.cs
+++ b/Engine/Source/Programs/UnrealBuildTool/Configuration/TargetRules.cs
@@ -1072,6 +1072,11 @@ namespace UnrealBuildTool
 		/// </summary>
 		private string LaunchModuleNamePrivate;
 
+		/// <summary>
+		/// Specifies the path to write a header containing public definitions for this target. Useful when building a DLL to be consumed by external build processes.
+		/// </summary>
+		public string ExportPublicHeader;
+
 		/// <summary>
 		/// List of additional modules to be compiled into the target.
 		/// </summary>
@@ -2177,6 +2182,11 @@ namespace UnrealBuildTool
 			get { return Inner.LaunchModuleName; }
 		}
 
+		public string ExportPublicHeader
+		{
+			get { return Inner.ExportPublicHeader; }
+		}
+
 		public IReadOnlyList<string> ExtraModuleNames
 		{
 			get { return Inner.ExtraModuleNames.AsReadOnly(); }
diff --git a/Engine/Source/Programs/UnrealBuildTool/Configuration/UEBuildTarget.cs b/Engine/Source/Programs/UnrealBuildTool/Configuration/UEBuildTarget.cs
--- a/Engine/Source/Programs/UnrealBuildTool/Configuration/UEBuildTarget.cs
+++ b/Engine/Source/Programs/UnrealBuildTool/Configuration/UEBuildTarget.cs
@@ -1509,6 +1509,14 @@ namespace UnrealBuildTool
 			}
 			Makefile.AdditionalDependencies.UnionWith(Makefile.PluginFiles);
 
+			// Write a header containing public definitions for this target
+			if(Rules.ExportPublicHeader != null)
+			{
+				UEBuildBinary Binary = Binaries[0];
+				FileReference Header = FileReference.Combine(Binary.OutputDir, Rules.ExportPublicHeader);
+				WritePublicHeader(Binary, Header, GlobalCompileEnvironment);
+			}
+
 			// Clean any stale modules which exist in multiple output directories. This can lead to the wrong DLL being loaded on Windows.
 			CleanStaleModules();
 
@@ -1724,6 +1732,44 @@ namespace UnrealBuildTool
 			}
 		}
 
+		/// <summary>
+		/// Writes a header for the given binary that allows including headers from it in an external application
+		/// </summary>
+		/// <param name="Binary">Binary to write a header for</param>
+		/// <param name="HeaderFile">Path to the header to output</param>
+		/// <param name="GlobalCompileEnvironment">The global compile environment for this target</param>
+		void WritePublicHeader(UEBuildBinary Binary, FileReference HeaderFile, CppCompileEnvironment GlobalCompileEnvironment)
+		{
+			DirectoryReference.CreateDirectory(HeaderFile.Directory);
+
+			// Find all the public definitions from each module. We pass null as the source binary to AddModuleToCompileEnvironment, which forces all _API macros to 'export' mode
+			List<string> Definitions = new List<string>(GlobalCompileEnvironment.Definitions);
+			foreach(UEBuildModule Module in Binary.Modules)
+			{
+				Module.AddModuleToCompileEnvironment(null, new HashSet<DirectoryReference>(), new HashSet<DirectoryReference>(), Definitions, new List<UEBuildFramework>(), false);
+			}
+
+			// Write the header
+			using(StreamWriter Writer = new StreamWriter(HeaderFile.FullName))
+			{
+				Writer.WriteLine("#pragma once");
+				Writer.WriteLine();
+				foreach(string Definition in Definitions)
+				{
+					int EqualsIdx = Definition.IndexOf('=');
+					if(EqualsIdx == -1)
+					{
+						Writer.WriteLine(String.Format("#define {0} 1", Definition));
+					}
+					else
+					{
+						Writer.WriteLine(String.Format("#define {0} {1}", Definition.Substring(0, EqualsIdx), Definition.Substring(EqualsIdx + 1)));
+					}
+				}
+			}
+			Log.TraceInformation("Written public header to {0}", HeaderFile);
+		}
+
 		/// <summary>
 		/// Check for EULA violation dependency issues.
 		/// </summary>
