Add a flag which allows overriding the shared compile environment with custom build settings.

diff --git a/Engine/Source/Programs/UnrealBuildTool/Configuration/TargetRules.cs b/Engine/Source/Programs/UnrealBuildTool/Configuration/TargetRules.cs
--- a/Engine/Source/Programs/UnrealBuildTool/Configuration/TargetRules.cs
+++ b/Engine/Source/Programs/UnrealBuildTool/Configuration/TargetRules.cs
@@ -1101,6 +1101,12 @@ namespace UnrealBuildTool
 		[CommandLine("-UniqueBuildEnvironment", Value = "Unique")]
 		public TargetBuildEnvironment BuildEnvironment = TargetBuildEnvironment.Default;
 
+		/// <summary>
+		/// Whether to ignore violations to the shared build environment (eg. editor targets modifying definitions)
+		/// </summary>
+		[CommandLine("-OverrideBuildEnvironment")]
+		public bool bOverrideBuildEnvironment = false;
+
 		/// <summary>
 		/// Specifies a list of steps which should be executed before this target is built, in the context of the host platform's shell.
 		/// The following variables will be expanded before execution: 
@@ -2207,6 +2213,11 @@ namespace UnrealBuildTool
 			get { return Inner.BuildEnvironment; }
 		}
 
+		public bool bOverrideBuildEnvironment
+		{
+			get { return Inner.bOverrideBuildEnvironment; }
+		}
+
 		public IReadOnlyList<string> PreBuildSteps
 		{
 			get { return Inner.PreBuildSteps; }
diff --git a/Engine/Source/Programs/UnrealBuildTool/Configuration/UEBuildTarget.cs b/Engine/Source/Programs/UnrealBuildTool/Configuration/UEBuildTarget.cs
--- a/Engine/Source/Programs/UnrealBuildTool/Configuration/UEBuildTarget.cs
+++ b/Engine/Source/Programs/UnrealBuildTool/Configuration/UEBuildTarget.cs
@@ -371,6 +371,12 @@ namespace UnrealBuildTool
 		/// </summary>
 		static void ValidateSharedEnvironment(RulesAssembly RulesAssembly, string ThisTargetName, TargetRules ThisRules)
 		{
+			// Allow disabling these checks
+			if(ThisRules.bOverrideBuildEnvironment)
+			{
+				return;
+			}
+
 			// Get the name of the target with default settings
 			string BaseTargetName;
 			switch(ThisRules.Type)
@@ -456,7 +462,7 @@ namespace UnrealBuildTool
 			// Throw an exception if they don't match
 			if(!bFieldsMatch)
 			{
-				throw new BuildException("{0} modifies the value of {1}. This is not allowed, as {0} has build products in common with {2}.\nRemove the modified setting or change {0} to use a unique build environment by setting 'BuildEnvironment = TargetBuildEnvironment.Unique;' in the {3} constructor.", ThisTargetName, FieldName, BaseTargetName, RulesType.Name);
+				throw new BuildException("{0} modifies the value of {1}. This is not allowed, as {0} has build products in common with {2}.\nRemove the modified setting, change {0} to use a unique build environment by setting 'BuildEnvironment = TargetBuildEnvironment.Unique;' in the {3} constructor, or set bOverrideSharedBuildEnvironment = true to force this setting on.", ThisTargetName, FieldName, BaseTargetName, RulesType.Name);
 			}
 		}
 
