Add a better error message when trying to use a module that is not included in an installed build.

diff --git a/Engine/Source/Programs/UnrealBuildTool/Configuration/UEBuildModuleCPP.cs b/Engine/Source/Programs/UnrealBuildTool/Configuration/UEBuildModuleCPP.cs
--- a/Engine/Source/Programs/UnrealBuildTool/Configuration/UEBuildModuleCPP.cs
+++ b/Engine/Source/Programs/UnrealBuildTool/Configuration/UEBuildModuleCPP.cs
@@ -258,6 +258,11 @@ namespace UnrealBuildTool
 			// If the module is precompiled, read the object files from the manifest
 			if(Rules.bUsePrecompiled && Target.LinkType == TargetLinkType.Monolithic)
 			{
+				if(!FileReference.Exists(PrecompiledManifestLocation))
+				{
+					throw new BuildException("Missing precompiled manifest for '{0}'. This module was most likely not flagged for being included in a precompiled build - set 'PrecompileForTargets = PrecompileTargetsType.Any;' in {0}.build.cs to override.", Name);
+				}
+
 				PrecompiledManifest Manifest = PrecompiledManifest.Read(PrecompiledManifestLocation);
 				foreach(FileReference OutputFile in Manifest.OutputFiles)
 				{
