Merging CL#5206353

Handle alternative splash screen resources removal

#android

diff --git a/Engine/Source/Programs/UnrealBuildTool/Platform/Android/UEDeployAndroid.cs b/Engine/Source/Programs/UnrealBuildTool/Platform/Android/UEDeployAndroid.cs
--- a/Engine/Source/Programs/UnrealBuildTool/Platform/Android/UEDeployAndroid.cs
+++ b/Engine/Source/Programs/UnrealBuildTool/Platform/Android/UEDeployAndroid.cs
@@ -1941,6 +1941,10 @@ namespace UnrealBuildTool
 				{
 					// Remove unused image
 					SafeDeleteFile(PortraitFilename);
+
+					// Remove optional extended resource
+					string PortraitXmlFilename = UE4BuildPath + ResolutionPath + "splashscreen_p.xml";
+					SafeDeleteFile(PortraitXmlFilename);
 				}
 
 				string LandscapeFilename = UE4BuildPath + ResolutionPath + "splashscreen_landscape.png";
@@ -1955,6 +1959,10 @@ namespace UnrealBuildTool
 				{
 					// Remove unused image
 					SafeDeleteFile(LandscapeFilename);
+
+					// Remove optional extended resource
+					string LandscapeXmlFilename = UE4BuildPath + ResolutionPath + "splashscreen_l.xml";
+					SafeDeleteFile(LandscapeXmlFilename);
 				}
 			}
 		}
