Fix default superclass for GameActivity

#android

diff --git a/Engine/Source/Programs/UnrealBuildTool/Platform/Android/UEDeployAndroid.cs b/Engine/Source/Programs/UnrealBuildTool/Platform/Android/UEDeployAndroid.cs
--- a/Engine/Source/Programs/UnrealBuildTool/Platform/Android/UEDeployAndroid.cs
+++ b/Engine/Source/Programs/UnrealBuildTool/Platform/Android/UEDeployAndroid.cs
@@ -4171,8 +4171,11 @@ bSaveSymbols = true;
 
 			string LoadLibraryDefaults = "";
 
-			string SuperClassDefault = "NativeActivity";
-			Ini.GetString("/Script/AndroidRuntimeSettings.AndroidRuntimeSettings", "GameActivitySuperClass", out SuperClassDefault);
+			string SuperClassDefault;
+			if (!Ini.GetString("/Script/AndroidRuntimeSettings.AndroidRuntimeSettings", "GameActivitySuperClass", out SuperClassDefault))
+			{
+				SuperClassDefault = "NativeActivity";
+			}
 
 			string AndroidGraphicsDebugger;
 			Ini.GetString("/Script/AndroidRuntimeSettings.AndroidRuntimeSettings", "AndroidGraphicsDebugger", out AndroidGraphicsDebugger);
