UE-70475 - Packaging for tvOS from Windows fails with rsync error unless iOS is packaged first

diff --git a/Engine/Source/Programs/UnrealBuildTool/ToolChain/RemoteMac.cs b/Engine/Source/Programs/UnrealBuildTool/ToolChain/RemoteMac.cs
--- a/Engine/Source/Programs/UnrealBuildTool/ToolChain/RemoteMac.cs
+++ b/Engine/Source/Programs/UnrealBuildTool/ToolChain/RemoteMac.cs
@@ -537,7 +537,7 @@ namespace UnrealBuildTool
 
 			// Copy remote FrameworkAssets directory as it could contain resource bundles that must be packaged locally.
 			DirectoryReference BaseDir = DirectoryReference.FromFile(TargetDesc.ProjectFile) ?? UnrealBuildTool.EngineDirectory;
-			DirectoryReference FrameworkAssetsDir = DirectoryReference.Combine(BaseDir, "Intermediate", "IOS", "FrameworkAssets");
+			DirectoryReference FrameworkAssetsDir = DirectoryReference.Combine(BaseDir, "Intermediate", TargetDesc.Platform == UnrealTargetPlatform.IOS ? "IOS" : "TVOS", "FrameworkAssets");
 
 			Log.TraceInformation("[Remote] Downloading {0}", FrameworkAssetsDir);
 			DownloadDirectory(FrameworkAssetsDir);
