UBT: Allow modules to disable PCHs entirely by setting PCHUsage = PCHUsageMode.NoPCHs.

diff --git a/Engine/Source/Programs/UnrealBuildTool/Configuration/ModuleRules.cs b/Engine/Source/Programs/UnrealBuildTool/Configuration/ModuleRules.cs
--- a/Engine/Source/Programs/UnrealBuildTool/Configuration/ModuleRules.cs
+++ b/Engine/Source/Programs/UnrealBuildTool/Configuration/ModuleRules.cs
@@ -72,6 +72,11 @@ namespace UnrealBuildTool
 			/// </summary>
 			Default,
 
+			/// <summary>
+			/// Never use any PCHs.
+			/// </summary>
+			NoPCHs,
+
 			/// <summary>
 			/// Never use shared PCHs.  Always generate a unique PCH for this module if appropriate
 			/// </summary>
diff --git a/Engine/Source/Programs/UnrealBuildTool/Configuration/UEBuildModuleCPP.cs b/Engine/Source/Programs/UnrealBuildTool/Configuration/UEBuildModuleCPP.cs
--- a/Engine/Source/Programs/UnrealBuildTool/Configuration/UEBuildModuleCPP.cs
+++ b/Engine/Source/Programs/UnrealBuildTool/Configuration/UEBuildModuleCPP.cs
@@ -346,7 +346,7 @@ namespace UnrealBuildTool
 
 			// Set up the environment with which to compile the CPP files
 			CppCompileEnvironment CompileEnvironment = ModuleCompileEnvironment;
-			if (Target.bUsePCHFiles)
+			if (Target.bUsePCHFiles && Rules.PCHUsage != ModuleRules.PCHUsageMode.NoPCHs)
 			{
 				// If this module doesn't need a shared PCH, configure that
 				if(Rules.PrivatePCHHeaderFile != null && (Rules.PCHUsage == ModuleRules.PCHUsageMode.NoSharedPCHs || Rules.PCHUsage == ModuleRules.PCHUsageMode.UseExplicitOrSharedPCHs))
diff --git a/Engine/Source/Programs/UnrealBuildTool/Configuration/UEBuildTarget.cs b/Engine/Source/Programs/UnrealBuildTool/Configuration/UEBuildTarget.cs
--- a/Engine/Source/Programs/UnrealBuildTool/Configuration/UEBuildTarget.cs
+++ b/Engine/Source/Programs/UnrealBuildTool/Configuration/UEBuildTarget.cs
@@ -3277,7 +3277,7 @@ namespace UnrealBuildTool
 				}
 
 				// If we can't use a shared PCH, check there's a private PCH set
-				if(RulesObject.PCHUsage != ModuleRules.PCHUsageMode.UseExplicitOrSharedPCHs && RulesObject.PrivatePCHHeaderFile == null)
+				if(RulesObject.PCHUsage != ModuleRules.PCHUsageMode.NoPCHs && RulesObject.PCHUsage != ModuleRules.PCHUsageMode.UseExplicitOrSharedPCHs && RulesObject.PrivatePCHHeaderFile == null)
 				{
 					// Try to figure out the legacy PCH file
 					FileReference CppFile = DirectoryReference.EnumerateFiles(RulesObject.Directory, "*.cpp", SearchOption.AllDirectories).FirstOrDefault();
