UBT: Include other files in the project directory within the generated projects.

#rnx

diff --git a/Engine/Source/Programs/UnrealBuildTool/ProjectFiles/ProjectFileGenerator.cs b/Engine/Source/Programs/UnrealBuildTool/ProjectFiles/ProjectFileGenerator.cs
--- a/Engine/Source/Programs/UnrealBuildTool/ProjectFiles/ProjectFileGenerator.cs
+++ b/Engine/Source/Programs/UnrealBuildTool/ProjectFiles/ProjectFileGenerator.cs
@@ -1169,6 +1169,8 @@ namespace UnrealBuildTool
 				{
 					// @todo projectfiles: We have engine localization files, but should we also add GAME localization files?
 
+					GameProject.AddFilesToProject( SourceFileSearch.FindFiles( GameProjectDirectory, SearchSubdirectories: false ), GameProjectDirectory );
+
 					// Game config files
 					if( bIncludeConfigFiles )
 					{
