UE-68313 Advertising Identified (IDFA) code for iOS should be optional in the Engine
#iOS
#4.22

Adding Enable Advertising Identified (IDFA) in Settings >> iOS >> Build. Enabled by default
Fom CL5287566[REDACTED_HANDLE]

diff --git a/Engine/Config/BaseEngine.ini b/Engine/Config/BaseEngine.ini
--- a/Engine/Config/BaseEngine.ini
+++ b/Engine/Config/BaseEngine.ini
@@ -1867,6 +1867,7 @@ UseFastIntrinsics=False
 ForceFloats=False
 EnableMathOptimisations=True
 MaxShaderLanguageVersion=2
+bEnableAdvertisingIdentifier=True
 
 ; These are the defaults for Android settings, and they need to be in the .ini since UBT reads the .ini settings, without instantiating the class
 [/Script/AndroidRuntimeSettings.AndroidRuntimeSettings]
diff --git a/Engine/Source/Programs/UnrealBuildTool/Platform/IOS/UEBuildIOS.cs b/Engine/Source/Programs/UnrealBuildTool/Platform/IOS/UEBuildIOS.cs
--- a/Engine/Source/Programs/UnrealBuildTool/Platform/IOS/UEBuildIOS.cs
+++ b/Engine/Source/Programs/UnrealBuildTool/Platform/IOS/UEBuildIOS.cs
@@ -283,6 +283,12 @@ namespace UnrealBuildTool
 		[ConfigFile(ConfigHierarchyType.Engine, "/Script/IOSRuntimeSettings.IOSRuntimeSettings", "bDisableForceInline")]
 		public readonly bool bDisableForceInline = false;
 		
+		/// <summary>
+		/// true if IDFA are enabled
+		/// </summary>
+		[ConfigFile(ConfigHierarchyType.Engine, "/Script/IOSRuntimeSettings.IOSRuntimeSettings", "bEnableAdvertisingIdentifier")]
+		public readonly bool bEnableAdvertisingIdentifier = false;
+
 		/// <summary>
 		/// Returns a list of all the non-shipping architectures which are supported
 		/// </summary>
@@ -987,6 +993,11 @@ namespace UnrealBuildTool
 				CompileEnvironment.Definitions.Add("WITH_SIMULATOR=0");
 			}
 
+			if (ProjectSettings.bEnableAdvertisingIdentifier)
+			{
+				CompileEnvironment.Definitions.Add("ENABLE_ADVERTISING_IDENTIFIER=1");
+			}
+
 			CompileEnvironment.Definitions.Add("HAS_OPENGL_ES=" + ((Target.IOSPlatform.RuntimeVersion < 12.0) ? "1" : "0"));
 
 			// if the project has an Oodle compression Dll, enable the decompressor on IOS
@@ -1011,6 +1022,7 @@ namespace UnrealBuildTool
 			LinkEnvironment.AdditionalFrameworks.Add(new UEBuildFramework("GameKit"));
 			LinkEnvironment.AdditionalFrameworks.Add(new UEBuildFramework("StoreKit"));
 			LinkEnvironment.AdditionalFrameworks.Add(new UEBuildFramework("DeviceCheck"));
+
 		}
 
 		/// <summary>
diff --git a/Engine/Source/Runtime/Core/Private/IOS/IOSPlatformMisc.cpp b/Engine/Source/Runtime/Core/Private/IOS/IOSPlatformMisc.cpp
--- a/Engine/Source/Runtime/Core/Private/IOS/IOSPlatformMisc.cpp
+++ b/Engine/Source/Runtime/Core/Private/IOS/IOSPlatformMisc.cpp
@@ -52,6 +52,10 @@
 #include "Misc/CoreDelegates.h"
 #endif
 
+#if !defined ENABLE_ADVERTISING_IDENTIFIER
+	#define ENABLE_ADVERTISING_IDENTIFIER 0
+#endif
+
 //#include <libproc.h>
 // @pjs commented out to resolve issue with PLATFORM_TVOS being defined by mach-o loader
 //#include <mach-o/dyld.h>
@@ -802,7 +806,7 @@ void FIOSPlatformMisc::RequestStoreReview()
 */
 FString FIOSPlatformMisc::GetUniqueAdvertisingId()
 {
-#if !PLATFORM_TVOS
+#if !PLATFORM_TVOS && ENABLE_ADVERTISING_IDENTIFIER
 	// Check to see if this OS has this function
 	if ([[ASIdentifierManager sharedManager] respondsToSelector:@selector(advertisingIdentifier)])
 	{
diff --git a/Engine/Source/Runtime/IOS/IOSRuntimeSettings/Classes/IOSRuntimeSettings.h b/Engine/Source/Runtime/IOS/IOSRuntimeSettings/Classes/IOSRuntimeSettings.h
--- a/Engine/Source/Runtime/IOS/IOSRuntimeSettings/Classes/IOSRuntimeSettings.h
+++ b/Engine/Source/Runtime/IOS/IOSRuntimeSettings/Classes/IOSRuntimeSettings.h
@@ -291,6 +291,10 @@ public:
 	// Enable bitcode compiling?
 	UPROPERTY(GlobalConfig, EditAnywhere, Category = Build, meta = (DisplayName = "Support bitcode in Shipping"))
 	bool bShipForBitcode;
+
+	// Enable Advertising Identified
+	UPROPERTY(GlobalConfig, EditAnywhere, Category = Build, meta = (DisplayName = "Enable Advertising Identified (IDFA)"))
+	bool bEnableAdvertisingIdentifier;
 	
 	// Any additional linker flags to pass to the linker in non-shipping builds
 	UPROPERTY(GlobalConfig, EditAnywhere, Category = Build, meta = (DisplayName = "Additional Non-Shipping Linker Flags", ConfigHierarchyEditable))
