- Changed how framework signing for dylibs works. Remove the framework searching from Xcode into UBT (likely a placeholder spot until we overhaul some things)
- This keeps Mac from failing with SG_Com because now it only signs what has been staged via UBT/UAT

diff --git a/Engine/Source/Programs/UnrealBuildTool/Platform/IOS/IOSToolChain.cs b/Engine/Source/Programs/UnrealBuildTool/Platform/IOS/IOSToolChain.cs
--- a/Engine/Source/Programs/UnrealBuildTool/Platform/IOS/IOSToolChain.cs
+++ b/Engine/Source/Programs/UnrealBuildTool/Platform/IOS/IOSToolChain.cs
@@ -412,7 +412,7 @@ namespace UnrealBuildTool
 				Result += " -fobjc-exceptions";
 			}
 			else
-			{
+		{
 				Result += " -fno-objc-exceptions";
 			}
 
@@ -736,9 +736,9 @@ namespace UnrealBuildTool
 				CompileAction.bShouldOutputStatusDescription = true;
 
 				foreach(UEBuildFramework Framework in CompileEnvironment.AdditionalFrameworks)
-			{
+		{
 					if(Framework.ZipFile != null)
-				{
+			{
 						FileItem ExtractedTokenFile = ExtractFramework(Framework, Actions);
 						CompileAction.PrerequisiteItems.Add(ExtractedTokenFile);
 				}
@@ -1422,12 +1422,22 @@ namespace UnrealBuildTool
 				Dictionary<string, DirectoryReference> FrameworkNameToSourceDir = new Dictionary<string, DirectoryReference>();
 				foreach (UEBuildFramework Framework in BinaryLinkEnvironment.AdditionalFrameworks)
 				{
-					if (Framework.OutputDirectory != null && !String.IsNullOrEmpty(Framework.CopyBundledAssets))
+					if (Framework.OutputDirectory != null)
 					{
-						// For now, this is hard coded, but we need to loop over all modules, and copy bundled assets that need it
-						DirectoryReference LocalSource = DirectoryReference.Combine(Framework.OutputDirectory, Framework.CopyBundledAssets);
-						string BundleName = Framework.CopyBundledAssets.Substring(Framework.CopyBundledAssets.LastIndexOf('/') + 1);
-						FrameworkNameToSourceDir[BundleName] = LocalSource;
+						// copy entire framework in for dynamic lib frameworks, not just a bundle
+						// @todo add a bool to the UEBuildFramework class that denotes it as a dylib framework, or possibly use RuntimeDependencies instead?
+						if (Framework.ZipFile.FullName.EndsWith(".framework.zip"))
+						{
+							string BundleName = Framework.OutputDirectory.FullName.Substring(Framework.OutputDirectory.FullName.LastIndexOf('/') + 1);
+							FrameworkNameToSourceDir[BundleName] = Framework.OutputDirectory;
+						}
+						else if (!String.IsNullOrEmpty(Framework.CopyBundledAssets))
+						{
+							// For now, this is hard coded, but we need to loop over all modules, and copy bundled assets that need it
+							DirectoryReference LocalSource = DirectoryReference.Combine(Framework.OutputDirectory, Framework.CopyBundledAssets);
+							string BundleName = Framework.CopyBundledAssets.Substring(Framework.CopyBundledAssets.LastIndexOf('/') + 1);
+							FrameworkNameToSourceDir[BundleName] = LocalSource;
+						}
 					}
 				}
 
@@ -1785,16 +1795,29 @@ namespace UnrealBuildTool
 				CleanIntermediateDirectory(LocalFrameworkAssets);
 
 				foreach (KeyValuePair<string, DirectoryReference> Pair in Target.FrameworkNameToSourceDir)
-					{
+				{
 					string UnpackedZipPath = Pair.Value.FullName;
 
-					// For now, this is hard coded, but we need to loop over all modules, and copy bundled assets that need it
-					string LocalDest = LocalFrameworkAssets + "/" + Pair.Key;
-
-					Log.TraceInformation("Copying bundled asset... LocalSource: {0}, LocalDest: {1}", Pair.Value, LocalDest);
-
-					string ResultsText;
-                    RunExecutableAndWait("cp", String.Format("-R -L \"{0}\" \"{1}\"", Pair.Value, LocalDest), out ResultsText);
+					// some framework bundles are an entire framework, not a .bundle, so copy them into a Frameworks directory
+					if (Pair.Key.EndsWith(".framework"))
+					{
+						string LocalDest = LocalFrameworkAssets + "/Frameworks";
+						Console.WriteLine("Copying dynamic framework... LocalSource: {0}, LocalDest: {1}", UnpackedZipPath, LocalDest);
+						string ResultsText;
+						// Create the intermediate local directory
+						RunExecutableAndWait("mkdir", String.Format("-p \"{0}\"", LocalDest), out ResultsText);
+						RunExecutableAndWait("cp", String.Format("-R -L \"{0}\" \"{1}\"", Pair.Value, LocalDest), out ResultsText);
+					}
+					else
+					{
+						// For now, this is hard coded, but we need to loop over all modules, and copy bundled assets that need it
+						string LocalDest = LocalFrameworkAssets + "/" + Pair.Key;
+	
+						Log.TraceInformation("Copying bundled asset... LocalSource: {0}, LocalDest: {1}", Pair.Value, LocalDest);
+	
+						string ResultsText;
+						RunExecutableAndWait("cp", String.Format("-R -L \"{0}\" \"{1}\"", Pair.Value, LocalDest), out ResultsText);
+					}
                 }
             }
 		}
diff --git a/Engine/Source/Programs/UnrealBuildTool/ProjectFiles/Xcode/XcodeProject.cs b/Engine/Source/Programs/UnrealBuildTool/ProjectFiles/Xcode/XcodeProject.cs
--- a/Engine/Source/Programs/UnrealBuildTool/ProjectFiles/Xcode/XcodeProject.cs
+++ b/Engine/Source/Programs/UnrealBuildTool/ProjectFiles/Xcode/XcodeProject.cs
@@ -335,64 +335,6 @@ namespace UnrealBuildTool
 			PBXFileReferenceSection.Append(string.Format("\t\t{0} /* {1} */ = {{isa = PBXFileReference; explicitFileType = wrapper.application; path = {1}; sourceTree = BUILT_PRODUCTS_DIR; }};" + ProjectFileGenerator.NewLine, TargetAppGuid, TargetName));
 		}
 
-
-		private void GenerateSectionsWithFrameworks(StringBuilder PBXBuildFileSection, StringBuilder PBXFileReferenceSection, StringBuilder PBXCopyFilesBuildPhaseSection, FileReference UProjectPath)
-		{
-			if (UProjectPath != null)
-			{
-				// @todo: look also in Project/Build/Frameworks directory!
-
-				ProjectDescriptor Project = ProjectDescriptor.FromFile(UProjectPath);
-				List<PluginInfo> AvailablePlugins = Plugins.ReadAvailablePlugins(UnrealBuildTool.EngineDirectory, UProjectPath, Project.AdditionalPluginDirectories);
-
-				// look in each plugin for frameworks
-				// @todo: Cache this kind of things since every target will re-do this work!
-				foreach (PluginInfo PI in AvailablePlugins)
-				{
-					if (!Plugins.IsPluginEnabledForProject(PI, Project, UnrealTargetPlatform.IOS, UnrealTargetConfiguration.Development, TargetRules.TargetType.Game))
-					{
-						continue;
-					}
-
-					// for now, we copy and code sign all frameworks, even if the have no code
-					DirectoryReference FrameworkDir = DirectoryReference.Combine(PI.Directory, "Source/Frameworks");
-					if (!DirectoryReference.Exists(FrameworkDir))
-					{
-						FrameworkDir = DirectoryReference.Combine(PI.Directory, "Frameworks");
-					}
-					if (DirectoryReference.Exists(FrameworkDir))
-					{
-						// look at each zip
-						foreach (FileInfo FI in new System.IO.DirectoryInfo(FrameworkDir.FullName).EnumerateFiles("*.zip"))
-						{
-							string Guid = XcodeProjectFileGenerator.MakeXcodeGuid();
-							string RefGuid = XcodeProjectFileGenerator.MakeXcodeGuid();
-
-							// for FI of foo.framework.zip, this will give us foo.framework
-							string Framework = Path.GetFileNameWithoutExtension(FI.FullName);
-
-							// hunt down in subdirectories for 
-							string FrameworkPath = Path.Combine("../../../Engine/Intermediate/UnzippedFrameworks", PI.Name, "Frameworks", Framework);
-
-							Console.WriteLine("  Plugin {0} has framework {1}, reading from intermediate location {2}!", PI.Name, FI, FrameworkPath);
-
-							PBXBuildFileSection.Append(string.Format("\t\t{0} /* {1} in Embed Frameworks */ = {{isa = PBXBuildFile; fileRef = {2} /* {1} */; settings = {{ATTRIBUTES = (CodeSignOnCopy, RemoveHeadersOnCopy, ); }}; }};" + ProjectFileGenerator.NewLine,
-								Guid,
-								Framework,
-								RefGuid));
-
-							PBXFileReferenceSection.Append(string.Format("\t\t{0} /* {1} */ = {{isa = PBXFileReference; lastKnownFileType = wrapper.framework; name = \"{1}\"; path = \"{2}\"; sourceTree = SOURCE_ROOT; }};" + ProjectFileGenerator.NewLine,
-								RefGuid,
-								Framework,
-								FrameworkPath));
-
-							PBXCopyFilesBuildPhaseSection.Append("\t\t\t\t" + Guid + " /* " + Framework + " in Embed Frameworks */," + ProjectFileGenerator.NewLine);
-						}
-					}
-				}
-			}
-		}
-
 		private void GenerateSectionsWithExtensions(StringBuilder PBXBuildFileSection, StringBuilder PBXFileReferenceSection, StringBuilder PBXCopyFilesBuildPhaseSection, StringBuilder PBXResourcesBuildPhaseSection,
 													List<XcodeExtensionInfo> AllExtensions, FileReference UProjectPath, List<XcodeBuildConfig> BuildConfigs)
 		{
@@ -593,23 +535,6 @@ namespace UnrealBuildTool
 			return null;
 		}
 
-		private void AppendCopyFrameworksBuildPhaseSection(StringBuilder Content, StringBuilder SectionContent, string CopyFilesBuildPhaseGuid)
-		{
-			Content.Append("/* Begin PBXCopyFilesBuildPhase section */" + ProjectFileGenerator.NewLine);
-			Content.Append(string.Format("\t{0} /* Embed Frameworks */ = {{{1}", CopyFilesBuildPhaseGuid, ProjectFileGenerator.NewLine));
-			Content.Append("\t\tisa = PBXCopyFilesBuildPhase;" + ProjectFileGenerator.NewLine);
-			Content.Append("\t\tbuildActionMask = 2147483647;" + ProjectFileGenerator.NewLine);
-			Content.Append("\t\tdstPath = \"\";" + ProjectFileGenerator.NewLine);
-			Content.Append("\t\tdstSubfolderSpec = 10;" + ProjectFileGenerator.NewLine);
-			Content.Append("\t\tfiles = (" + ProjectFileGenerator.NewLine);
-			Content.Append(SectionContent);
-			Content.Append("\t\t);" + ProjectFileGenerator.NewLine);
-			Content.Append("\t\tname = \"Embed Frameworks\";" + ProjectFileGenerator.NewLine);
-			Content.Append("\t\trunOnlyForDeploymentPostprocessing = 0;" + ProjectFileGenerator.NewLine);
-			Content.Append("\t};" + ProjectFileGenerator.NewLine);
-			Content.Append("/* End PBXCopyFilesBuildPhase section */" + ProjectFileGenerator.NewLine);
-		}
-
 		private void AppendCopyExtensionsBuildPhaseSection(StringBuilder Content, StringBuilder SectionContent, string CopyFilesBuildPhaseGuid)
 		{
 			Content.Append("/* Begin PBXCopyFilesBuildPhase section */" + ProjectFileGenerator.NewLine);
@@ -705,7 +630,8 @@ namespace UnrealBuildTool
 			Content.Append("/* End PBXLegacyTarget section */" + ProjectFileGenerator.NewLine + ProjectFileGenerator.NewLine);
 		}
 
-		private void AppendRunTargetSection(StringBuilder Content, string TargetName, string TargetGuid, string TargetBuildConfigGuid, string TargetDependencyGuid, string TargetAppGuid, string CopyFrameworksBuildPhaseGuid, string CopyExtensionsBuildPhaseGuid, List<XcodeExtensionInfo> AllExtensions)
+		private void AppendRunTargetSection(StringBuilder Content, string TargetName, string TargetGuid, string TargetBuildConfigGuid, string TargetDependencyGuid, 
+				string TargetAppGuid, string CopyExtensionsBuildPhaseGuid, string ShellScriptSectionGuid, List<XcodeExtensionInfo> AllExtensions)
 		{
 			List<string> DependencyGuids = new List<string>();
 			// depends on the Run target if we want one
@@ -721,14 +647,14 @@ namespace UnrealBuildTool
 
 			Dictionary<string, string> BuildPhases = new Dictionary<string, string>();
 			// add optional build phases
-			if (!string.IsNullOrEmpty(CopyFrameworksBuildPhaseGuid))
-			{
-				BuildPhases.Add(CopyFrameworksBuildPhaseGuid, "Embed Frameworks");
-			}
 			if (!string.IsNullOrEmpty(CopyExtensionsBuildPhaseGuid))
 			{
 				BuildPhases.Add(CopyExtensionsBuildPhaseGuid, "Embed App Extensions");
 			}
+			if (!string.IsNullOrEmpty(ShellScriptSectionGuid))
+			{
+				BuildPhases.Add(ShellScriptSectionGuid, "Shell Script");
+			}
 
 			// use generica target section function for an application type
 			AppendGenericTargetSection(Content, TargetName, TargetGuid, "com.apple.product-type.application", TargetBuildConfigGuid, TargetAppGuid, DependencyGuids, BuildPhases);
@@ -790,6 +716,34 @@ namespace UnrealBuildTool
 			Content.Append("/* End PBXNativeTarget section */" + ProjectFileGenerator.NewLine + ProjectFileGenerator.NewLine);
 		}
 
+		private void AppendShellScriptSection(StringBuilder Content, string ShellScriptGuid)
+		{
+			string ShellScript = "set -e\\n\\n" +
+				"FRAMEWORK_DIR=$TARGET_BUILD_DIR/$EXECUTABLE_FOLDER_PATH/Frameworks\\n" +
+				"for FRAMEWORK in ${FRAMEWORK_DIR}/*.framework; do\\n" +
+					"\\t[ -d \\\"${FRAMEWORK}\\\" ] || continue\\n" + 
+					"\\techo Codesigning ${FRAMEWORK}\\n" +
+					"\\tcodesign --force --sign ${EXPANDED_CODE_SIGN_IDENTITY} --verbose --preserve-metadata=identifier,entitlements,flags --timestamp=none \\\"${FRAMEWORK}\\\"\\n" +
+				"done";
+
+			Content.Append("/* Begin PBXShellScriptBuildPhase section */" + ProjectFileGenerator.NewLine);
+			Content.Append(string.Format("\t\t{0} /* Sign Frameworks */ = {{" + ProjectFileGenerator.NewLine, ShellScriptGuid));
+			Content.Append("\t\t\tisa = PBXShellScriptBuildPhase;" + ProjectFileGenerator.NewLine);
+			Content.Append("\t\t\tbuildActionMask = 2147483647;" + ProjectFileGenerator.NewLine);
+			Content.Append("\t\t\tfiles = (" + ProjectFileGenerator.NewLine);
+			Content.Append("\t\t\t);" + ProjectFileGenerator.NewLine);
+			Content.Append("\t\t\tinputPaths = (" + ProjectFileGenerator.NewLine);
+			Content.Append("\t\t\t);" + ProjectFileGenerator.NewLine);
+			Content.Append("\t\t\toutputPaths = (" + ProjectFileGenerator.NewLine);
+			Content.Append("\t\t\t);" + ProjectFileGenerator.NewLine);
+			Content.Append("\t\t\tname = \"Sign Manual Frameworks\";" + ProjectFileGenerator.NewLine);
+			Content.Append("\t\t\trunOnlyForDeploymentPostprocessing = 0;" + ProjectFileGenerator.NewLine);
+			Content.Append("\t\t\tshellPath = /bin/sh;" + ProjectFileGenerator.NewLine);
+			Content.Append(string.Format("\t\t\tshellScript = \"{0}\";" + ProjectFileGenerator.NewLine, ShellScript));
+			Content.Append("\t\t};" + ProjectFileGenerator.NewLine);
+			Content.Append("/* End PBXShellScriptBuildPhase section */" + ProjectFileGenerator.NewLine);
+		}
+
 		private void AppendExtensionTargetSections(StringBuilder ProjectFileContent, List<XcodeExtensionInfo> AllExtensions)
 		{
 			foreach (XcodeExtensionInfo EI in AllExtensions)
@@ -1719,8 +1673,9 @@ namespace UnrealBuildTool
 			string MainGroupGuid = XcodeProjectFileGenerator.MakeXcodeGuid();
 			string ProductRefGroupGuid = XcodeProjectFileGenerator.MakeXcodeGuid();
 			string SourcesBuildPhaseGuid = XcodeProjectFileGenerator.MakeXcodeGuid();
-			string CopyFrameworksBuildPhaseGuid = XcodeProjectFileGenerator.MakeXcodeGuid();
 			string CopyExtensionsBuildPhaseGuid = XcodeProjectFileGenerator.MakeXcodeGuid();
+			string ShellScriptSectionGuid = XcodeProjectFileGenerator.MakeXcodeGuid();
+
 
 			// Figure out all the desired configurations
 			List<XcodeBuildConfig> BuildConfigs = GetSupportedBuildConfigs(InPlatforms, InConfigurations, PlatformProjectGenerators);
@@ -1761,12 +1716,10 @@ namespace UnrealBuildTool
 			StringBuilder PBXBuildFileSection = new StringBuilder();
 			StringBuilder PBXFileReferenceSection = new StringBuilder();
 			StringBuilder PBXSourcesBuildPhaseSection = new StringBuilder();
-			StringBuilder PBXCopyFrameworksBuildPhaseSection = new StringBuilder();
 			StringBuilder PBXCopyExtensionsBuildPhaseSection = new StringBuilder();
 			StringBuilder PBXResourcesBuildPhaseSection = new StringBuilder();
 			List<XcodeExtensionInfo> AllExtensions = new List<XcodeExtensionInfo>();
 			GenerateSectionsWithSourceFiles(PBXBuildFileSection, PBXFileReferenceSection, PBXSourcesBuildPhaseSection, TargetAppGuid, TargetName);
-			GenerateSectionsWithFrameworks(PBXBuildFileSection, PBXFileReferenceSection, PBXCopyFrameworksBuildPhaseSection, GameProjectPath);
 			GenerateSectionsWithExtensions(PBXBuildFileSection, PBXFileReferenceSection, PBXCopyExtensionsBuildPhaseSection, PBXResourcesBuildPhaseSection, AllExtensions, GameProjectPath, BuildConfigs);
 
 			StringBuilder ProjectFileContent = new StringBuilder();
@@ -1782,7 +1735,6 @@ namespace UnrealBuildTool
 			AppendBuildFileSection(ProjectFileContent, PBXBuildFileSection);
 			AppendFileReferenceSection(ProjectFileContent, PBXFileReferenceSection);
 			AppendSourcesBuildPhaseSection(ProjectFileContent, PBXSourcesBuildPhaseSection, SourcesBuildPhaseGuid);
-			AppendCopyFrameworksBuildPhaseSection(ProjectFileContent, PBXCopyFrameworksBuildPhaseSection, CopyFrameworksBuildPhaseGuid);
 			AppendCopyExtensionsBuildPhaseSection(ProjectFileContent, PBXCopyExtensionsBuildPhaseSection, CopyExtensionsBuildPhaseGuid);
 			ProjectFileContent.Append(PBXResourcesBuildPhaseSection);
 			AppendContainerItemProxySection(ProjectFileContent, BuildTargetName, BuildTargetGuid, TargetProxyGuid, ProjectGuid);
@@ -1797,13 +1749,14 @@ namespace UnrealBuildTool
 			}
 			AppendGroupSection(ProjectFileContent, MainGroupGuid, ProductRefGroupGuid, TargetAppGuid, TargetName, AllExtensions);
 			AppendLegacyTargetSection(ProjectFileContent, BuildTargetName, BuildTargetGuid, BuildTargetConfigListGuid, GameProjectPath, bHasEditorConfiguration);
-			AppendRunTargetSection(ProjectFileContent, TargetName, TargetGuid, TargetConfigListGuid, TargetDependencyGuid, TargetAppGuid, CopyFrameworksBuildPhaseGuid, CopyExtensionsBuildPhaseGuid, AllExtensions);
+			AppendRunTargetSection(ProjectFileContent, TargetName, TargetGuid, TargetConfigListGuid, TargetDependencyGuid, TargetAppGuid, CopyExtensionsBuildPhaseGuid, ShellScriptSectionGuid, AllExtensions);
 			AppendIndexTargetSection(ProjectFileContent, IndexTargetName, IndexTargetGuid, IndexTargetConfigListGuid, SourcesBuildPhaseGuid);
 			AppendExtensionTargetSections(ProjectFileContent, AllExtensions);
 			AppendProjectSection(ProjectFileContent, TargetName, TargetGuid, BuildTargetName, BuildTargetGuid, IndexTargetName, IndexTargetGuid, MainGroupGuid, ProductRefGroupGuid, ProjectGuid, ProjectConfigListGuid, GameProjectPath, AllExtensions);
 			AppendXCBuildConfigurationSection(ProjectFileContent, ProjectBuildConfigs, TargetBuildConfigs, BuildTargetBuildConfigs, IndexTargetBuildConfigs, GameProjectPath, AllExtensions);
 			AppendXCConfigurationListSection(ProjectFileContent, TargetName, BuildTargetName, IndexTargetName, ProjectConfigListGuid, ProjectBuildConfigs,
 				TargetConfigListGuid, TargetBuildConfigs, BuildTargetConfigListGuid, BuildTargetBuildConfigs, IndexTargetConfigListGuid, IndexTargetBuildConfigs, AllExtensions);
+			AppendShellScriptSection(ProjectFileContent, ShellScriptSectionGuid);
 
 			ProjectFileContent.Append("\t};" + ProjectFileGenerator.NewLine);
 			ProjectFileContent.Append("\trootObject = " + ProjectGuid + " /* Project object */;" + ProjectFileGenerator.NewLine);
