Fix exception when trying to get path for main module for hot reload.

#rnx

diff --git a/Engine/Source/Programs/UnrealBuildTool/System/HotReload.cs b/Engine/Source/Programs/UnrealBuildTool/System/HotReload.cs
--- a/Engine/Source/Programs/UnrealBuildTool/System/HotReload.cs
+++ b/Engine/Source/Programs/UnrealBuildTool/System/HotReload.cs
@@ -202,7 +202,17 @@ namespace UnrealBuildTool
 								continue;
 							}
 
-							if(!bIsRunning && EditorLocation == new FileReference(RunningProcess.MainModule.FileName))
+							FileReference MainModuleFile;
+							try
+							{
+								MainModuleFile = new FileReference(RunningProcess.MainModule.FileName);
+							}
+							catch
+							{
+								MainModuleFile = null;
+							}
+
+							if(!bIsRunning && EditorLocation == MainModuleFile)
 							{
 								bIsRunning = true;
 							}
