Fix settings for enabling shadow variable and undefined identifier warnings not being treated as requiring a unique PCH. Also fix shared PCHs not enabling each by default. The compiler (at least on Windows) does not produce a warning/error on this mismatch, but does not respect the different settings passed via the command line.

diff --git a/Engine/Source/Programs/UnrealBuildTool/Configuration/UEBuildModuleCPP.cs b/Engine/Source/Programs/UnrealBuildTool/Configuration/UEBuildModuleCPP.cs
--- a/Engine/Source/Programs/UnrealBuildTool/Configuration/UEBuildModuleCPP.cs
+++ b/Engine/Source/Programs/UnrealBuildTool/Configuration/UEBuildModuleCPP.cs
@@ -373,7 +373,7 @@ namespace UnrealBuildTool
 					PrecompiledHeaderTemplate Template = CompileEnvironment.SharedPCHs.FirstOrDefault(x => ReferencedModules.Contains(x.Module));
 					if(Template != null && Template.IsValidFor(CompileEnvironment))
 					{
-						PrecompiledHeaderInstance Instance = FindOrCreateSharedPCH(ToolChain, Template, ModuleCompileEnvironment.bOptimizeCode, ModuleCompileEnvironment.bUseRTTI, ModuleCompileEnvironment.bEnableExceptions, Makefile.Actions);
+						PrecompiledHeaderInstance Instance = FindOrCreateSharedPCH(ToolChain, Template, ModuleCompileEnvironment, Makefile.Actions);
 
 						FileReference PrivateDefinitionsFile = FileReference.Combine(IntermediateDirectory, String.Format("Definitions.{0}.h", Name));
 
@@ -565,7 +565,7 @@ namespace UnrealBuildTool
 
 			// Create the action to compile the PCH file.
 			CPPOutput Output = ToolChain.CompileCPPFiles(CompileEnvironment, new List<FileItem>() { WrapperFile }, IntermediateDirectory, Name, Actions);
-			return new PrecompiledHeaderInstance(WrapperFile, CompileEnvironment.bOptimizeCode, CompileEnvironment.bUseRTTI, CompileEnvironment.bEnableExceptions, Output);
+			return new PrecompiledHeaderInstance(WrapperFile, CompileEnvironment, Output);
 		}
 
 		/// <summary>
@@ -573,51 +573,16 @@ namespace UnrealBuildTool
 		/// </summary>
 		/// <param name="ToolChain">The toolchain being used to build this module</param>
 		/// <param name="Template">The PCH template</param>
-		/// <param name="bOptimizeCode">Whether optimization should be enabled for this PCH</param>
-		/// <param name="bUseRTTI">Whether to enable RTTI for this PCH</param>
-		/// <param name="bEnableExceptions">Whether to enable exceptions for this PCH</param>
+		/// <param name="ModuleCompileEnvironment">Compile environment for the current module</param>
 		/// <param name="Actions">List of actions to be executed. Additional actions will be added to this list.</param>
 		/// <returns>Instance of a PCH</returns>
-		public PrecompiledHeaderInstance FindOrCreateSharedPCH(UEToolChain ToolChain, PrecompiledHeaderTemplate Template, bool bOptimizeCode, bool bUseRTTI, bool bEnableExceptions, List<Action> Actions)
+		public PrecompiledHeaderInstance FindOrCreateSharedPCH(UEToolChain ToolChain, PrecompiledHeaderTemplate Template, CppCompileEnvironment ModuleCompileEnvironment, List<Action> Actions)
 		{
-			PrecompiledHeaderInstance Instance = Template.Instances.Find(x => x.bOptimizeCode == bOptimizeCode && x.bUseRTTI == bUseRTTI && x.bEnableExceptions == bEnableExceptions);
+			PrecompiledHeaderInstance Instance = Template.Instances.Find(x => IsCompatibleForSharedPCH(x.CompileEnvironment, ModuleCompileEnvironment));
 			if(Instance == null)
 			{
 				// Create a suffix to distinguish this shared PCH variant from any others. Currently only optimized and non-optimized shared PCHs are supported.
-				string Variant = "";
-				if(bOptimizeCode != Template.BaseCompileEnvironment.bOptimizeCode)
-				{
-					if(bOptimizeCode)
-					{
-						Variant += ".Optimized";
-					}
-					else
-					{
-						Variant += ".NonOptimized";
-					}
-				}
-				if(bUseRTTI != Template.BaseCompileEnvironment.bUseRTTI)
-				{
-					if (bUseRTTI)
-					{
-						Variant += ".RTTI";
-					}
-					else
-					{
-						Variant += ".NonRTTI";
-					}
-				}
-				if (bEnableExceptions != Template.BaseCompileEnvironment.bEnableExceptions)
-				{
-					if (bEnableExceptions)
-					{
-						Variant += ".Exceptions";
-					}
-					else
-					{
-						Variant += ".NoExceptions";
-					}
-				}
+				string Variant = GetSuffixForSharedPCH(ModuleCompileEnvironment, Template.BaseCompileEnvironment);
 
 				// Create the wrapper file, which sets all the definitions needed to compile it
 				FileReference WrapperLocation = FileReference.Combine(Template.OutputDir, String.Format("SharedPCH.{0}{1}.h", Template.Module.Name, Variant));
@@ -628,18 +593,128 @@ namespace UnrealBuildTool
 				CompileEnvironment.Definitions.Clear();
 				CompileEnvironment.PrecompiledHeaderAction = PrecompiledHeaderAction.Create;
 				CompileEnvironment.PrecompiledHeaderIncludeFilename = WrapperFile.Location;
-				CompileEnvironment.bOptimizeCode = bOptimizeCode;
-				CompileEnvironment.bUseRTTI = bUseRTTI;
-				CompileEnvironment.bEnableExceptions = bEnableExceptions;
+				CopySettingsForSharedPCH(ModuleCompileEnvironment, CompileEnvironment);
 
 				// Create the PCH
 				CPPOutput Output = ToolChain.CompileCPPFiles(CompileEnvironment, new List<FileItem>() { WrapperFile }, Template.OutputDir, "Shared", Actions);
-				Instance = new PrecompiledHeaderInstance(WrapperFile, bOptimizeCode, bUseRTTI, bEnableExceptions, Output);
+				Instance = new PrecompiledHeaderInstance(WrapperFile, CompileEnvironment, Output);
 				Template.Instances.Add(Instance);
 			}
 			return Instance;
 		}
 
+		/// <summary>
+		/// Determines if a module compile environment is compatible with the given shared PCH compile environment
+		/// </summary>
+		/// <param name="ModuleCompileEnvironment">The module compile environment</param>
+		/// <param name="CompileEnvironment">The shared PCH compile environment</param>
+		/// <returns>True if the two compile enviroments are compatible</returns>
+		private bool IsCompatibleForSharedPCH(CppCompileEnvironment ModuleCompileEnvironment, CppCompileEnvironment CompileEnvironment)
+		{
+			if(ModuleCompileEnvironment.bOptimizeCode != CompileEnvironment.bOptimizeCode)
+			{
+				return false;
+			}
+			if(ModuleCompileEnvironment.bUseRTTI != CompileEnvironment.bUseRTTI)
+			{
+				return false;
+			}
+			if(ModuleCompileEnvironment.bEnableExceptions != CompileEnvironment.bEnableExceptions)
+			{
+				return false;
+			}
+			if(ModuleCompileEnvironment.bEnableShadowVariableWarnings != CompileEnvironment.bEnableShadowVariableWarnings)
+			{
+				return false;
+			}
+			if(ModuleCompileEnvironment.bEnableUndefinedIdentifierWarnings != CompileEnvironment.bEnableUndefinedIdentifierWarnings)
+			{
+				return false;
+			}
+			return true;
+		}
+
+		/// <summary>
+		/// Gets the unique suffix for a shared PCH
+		/// </summary>
+		/// <param name="CompileEnvironment">The shared PCH compile environment</param>
+		/// <param name="BaseCompileEnvironment">The base compile environment</param>
+		/// <returns>The unique suffix for the shared PCH</returns>
+		private string GetSuffixForSharedPCH(CppCompileEnvironment CompileEnvironment, CppCompileEnvironment BaseCompileEnvironment)
+		{
+			string Variant = "";
+			if(CompileEnvironment.bOptimizeCode != BaseCompileEnvironment.bOptimizeCode)
+			{
+				if(CompileEnvironment.bOptimizeCode)
+				{
+					Variant += ".Optimized";
+				}
+				else
+				{
+					Variant += ".NonOptimized";
+				}
+			}
+			if(CompileEnvironment.bUseRTTI != BaseCompileEnvironment.bUseRTTI)
+			{
+				if (CompileEnvironment.bUseRTTI)
+				{
+					Variant += ".RTTI";
+				}
+				else
+				{
+					Variant += ".NonRTTI";
+				}
+			}
+			if (CompileEnvironment.bEnableExceptions != BaseCompileEnvironment.bEnableExceptions)
+			{
+				if (CompileEnvironment.bEnableExceptions)
+				{
+					Variant += ".Exceptions";
+				}
+				else
+				{
+					Variant += ".NoExceptions";
+				}
+			}
+			if (CompileEnvironment.bEnableShadowVariableWarnings != BaseCompileEnvironment.bEnableShadowVariableWarnings)
+			{
+				if (CompileEnvironment.bEnableShadowVariableWarnings)
+				{
+					Variant += ".Shadow";
+				}
+				else
+				{
+					Variant += ".NoShadow";
+				}
+			}
+			if (CompileEnvironment.bEnableUndefinedIdentifierWarnings != BaseCompileEnvironment.bEnableUndefinedIdentifierWarnings)
+			{
+				if (CompileEnvironment.bEnableUndefinedIdentifierWarnings)
+				{
+					Variant += ".Undef";
+				}
+				else
+				{
+					Variant += ".NoUndef";
+				}
+			}
+			return Variant;
+		}
+
+		/// <summary>
+		/// Copy settings from the module's compile environment into the environment for the shared PCH
+		/// </summary>
+		/// <param name="ModuleCompileEnvironment">The module compile environment</param>
+		/// <param name="CompileEnvironment">The shared PCH compile environment</param>
+		private void CopySettingsForSharedPCH(CppCompileEnvironment ModuleCompileEnvironment, CppCompileEnvironment CompileEnvironment)
+		{
+			CompileEnvironment.bOptimizeCode = ModuleCompileEnvironment.bOptimizeCode;
+			CompileEnvironment.bUseRTTI = ModuleCompileEnvironment.bUseRTTI;
+			CompileEnvironment.bEnableExceptions = ModuleCompileEnvironment.bEnableExceptions;
+			CompileEnvironment.bEnableShadowVariableWarnings = ModuleCompileEnvironment.bEnableShadowVariableWarnings;
+			CompileEnvironment.bEnableUndefinedIdentifierWarnings = ModuleCompileEnvironment.bEnableUndefinedIdentifierWarnings;
+		}
+
 		/// <summary>
 		/// Compiles the provided CPP unity files. Will
 		/// </summary>
diff --git a/Engine/Source/Programs/UnrealBuildTool/System/CppCompileEnvironment.cs b/Engine/Source/Programs/UnrealBuildTool/System/CppCompileEnvironment.cs
--- a/Engine/Source/Programs/UnrealBuildTool/System/CppCompileEnvironment.cs
+++ b/Engine/Source/Programs/UnrealBuildTool/System/CppCompileEnvironment.cs
@@ -195,7 +195,7 @@ namespace UnrealBuildTool
 		/// <summary>
 		/// Whether to warn about the use of shadow variables
 		/// </summary>
-		public bool bEnableShadowVariableWarnings = false;
+		public bool bEnableShadowVariableWarnings = true;
 
 		/// <summary>
 		/// Whether to treat shadow variable warnings as errors.
@@ -205,7 +205,7 @@ namespace UnrealBuildTool
 		/// <summary>
 		/// Whether to warn about the use of undefined identifiers in #if expressions
 		/// </summary>
-		public bool bEnableUndefinedIdentifierWarnings = false;
+		public bool bEnableUndefinedIdentifierWarnings = true;
 
 		/// <summary>
 		/// Whether to treat undefined identifier warnings as errors.
diff --git a/Engine/Source/Programs/UnrealBuildTool/System/PrecompiledHeaderInstance.cs b/Engine/Source/Programs/UnrealBuildTool/System/PrecompiledHeaderInstance.cs
--- a/Engine/Source/Programs/UnrealBuildTool/System/PrecompiledHeaderInstance.cs
+++ b/Engine/Source/Programs/UnrealBuildTool/System/PrecompiledHeaderInstance.cs
@@ -1,4 +1,4 @@
-﻿// Copyright 1998-2019 Epic Games, Inc. All Rights Reserved.
+// Copyright 1998-2019 Epic Games, Inc. All Rights Reserved.
 
 using System;
 using System.Collections.Generic;
@@ -19,19 +19,9 @@ namespace UnrealBuildTool
 		public FileItem HeaderFile;
 
 		/// <summary>
-		/// Whether optimization is enabled
+		/// The compile environment for this shared PCH
 		/// </summary>
-		public bool bOptimizeCode;
-
-		/// <summary>
-		/// Whether to enable RTTI
-		/// </summary>
-		public bool bUseRTTI;
-
-		/// <summary>
-		/// Whether to enable exceptions
-		/// </summary>
-		public bool bEnableExceptions;
+		public CppCompileEnvironment CompileEnvironment;
 
 		/// <summary>
 		/// The output files for the shared PCH
@@ -41,12 +31,10 @@ namespace UnrealBuildTool
 		/// <summary>
 		/// Constructor
 		/// </summary>
-		public PrecompiledHeaderInstance(FileItem HeaderFile, bool bOptimizeCode, bool bUseRTTI, bool bEnableExceptions, CPPOutput Output)
+		public PrecompiledHeaderInstance(FileItem HeaderFile, CppCompileEnvironment CompileEnvironment, CPPOutput Output)
 		{
 			this.HeaderFile = HeaderFile;
-			this.bOptimizeCode = bOptimizeCode;
-			this.bUseRTTI = bUseRTTI;
-			this.bEnableExceptions = bEnableExceptions;
+			this.CompileEnvironment = CompileEnvironment;
 			this.Output = Output;
 		}
 
@@ -56,7 +44,7 @@ namespace UnrealBuildTool
 		/// <returns>String representation of the object</returns>
 		public override string ToString()
 		{
-			return String.Format("{0} (Optimized={1}, RTTI={2}, Exceptions={3})", HeaderFile.Location.GetFileName(), bOptimizeCode, bUseRTTI, bEnableExceptions);
+			return HeaderFile.Location.GetFileName();
 		}
 	}
 }
