alllow incredibuild to restart processes on local machine.

diff --git a/Engine/Source/Programs/UnrealBuildTool/Executors/XGE.cs b/Engine/Source/Programs/UnrealBuildTool/Executors/XGE.cs
--- a/Engine/Source/Programs/UnrealBuildTool/Executors/XGE.cs
+++ b/Engine/Source/Programs/UnrealBuildTool/Executors/XGE.cs
@@ -393,6 +393,7 @@ namespace UnrealBuildTool
 				TaskElement.SetAttribute("Tool", string.Format("Tool{0}", ActionIndex));
 				TaskElement.SetAttribute("WorkingDir", Action.WorkingDirectory.FullName);
 				TaskElement.SetAttribute("SkipIfProjectFailed", "true");
+				TaskElement.SetAttribute("AllowRestartOnLocal", "true");
 
 				// Create a semi-colon separated list of the other tasks this task depends on the results of.
 				List<string> DependencyNames = new List<string>();
