UBT: Add additional timing info into logs.

#rnx

diff --git a/Engine/Source/Programs/UnrealBuildTool/Configuration/UEBuildPlatform.cs b/Engine/Source/Programs/UnrealBuildTool/Configuration/UEBuildPlatform.cs
--- a/Engine/Source/Programs/UnrealBuildTool/Configuration/UEBuildPlatform.cs
+++ b/Engine/Source/Programs/UnrealBuildTool/Configuration/UEBuildPlatform.cs
@@ -60,10 +60,18 @@ namespace UnrealBuildTool
 		/// <param name="bIncludeNonInstalledPlatforms">Whether to register platforms that are not installed</param>
 		public static void RegisterPlatforms(bool bIncludeNonInstalledPlatforms)
 		{
-			// Find and register all tool chains and build platforms that are present
-			Assembly UBTAssembly = Assembly.GetExecutingAssembly();
+			// Initialize the installed platform info
+			using(Timeline.ScopeEvent("Initializing InstalledPlatformInfo"))
+			{
+				InstalledPlatformInfo.Initialize();
+			}
 
-			Type[] AllTypes = UBTAssembly.GetTypes();
+			// Find and register all tool chains and build platforms that are present
+			Type[] AllTypes;
+			using(Timeline.ScopeEvent("Querying types"))
+			{
+				AllTypes = Assembly.GetExecutingAssembly().GetTypes();
+			}
 
 			// register all build platforms first, since they implement SDK-switching logic that can set environment variables
 			foreach (Type CheckType in AllTypes)
@@ -73,12 +81,15 @@ namespace UnrealBuildTool
 					if (CheckType.IsSubclassOf(typeof(UEBuildPlatformFactory)))
 					{
 						Log.TraceVerbose("    Registering build platform: {0}", CheckType.ToString());
-						UEBuildPlatformFactory TempInst = (UEBuildPlatformFactory)(UBTAssembly.CreateInstance(CheckType.FullName, true));
-
-						// We need all platforms to be registered when we run -validateplatform command to check SDK status of each
-						if (bIncludeNonInstalledPlatforms || InstalledPlatformInfo.IsValidPlatform(TempInst.TargetPlatform))
+						using(Timeline.ScopeEvent(CheckType.Name))
 						{
-							TempInst.RegisterBuildPlatforms();
+							UEBuildPlatformFactory TempInst = (UEBuildPlatformFactory)Activator.CreateInstance(CheckType);
+							
+							// We need all platforms to be registered when we run -validateplatform command to check SDK status of each
+							if (bIncludeNonInstalledPlatforms || InstalledPlatformInfo.IsValidPlatform(TempInst.TargetPlatform))
+							{
+								TempInst.RegisterBuildPlatforms();
+							}
 						}
 					}
 				}
diff --git a/Engine/Source/Programs/UnrealBuildTool/System/ActionGraph.cs b/Engine/Source/Programs/UnrealBuildTool/System/ActionGraph.cs
--- a/Engine/Source/Programs/UnrealBuildTool/System/ActionGraph.cs
+++ b/Engine/Source/Programs/UnrealBuildTool/System/ActionGraph.cs
@@ -563,9 +563,12 @@ namespace UnrealBuildTool
 				Parallel.ForEach(Dependencies, File => { List<FileItem> Temp; CppDependencies.TryGetDependencies(File, out Temp); });
 			}
 
-			foreach (Action Action in Actions)
+			using(Timeline.ScopeEvent("Cache outdated actions"))
 			{
-				IsActionOutdated(Action, OutdatedActions, ActionHistory, CppDependencies, bIgnoreOutdatedImportLibraries);
+				foreach (Action Action in Actions)
+				{
+					IsActionOutdated(Action, OutdatedActions, ActionHistory, CppDependencies, bIgnoreOutdatedImportLibraries);
+				}
 			}
 		}
 
diff --git a/Engine/Source/Programs/UnrealBuildTool/System/InstalledPlatformInfo.cs b/Engine/Source/Programs/UnrealBuildTool/System/InstalledPlatformInfo.cs
--- a/Engine/Source/Programs/UnrealBuildTool/System/InstalledPlatformInfo.cs
+++ b/Engine/Source/Programs/UnrealBuildTool/System/InstalledPlatformInfo.cs
@@ -124,6 +124,14 @@ namespace UnrealBuildTool
 			}
 		}
 
+		/// <summary>
+		/// Initializes the InstalledPlatformInfo. While this is not necessary to be called, it allows timing the static constructor call.
+		/// </summary>
+		internal static void Initialize()
+		{
+			// Unused, but allows timing call of the static constructor
+		}
+
 		private static void ParsePlatformConfiguration(string PlatformConfiguration)
 		{
 			// Trim whitespace at the beginning.
