UBT: Figure out outdated actions in parallel. Saves 0.5s on iterative compile.

#rnx

diff --git a/Engine/Source/Programs/UnrealBuildTool/System/ActionGraph.cs b/Engine/Source/Programs/UnrealBuildTool/System/ActionGraph.cs
--- a/Engine/Source/Programs/UnrealBuildTool/System/ActionGraph.cs
+++ b/Engine/Source/Programs/UnrealBuildTool/System/ActionGraph.cs
@@ -370,136 +370,147 @@ namespace UnrealBuildTool
 		{
 			// Only compute the outdated-ness for actions that don't aren't cached in the outdated action dictionary.
 			bool bIsOutdated = false;
-			if (!OutdatedActionDictionary.TryGetValue(RootAction, out bIsOutdated))
+			lock(OutdatedActionDictionary)
 			{
-				// Determine the last time the action was run based on the write times of its produced files.
-				string LatestUpdatedProducedItemName = null;
-				DateTimeOffset LastExecutionTimeUtc = DateTimeOffset.MaxValue;
-				foreach (FileItem ProducedItem in RootAction.ProducedItems)
-				{
-					// Check if the command-line of the action previously used to produce the item is outdated.
-					string OldProducingCommandLine = "";
-					string NewProducingCommandLine = RootAction.CommandPath.FullName + " " + RootAction.CommandArguments;
-					if (!ActionHistory.TryGetProducingCommandLine(ProducedItem, out OldProducingCommandLine)
-					|| !String.Equals(OldProducingCommandLine, NewProducingCommandLine, StringComparison.InvariantCultureIgnoreCase))
-					{
-						if(ProducedItem.Exists)
-						{
-							Log.TraceLog(
-								"{0}: Produced item \"{1}\" was produced by outdated command-line.\n  Old command-line: {2}\n  New command-line: {3}",
-								RootAction.StatusDescription,
-								Path.GetFileName(ProducedItem.AbsolutePath),
-								OldProducingCommandLine,
-								NewProducingCommandLine
-								);
-						}
-
-						bIsOutdated = true;
-
-						// Update the command-line used to produce this item in the action history.
-						ActionHistory.SetProducingCommandLine(ProducedItem, NewProducingCommandLine);
-					}
+				if (OutdatedActionDictionary.TryGetValue(RootAction, out bIsOutdated))
+				{
+					return bIsOutdated;
+				}
+			}
 
-					// If the produced file doesn't exist or has zero size, consider it outdated.  The zero size check is to detect cases
-					// where aborting an earlier compile produced invalid zero-sized obj files, but that may cause actions where that's
-					// legitimate output to always be considered outdated.
-					if (ProducedItem.Exists && (RootAction.ActionType != ActionType.Compile || ProducedItem.Length > 0 || (!ProducedItem.Location.HasExtension(".obj") && !ProducedItem.Location.HasExtension(".o"))))
-					{
-						// Use the oldest produced item's time as the last execution time.
-						if (ProducedItem.LastWriteTimeUtc < LastExecutionTimeUtc)
-						{
-							LastExecutionTimeUtc = ProducedItem.LastWriteTimeUtc;
-							LatestUpdatedProducedItemName = ProducedItem.AbsolutePath;
-						}
-					}
-					else
+			// Determine the last time the action was run based on the write times of its produced files.
+			string LatestUpdatedProducedItemName = null;
+			DateTimeOffset LastExecutionTimeUtc = DateTimeOffset.MaxValue;
+			foreach (FileItem ProducedItem in RootAction.ProducedItems)
+			{
+				// Check if the command-line of the action previously used to produce the item is outdated.
+				string OldProducingCommandLine = "";
+				string NewProducingCommandLine = RootAction.CommandPath.FullName + " " + RootAction.CommandArguments;
+				if (!ActionHistory.TryGetProducingCommandLine(ProducedItem, out OldProducingCommandLine)
+				|| !String.Equals(OldProducingCommandLine, NewProducingCommandLine, StringComparison.OrdinalIgnoreCase))
+				{
+					if(ProducedItem.Exists)
 					{
-						// If any of the produced items doesn't exist, the action is outdated.
 						Log.TraceLog(
-							"{0}: Produced item \"{1}\" doesn't exist.",
+							"{0}: Produced item \"{1}\" was produced by outdated command-line.\n  Old command-line: {2}\n  New command-line: {3}",
 							RootAction.StatusDescription,
-							Path.GetFileName(ProducedItem.AbsolutePath)
+							Path.GetFileName(ProducedItem.AbsolutePath),
+							OldProducingCommandLine,
+							NewProducingCommandLine
 							);
-						bIsOutdated = true;
 					}
+
+					bIsOutdated = true;
+
+					// Update the command-line used to produce this item in the action history.
+					ActionHistory.SetProducingCommandLine(ProducedItem, NewProducingCommandLine);
 				}
 
-				// Check if any of the prerequisite actions are out of date
-				if (!bIsOutdated)
+				// If the produced file doesn't exist or has zero size, consider it outdated.  The zero size check is to detect cases
+				// where aborting an earlier compile produced invalid zero-sized obj files, but that may cause actions where that's
+				// legitimate output to always be considered outdated.
+				if (ProducedItem.Exists && (RootAction.ActionType != ActionType.Compile || ProducedItem.Length > 0 || (!ProducedItem.Location.HasExtension(".obj") && !ProducedItem.Location.HasExtension(".o"))))
 				{
-					foreach (Action PrerequisiteAction in RootAction.PrerequisiteActions)
+					// Use the oldest produced item's time as the last execution time.
+					if (ProducedItem.LastWriteTimeUtc < LastExecutionTimeUtc)
 					{
-						if (IsActionOutdated(PrerequisiteAction, OutdatedActionDictionary, ActionHistory, CppDependencies, bIgnoreOutdatedImportLibraries))
-						{
-							// Only check for outdated import libraries if we were configured to do so.  Often, a changed import library
-							// won't affect a dependency unless a public header file was also changed, in which case we would be forced
-							// to recompile anyway.  This just allows for faster iteration when working on a subsystem in a DLL, as we
-							// won't have to wait for dependent targets to be relinked after each change.
-							if(!bIgnoreOutdatedImportLibraries || !IsImportLibraryDependency(RootAction, PrerequisiteAction))
-							{
-								Log.TraceLog("{0}: Prerequisite {1} is produced by outdated action.", RootAction.StatusDescription, PrerequisiteAction.StatusDescription);
-								bIsOutdated = true;
-								break;
-							}
-						}
+						LastExecutionTimeUtc = ProducedItem.LastWriteTimeUtc;
+						LatestUpdatedProducedItemName = ProducedItem.AbsolutePath;
 					}
-				} 
+				}
+				else
+				{
+					// If any of the produced items doesn't exist, the action is outdated.
+					Log.TraceLog(
+						"{0}: Produced item \"{1}\" doesn't exist.",
+						RootAction.StatusDescription,
+						Path.GetFileName(ProducedItem.AbsolutePath)
+						);
+					bIsOutdated = true;
+				}
+			}
 
-				// Check if any prerequisite item has a newer timestamp than the last execution time of this action
-				if(!bIsOutdated)
+			// Check if any of the prerequisite actions are out of date
+			if (!bIsOutdated)
+			{
+				foreach (Action PrerequisiteAction in RootAction.PrerequisiteActions)
 				{
-					foreach (FileItem PrerequisiteItem in RootAction.PrerequisiteItems)
+					if (IsActionOutdated(PrerequisiteAction, OutdatedActionDictionary, ActionHistory, CppDependencies, bIgnoreOutdatedImportLibraries))
 					{
-						if (PrerequisiteItem.Exists)
+						// Only check for outdated import libraries if we were configured to do so.  Often, a changed import library
+						// won't affect a dependency unless a public header file was also changed, in which case we would be forced
+						// to recompile anyway.  This just allows for faster iteration when working on a subsystem in a DLL, as we
+						// won't have to wait for dependent targets to be relinked after each change.
+						if(!bIgnoreOutdatedImportLibraries || !IsImportLibraryDependency(RootAction, PrerequisiteAction))
 						{
-							// allow a 1 second slop for network copies
-							TimeSpan TimeDifference = PrerequisiteItem.LastWriteTimeUtc - LastExecutionTimeUtc;
-							bool bPrerequisiteItemIsNewerThanLastExecution = TimeDifference.TotalSeconds > 1;
-							if (bPrerequisiteItemIsNewerThanLastExecution)
-							{
-								// Need to check for import libraries here too
-								if(!bIgnoreOutdatedImportLibraries || !IsImportLibraryDependency(RootAction, PrerequisiteItem))
-								{
-									Log.TraceLog("{0}: Prerequisite {1} is newer than the last execution of the action: {2} vs {3}", RootAction.StatusDescription, Path.GetFileName(PrerequisiteItem.AbsolutePath), PrerequisiteItem.LastWriteTimeUtc.ToLocalTime(), LastExecutionTimeUtc.LocalDateTime);
-									bIsOutdated = true;
-									break;
-								}
-							}
+							Log.TraceLog("{0}: Prerequisite {1} is produced by outdated action.", RootAction.StatusDescription, PrerequisiteAction.StatusDescription);
+							bIsOutdated = true;
+							break;
 						}
 					}
 				}
+			} 
 
-				// Check the dependency list
-				if(!bIsOutdated && RootAction.DependencyListFile != null)
+			// Check if any prerequisite item has a newer timestamp than the last execution time of this action
+			if(!bIsOutdated)
+			{
+				foreach (FileItem PrerequisiteItem in RootAction.PrerequisiteItems)
 				{
-					List<FileItem> DependencyFiles;
-					if(!CppDependencies.TryGetDependencies(RootAction.DependencyListFile, out DependencyFiles))
-					{
-						Log.TraceLog("{0}: Missing dependency list file \"{1}\"", RootAction.StatusDescription, RootAction.DependencyListFile);
-						bIsOutdated = true;
-					}
-					else
+					if (PrerequisiteItem.Exists)
 					{
-						foreach(FileItem DependencyFile in DependencyFiles)
+						// allow a 1 second slop for network copies
+						TimeSpan TimeDifference = PrerequisiteItem.LastWriteTimeUtc - LastExecutionTimeUtc;
+						bool bPrerequisiteItemIsNewerThanLastExecution = TimeDifference.TotalSeconds > 1;
+						if (bPrerequisiteItemIsNewerThanLastExecution)
 						{
-							if(!DependencyFile.Exists || DependencyFile.LastWriteTimeUtc > LastExecutionTimeUtc)
+							// Need to check for import libraries here too
+							if(!bIgnoreOutdatedImportLibraries || !IsImportLibraryDependency(RootAction, PrerequisiteItem))
 							{
-								Log.TraceLog(
-									"{0}: Dependency {1} is newer than the last execution of the action: {2} vs {3}",
-									RootAction.StatusDescription,
-									Path.GetFileName(DependencyFile.AbsolutePath),
-									DependencyFile.LastWriteTimeUtc.ToLocalTime(),
-									LastExecutionTimeUtc.LocalDateTime
-									);
+								Log.TraceLog("{0}: Prerequisite {1} is newer than the last execution of the action: {2} vs {3}", RootAction.StatusDescription, Path.GetFileName(PrerequisiteItem.AbsolutePath), PrerequisiteItem.LastWriteTimeUtc.ToLocalTime(), LastExecutionTimeUtc.LocalDateTime);
 								bIsOutdated = true;
 								break;
 							}
 						}
 					}
 				}
+			}
 
-				// Cache the outdated-ness of this action.
-				OutdatedActionDictionary.Add(RootAction, bIsOutdated);
+			// Check the dependency list
+			if(!bIsOutdated && RootAction.DependencyListFile != null)
+			{
+				List<FileItem> DependencyFiles;
+				if(!CppDependencies.TryGetDependencies(RootAction.DependencyListFile, out DependencyFiles))
+				{
+					Log.TraceLog("{0}: Missing dependency list file \"{1}\"", RootAction.StatusDescription, RootAction.DependencyListFile);
+					bIsOutdated = true;
+				}
+				else
+				{
+					foreach(FileItem DependencyFile in DependencyFiles)
+					{
+						if(!DependencyFile.Exists || DependencyFile.LastWriteTimeUtc > LastExecutionTimeUtc)
+						{
+							Log.TraceLog(
+								"{0}: Dependency {1} is newer than the last execution of the action: {2} vs {3}",
+								RootAction.StatusDescription,
+								Path.GetFileName(DependencyFile.AbsolutePath),
+								DependencyFile.LastWriteTimeUtc.ToLocalTime(),
+								LastExecutionTimeUtc.LocalDateTime
+								);
+							bIsOutdated = true;
+							break;
+						}
+					}
+				}
+			}
+
+			// Cache the outdated-ness of this action.
+			lock(OutdatedActionDictionary)
+			{
+				if(!OutdatedActionDictionary.ContainsKey(RootAction))
+				{
+					OutdatedActionDictionary.Add(RootAction, bIsOutdated);
+				}
 			}
 
 			return bIsOutdated;
@@ -565,10 +576,7 @@ namespace UnrealBuildTool
 
 			using(Timeline.ScopeEvent("Cache outdated actions"))
 			{
-				foreach (Action Action in Actions)
-				{
-					IsActionOutdated(Action, OutdatedActions, ActionHistory, CppDependencies, bIgnoreOutdatedImportLibraries);
-				}
+				Parallel.ForEach(Actions, Action => IsActionOutdated(Action, OutdatedActions, ActionHistory, CppDependencies, bIgnoreOutdatedImportLibraries));
 			}
 		}
 
diff --git a/Engine/Source/Programs/UnrealBuildTool/System/ActionHistory.cs b/Engine/Source/Programs/UnrealBuildTool/System/ActionHistory.cs
--- a/Engine/Source/Programs/UnrealBuildTool/System/ActionHistory.cs
+++ b/Engine/Source/Programs/UnrealBuildTool/System/ActionHistory.cs
@@ -44,6 +44,11 @@ namespace UnrealBuildTool
 		/// </summary>
 		bool bModified;
 
+		/// <summary>
+		/// Object to use for guarding access to the OutputItemToCommandLine dictionary
+		/// </summary>
+		object LockObject = new object();
+
 		/// <summary>
 		/// Static cache of all loaded action history files
 		/// </summary>
@@ -117,7 +122,10 @@ namespace UnrealBuildTool
 		{
 			if(File.Location.IsUnderDirectory(BaseDirectory) || Parent == null)
 			{
-				return OutputItemToCommandLine.TryGetValue(File, out CommandLine);
+				lock (LockObject)
+				{
+					return OutputItemToCommandLine.TryGetValue(File, out CommandLine);
+				}
 			}
 			else
 			{
@@ -134,8 +142,11 @@ namespace UnrealBuildTool
 		{
 			if(File.Location.IsUnderDirectory(BaseDirectory) || Parent == null)
 			{
-				OutputItemToCommandLine[File] = CommandLine;
-				bModified = true;
+				lock(LockObject)
+				{
+					OutputItemToCommandLine[File] = CommandLine;
+					bModified = true;
+				}
 			}
 			else
 			{
