UBT: Exclude projects with temporary targets from generated project files by default. Can be re-enabled by passing -IncludeTempTargets to GenerateProjectFiles.bat.

diff --git a/Engine/Source/Programs/UnrealBuildTool/ProjectFiles/ProjectFileGenerator.cs b/Engine/Source/Programs/UnrealBuildTool/ProjectFiles/ProjectFileGenerator.cs
--- a/Engine/Source/Programs/UnrealBuildTool/ProjectFiles/ProjectFileGenerator.cs
+++ b/Engine/Source/Programs/UnrealBuildTool/ProjectFiles/ProjectFileGenerator.cs
@@ -266,6 +266,11 @@ namespace UnrealBuildTool
 		/// </summary>
 		protected bool IncludeEnginePrograms = true;
 
+		/// <summary>
+		/// Whether to include temporary targets generated by UAT to support content only projects with non-default settings
+		/// </summary>
+		bool bIncludeTempTargets = false;
+
 		/// <summary>
 		/// True if we should include .NET Core projects in the generated solution
 		/// </summary>
@@ -474,24 +479,11 @@ namespace UnrealBuildTool
 			}
 			else
 			{
-				ProjectFiles.AddRange(NativeProjects.EnumerateProjectFiles().Where(x => ShouldGenerateProject(x)));
+				ProjectFiles.AddRange(NativeProjects.EnumerateProjectFiles());
 			}
 			return ProjectFiles;
 		}
 
-		/// <summary>
-		/// Determines if we should generate IDE project files for for the given project
-		/// </summary>
-		/// <param name="ProjectFile">The project to check</param>
-		/// <returns>True to generate a project, false otherwise</returns>
-		private static bool ShouldGenerateProject(FileReference ProjectFile)
-		{
-			DirectoryReference ProjectDirectory = ProjectFile.Directory;
-			DirectoryReference SourceFolder = DirectoryReference.Combine(ProjectDirectory, "Source");
-			DirectoryReference IntermediateSourceFolder = DirectoryReference.Combine(ProjectDirectory, "Intermediate", "Source");
-			return DirectoryReference.Exists(SourceFolder) || DirectoryReference.Exists(IntermediateSourceFolder);
-		}
-
 		/// <summary>
 		/// Gets the user's preferred IDE from their editor settings
 		/// </summary>
@@ -666,6 +658,16 @@ namespace UnrealBuildTool
 			// Build the list of games to generate projects for
 			List<FileReference> AllGameProjects = FindGameProjects();
 
+			// Find all of the target files.  This will filter out any modules or targets that don't
+			// belong to platforms we're generating project files for.
+			List<FileReference> AllTargetFiles = DiscoverTargets(AllGameProjects);
+
+			// Sort the targets by name. When we have multiple targets of a given type for a project, we'll use the order to determine which goes in the primary project file (so that client names with a suffix will go into their own project).
+			AllTargetFiles = AllTargetFiles.OrderBy(x => x.FullName, StringComparer.OrdinalIgnoreCase).ToList();
+
+			// Remove any game projects that don't have a target
+			AllGameProjects.RemoveAll(x => !AllTargetFiles.Any(y => y.IsUnderDirectory(x.Directory)));
+
 			// Find all of the module files.  This will filter out any modules or targets that don't belong to platforms
 			// we're generating project files for.
 			List<FileReference> AllModuleFiles = DiscoverModules(AllGameProjects);
@@ -677,7 +679,7 @@ namespace UnrealBuildTool
 			Dictionary<FileReference, ProjectFile> ProgramProjects;
 			{
 				// Setup buildable projects for all targets
-				AddProjectsForAllTargets( PlatformProjectGenerators, AllGameProjects, out EngineProject, out EnterpriseProject, out GameProjects, out ProgramProjects );
+				AddProjectsForAllTargets( PlatformProjectGenerators, AllGameProjects, AllTargetFiles, out EngineProject, out EnterpriseProject, out GameProjects, out ProgramProjects );
 
 				// Add projects for mods
 				AddProjectsForMods(GameProjects, out ModProjects);
@@ -1112,6 +1114,10 @@ namespace UnrealBuildTool
 						case "-VSCODE":
 							bIncludeDotNETCoreProjects = true;
 							break;
+
+						case "-INCLUDETEMPTARGETS":
+							bIncludeTempTargets = true;
+							break;
 					}
 				}
 			}
@@ -1334,7 +1340,7 @@ namespace UnrealBuildTool
 			List<string> UnsupportedPlatformNameStrings = Utils.MakeListOfUnsupportedPlatforms( SupportedPlatforms );
 
 			// Locate all targets (*.Target.cs files)
-			List<FileReference> FoundTargetFiles = RulesCompiler.FindAllRulesSourceFiles( RulesCompiler.RulesFileType.Target, AllGameProjects.Select(x => x.Directory).ToList(), ForeignPlugins: DiscoverExtraPlugins(AllGameProjects), AdditionalSearchPaths:null );
+			List<FileReference> FoundTargetFiles = RulesCompiler.FindAllRulesSourceFiles( RulesCompiler.RulesFileType.Target, AllGameProjects.Select(x => x.Directory).ToList(), ForeignPlugins: DiscoverExtraPlugins(AllGameProjects), AdditionalSearchPaths:null, bIncludeTempTargets: bIncludeTempTargets );
 			foreach( FileReference CurTargetFile in FoundTargetFiles )
 			{
 				string CleanTargetFileName = Utils.CleanDirectorySeparators( CurTargetFile.FullName );
@@ -2026,6 +2032,7 @@ namespace UnrealBuildTool
 		/// </summary>
 		/// <param name="PlatformProjectGenerators">The registered platform project generators</param>
 		/// <param name="AllGames">All game folders</param>
+		/// <param name="AllTargetFiles">All the target files to add</param>
 		/// <param name="EngineProject">The engine project we created</param>
 		/// <param name="EnterpriseProject">The enterprise project we created</param>
 		/// <param name="GameProjects">Map of game folder name to all of the game projects we created</param>
@@ -2033,6 +2040,7 @@ namespace UnrealBuildTool
 		private void AddProjectsForAllTargets( 
 			PlatformProjectGeneratorCollection PlatformProjectGenerators, 
 			List<FileReference> AllGames, 
+			List<FileReference> AllTargetFiles,
 			out ProjectFile EngineProject, 
 			out ProjectFile EnterpriseProject,
 			out List<ProjectFile> GameProjects, 
@@ -2048,13 +2056,6 @@ namespace UnrealBuildTool
 			DirectoryReference EngineSourceProgramsDirectory = DirectoryReference.Combine(UnrealBuildTool.EngineSourceDirectory, "Programs");
 			DirectoryReference EnterpriseSourceProgramsDirectory = DirectoryReference.Combine(UnrealBuildTool.EnterpriseSourceDirectory, "Programs");
 
-			// Find all of the target files.  This will filter out any modules or targets that don't
-			// belong to platforms we're generating project files for.
-			List<FileReference> AllTargetFiles = DiscoverTargets(AllGames);
-
-			// Sort the targets by name. When we have multiple targets of a given type for a project, we'll use the order to determine which goes in the primary project file (so that client names with a suffix will go into their own project).
-			AllTargetFiles = AllTargetFiles.OrderBy(x => x.FullName, StringComparer.OrdinalIgnoreCase).ToList();
-
 			foreach( FileReference TargetFilePath in AllTargetFiles )
 			{
 				string TargetName = TargetFilePath.GetFileNameWithoutAnyExtensions();		// Remove both ".cs" and ".Target"
diff --git a/Engine/Source/Programs/UnrealBuildTool/System/RulesCompiler.cs b/Engine/Source/Programs/UnrealBuildTool/System/RulesCompiler.cs
--- a/Engine/Source/Programs/UnrealBuildTool/System/RulesCompiler.cs
+++ b/Engine/Source/Programs/UnrealBuildTool/System/RulesCompiler.cs
@@ -70,8 +70,9 @@ namespace UnrealBuildTool
 		/// <param name="AdditionalSearchPaths"></param>
 		/// <param name="bIncludeEngine"></param>
 		/// <param name="bIncludeEnterprise"></param>
+		/// <param name="bIncludeTempTargets">Whether to include targets generated by UAT to accomodate content-only projects that need to be compiled to include plugins</param>
 		/// <returns></returns>
-		public static List<FileReference> FindAllRulesSourceFiles(RulesFileType RulesFileType, List<DirectoryReference> GameFolders, List<FileReference> ForeignPlugins, List<DirectoryReference> AdditionalSearchPaths, bool bIncludeEngine = true, bool bIncludeEnterprise = true)
+		public static List<FileReference> FindAllRulesSourceFiles(RulesFileType RulesFileType, List<DirectoryReference> GameFolders, List<FileReference> ForeignPlugins, List<DirectoryReference> AdditionalSearchPaths, bool bIncludeEngine = true, bool bIncludeEnterprise = true, bool bIncludeTempTargets = true)
 		{
 			List<DirectoryReference> Folders = new List<DirectoryReference>();
 
@@ -128,8 +129,11 @@ namespace UnrealBuildTool
 				{
 					DirectoryReference GameSourceFolder = DirectoryReference.Combine(GameFolder, "Source");
 					Folders.Add(GameSourceFolder);
-					DirectoryReference GameIntermediateSourceFolder = DirectoryReference.Combine(GameFolder, "Intermediate", "Source");
-					Folders.Add(GameIntermediateSourceFolder);
+					if(bIncludeTempTargets)
+					{
+						DirectoryReference GameIntermediateSourceFolder = DirectoryReference.Combine(GameFolder, "Intermediate", "Source");
+						Folders.Add(GameIntermediateSourceFolder);
+					}
 				}
 			}
 
