UBT: Markup configurable subobjects in the TargetRules class with an attribute.

diff --git a/Engine/Source/Programs/UnrealBuildTool/Configuration/TargetRules.cs b/Engine/Source/Programs/UnrealBuildTool/Configuration/TargetRules.cs
--- a/Engine/Source/Programs/UnrealBuildTool/Configuration/TargetRules.cs
+++ b/Engine/Source/Programs/UnrealBuildTool/Configuration/TargetRules.cs
@@ -5,6 +5,7 @@ using System.Collections;
 using System.Collections.Generic;
 using System.IO;
 using System.Linq;
+using System.Reflection;
 using System.Text;
 using System.Threading;
 using System.Threading.Tasks;
@@ -99,6 +100,14 @@ namespace UnrealBuildTool
 	{
 	}
 
+	/// <summary>
+	/// Attribute used to mark configurable sub-objects
+	/// </summary>
+	[AttributeUsage(AttributeTargets.Field, AllowMultiple = false)]
+	class ConfigSubObjectAttribute : Attribute
+	{
+	}
+
 	/// <summary>
 	/// TargetRules is a data structure that contains the rules for defining a target (application/executable)
 	/// </summary>
@@ -1155,51 +1164,61 @@ namespace UnrealBuildTool
 		/// <summary>
 		/// Android-specific target settings.
 		/// </summary>
+		[ConfigSubObject]
 		public AndroidTargetRules AndroidPlatform = new AndroidTargetRules();
 
 		/// <summary>
 		/// HTML5-specific target settings.
 		/// </summary>
+		[ConfigSubObject]
 		public HTML5TargetRules HTML5Platform = new HTML5TargetRules();
 
 		/// <summary>
 		/// IOS-specific target settings.
 		/// </summary>
+		[ConfigSubObject]
 		public IOSTargetRules IOSPlatform = new IOSTargetRules();
 
 		/// <summary>
 		/// Lumin-specific target settings.
 		/// </summary>
+		[ConfigSubObject]
 		public LuminTargetRules LuminPlatform = new LuminTargetRules();
 
 		/// <summary>
 		/// Linux-specific target settings.
 		/// </summary>
+		[ConfigSubObject]
 		public LinuxTargetRules LinuxPlatform = new LinuxTargetRules();
 
 		/// <summary>
 		/// Mac-specific target settings.
 		/// </summary>
+		[ConfigSubObject]
 		public MacTargetRules MacPlatform = new MacTargetRules();
 
 		/// <summary>
 		/// PS4-specific target settings.
 		/// </summary>
+		[ConfigSubObject]
 		public PS4TargetRules PS4Platform = new PS4TargetRules();
 
 		/// <summary>
 		/// Switch-specific target settings.
 		/// </summary>
+		[ConfigSubObject]
 		public SwitchTargetRules SwitchPlatform = new SwitchTargetRules();
 
 		/// <summary>
 		/// Windows-specific target settings.
 		/// </summary>
+		[ConfigSubObject]
 		public WindowsTargetRules WindowsPlatform = new WindowsTargetRules();
 
 		/// <summary>
 		/// Xbox One-specific target settings.
 		/// </summary>
+		[ConfigSubObject]
 		public XboxOneTargetRules XboxOnePlatform = new XboxOneTargetRules();
 
 		/// <summary>
@@ -1453,16 +1472,14 @@ namespace UnrealBuildTool
 		internal IEnumerable<object> GetConfigurableObjects()
 		{
 			yield return this;
-			yield return AndroidPlatform;
-			yield return HTML5Platform;
-			yield return IOSPlatform;
-			yield return LuminPlatform;
-			yield return LinuxPlatform;
-			yield return MacPlatform;
-			yield return PS4Platform;
-			yield return SwitchPlatform;
-			yield return WindowsPlatform;
-			yield return XboxOnePlatform;
+
+			foreach(FieldInfo Field in GetType().GetFields(BindingFlags.Public | BindingFlags.Instance))
+			{
+				if(Field.GetCustomAttribute<ConfigSubObjectAttribute>() != null)
+				{
+					yield return Field.GetValue(this);
+				}
+			}
 		}
 
 		/// <summary>
