null-merge of Android toolchain
Fixed issue with code stripping of frameworks not being restricted to local symbols

diff --git a/Engine/Source/Programs/UnrealBuildTool/Platform/IOS/IOSToolChain.cs b/Engine/Source/Programs/UnrealBuildTool/Platform/IOS/IOSToolChain.cs
--- a/Engine/Source/Programs/UnrealBuildTool/Platform/IOS/IOSToolChain.cs
+++ b/Engine/Source/Programs/UnrealBuildTool/Platform/IOS/IOSToolChain.cs
@@ -1349,9 +1349,9 @@ namespace UnrealBuildTool
             List<FileItem> OutputFiles = new List<FileItem>(base.PostBuild(Executable, BinaryLinkEnvironment, Actions));
 
 			if (BinaryLinkEnvironment.bIsBuildingLibrary)
-            {
-                return OutputFiles;
-            }
+			{
+				return OutputFiles;
+			}
 
             // For IOS/tvOS, generate the dSYM file if the config file is set to do so
 			if (ProjectSettings.bGeneratedSYMFile == true || ProjectSettings.bGeneratedSYMBundle == true || BinaryLinkEnvironment.bUsePDBFiles == true)
@@ -1368,10 +1368,13 @@ namespace UnrealBuildTool
 			{
 				FileItem StripCompleteFile = FileItem.GetItemByFileReference(FileReference.Combine(BinaryLinkEnvironment.IntermediateDirectory, Executable.Location.GetFileName() + ".stripped"));
 
+				// If building a framework we can only strip local symbols, need to leave global in place
+				string StripArguments = BinaryLinkEnvironment.bIsBuildingDLL ? "-x" : "";
+
 				Action StripAction = new Action(ActionType.CreateAppBundle);
 				StripAction.WorkingDirectory = GetMacDevSrcRoot();
 				StripAction.CommandPath = BuildHostPlatform.Current.Shell;
-				StripAction.CommandArguments = String.Format("-c '\"{0}strip\" \"{1}\" && touch \"{2}\"'", Settings.Value.ToolchainDir, Executable.Location, StripCompleteFile);
+				StripAction.CommandArguments = String.Format("-c '\"{0}strip\" {1} \"{2}\" && touch \"{3}\"'", Settings.Value.ToolchainDir, StripArguments, Executable.Location, StripCompleteFile);
 				StripAction.PrerequisiteItems.Add(Executable);
 				StripAction.PrerequisiteItems.AddRange(OutputFiles);
 				StripAction.ProducedItems.Add(StripCompleteFile);
