Escape all paths when generating VS project filter files

diff --git a/Engine/Source/Programs/UnrealBuildTool/ProjectFiles/VisualStudio/VCProject.cs b/Engine/Source/Programs/UnrealBuildTool/ProjectFiles/VisualStudio/VCProject.cs
--- a/Engine/Source/Programs/UnrealBuildTool/ProjectFiles/VisualStudio/VCProject.cs
+++ b/Engine/Source/Programs/UnrealBuildTool/ProjectFiles/VisualStudio/VCProject.cs
@@ -378,7 +378,7 @@ namespace UnrealBuildTool
 				return new MSBuildProjectContext(StubProjectConfigurationName, StubProjectPlatformName);
 			}
 
-			// Have to match every solution configuration combination to a project configuration (or use the invalid one) 
+			// Have to match every solution configuration combination to a project configuration (or use the invalid one)
 			string ProjectConfigurationName = "Invalid";
 
 			// Get the default platform. If there were not valid platforms for this project, just use one that will always be available in VS.
@@ -944,7 +944,7 @@ namespace UnrealBuildTool
 					if (!String.IsNullOrWhiteSpace(AliasedFile.ProjectPath))
 					{
 						VCFiltersFileContent.AppendLine("    <{0} Include=\"{1}\">", VCFileType, EscapeFileName(AliasedFile.FileSystemPath));
-						VCFiltersFileContent.AppendLine("      <Filter>{0}</Filter>", Utils.CleanDirectorySeparators(AliasedFile.ProjectPath));
+						VCFiltersFileContent.AppendLine("      <Filter>{0}</Filter>", Utils.CleanDirectorySeparators(EscapeFileName(AliasedFile.ProjectPath)));
 						VCFiltersFileContent.AppendLine("    </{0}>", VCFileType);
 
 						FiltersFileIsNeeded = true;
@@ -1148,7 +1148,7 @@ namespace UnrealBuildTool
 						//       matches identically with the pre-existing file
 						string FilterGUID = Guid.NewGuid().ToString("B").ToUpperInvariant();
 
-						VCFiltersFileContent.AppendLine("    <Filter Include=\"{0}\">", LeadingDirectory);
+						VCFiltersFileContent.AppendLine("    <Filter Include=\"{0}\">", EscapeFileName(LeadingDirectory));
 						VCFiltersFileContent.AppendLine("      <UniqueIdentifier>{0}</UniqueIdentifier>", FilterGUID);
 						VCFiltersFileContent.AppendLine("    </Filter>");
 
@@ -1517,7 +1517,7 @@ namespace UnrealBuildTool
 					throw new BuildException("Unexpected root element '{0}' in project file", Document.DocumentElement.Name);
 				}
 
-				// Parse all the configurations and platforms 
+				// Parse all the configurations and platforms
 				// Parse the basic structure of the document, updating properties and recursing into other referenced projects as we go
 				foreach (XmlElement Element in Document.DocumentElement.ChildNodes.OfType<XmlElement>())
 				{
@@ -1624,7 +1624,7 @@ namespace UnrealBuildTool
 			{
 				bBuildByDefault = true;
 			}
-			
+
 			// Create the context
 			return new MSBuildProjectContext(ProjectConfigurationName, ProjectPlatformName){ bBuildByDefault = bBuildByDefault };
 		}
