Add an error for using a unique build environment with an installed build.

diff --git a/Engine/Source/Programs/UnrealBuildTool/System/RulesAssembly.cs b/Engine/Source/Programs/UnrealBuildTool/System/RulesAssembly.cs
--- a/Engine/Source/Programs/UnrealBuildTool/System/RulesAssembly.cs
+++ b/Engine/Source/Programs/UnrealBuildTool/System/RulesAssembly.cs
@@ -477,6 +477,13 @@ namespace UnrealBuildTool
 					Rules.BuildEnvironment = TargetBuildEnvironment.Unique;
 				}
 			}
+			else
+			{
+				if(Rules.BuildEnvironment == TargetBuildEnvironment.Unique && UnrealBuildTool.IsEngineInstalled())
+				{
+					throw new BuildException("Targets with a unique build environment cannot be built an installed engine.");
+				}
+			}
 
 			// Automatically include CoreUObject
 			if (Rules.bCompileAgainstEngine)
