UBT: Disable the "bAdaptiveUnityDisablesPCH" build flag by default. This significantly improves iteration times on individual files, at the risk of introducing missing headers. Will validate this in CIS instead.

diff --git a/Engine/Source/Programs/UnrealBuildTool/Configuration/TargetRules.cs b/Engine/Source/Programs/UnrealBuildTool/Configuration/TargetRules.cs
--- a/Engine/Source/Programs/UnrealBuildTool/Configuration/TargetRules.cs
+++ b/Engine/Source/Programs/UnrealBuildTool/Configuration/TargetRules.cs
@@ -677,7 +677,7 @@ namespace UnrealBuildTool
 		/// Disables force-included PCHs for files that are in the adaptive non-unity working set.
 		/// </summary>
 		[XmlConfigFile(Category = "BuildConfiguration")]
-		public bool bAdaptiveUnityDisablesPCH = true;
+		public bool bAdaptiveUnityDisablesPCH = false;
 
 		/// <summary>
 		/// Backing storage for bAdaptiveUnityDisablesProjectPCH.
