UE-52346 TM-Gameplay: Load Level Volume Fails to Read TestSlot.sav file

#iOS
#tvOS

Created a pList key named ICloudContainerIdentifier to be used at run-time when intializing the CloudKit services

diff --git a/Engine/Plugins/Online/IOS/OnlineSubsystemIOS/Source/Private/OnlineUserCloudInterfaceIOS.cpp b/Engine/Plugins/Online/IOS/OnlineSubsystemIOS/Source/Private/OnlineUserCloudInterfaceIOS.cpp
--- a/Engine/Plugins/Online/IOS/OnlineSubsystemIOS/Source/Private/OnlineUserCloudInterfaceIOS.cpp
+++ b/Engine/Plugins/Online/IOS/OnlineSubsystemIOS/Source/Private/OnlineUserCloudInterfaceIOS.cpp
@@ -24,7 +24,16 @@
 			[[NSNotificationCenter defaultCenter] addObserver:self selector : @selector(iCloudAccountAvailabilityChanged:) name: NSUbiquityIdentityDidChangeNotification object : nil];
 		}
 
-		CloudContainer = [CKContainer defaultContainer];
+		NSString *ICloudContainerIdentifier = [[NSBundle mainBundle].infoDictionary objectForKey : @"ICloudContainerIdentifier"];
+		if (ICloudContainerIdentifier != nil)
+		{
+			NSLog(@"Using a custom CloudKit container: %@", ICloudContainerIdentifier);
+			CloudContainer = [CKContainer containerWithIdentifier:ICloudContainerIdentifier];
+		}
+		else
+		{
+			CloudContainer = [CKContainer defaultContainer];
+		}
 		SharedDatabase = [CloudContainer publicCloudDatabase];
 		UserDatabase = [CloudContainer privateCloudDatabase];
 	}
diff --git a/Engine/Source/Programs/UnrealBuildTool/Platform/IOS/IOSToolChain.cs b/Engine/Source/Programs/UnrealBuildTool/Platform/IOS/IOSToolChain.cs
--- a/Engine/Source/Programs/UnrealBuildTool/Platform/IOS/IOSToolChain.cs
+++ b/Engine/Source/Programs/UnrealBuildTool/Platform/IOS/IOSToolChain.cs
@@ -9,6 +9,7 @@ using System.IO;
 using System.Diagnostics;
 using System.Security.AccessControl;
 using System.Xml;
+using System.Xml.Linq;
 using System.Text;
 using Ionic.Zip;
 using Ionic.Zlib;
@@ -1833,18 +1834,53 @@ namespace UnrealBuildTool
 					// get some info from the mobileprovisioning file
 					// the iCloud identifier and the bundle id may differ
 					string iCloudContainerIdentifiersXML = "";
+					string iCloudContainerIdentifier = "";
 					string UbiquityContainerIdentifiersXML = "";
 					if (File.Exists(MobileProvisionFile.FullName))
 					{
 						MobileProvisionContents MobileProvisionContent = MobileProvisionContents.Read(MobileProvisionFile);
 
-						iCloudContainerIdentifiersXML = MobileProvisionContent.GetNodeValueByName("com.apple.developer.icloud-container-identifiers");
-						UbiquityContainerIdentifiersXML = MobileProvisionContent.GetNodeValueByName("com.apple.developer.ubiquity-container-identifiers");
+						iCloudContainerIdentifier = MobileProvisionContent.GetNodeValueByName("com.apple.developer.icloud-container-identifiers");
+						iCloudContainerIdentifiersXML = MobileProvisionContent.GetNodeXMLValueByName("com.apple.developer.icloud-container-identifiers");
+						UbiquityContainerIdentifiersXML = MobileProvisionContent.GetNodeXMLValueByName("com.apple.developer.ubiquity-container-identifiers");
 					}
 					// create the entitlements file
 					string IntermediateDir = (((Target.ProjectFile != null) ? Target.ProjectFile.Directory.ToString() : 
 						UnrealBuildTool.EngineDirectory.ToString())) + "/Intermediate/" + (Target.Platform == UnrealTargetPlatform.IOS ? "IOS" : "TVOS");
 					WriteEntitlementsFile(Path.Combine(IntermediateDir, AppName + ".entitlements"), Target.ProjectFile, Target.bForDistribution, iCloudContainerIdentifiersXML, UbiquityContainerIdentifiersXML);
+
+					// create a pList key named ICloudContainerIdentifier
+					// to be used at run-time when intializing the CloudKit services
+					if (iCloudContainerIdentifier != "")
+					{
+						string PListFile = IntermediateDir + "/" + AppName + "-Info.plist";
+						if (File.Exists(PListFile))
+						{
+							string OldPListData = File.ReadAllText(PListFile);
+							XDocument XDoc;
+							try
+							{
+								XDoc = XDocument.Parse(OldPListData);
+								if (XDoc.DocumentType != null)
+									XDoc.DocumentType.InternalSubset = null;
+
+								XElement dictElement = XDoc.Root.Element("dict");
+								if (dictElement != null)
+								{
+									dictElement.Add(new XElement("key", "ICloudContainerIdentifier"));
+									dictElement.Add(new XElement("string", iCloudContainerIdentifier));
+
+									XDoc.Save(PListFile);
+								}
+							}
+							catch (Exception e)
+							{
+								throw new BuildException("plist is invalid {0}\n{1}", e, OldPListData);
+							}
+
+						}
+					}
+
 				}
 
 				Log.TraceInformation("Executing {0}", SignProjectScript);
diff --git a/Engine/Source/Programs/UnrealBuildTool/Platform/IOS/MobileProvisionContents.cs b/Engine/Source/Programs/UnrealBuildTool/Platform/IOS/MobileProvisionContents.cs
--- a/Engine/Source/Programs/UnrealBuildTool/Platform/IOS/MobileProvisionContents.cs
+++ b/Engine/Source/Programs/UnrealBuildTool/Platform/IOS/MobileProvisionContents.cs
@@ -146,6 +146,25 @@ namespace UnrealBuildTool
 		}
 
 		// return the outerXML of the node's value
+		public string GetNodeXMLValueByName(string InValue)
+		{
+			XmlNodeList elemList = this.Document.GetElementsByTagName("key");
+			for (int i = 0; i < elemList.Count; i++)
+			{
+				if (elemList[i].InnerXml.Equals(InValue))
+				{
+					XmlNode valueNode = elemList[i].NextSibling;
+
+					if (valueNode != null)
+					{
+						return valueNode.OuterXml;
+					}
+				}
+			}
+			return "";
+		}
+
+		// return the innerXML of the node's value
 		public string GetNodeValueByName(string InValue)
 		{
 			XmlNodeList elemList = this.Document.GetElementsByTagName("key");
@@ -153,9 +172,21 @@ namespace UnrealBuildTool
 			{
 				if (elemList[i].InnerXml.Equals(InValue))
 				{
-					if (elemList[i].NextSibling != null)
+					XmlNode valueNode = elemList[i].NextSibling;
+					if (valueNode != null)
 					{
-						return elemList[i].NextSibling.OuterXml;
+						if (valueNode.Name.Equals("array"))
+						{
+							XmlNode firstChildNode = valueNode.FirstChild;
+							if (firstChildNode != null)
+							{
+								return firstChildNode.InnerXml;
+							}
+						}
+						else
+						{
+							return valueNode.InnerXml;
+						}
 					}
 				}
 			}
