Fix exception when hot-reloading a module that contains an additional bundle resource. List of hot-reload build products for each module now only includes import library and DLL.

diff --git a/Engine/Source/Programs/UnrealBuildTool/Configuration/UEBuildBinary.cs b/Engine/Source/Programs/UnrealBuildTool/Configuration/UEBuildBinary.cs
--- a/Engine/Source/Programs/UnrealBuildTool/Configuration/UEBuildBinary.cs
+++ b/Engine/Source/Programs/UnrealBuildTool/Configuration/UEBuildBinary.cs
@@ -201,6 +201,12 @@ namespace UnrealBuildTool
 			FileItem[] Executables = ToolChain.LinkAllFiles(BinaryLinkEnvironment, false, Makefile.Actions);
 			OutputFiles.AddRange(Executables);
 
+			// Save all the output items for this binary. This is used for hot-reload, and excludes any items added in PostBuild (such as additional files copied into the app).
+			if(Target.LinkType == TargetLinkType.Modular)
+			{
+				Makefile.ModuleNameToOutputItems[PrimaryModule.Name] = OutputFiles.ToArray();
+			}
+
 			// Produce additional console app if requested
 			if (bBuildAdditionalConsoleApp)
 			{
diff --git a/Engine/Source/Programs/UnrealBuildTool/Configuration/UEBuildTarget.cs b/Engine/Source/Programs/UnrealBuildTool/Configuration/UEBuildTarget.cs
--- a/Engine/Source/Programs/UnrealBuildTool/Configuration/UEBuildTarget.cs
+++ b/Engine/Source/Programs/UnrealBuildTool/Configuration/UEBuildTarget.cs
@@ -1314,10 +1314,6 @@ namespace UnrealBuildTool
 				foreach (UEBuildBinary Binary in Binaries)
 				{
 					List<FileItem> BinaryOutputItems = Binary.Build(Rules, TargetToolChain, GlobalCompileEnvironment, GlobalLinkEnvironment, WorkingSet, ExeDir, Makefile);
-					if(!bCompileMonolithic)
-					{
-						Makefile.ModuleNameToOutputItems[Binary.PrimaryModule.Name] = BinaryOutputItems.ToArray();
-					}
 					Makefile.OutputItems.AddRange(BinaryOutputItems);
 				}
 			}
