Add a warning when a plugin reference overrides the plugin's list of supported target platforms. Silently failing at build time will cause errors at runtime.

diff --git a/Engine/Source/Programs/UnrealBuildTool/Configuration/UEBuildTarget.cs b/Engine/Source/Programs/UnrealBuildTool/Configuration/UEBuildTarget.cs
--- a/Engine/Source/Programs/UnrealBuildTool/Configuration/UEBuildTarget.cs
+++ b/Engine/Source/Programs/UnrealBuildTool/Configuration/UEBuildTarget.cs
@@ -2614,7 +2614,7 @@ namespace UnrealBuildTool
 			{
 				foreach(PluginInfo Plugin in NameToInfo.Values)
 				{
-					if(Plugin.EnabledByDefault && !ReferencedNames.Contains(Plugin.Name))
+					if(Plugin.EnabledByDefault && !ReferencedNames.Contains(Plugin.Name) && Plugin.Descriptor.SupportsTargetPlatform(Platform))
 					{
 						ReferencedNames.Add(Plugin.Name);
 
@@ -2723,7 +2723,7 @@ namespace UnrealBuildTool
 				// Disable any plugin which does not support the target platform. The editor should update such references in the .uproject file on load.
 				if (!Rules.bIncludePluginsForTargetPlatforms && !Info.Descriptor.SupportsTargetPlatform(Platform))
 				{
-					Log.TraceLog("Ignoring plugin '{0}' (referenced via {1}) due to target platform not supported by descriptor.", Reference.Name, ReferenceChain);
+					Log.TraceWarning("Plugin '{0}' is referenced via {1} without a matching 'SupportedTargetPlatforms' field. This may cause problems in packaged builds, because the .uplugin file will not be staged.", Reference.Name, ReferenceChain);
 					return null;
 				}
 
