UBT: Make mismatched supported target platforms an error.

diff --git a/Engine/Source/Programs/UnrealBuildTool/Configuration/UEBuildTarget.cs b/Engine/Source/Programs/UnrealBuildTool/Configuration/UEBuildTarget.cs
--- a/Engine/Source/Programs/UnrealBuildTool/Configuration/UEBuildTarget.cs
+++ b/Engine/Source/Programs/UnrealBuildTool/Configuration/UEBuildTarget.cs
@@ -2723,8 +2723,7 @@ namespace UnrealBuildTool
 				// Disable any plugin which does not support the target platform. The editor should update such references in the .uproject file on load.
 				if (!Rules.bIncludePluginsForTargetPlatforms && !Info.Descriptor.SupportsTargetPlatform(Platform))
 				{
-					Log.TraceWarning("Plugin '{0}' is referenced via {1} without a matching 'SupportedTargetPlatforms' field. This may cause problems in packaged builds, because the .uplugin file will not be staged.", Reference.Name, ReferenceChain);
-					return null;
+					throw new BuildException("{0} is referenced via {1} with a mismatched 'SupportedTargetPlatforms' field. This will cause problems in packaged builds, because the .uplugin file will not be staged. Launch the editor to update references from your project file, or update references from other plugins manually.", Info.File.GetFileName(), ReferenceChain);
 				}
 
 				// Disable any plugin that requires the build platform
