Fix missing plugin error at startup.

diff --git a/Engine/Source/Programs/UnrealBuildTool/Configuration/UEBuildTarget.cs b/Engine/Source/Programs/UnrealBuildTool/Configuration/UEBuildTarget.cs
--- a/Engine/Source/Programs/UnrealBuildTool/Configuration/UEBuildTarget.cs
+++ b/Engine/Source/Programs/UnrealBuildTool/Configuration/UEBuildTarget.cs
@@ -2614,11 +2614,12 @@ namespace UnrealBuildTool
 			{
 				foreach(PluginInfo Plugin in NameToInfo.Values)
 				{
-					if(Plugin.EnabledByDefault && !ReferencedNames.Contains(Plugin.Name) && Plugin.Descriptor.SupportsTargetPlatform(Platform))
+					if(Plugin.EnabledByDefault && !ReferencedNames.Contains(Plugin.Name))
 					{
 						ReferencedNames.Add(Plugin.Name);
 
 						PluginReferenceDescriptor PluginReference = new PluginReferenceDescriptor(Plugin.Name, null, true);
+						PluginReference.SupportedTargetPlatforms = Plugin.Descriptor.SupportedTargetPlatforms;
 						PluginReference.bOptional = true;
 
 						AddPlugin(PluginReference, "default plugins", ExcludeFolders, NameToInstance, NameToInfo);
