Remove LinkEnvironment.InputLibraries. This is always empty.

diff --git a/Engine/Source/Programs/UnrealBuildTool/Platform/Android/AndroidToolChain.cs b/Engine/Source/Programs/UnrealBuildTool/Platform/Android/AndroidToolChain.cs
--- a/Engine/Source/Programs/UnrealBuildTool/Platform/Android/AndroidToolChain.cs
+++ b/Engine/Source/Programs/UnrealBuildTool/Platform/Android/AndroidToolChain.cs
@@ -1809,10 +1809,6 @@ namespace UnrealBuildTool
 			{
 				ObjectFileDirectories.Add(InputFile.Location.Directory);
 			}
-			foreach (FileItem InputLibrary in LinkEnvironment.InputLibraries)
-			{
-				ObjectFileDirectories.Add(InputLibrary.Location.Directory);
-			}
 			foreach (string AdditionalLibrary in LinkEnvironment.AdditionalLibraries.Where(x => Path.IsPathRooted(x)))
 			{
 				ObjectFileDirectories.Add(new FileReference(AdditionalLibrary).Directory);
diff --git a/Engine/Source/Programs/UnrealBuildTool/Platform/Windows/VCToolChain.cs b/Engine/Source/Programs/UnrealBuildTool/Platform/Windows/VCToolChain.cs
--- a/Engine/Source/Programs/UnrealBuildTool/Platform/Windows/VCToolChain.cs
+++ b/Engine/Source/Programs/UnrealBuildTool/Platform/Windows/VCToolChain.cs
@@ -1420,16 +1420,6 @@ namespace UnrealBuildTool
 				PrerequisiteItems.Add(InputFile);
 			}
 
-			if (!bBuildImportLibraryOnly)
-			{
-				// Add input libraries as prerequisites, too!
-				foreach (FileItem InputLibrary in LinkEnvironment.InputLibraries)
-				{
-					InputFileNames.Add(string.Format("\"{0}\"", InputLibrary.AbsolutePath));
-					PrerequisiteItems.Add(InputLibrary);
-				}
-			}
-
 			if (!bIsBuildingLibraryOrImportLibrary)
 			{
 				foreach (string AdditionalLibrary in LinkEnvironment.AdditionalLibraries)
@@ -1578,10 +1568,6 @@ namespace UnrealBuildTool
 			{
 				ObjectFileDirectories.Add(InputFile.Location.Directory);
 			}
-			foreach(FileItem InputLibrary in LinkEnvironment.InputLibraries)
-			{
-				ObjectFileDirectories.Add(InputLibrary.Location.Directory);
-			}
 			foreach(string AdditionalLibrary in LinkEnvironment.AdditionalLibraries)
 			{
 				// Need to handle import libraries that are about to be built (but may not exist yet), third party libraries with relative paths in the UE4 tree, and system libraries in the system path
diff --git a/Engine/Source/Programs/UnrealBuildTool/System/LinkEnvironment.cs b/Engine/Source/Programs/UnrealBuildTool/System/LinkEnvironment.cs
--- a/Engine/Source/Programs/UnrealBuildTool/System/LinkEnvironment.cs
+++ b/Engine/Source/Programs/UnrealBuildTool/System/LinkEnvironment.cs
@@ -264,11 +264,6 @@ namespace UnrealBuildTool
 		/// </summary>
 		public List<FileItem> InputFiles = new List<FileItem>();
 
-		/// <summary>
-		/// A list of dependent static or import libraries that need to be linked.
-		/// </summary>
-		public List<FileItem> InputLibraries = new List<FileItem>();
-
 		/// <summary>
 		/// The default resource file to link in to every binary if a custom one is not provided
 		/// </summary>
@@ -356,7 +351,6 @@ namespace UnrealBuildTool
 			BundleVersion = Other.BundleVersion;
 			InstallName = Other.InstallName;
 			InputFiles.AddRange(Other.InputFiles);
-			InputLibraries.AddRange(Other.InputLibraries);
 			DefaultResourceFiles.AddRange(Other.DefaultResourceFiles);
 			CommonResourceFiles.AddRange(Other.CommonResourceFiles);
 			IncludeFunctions.AddRange(Other.IncludeFunctions);
