Added some logging around Android symbol generaton

diff --git a/Engine/Source/Programs/UnrealBuildTool/Platform/Android/UEDeployAndroid.cs b/Engine/Source/Programs/UnrealBuildTool/Platform/Android/UEDeployAndroid.cs
--- a/Engine/Source/Programs/UnrealBuildTool/Platform/Android/UEDeployAndroid.cs
+++ b/Engine/Source/Programs/UnrealBuildTool/Platform/Android/UEDeployAndroid.cs
@@ -860,8 +860,9 @@ namespace UnrealBuildTool
 				bool bUseChangeListAsStoreVersion = false;
 				Ini.GetBool("/Script/AndroidRuntimeSettings.AndroidRuntimeSettings", "bUseChangeListAsStoreVersion", out bUseChangeListAsStoreVersion);
 
+				bool IsBuildMachine = Environment.GetEnvironmentVariable("IsBuildMachine") == "1";
 				// override store version with changelist if enabled and is build machine
-				if (bUseChangeListAsStoreVersion && Environment.GetEnvironmentVariable("IsBuildMachine") == "1")
+				if (bUseChangeListAsStoreVersion && IsBuildMachine)
 				{
 					// make sure changelist is cached
 					string EngineVersion = ReadEngineVersion();
@@ -876,6 +877,8 @@ namespace UnrealBuildTool
 					}
 				}
 
+				Log.TraceInformation("GotStoreVersion found v{0}. (bUseChangeListAsStoreVersion={1} IsBuildMachine={2} EngineChangeList={3)}", StoreVersion, bUseChangeListAsStoreVersion, IsBuildMachine, EngineChangelist);
+
 				CachedStoreVersion = StoreVersion;
 			}
 
@@ -3702,6 +3705,10 @@ namespace UnrealBuildTool
 				if (bSkipGradleBuild)
 				{
 					FinalSOName = OutputPath;
+					if (!File.Exists(FinalSOName))
+					{
+						Log.TraceWarning("Did not find compiled .so [{0}]", FinalSOName);
+					}
 				}
 				else
 				{
@@ -3995,7 +4002,7 @@ namespace UnrealBuildTool
 					bool bSaveSymbols = false;
 					Ini.GetBool("/Script/AndroidRuntimeSettings.AndroidRuntimeSettings", "bBuildWithHiddenSymbolVisibility", out bBuildWithHiddenSymbolVisibility);
 					Ini.GetBool("/Script/AndroidRuntimeSettings.AndroidRuntimeSettings", "bSaveSymbols", out bSaveSymbols);
-bSaveSymbols = true;
+					bSaveSymbols = true;
 					if (bSaveSymbols || (Configuration == UnrealTargetConfiguration.Shipping && bBuildWithHiddenSymbolVisibility))
 					{
 						// Copy .so with symbols to 
@@ -4003,6 +4010,7 @@ bSaveSymbols = true;
 						string SymbolSODirectory = Path.Combine(DestApkDirectory, ProjectName + "_Symbols_v" + StoreVersion + "/" + ProjectName + Arch + GPUArchitecture);
 						string SymbolifiedSOPath = Path.Combine(SymbolSODirectory, Path.GetFileName(FinalSOName));
 						MakeDirectoryIfRequired(SymbolifiedSOPath);
+						Log.TraceInformation("Writing symbols to {0}", SymbolifiedSOPath);
 
 						File.Copy(FinalSOName, SymbolifiedSOPath, true);
 					}
