Fix logging of store version

diff --git a/Engine/Source/Programs/UnrealBuildTool/Platform/Android/UEDeployAndroid.cs b/Engine/Source/Programs/UnrealBuildTool/Platform/Android/UEDeployAndroid.cs
--- a/Engine/Source/Programs/UnrealBuildTool/Platform/Android/UEDeployAndroid.cs
+++ b/Engine/Source/Programs/UnrealBuildTool/Platform/Android/UEDeployAndroid.cs
@@ -877,7 +877,7 @@ namespace UnrealBuildTool
 					}
 				}
 
-				Log.TraceInformation("GotStoreVersion found v{0}. (bUseChangeListAsStoreVersion={1} IsBuildMachine={2} EngineChangeList={3)}", StoreVersion, bUseChangeListAsStoreVersion, IsBuildMachine, EngineChangelist);
+				Log.TraceInformation("GotStoreVersion found v{0}. (bUseChangeListAsStoreVersion={1} IsBuildMachine={2} EngineChangeList={3})", StoreVersion, bUseChangeListAsStoreVersion, IsBuildMachine, EngineChangelist);
 
 				CachedStoreVersion = StoreVersion;
 			}
