Add new config values to go alongside "PlatformRequiresDataCrypto" which can specifically disable signing or encryption. Allows selective disabling of signing OR encryption on some platforms, rather than the old config value which would disable both.
Disable signing on platforms that we wanted encryption enabled for by default (mobile mainly)
Add a warning if pak signing or encryption is enabled but there is no key, then disable the offending functionality for that build

diff --git a/Engine/Config/Android/AndroidEngine.ini b/Engine/Config/Android/AndroidEngine.ini
--- a/Engine/Config/Android/AndroidEngine.ini
+++ b/Engine/Config/Android/AndroidEngine.ini
@@ -4,6 +4,7 @@ DefaultBloomKernelTextureName=/Engine/EngineResources/DefaultTexture.DefaultText
 
 [PlatformCrypto]
 PlatformRequiresDataCrypto=True
+PakSigningRequired=False
 
 [/Script/Engine.GarbageCollectionSettings]
 gc.MaxObjectsInGame=131072
diff --git a/Engine/Config/BaseEngine.ini b/Engine/Config/BaseEngine.ini
--- a/Engine/Config/BaseEngine.ini
+++ b/Engine/Config/BaseEngine.ini
@@ -2222,6 +2222,8 @@ DefaultCompletionMode=RestoreState
 
 [PlatformCrypto]
 PlatformRequiresDataCrypto=False
+PakSigningRequired=True
+PakEncryptionRequired=True
 
 [/Script/AppleARKit.AppleARKitSettings]
 bEnableLiveLinkForFaceTracking=true
diff --git a/Engine/Config/IOS/IOSEngine.ini b/Engine/Config/IOS/IOSEngine.ini
--- a/Engine/Config/IOS/IOSEngine.ini
+++ b/Engine/Config/IOS/IOSEngine.ini
@@ -7,6 +7,7 @@ gc.MaxObjectsInGame=131072
 
 [PlatformCrypto]
 PlatformRequiresDataCrypto=True
+PakSigningRequired=False
 
 [Audio]
 AudioDeviceModuleName=IOSAudio
diff --git a/Engine/Source/Programs/UnrealBuildTool/System/EncryptionAndSigning.cs b/Engine/Source/Programs/UnrealBuildTool/System/EncryptionAndSigning.cs
--- a/Engine/Source/Programs/UnrealBuildTool/System/EncryptionAndSigning.cs
+++ b/Engine/Source/Programs/UnrealBuildTool/System/EncryptionAndSigning.cs
@@ -147,6 +147,16 @@ namespace UnrealBuildTool
 			/// </summary>
 			public bool bDataCryptoRequired = false;
 
+			/// <summary>
+			/// Config setting to enable pak signing
+			/// </summary>
+			public bool PakEncryptionRequired = true;
+
+			/// <summary>
+			/// Config setting to enable pak encryption
+			/// </summary>
+			public bool PakSigningRequired = true;
+			
 			/// <summary>
 			/// A set of named encryption keys that can be used to encrypt different sets of data with a different key that is delivered dynamically (i.e. not embedded within the game executable)
 			/// </summary>
@@ -217,10 +227,12 @@ namespace UnrealBuildTool
 		public static CryptoSettings ParseCryptoSettings(DirectoryReference InProjectDirectory, UnrealTargetPlatform InTargetPlatform)
 		{
 			CryptoSettings Settings = new CryptoSettings();
-
+			
 			ConfigHierarchy Ini = ConfigCache.ReadHierarchy(ConfigHierarchyType.Engine, InProjectDirectory, InTargetPlatform);
 			Ini.GetBool("PlatformCrypto", "PlatformRequiresDataCrypto", out Settings.bDataCryptoRequired);
-			
+			Ini.GetBool("PlatformCrypto", "PakSigningRequired", out Settings.PakSigningRequired);
+			Ini.GetBool("PlatformCrypto", "PakEncryptionRequired", out Settings.PakEncryptionRequired);
+
 			{
 				// Start by parsing the legacy encryption.ini settings
 				Ini = ConfigCache.ReadHierarchy(ConfigHierarchyType.Encryption, InProjectDirectory, InTargetPlatform);
@@ -406,6 +418,27 @@ namespace UnrealBuildTool
 				NewSettings.SecondaryEncryptionKeys = Settings.SecondaryEncryptionKeys;
 				Settings = NewSettings;
 			}
+			else
+			{
+				if (!Settings.PakSigningRequired)
+				{
+					Settings.bEnablePakSigning = false;
+					Settings.SigningKey = null;
+				}
+
+				if (!Settings.PakEncryptionRequired)
+				{
+					Settings.bEnablePakFullAssetEncryption = false;
+					Settings.bEnablePakIndexEncryption = false;
+					Settings.bEnablePakIniEncryption = false;
+					Settings.EncryptionKey = null;
+					Settings.SigningKey = null;
+				}
+			}
+
+			{
+				Log.TraceWarningOnce("Project signing keys found in '{0}' are of the old insecure short format. Please regenerate them using the project crypto settings panel in the editor!", InProjectDirectory);
+			}
 
 			// Validate the settings we have read
 			if (Settings.bDataCryptoRequired && Settings.bEnablePakSigning && (Settings.SigningKey == null || !Settings.SigningKey.IsValid()))
