UE-71596 iOS packaging error with CloudKit enabled - Provisioning profile <...> doesn't match the entitlements file's values

#4.22
#iOS
#tvOS

Write the entitlements in PostBuldSync.
Target.bCreateStubIPA is false when building on mac, true on windows (using the -CreateStub argument)

diff --git a/Engine/Source/Programs/UnrealBuildTool/Platform/IOS/IOSToolChain.cs b/Engine/Source/Programs/UnrealBuildTool/Platform/IOS/IOSToolChain.cs
--- a/Engine/Source/Programs/UnrealBuildTool/Platform/IOS/IOSToolChain.cs
+++ b/Engine/Source/Programs/UnrealBuildTool/Platform/IOS/IOSToolChain.cs
@@ -1570,6 +1570,76 @@ namespace UnrealBuildTool
 			return FileReference.Combine(Executable.Directory, "Payload", TargetName + ".app", TargetName);
 		}
 
+
+		private static void WriteEntitlements(IOSPostBuildSyncTarget Target)
+		{
+			// get some info from the mobileprovisioning file
+			// the iCloud identifier and the bundle id may differ
+			string AppName = Target.TargetName;
+			FileReference MobileProvisionFile;
+			IOSProjectSettings ProjectSettings = ((IOSPlatform)UEBuildPlatform.GetBuildPlatform(Target.Platform)).ReadProjectSettings(Target.ProjectFile);
+
+			if (Target.ImportProvision == null)
+			{
+				IOSProvisioningData ProvisioningData = ((IOSPlatform)UEBuildPlatform.GetBuildPlatform(Target.Platform)).ReadProvisioningData(ProjectSettings, Target.bForDistribution);
+				MobileProvisionFile = ProvisioningData.MobileProvisionFile;
+			}
+			else
+			{
+				MobileProvisionFile = new FileReference(Target.ImportProvision);
+			}
+
+			string iCloudContainerIdentifiersXML = "";
+			string iCloudContainerIdentifier = "";
+			string UbiquityContainerIdentifiersXML = "";
+			if (MobileProvisionFile!= null && File.Exists(MobileProvisionFile.FullName))
+			{
+				MobileProvisionContents MobileProvisionContent = MobileProvisionContents.Read(MobileProvisionFile);
+
+				iCloudContainerIdentifier = MobileProvisionContent.GetNodeValueByName("com.apple.developer.icloud-container-identifiers");
+				iCloudContainerIdentifiersXML = MobileProvisionContent.GetNodeXMLValueByName("com.apple.developer.icloud-container-identifiers");
+				UbiquityContainerIdentifiersXML = MobileProvisionContent.GetNodeXMLValueByName("com.apple.developer.ubiquity-container-identifiers");
+			}
+			// create the entitlements file
+			string IntermediateDir = (((Target.ProjectFile != null) ? Target.ProjectFile.Directory.ToString() :
+				UnrealBuildTool.EngineDirectory.ToString())) + "/Intermediate/" + (Target.Platform == UnrealTargetPlatform.IOS ? "IOS" : "TVOS");
+			WriteEntitlementsFile(Path.Combine(IntermediateDir, AppName + ".entitlements"), Target.ProjectFile, Target.bForDistribution, iCloudContainerIdentifiersXML, UbiquityContainerIdentifiersXML);
+
+			// create a pList key named ICloudContainerIdentifier
+			// to be used at run-time when intializing the CloudKit services
+			if (iCloudContainerIdentifier != "")
+			{
+				string PListFile = IntermediateDir + "/" + AppName + "-Info.plist";
+				if (File.Exists(PListFile))
+				{
+					string OldPListData = File.ReadAllText(PListFile);
+					XDocument XDoc;
+					try
+					{
+						XDoc = XDocument.Parse(OldPListData);
+						if (XDoc.DocumentType != null)
+						{
+							XDoc.DocumentType.InternalSubset = null;
+						}
+
+						XElement dictElement = XDoc.Root.Element("dict");
+						if (dictElement != null)
+						{
+							dictElement.Add(new XElement("key", "ICloudContainerIdentifier"));
+							dictElement.Add(new XElement("string", iCloudContainerIdentifier));
+
+							XDoc.Save(PListFile);
+						}
+					}
+					catch (Exception e)
+					{
+						throw new BuildException("plist is invalid {0}\n{1}", e, OldPListData);
+					}
+
+				}
+			}
+		}
+
 		private static void WriteEntitlementsFile(string OutputFilename, FileReference ProjectFile, bool bForDistribution, string iCloudContainerIdentifiersXML, string UbiquityContainerIdentifiersXML)
 		{
 			// get the settings from the ini file
@@ -1830,61 +1900,14 @@ namespace UnrealBuildTool
 						CmdLine += (!string.IsNullOrEmpty(MobileProvisionUUID) ? (" PROVISIONING_PROFILE_SPECIFIER=" + MobileProvisionUUID) : "");
 					}
 					Writer.WriteLine("/usr/bin/xcrun {0}", CmdLine);
-					
-					// get some info from the mobileprovisioning file
-					// the iCloud identifier and the bundle id may differ
-					string iCloudContainerIdentifiersXML = "";
-					string iCloudContainerIdentifier = "";
-					string UbiquityContainerIdentifiersXML = "";
-					if (File.Exists(MobileProvisionFile.FullName))
-					{
-						MobileProvisionContents MobileProvisionContent = MobileProvisionContents.Read(MobileProvisionFile);
-
-						iCloudContainerIdentifier = MobileProvisionContent.GetNodeValueByName("com.apple.developer.icloud-container-identifiers");
-						iCloudContainerIdentifiersXML = MobileProvisionContent.GetNodeXMLValueByName("com.apple.developer.icloud-container-identifiers");
-						UbiquityContainerIdentifiersXML = MobileProvisionContent.GetNodeXMLValueByName("com.apple.developer.ubiquity-container-identifiers");
-					}
-					// create the entitlements file
-					string IntermediateDir = (((Target.ProjectFile != null) ? Target.ProjectFile.Directory.ToString() : 
-						UnrealBuildTool.EngineDirectory.ToString())) + "/Intermediate/" + (Target.Platform == UnrealTargetPlatform.IOS ? "IOS" : "TVOS");
-					WriteEntitlementsFile(Path.Combine(IntermediateDir, AppName + ".entitlements"), Target.ProjectFile, Target.bForDistribution, iCloudContainerIdentifiersXML, UbiquityContainerIdentifiersXML);
-
-					// create a pList key named ICloudContainerIdentifier
-					// to be used at run-time when intializing the CloudKit services
-					if (iCloudContainerIdentifier != "")
-					{
-						string PListFile = IntermediateDir + "/" + AppName + "-Info.plist";
-						if (File.Exists(PListFile))
-						{
-							string OldPListData = File.ReadAllText(PListFile);
-							XDocument XDoc;
-							try
-							{
-								XDoc = XDocument.Parse(OldPListData);
-								if (XDoc.DocumentType != null)
-									XDoc.DocumentType.InternalSubset = null;
-
-								XElement dictElement = XDoc.Root.Element("dict");
-								if (dictElement != null)
-								{
-									dictElement.Add(new XElement("key", "ICloudContainerIdentifier"));
-									dictElement.Add(new XElement("string", iCloudContainerIdentifier));
-
-									XDoc.Save(PListFile);
-								}
-							}
-							catch (Exception e)
-							{
-								throw new BuildException("plist is invalid {0}\n{1}", e, OldPListData);
-							}
-
-						}
-					}
 
 				}
 
 				Log.TraceInformation("Executing {0}", SignProjectScript);
 
+				// write the entitlements file (building remotely)
+				WriteEntitlements(Target);
+
 				Process SignProcess = new Process();
 				SignProcess.StartInfo.WorkingDirectory = RemoteShadowDirectoryMac;
 				SignProcess.StartInfo.FileName = "/bin/sh";
@@ -1936,6 +1959,12 @@ namespace UnrealBuildTool
 				// Package the stub
 				PackageStub(RemoteShadowDirectoryMac, AppName, Target.OutputPath.GetFileNameWithoutExtension());
 			}
+			else
+			{
+				// write the entitlements file (building on Mac)
+				WriteEntitlements(Target);
+			}
+			
 
 			{
 				// Copy bundled assets from additional frameworks to the intermediate assets directory (so they can get picked up during staging)
