Moved dsymutil with fix for large sizes into Binaries for the time being so people with Xcode installed don't need to also set up autosdks

michael.trepka

diff --git a/Engine/Source/Programs/UnrealBuildTool/ToolChain/AppleToolChain.cs b/Engine/Source/Programs/UnrealBuildTool/ToolChain/AppleToolChain.cs
--- a/Engine/Source/Programs/UnrealBuildTool/ToolChain/AppleToolChain.cs
+++ b/Engine/Source/Programs/UnrealBuildTool/ToolChain/AppleToolChain.cs
@@ -165,6 +165,15 @@ namespace UnrealBuildTool
 		{
 			FileReference DsymutilLocation = new FileReference("/usr/bin/dsymutil");
 
+			// Default dsymutil has a bug that causes issues, it's fixed in autosdks but not everyone has those set up so for the timebeing we have
+			// a version in P4...
+			FileReference PatchedDsymutilLocation = FileReference.Combine(UnrealBuildTool.EngineDirectory, "Binaries/Mac/NotForLicensees/LLVM/bin/dsymutil");
+			
+			if (File.Exists(PatchedDsymutilLocation.FullName))
+			{
+				DsymutilLocation = PatchedDsymutilLocation;
+			}
+			
 			DirectoryReference AutoSdkDir;
 			if (UEBuildPlatformSDK.TryGetHostPlatformAutoSDKDir(out AutoSdkDir))
 			{
