#UE4 Disable warning while we wait for a key cycle

diff --git a/Engine/Source/Programs/UnrealBuildTool/System/EncryptionAndSigning.cs b/Engine/Source/Programs/UnrealBuildTool/System/EncryptionAndSigning.cs
--- a/Engine/Source/Programs/UnrealBuildTool/System/EncryptionAndSigning.cs
+++ b/Engine/Source/Programs/UnrealBuildTool/System/EncryptionAndSigning.cs
@@ -437,7 +437,7 @@ namespace UnrealBuildTool
 			}
 
 			{
-				Log.TraceWarningOnce("Project signing keys found in '{0}' are of the old insecure short format. Please regenerate them using the project crypto settings panel in the editor!", InProjectDirectory);
+				//Log.TraceWarningOnce("Project signing keys found in '{0}' are of the old insecure short format. Please regenerate them using the project crypto settings panel in the editor!", InProjectDirectory);
 			}
 
 			// Validate the settings we have read
