Fix problems with HotReload on Mac.

* Linker arguments for each hot-reloaded module where only being patched with that module's new name (and not *all* the new module names). Windows was unaffected by this, since it uses response files (which were being handled correctly).
* Loop to rename modules in the manifest was terminating after the first modified module was found. This prevented other modules from being renamed.

diff --git a/Engine/Source/Programs/UnrealBuildTool/System/HotReload.cs b/Engine/Source/Programs/UnrealBuildTool/System/HotReload.cs
--- a/Engine/Source/Programs/UnrealBuildTool/System/HotReload.cs
+++ b/Engine/Source/Programs/UnrealBuildTool/System/HotReload.cs
@@ -527,9 +527,6 @@ namespace UnrealBuildTool
 					}
 				}
 
-				// Go ahead and replace all occurrences of our file name in the command-line (ignoring extensions)
-				Action.CommandArguments = ReplaceBaseFileName(Action.CommandArguments, OriginalFileNameWithoutExtension, NewFileNameWithoutExtension);
-
 				// Update this action's list of produced items too
 				for (int ItemIndex = 0; ItemIndex < Action.ProducedItems.Count; ++ItemIndex)
 				{
@@ -581,6 +578,18 @@ namespace UnrealBuildTool
 
 			if (OriginalFileNameAndNewFileNameList_NoExtensions.Count > 0)
 			{
+				// Update all the paths in link actions
+				foreach (Action Action in Actions.Where((Action) => Action.ActionType == ActionType.Link))
+				{
+					foreach (KeyValuePair<string, string> FileNameTuple in OriginalFileNameAndNewFileNameList_NoExtensions)
+					{
+						string OriginalFileNameWithoutExtension = FileNameTuple.Key;
+						string NewFileNameWithoutExtension = FileNameTuple.Value;
+
+						Action.CommandArguments = ReplaceBaseFileName(Action.CommandArguments, OriginalFileNameWithoutExtension, NewFileNameWithoutExtension);
+					}
+				}
+
 				foreach (string ResponseFilePath in ResponseFilePaths)
 				{
 					// Load the file up
@@ -660,7 +669,8 @@ namespace UnrealBuildTool
 					WriteMetadataTargetInfo TargetInfo = BinaryFormatterUtils.Load<WriteMetadataTargetInfo>(TargetInfoFile);
 					foreach (KeyValuePair<FileReference, ModuleManifest> FileNameToVersionManifest in TargetInfo.FileToManifest)
 					{
-						foreach (KeyValuePair<string, string> Manifest in FileNameToVersionManifest.Value.ModuleNameToFileName)
+						KeyValuePair<string, string>[] ManifestEntries = FileNameToVersionManifest.Value.ModuleNameToFileName.ToArray();
+						foreach (KeyValuePair<string, string> Manifest in ManifestEntries)
 						{
 							FileReference OriginalFile = FileReference.Combine(FileNameToVersionManifest.Key.Directory, Manifest.Value);
 
@@ -668,7 +678,6 @@ namespace UnrealBuildTool
 							if(OriginalFileToHotReloadFile.TryGetValue(OriginalFile, out HotReloadFile))
 							{
 								FileNameToVersionManifest.Value.ModuleNameToFileName[Manifest.Key] = HotReloadFile.GetFileName();
-								break;
 							}
 						}
 					}
