UBT: Fix makefile not being invalidated if a directory containing source files is added.

diff --git a/Engine/Source/Programs/UnrealBuildTool/Modes/BuildMode.cs b/Engine/Source/Programs/UnrealBuildTool/Modes/BuildMode.cs
--- a/Engine/Source/Programs/UnrealBuildTool/Modes/BuildMode.cs
+++ b/Engine/Source/Programs/UnrealBuildTool/Modes/BuildMode.cs
@@ -397,7 +397,7 @@ namespace UnrealBuildTool
 
 				// Check that the makefile is still valid
 				string Reason;
-				if(!TargetMakefile.IsValidForSourceFiles(Makefile, TargetDescriptor.ProjectFile, WorkingSet, out Reason))
+				if(!TargetMakefile.IsValidForSourceFiles(Makefile, TargetDescriptor.ProjectFile, TargetDescriptor.Platform, WorkingSet, out Reason))
 				{
 					Log.TraceInformation("Invalidating makefile for {0} ({1})", TargetDescriptor.Name, Reason);
 					Makefile = null;
diff --git a/Engine/Source/Programs/UnrealBuildTool/System/TargetMakefile.cs b/Engine/Source/Programs/UnrealBuildTool/System/TargetMakefile.cs
--- a/Engine/Source/Programs/UnrealBuildTool/System/TargetMakefile.cs
+++ b/Engine/Source/Programs/UnrealBuildTool/System/TargetMakefile.cs
@@ -414,13 +414,17 @@ namespace UnrealBuildTool
 		/// </summary>
 		/// <param name="Makefile">The makefile that has been loaded</param>
 		/// <param name="ProjectFile">Path to the project file</param>
+		/// <param name="Platform">The platform being built</param>
 		/// <param name="WorkingSet">The current working set of source files</param>
 		/// <param name="ReasonNotLoaded">If the makefile is not valid, is set to a message describing why</param>
 		/// <returns>True if the makefile is valid, false otherwise</returns>
-		public static bool IsValidForSourceFiles(TargetMakefile Makefile, FileReference ProjectFile, ISourceFileWorkingSet WorkingSet, out string ReasonNotLoaded)
+		public static bool IsValidForSourceFiles(TargetMakefile Makefile, FileReference ProjectFile, UnrealTargetPlatform Platform, ISourceFileWorkingSet WorkingSet, out string ReasonNotLoaded)
 		{
 			using(Timeline.ScopeEvent("TargetMakefile.IsValidForSourceFiles()"))
 			{
+				// Get the list of excluded folder names for this platform
+				ReadOnlyHashSet<string> ExcludedFolderNames = UEBuildPlatform.GetBuildPlatform(Platform).GetExcludedFolderNames();
+
 				// Check if any source files have been added or removed
 				foreach(KeyValuePair<DirectoryItem, FileItem[]> Pair in Makefile.DirectoryToSourceFiles)
 				{
@@ -443,6 +447,15 @@ namespace UnrealBuildTool
 							ReasonNotLoaded = "source file modified";
 							return false;
 						}
+
+						foreach(DirectoryItem Directory in InputDirectory.EnumerateDirectories())
+						{
+							if(!Makefile.DirectoryToSourceFiles.ContainsKey(Directory) && ContainsSourceFiles(Directory, ExcludedFolderNames))
+							{
+								ReasonNotLoaded = "directory added";
+								return false;
+							}
+						}
 					}
 				}
 
@@ -556,6 +569,36 @@ namespace UnrealBuildTool
 			return true;
 		}
 
+		/// <summary>
+		/// Determines if a directory, or any subdirectory of it, contains new source files
+		/// </summary>
+		/// <param name="Directory">Directory to search through</param>
+		/// <param name="ExcludedFolderNames">Set of directory names to exclude</param>
+		/// <returns>True if the directory contains any source files</returns>
+		static bool ContainsSourceFiles(DirectoryItem Directory, ReadOnlyHashSet<string> ExcludedFolderNames)
+		{
+			// Check this directory isn't ignored
+			if(!ExcludedFolderNames.Contains(Directory.Name))
+			{
+				// Check for any source files in this actual directory
+				FileItem[] SourceFiles = UEBuildModuleCPP.GetSourceFiles(Directory);
+				if(SourceFiles.Length > 0)
+				{
+					return true;
+				}
+
+				// Check for any source files in a subdirectory
+				foreach(DirectoryItem SubDirectory in Directory.EnumerateDirectories())
+				{
+					if(ContainsSourceFiles(SubDirectory, ExcludedFolderNames))
+					{
+						return true;
+					}
+				}
+			}
+			return false;
+		}
+
 		/// <summary>
 		/// Gets the location of the makefile for particular target
 		/// </summary>
