Fix issue where live++ recompile would execute old hotreload system code that would trigger an exception in monolithic builds

diff --git a/Engine/Source/Programs/UnrealBuildTool/Configuration/UEBuildTarget.cs b/Engine/Source/Programs/UnrealBuildTool/Configuration/UEBuildTarget.cs
--- a/Engine/Source/Programs/UnrealBuildTool/Configuration/UEBuildTarget.cs
+++ b/Engine/Source/Programs/UnrealBuildTool/Configuration/UEBuildTarget.cs
@@ -2084,14 +2084,17 @@ namespace UnrealBuildTool
 		private HashSet<string> GetHotReloadModuleNames()
 		{
 			HashSet<string> HotReloadModuleNames = new HashSet<string>(StringComparer.OrdinalIgnoreCase);
-			foreach (UEBuildBinary Binary in Binaries)
+			if (!ShouldCompileMonolithic())
 			{
-				List<UEBuildModule> GameModules = Binary.FindGameModules();
-				if (GameModules != null && GameModules.Count > 0)
+				foreach (UEBuildBinary Binary in Binaries)
 				{
-					if(!UnrealBuildTool.IsProjectInstalled() || EnabledPlugins.Where(x => x.Type == PluginType.Mod).Any(x => Binary.OutputFilePaths[0].IsUnderDirectory(x.Directory)))
+					List<UEBuildModule> GameModules = Binary.FindGameModules();
+					if (GameModules != null && GameModules.Count > 0)
 					{
-						HotReloadModuleNames.UnionWith(GameModules.OfType<UEBuildModuleCPP>().Select(x => x.Name));
+						if (!UnrealBuildTool.IsProjectInstalled() || EnabledPlugins.Where(x => x.Type == PluginType.Mod).Any(x => Binary.OutputFilePaths[0].IsUnderDirectory(x.Directory)))
+						{
+							HotReloadModuleNames.UnionWith(GameModules.OfType<UEBuildModuleCPP>().Select(x => x.Name));
+						}
 					}
 				}
 			}
