- Now using new dsymutil with Xcode 10.2 (checking version of dsymutil that's builtin before looking for alternates)

diff --git a/Engine/Source/Programs/UnrealBuildTool/ToolChain/AppleToolChain.cs b/Engine/Source/Programs/UnrealBuildTool/ToolChain/AppleToolChain.cs
--- a/Engine/Source/Programs/UnrealBuildTool/ToolChain/AppleToolChain.cs
+++ b/Engine/Source/Programs/UnrealBuildTool/ToolChain/AppleToolChain.cs
@@ -165,22 +165,27 @@ namespace UnrealBuildTool
 		{
 			FileReference DsymutilLocation = new FileReference("/usr/bin/dsymutil");
 
-			// Default dsymutil has a bug that causes issues, it's fixed in autosdks but not everyone has those set up so for the timebeing we have
+			string DsymutilVersionString = Utils.RunLocalProcessAndReturnStdOut(DsymutilLocation.FullName, "-version");
+
+			// dsymutil 10.0.0 has a bug that causes issues, it's fixed in autosdks but not everyone has those set up so for the timebeing we have
 			// a version in P4...
-			FileReference PatchedDsymutilLocation = FileReference.Combine(UnrealBuildTool.EngineDirectory, "Binaries/Mac/NotForLicensees/LLVM/bin/dsymutil");
-			
-			if (File.Exists(PatchedDsymutilLocation.FullName))
-			{
-				DsymutilLocation = PatchedDsymutilLocation;
-			}
-			
-			DirectoryReference AutoSdkDir;
-			if (UEBuildPlatformSDK.TryGetHostPlatformAutoSDKDir(out AutoSdkDir))
+			if (DsymutilVersionString.StartsWith("Apple LLVM version 10.0.0", StringComparison.Ordinal))
 			{
-				FileReference AutoSdkDsymutilLocation = FileReference.Combine(AutoSdkDir, "Mac", "LLVM", "bin", "dsymutil");
-				if (FileReference.Exists(AutoSdkDsymutilLocation))
+				FileReference PatchedDsymutilLocation = FileReference.Combine(UnrealBuildTool.EngineDirectory, "Binaries/Mac/NotForLicensees/LLVM/bin/dsymutil");
+
+				if (File.Exists(PatchedDsymutilLocation.FullName))
 				{
-					DsymutilLocation = AutoSdkDsymutilLocation;
+					DsymutilLocation = PatchedDsymutilLocation;
+				}
+
+				DirectoryReference AutoSdkDir;
+				if (UEBuildPlatformSDK.TryGetHostPlatformAutoSDKDir(out AutoSdkDir))
+				{
+					FileReference AutoSdkDsymutilLocation = FileReference.Combine(AutoSdkDir, "Mac", "LLVM", "bin", "dsymutil");
+					if (FileReference.Exists(AutoSdkDsymutilLocation))
+					{
+						DsymutilLocation = AutoSdkDsymutilLocation;
+					}
 				}
 			}
 
