UE-52346 TM-Gameplay: Load Level Volume Fails to Read TestSlot.sav file

#iOS
#tvOS
#4.22

On Mac the pList file is overriden before the final packaging. Adding iCloudContainerIdentifier.

diff --git a/Engine/Source/Programs/UnrealBuildTool/Platform/IOS/UEDeployIOS.cs b/Engine/Source/Programs/UnrealBuildTool/Platform/IOS/UEDeployIOS.cs
--- a/Engine/Source/Programs/UnrealBuildTool/Platform/IOS/UEDeployIOS.cs
+++ b/Engine/Source/Programs/UnrealBuildTool/Platform/IOS/UEDeployIOS.cs
@@ -720,6 +720,20 @@ namespace UnrealBuildTool
 				Text.AppendLine("\t</array>");
 			}
 
+			// write the iCloud container identifier, if present in the old file
+			if (!string.IsNullOrEmpty(OldPListData))
+			{
+				int index = OldPListData.IndexOf("ICloudContainerIdentifier");
+				if (index > 0)
+				{
+					index = OldPListData.IndexOf("<string>", index) + 8;
+					int length = OldPListData.IndexOf("</string>", index) - index;
+					string ICloudContainerIdentifier = OldPListData.Substring(index, length);
+					Text.AppendLine("\t<key>ICloudContainerIdentifier</key>");
+					Text.AppendLine(string.Format("\t<string>{0}</string>", ICloudContainerIdentifier));
+				}
+			}
+
 			Text.AppendLine("</dict>");
 			Text.AppendLine("</plist>");
 
