Undo my dsymutil change as it caused some odd buildmachine issue that needs to be investigated with preflights

diff --git a/Engine/Source/Programs/UnrealBuildTool/ToolChain/AppleToolChain.cs b/Engine/Source/Programs/UnrealBuildTool/ToolChain/AppleToolChain.cs
--- a/Engine/Source/Programs/UnrealBuildTool/ToolChain/AppleToolChain.cs
+++ b/Engine/Source/Programs/UnrealBuildTool/ToolChain/AppleToolChain.cs
@@ -165,27 +165,22 @@ namespace UnrealBuildTool
 		{
 			FileReference DsymutilLocation = new FileReference("/usr/bin/dsymutil");
 
-			string DsymutilVersionString = Utils.RunLocalProcessAndReturnStdOut(DsymutilLocation.FullName, "-version");
-
-			// dsymutil 10.0.0 has a bug that causes issues, it's fixed in autosdks but not everyone has those set up so for the timebeing we have
+			// Default dsymutil has a bug that causes issues, it's fixed in autosdks but not everyone has those set up so for the timebeing we have
 			// a version in P4...
-			if (DsymutilVersionString.StartsWith("Apple LLVM version 10.0.0", StringComparison.Ordinal))
+			FileReference PatchedDsymutilLocation = FileReference.Combine(UnrealBuildTool.EngineDirectory, "Binaries/Mac/NotForLicensees/LLVM/bin/dsymutil");
+			
+			if (File.Exists(PatchedDsymutilLocation.FullName))
 			{
-				FileReference PatchedDsymutilLocation = FileReference.Combine(UnrealBuildTool.EngineDirectory, "Binaries/Mac/NotForLicensees/LLVM/bin/dsymutil");
-
-				if (File.Exists(PatchedDsymutilLocation.FullName))
-				{
-					DsymutilLocation = PatchedDsymutilLocation;
-				}
-
-				DirectoryReference AutoSdkDir;
-				if (UEBuildPlatformSDK.TryGetHostPlatformAutoSDKDir(out AutoSdkDir))
+				DsymutilLocation = PatchedDsymutilLocation;
+			}
+			
+			DirectoryReference AutoSdkDir;
+			if (UEBuildPlatformSDK.TryGetHostPlatformAutoSDKDir(out AutoSdkDir))
+			{
+				FileReference AutoSdkDsymutilLocation = FileReference.Combine(AutoSdkDir, "Mac", "LLVM", "bin", "dsymutil");
+				if (FileReference.Exists(AutoSdkDsymutilLocation))
 				{
-					FileReference AutoSdkDsymutilLocation = FileReference.Combine(AutoSdkDir, "Mac", "LLVM", "bin", "dsymutil");
-					if (FileReference.Exists(AutoSdkDsymutilLocation))
-					{
-						DsymutilLocation = AutoSdkDsymutilLocation;
-					}
+					DsymutilLocation = AutoSdkDsymutilLocation;
 				}
 			}
 
