UBT: Allow compiled rules assemblies to use the System.Core DLL, so Linq is available.

diff --git a/Engine/Source/Programs/UnrealBuildTool/System/DynamicCompilation.cs b/Engine/Source/Programs/UnrealBuildTool/System/DynamicCompilation.cs
--- a/Engine/Source/Programs/UnrealBuildTool/System/DynamicCompilation.cs
+++ b/Engine/Source/Programs/UnrealBuildTool/System/DynamicCompilation.cs
@@ -279,6 +279,7 @@ namespace UnrealBuildTool
 					{
 						// Always depend on the CLR System assembly
 						CompileParams.ReferencedAssemblies.Add("System.dll");
+						CompileParams.ReferencedAssemblies.Add("System.Core.dll");
 					}
 					else
 					{
