UBT: Exclude any source files in a ThirdParty folder from the generated solution, not just those under Engine/Source/ThirdParty. There are a lot of large third party modules in engine plugins.

diff --git a/Engine/Source/Programs/UnrealBuildTool/ProjectFiles/ProjectFileGenerator.cs b/Engine/Source/Programs/UnrealBuildTool/ProjectFiles/ProjectFileGenerator.cs
--- a/Engine/Source/Programs/UnrealBuildTool/ProjectFiles/ProjectFileGenerator.cs
+++ b/Engine/Source/Programs/UnrealBuildTool/ProjectFiles/ProjectFileGenerator.cs
@@ -1879,7 +1879,7 @@ namespace UnrealBuildTool
 
 				// We'll keep track of whether this is an "engine" or "external" module.  This is determined below while loading module rules.
 				bool IsEngineModule = UnrealBuildTool.IsUnderAnEngineDirectory(CurModuleFile.Directory);
-				bool IsThirdPartyModule = CurModuleFile.IsUnderDirectory(UnrealBuildTool.EngineSourceThirdPartyDirectory);
+				bool IsThirdPartyModule = CurModuleFile.ContainsName("ThirdParty", UnrealBuildTool.RootDirectory);
 
 				if( IsEngineModule && !bIncludeEngineSource )
 				{
