Raise the minimum value for PCH compilation to avoid compiler errors in VS2017

diff --git a/Engine/Source/Programs/UnrealBuildTool/Platform/Windows/VCToolChain.cs b/Engine/Source/Programs/UnrealBuildTool/Platform/Windows/VCToolChain.cs
--- a/Engine/Source/Programs/UnrealBuildTool/Platform/Windows/VCToolChain.cs
+++ b/Engine/Source/Programs/UnrealBuildTool/Platform/Windows/VCToolChain.cs
@@ -199,7 +199,7 @@ namespace UnrealBuildTool
 			// Separate functions for linker.
 			Arguments.Add("/Gy");
 
-			// Allow 750% of the default memory allocation limit when using the static analyzer, and 850% at other times.
+			// Allow 750% of the default memory allocation limit when using the static analyzer, and 1000% at other times.
 			if(Target.WindowsPlatform.PCHMemoryAllocationFactor == 0)
 			{
 				if (Target.WindowsPlatform.StaticAnalyzer == WindowsStaticAnalyzer.VisualCpp)
@@ -208,7 +208,7 @@ namespace UnrealBuildTool
 				}
 				else
 				{
-					Arguments.Add("/Zm850");
+					Arguments.Add("/Zm1000");
 				}
 			}
 			else
