ROBOMERGE conflict 5630547 resolved UE-71953 - Blueprint iOS projects can't package For Distribution

diff --git a/Engine/Source/Programs/UnrealBuildTool/Platform/IOS/IOSToolChain.cs b/Engine/Source/Programs/UnrealBuildTool/Platform/IOS/IOSToolChain.cs
--- a/Engine/Source/Programs/UnrealBuildTool/Platform/IOS/IOSToolChain.cs
+++ b/Engine/Source/Programs/UnrealBuildTool/Platform/IOS/IOSToolChain.cs
@@ -413,7 +413,7 @@ namespace UnrealBuildTool
 				Result += " -fobjc-exceptions";
 			}
 			else
-		{
+			{
 				Result += " -fno-objc-exceptions";
 			}
 
@@ -1718,6 +1718,10 @@ namespace UnrealBuildTool
 
         public static void PostBuildSync(IOSPostBuildSyncTarget Target)
 		{
+			ConfigHierarchy Ini = ConfigCache.ReadHierarchy(ConfigHierarchyType.Engine, DirectoryReference.FromFile(Target.ProjectFile), UnrealTargetPlatform.IOS);
+			string BundleID;
+			Ini.GetString("/Script/IOSRuntimeSettings.IOSRuntimeSettings", "BundleIdentifier", out BundleID);
+			
 			IOSProjectSettings ProjectSettings = ((IOSPlatform)UEBuildPlatform.GetBuildPlatform(Target.Platform)).ReadProjectSettings(Target.ProjectFile);
 
 			if (BuildHostPlatform.Current.Platform == UnrealTargetPlatform.Mac)
@@ -1749,11 +1753,6 @@ namespace UnrealBuildTool
 			string RemoteShadowDirectoryMac = Target.OutputPath.Directory.FullName;
 			FileReference StagedExecutablePath = GetStagedExecutablePath(Target.OutputPath, Target.TargetName);
 
-            // ensure the plist, entitlements, and provision files are properly copied
-            UEDeployIOS DeployHandler = (Target.Platform == UnrealTargetPlatform.IOS ? new UEDeployIOS() : new UEDeployTVOS());
-			DeployHandler.ForDistribution = Target.bForDistribution;
-            DeployHandler.PrepTargetForDeployment(Target.ProjectFile, Target.TargetName, Target.Platform, Target.Configuration, Target.UPLScripts, Target.SdkVersion, Target.bCreateStubIPA);
-
 			// copy the executable
 			DirectoryReference.CreateDirectory(StagedExecutablePath.Directory);
 			FileReference.Copy(Target.OutputPath, StagedExecutablePath, true);
@@ -1769,26 +1768,13 @@ namespace UnrealBuildTool
 				DirectoryReference XcodeWorkspaceDir;
 				if (AppName == "UE4Game" || AppName == "UE4Client" || Target.ProjectFile == null || Target.ProjectFile.IsUnderDirectory(UnrealBuildTool.EngineDirectory))
 				{
-					GenerateProjectFiles(Target.ProjectFile, new string[] { "-platforms=" + (Target.Platform == UnrealTargetPlatform.IOS ? "IOS" : "TVOS"), "-NoIntellIsense", (Target.Platform == UnrealTargetPlatform.IOS ? "-iosdeployonly" : "-tvosdeployonly"), "-ignorejunk", (Target.bForDistribution ? "-distribution" : "-development") });
 					XcodeWorkspaceDir = DirectoryReference.Combine(UnrealBuildTool.RootDirectory, String.Format("UE4_{0}.xcworkspace", (Target.Platform == UnrealTargetPlatform.IOS ? "IOS" : "TVOS")));
 				}
 				else
 				{
-					GenerateProjectFiles(Target.ProjectFile, new string[] { "-platforms=" + (Target.Platform == UnrealTargetPlatform.IOS ? "IOS" : "TVOS"), "-NoIntellIsense", (Target.Platform == UnrealTargetPlatform.IOS ? "-iosdeployonly" : "-tvosdeployonly"), "-ignorejunk", (Target.bForDistribution ? "-distribution" : "-development"), String.Format("-project={0}", Target.ProjectFile), "-game" });
 					XcodeWorkspaceDir = DirectoryReference.Combine(Target.ProjectDirectory, String.Format("{0}_{1}.xcworkspace", Target.ProjectFile.GetFileNameWithoutExtension(), (Target.Platform == UnrealTargetPlatform.IOS ? "IOS" : "TVOS")));
 				}
 
-				// Make sure it exists
-				if (!DirectoryReference.Exists(XcodeWorkspaceDir))
-				{
-					throw new BuildException("Unable to create stub IPA; Xcode workspace not found at {0}", XcodeWorkspaceDir);
-				}
-
-				// ensure the plist, entitlements, and provision files are properly copied
-				DeployHandler = (Target.Platform == UnrealTargetPlatform.IOS ? new UEDeployIOS() : new UEDeployTVOS());
-				DeployHandler.ForDistribution = Target.bForDistribution;
-				DeployHandler.PrepTargetForDeployment(Target.ProjectFile, Target.TargetName, Target.Platform, Target.Configuration, Target.UPLScripts, Target.SdkVersion, true);
-
 				// Path to the temporary keychain. When -ImportCertificate is specified, we will temporarily add this to the list of keychains to search, and remove it later.
 				FileReference TempKeychain = FileReference.Combine(Target.ProjectIntermediateDirectory, "TempKeychain.keychain");
 
@@ -1859,6 +1845,7 @@ namespace UnrealBuildTool
 						MobileProvisionFile = ProvisioningData.MobileProvisionFile;
 						MobileProvisionUUID = ProvisioningData.MobileProvisionUUID;
 						TeamUUID = ProvisioningData.TeamUUID;
+						BundleID = ProvisioningData.BundleIdentifier;
 					}
 					else
 					{
@@ -1867,6 +1854,7 @@ namespace UnrealBuildTool
 						MobileProvisionContents MobileProvision = MobileProvisionContents.Read(MobileProvisionFile);
 						MobileProvisionUUID = MobileProvision.GetUniqueId();
 						MobileProvision.TryGetTeamUniqueId(out TeamUUID);
+						BundleID = MobileProvision.GetBundleIdentifier();
 					}
 
 					string ConfigName = Target.Configuration.ToString();
@@ -1903,6 +1891,25 @@ namespace UnrealBuildTool
 
 				}
 
+				if (AppName == "UE4Game" || AppName == "UE4Client" || Target.ProjectFile == null || Target.ProjectFile.IsUnderDirectory(UnrealBuildTool.EngineDirectory))
+				{
+					GenerateProjectFiles(Target.ProjectFile, new string[] { "-platforms=" + (Target.Platform == UnrealTargetPlatform.IOS ? "IOS" : "TVOS"), "-NoIntellIsense", (Target.Platform == UnrealTargetPlatform.IOS ? "-iosdeployonly" : "-tvosdeployonly"), "-ignorejunk", (Target.bForDistribution ? "-distribution" : "-development"), "-bundleID=" + BundleID });
+				}
+				else
+				{
+					GenerateProjectFiles(Target.ProjectFile, new string[] { "-platforms=" + (Target.Platform == UnrealTargetPlatform.IOS ? "IOS" : "TVOS"), "-NoIntellIsense", (Target.Platform == UnrealTargetPlatform.IOS ? "-iosdeployonly" : "-tvosdeployonly"), "-ignorejunk", (Target.bForDistribution ? "-distribution" : "-development"), String.Format("-project={0}", Target.ProjectFile), "-game", "-bundleID=" + BundleID });
+				}
+				// Make sure it exists
+				if (!DirectoryReference.Exists(XcodeWorkspaceDir))
+				{
+					throw new BuildException("Unable to create stub IPA; Xcode workspace not found at {0}", XcodeWorkspaceDir);
+				}
+
+				// ensure the plist, entitlements, and provision files are properly copied
+				UEDeployIOS DeployHandler = (Target.Platform == UnrealTargetPlatform.IOS ? new UEDeployIOS() : new UEDeployTVOS());
+				DeployHandler.ForDistribution = Target.bForDistribution;
+				DeployHandler.PrepTargetForDeployment(Target.ProjectFile, Target.TargetName, Target.Platform, Target.Configuration, Target.UPLScripts, Target.SdkVersion, Target.bCreateStubIPA, BundleID);
+
 				Log.TraceInformation("Executing {0}", SignProjectScript);
 
 				// write the entitlements file (building remotely)
@@ -1961,6 +1968,11 @@ namespace UnrealBuildTool
 			}
 			else
 			{
+				// ensure the plist, entitlements, and provision files are properly copied
+				UEDeployIOS DeployHandler = (Target.Platform == UnrealTargetPlatform.IOS ? new UEDeployIOS() : new UEDeployTVOS());
+				DeployHandler.ForDistribution = Target.bForDistribution;
+				DeployHandler.PrepTargetForDeployment(Target.ProjectFile, Target.TargetName, Target.Platform, Target.Configuration, Target.UPLScripts, Target.SdkVersion, Target.bCreateStubIPA, BundleID);
+
 				// write the entitlements file (building on Mac)
 				WriteEntitlements(Target);
 			}
diff --git a/Engine/Source/Programs/UnrealBuildTool/Platform/IOS/MobileProvisionContents.cs b/Engine/Source/Programs/UnrealBuildTool/Platform/IOS/MobileProvisionContents.cs
--- a/Engine/Source/Programs/UnrealBuildTool/Platform/IOS/MobileProvisionContents.cs
+++ b/Engine/Source/Programs/UnrealBuildTool/Platform/IOS/MobileProvisionContents.cs
@@ -63,6 +63,37 @@ namespace UnrealBuildTool
 			return UniqueIdElement.InnerText;
 		}
 
+		/// <summary>
+		/// Gets the bundle id for this mobileprovision
+		/// </summary>
+		/// <returns>Bundle Identifier for the provision</returns>
+		public string GetBundleIdentifier()
+		{
+			XmlElement UniqueIdElement = null, UniqueIdEntitlement;
+			if (!NameToValue.TryGetValue("Entitlements", out UniqueIdEntitlement) || UniqueIdEntitlement.Name != "dict")
+			{
+				throw new BuildException("Missing Entitlements in MobileProvision");
+			}
+
+			foreach (XmlElement KeyElement in UniqueIdEntitlement.SelectNodes("key"))
+			{
+				Console.WriteLine("Found entitlement node:" + KeyElement.InnerText);
+				if (!KeyElement.InnerText.Equals("application-identifier"))
+				{
+					continue;
+				}
+				UniqueIdElement = KeyElement.NextSibling as XmlElement;
+				break;
+			}
+
+
+			if (UniqueIdElement == null)
+			{
+				throw new BuildException("Missing Bundle Identifier in MobileProvision");
+			}
+			return UniqueIdElement.InnerText.Substring(UniqueIdElement.InnerText.IndexOf('.') + 1);
+		}
+
 		/// <summary>
 		/// Gets the team unique id for this mobileprovision
 		/// </summary>
diff --git a/Engine/Source/Programs/UnrealBuildTool/Platform/IOS/UEBuildIOS.cs b/Engine/Source/Programs/UnrealBuildTool/Platform/IOS/UEBuildIOS.cs
--- a/Engine/Source/Programs/UnrealBuildTool/Platform/IOS/UEBuildIOS.cs
+++ b/Engine/Source/Programs/UnrealBuildTool/Platform/IOS/UEBuildIOS.cs
@@ -395,8 +395,9 @@ namespace UnrealBuildTool
 		/// Constructor
 		/// </summary>
 		/// <param name="ProjectFile">The project file to read settings for</param>
-		public IOSProjectSettings(FileReference ProjectFile) 
-			: this(ProjectFile, UnrealTargetPlatform.IOS)
+		/// <param name="Bundle">Bundle identifier needed when project file is empty</param>
+		public IOSProjectSettings(FileReference ProjectFile, string Bundle) 
+			: this(ProjectFile, UnrealTargetPlatform.IOS, Bundle)
 		{
 		}
 
@@ -405,11 +406,16 @@ namespace UnrealBuildTool
 		/// </summary>
 		/// <param name="ProjectFile">The project file to read settings for</param>
 		/// <param name="Platform">The platform to read settings for</param>
-		protected IOSProjectSettings(FileReference ProjectFile, UnrealTargetPlatform Platform)
+		/// <param name="Bundle">Bundle identifier needed when project file is empty</param>
+		protected IOSProjectSettings(FileReference ProjectFile, UnrealTargetPlatform Platform, string Bundle)
 		{
 			this.ProjectFile = ProjectFile;
 			ConfigCache.ReadSettings(DirectoryReference.FromFile(ProjectFile), Platform, this);
-            BundleIdentifier = BundleIdentifier.Replace("[PROJECT_NAME]", ((ProjectFile != null) ? ProjectFile.GetFileNameWithoutAnyExtensions() : "UE4Game")).Replace("_", "");
+			if ((ProjectFile == null || string.IsNullOrEmpty(ProjectFile.FullName)) && !string.IsNullOrEmpty(Bundle))
+			{
+				BundleIdentifier = Bundle;
+			}
+			BundleIdentifier = BundleIdentifier.Replace("[PROJECT_NAME]", ((ProjectFile != null) ? ProjectFile.GetFileNameWithoutAnyExtensions() : "UE4Game")).Replace("_", "");
 		}
 	}
 
@@ -423,6 +429,7 @@ namespace UnrealBuildTool
         public string MobileProvisionUUID;
         public string MobileProvisionName;
         public string TeamUUID;
+		public string BundleIdentifier;
 		public bool bHaveCertificate = false;
 
 		public string MobileProvision
@@ -589,7 +596,18 @@ namespace UnrealBuildTool
 							TeamUUID = AllText.Substring(idx, AllText.IndexOf("</string>", idx) - idx);
 						}
 					}
-                    idx = AllText.IndexOf("<key>Name</key>");
+					idx = AllText.IndexOf("<key>application-identifier</key>");
+					if (idx > 0)
+					{
+						idx = AllText.IndexOf("<string>", idx);
+						if (idx > 0)
+						{
+							idx += "<string>".Length;
+							String FullID = AllText.Substring(idx, AllText.IndexOf("</string>", idx) - idx);
+							BundleIdentifier = FullID.Substring(FullID.IndexOf('.') + 1);
+						}
+					}
+					idx = AllText.IndexOf("<key>Name</key>");
                     if (idx > 0)
                     {
                         idx = AllText.IndexOf("<string>", idx);
@@ -756,25 +774,25 @@ namespace UnrealBuildTool
 			return base.GetBinaryExtension(InBinaryType);
 		}
 
-		public IOSProjectSettings ReadProjectSettings(FileReference ProjectFile)
+		public IOSProjectSettings ReadProjectSettings(FileReference ProjectFile, string Bundle = "")
 		{
 			IOSProjectSettings ProjectSettings = CachedProjectSettings.FirstOrDefault(x => x.ProjectFile == ProjectFile);
 			if(ProjectSettings == null)
 			{
-				ProjectSettings = CreateProjectSettings(ProjectFile);
+				ProjectSettings = CreateProjectSettings(ProjectFile, Bundle);
 				CachedProjectSettings.Add(ProjectSettings);
 			}
 			return ProjectSettings;
 		}
 
-		protected virtual IOSProjectSettings CreateProjectSettings(FileReference ProjectFile)
+		protected virtual IOSProjectSettings CreateProjectSettings(FileReference ProjectFile, string Bundle)
 		{
-			return new IOSProjectSettings(ProjectFile);
+			return new IOSProjectSettings(ProjectFile, Bundle);
 		}
 
-		public IOSProvisioningData ReadProvisioningData(FileReference ProjectFile, bool bForDistribution = false)
+		public IOSProvisioningData ReadProvisioningData(FileReference ProjectFile, bool bForDistribution = false, string Bundle = "")
 		{
-			IOSProjectSettings ProjectSettings = ReadProjectSettings(ProjectFile);
+			IOSProjectSettings ProjectSettings = ReadProjectSettings(ProjectFile, Bundle);
 			return ReadProvisioningData(ProjectSettings, bForDistribution);
 		}
 
diff --git a/Engine/Source/Programs/UnrealBuildTool/Platform/IOS/UEDeployIOS.cs b/Engine/Source/Programs/UnrealBuildTool/Platform/IOS/UEDeployIOS.cs
--- a/Engine/Source/Programs/UnrealBuildTool/Platform/IOS/UEDeployIOS.cs
+++ b/Engine/Source/Programs/UnrealBuildTool/Platform/IOS/UEDeployIOS.cs
@@ -195,7 +195,7 @@ namespace UnrealBuildTool
 			return result;
 		}
 
-		public static bool GenerateIOSPList(FileReference ProjectFile, UnrealTargetConfiguration Config, string ProjectDirectory, bool bIsUE4Game, string GameName, string ProjectName, string InEngineDir, string AppDirectory, VersionNumber SdkVersion, UnrealPluginLanguage UPL, out bool bSupportsPortrait, out bool bSupportsLandscape, out bool bSkipIcons)
+		public static bool GenerateIOSPList(FileReference ProjectFile, UnrealTargetConfiguration Config, string ProjectDirectory, bool bIsUE4Game, string GameName, string ProjectName, string InEngineDir, string AppDirectory, VersionNumber SdkVersion, UnrealPluginLanguage UPL, string BundleID, out bool bSupportsPortrait, out bool bSupportsLandscape, out bool bSkipIcons)
 		{
 			// generate the Info.plist for future use
 			string BuildDirectory = ProjectDirectory + "/Build/IOS";
@@ -276,6 +276,10 @@ namespace UnrealBuildTool
 			// bundle identifier
 			string BundleIdentifier;
 			Ini.GetString("/Script/IOSRuntimeSettings.IOSRuntimeSettings", "BundleIdentifier", out BundleIdentifier);
+			if (!string.IsNullOrEmpty(BundleID))
+			{
+				BundleIdentifier = BundleID; // overriding bundle ID
+			}
 
 			// bundle name
 			string BundleName;
@@ -801,10 +805,10 @@ namespace UnrealBuildTool
 		{
 			List<string> UPLScripts = CollectPluginDataPaths(Receipt.AdditionalProperties);
 			VersionNumber SdkVersion = GetSdkVersion(Receipt);
-			return GeneratePList(ProjectFile, Config, ProjectDirectory, bIsUE4Game, GameName, ProjectName, InEngineDir, AppDirectory, UPLScripts, SdkVersion, out bSupportsPortrait, out bSupportsLandscape, out bSkipIcons);
+			return GeneratePList(ProjectFile, Config, ProjectDirectory, bIsUE4Game, GameName, ProjectName, InEngineDir, AppDirectory, UPLScripts, SdkVersion, "", out bSupportsPortrait, out bSupportsLandscape, out bSkipIcons);
 		}
 
-		public virtual bool GeneratePList(FileReference ProjectFile, UnrealTargetConfiguration Config, string ProjectDirectory, bool bIsUE4Game, string GameName, string ProjectName, string InEngineDir, string AppDirectory, List<string> UPLScripts, VersionNumber SdkVersion, out bool bSupportsPortrait, out bool bSupportsLandscape, out bool bSkipIcons)
+		public virtual bool GeneratePList(FileReference ProjectFile, UnrealTargetConfiguration Config, string ProjectDirectory, bool bIsUE4Game, string GameName, string ProjectName, string InEngineDir, string AppDirectory, List<string> UPLScripts, VersionNumber SdkVersion, string BundleID, out bool bSupportsPortrait, out bool bSupportsLandscape, out bool bSkipIcons)
 		{
 			// remember name with -IOS-Shipping, etc
 			string ExeName = GameName;
@@ -836,7 +840,7 @@ namespace UnrealBuildTool
 			// Passing in true for distribution is not ideal here but given the way that ios packaging happens and this call chain it seems unavoidable for now, maybe there is a way to correctly pass it in that I can't find?
 			UPL.Init(ProjectArches, true, RelativeEnginePath, BundlePath, ProjectDirectory, Config.ToString());
 
-			return GenerateIOSPList(ProjectFile, Config, ProjectDirectory, bIsUE4Game, GameName, ProjectName, InEngineDir, AppDirectory, SdkVersion, UPL, out bSupportsPortrait, out bSupportsLandscape, out bSkipIcons);
+			return GenerateIOSPList(ProjectFile, Config, ProjectDirectory, bIsUE4Game, GameName, ProjectName, InEngineDir, AppDirectory, SdkVersion, UPL, BundleID, out bSupportsPortrait, out bSupportsLandscape, out bSkipIcons);
 		}
 
 		protected virtual void CopyCloudResources(string InEngineDir, string AppDirectory)
@@ -941,10 +945,10 @@ namespace UnrealBuildTool
 		{
 			List<string> UPLScripts = CollectPluginDataPaths(Receipt.AdditionalProperties);
 			VersionNumber SdkVersion = GetSdkVersion(Receipt);
-			return PrepForUATPackageOrDeploy(Config, ProjectFile, InProjectName, InProjectDirectory, InExecutablePath, InEngineDir, bForDistribution, CookFlavor, bIsDataDeploy, bCreateStubIPA, UPLScripts, SdkVersion);
+			return PrepForUATPackageOrDeploy(Config, ProjectFile, InProjectName, InProjectDirectory, InExecutablePath, InEngineDir, bForDistribution, CookFlavor, bIsDataDeploy, bCreateStubIPA, UPLScripts, SdkVersion, "");
 		}
 
-		public bool PrepForUATPackageOrDeploy(UnrealTargetConfiguration Config, FileReference ProjectFile, string InProjectName, string InProjectDirectory, string InExecutablePath, string InEngineDir, bool bForDistribution, string CookFlavor, bool bIsDataDeploy, bool bCreateStubIPA, List<string> UPLScripts, VersionNumber SdkVersion)
+		public bool PrepForUATPackageOrDeploy(UnrealTargetConfiguration Config, FileReference ProjectFile, string InProjectName, string InProjectDirectory, string InExecutablePath, string InEngineDir, bool bForDistribution, string CookFlavor, bool bIsDataDeploy, bool bCreateStubIPA, List<string> UPLScripts, VersionNumber SdkVersion, string BundleID)
 		{
 			if (BuildHostPlatform.Current.Platform != UnrealTargetPlatform.Mac)
 			{
@@ -1106,7 +1110,7 @@ namespace UnrealBuildTool
 			bool bSupportsPortrait = true;
 			bool bSupportsLandscape = false;
 			bool bSkipIcons = false;
-			bool bSkipDefaultPNGs = GeneratePList(ProjectFile, Config, InProjectDirectory, bIsUE4Game, GameExeName, InProjectName, InEngineDir, AppDirectory, UPLScripts, SdkVersion, out bSupportsPortrait, out bSupportsLandscape, out bSkipIcons);
+			bool bSkipDefaultPNGs = GeneratePList(ProjectFile, Config, InProjectDirectory, bIsUE4Game, GameExeName, InProjectName, InEngineDir, AppDirectory, UPLScripts, SdkVersion, BundleID, out bSupportsPortrait, out bSupportsLandscape, out bSkipIcons);
 
 			// ensure the destination is writable
 			if (File.Exists(AppDirectory + "/" + GameName))
@@ -1150,10 +1154,10 @@ namespace UnrealBuildTool
 		{
 			List<string> UPLScripts = CollectPluginDataPaths(Receipt.AdditionalProperties);
 			VersionNumber SdkVersion = GetSdkVersion(Receipt);
-			return PrepTargetForDeployment(Receipt.ProjectFile, Receipt.TargetName, Receipt.Platform, Receipt.Configuration, UPLScripts, SdkVersion, false);
+			return PrepTargetForDeployment(Receipt.ProjectFile, Receipt.TargetName, Receipt.Platform, Receipt.Configuration, UPLScripts, SdkVersion, false, "");
 		}
 
-		public bool PrepTargetForDeployment(FileReference ProjectFile, string TargetName, UnrealTargetPlatform Platform, UnrealTargetConfiguration Configuration, List<string> UPLScripts, VersionNumber SdkVersion, bool bCreateStubIPA)
+		public bool PrepTargetForDeployment(FileReference ProjectFile, string TargetName, UnrealTargetPlatform Platform, UnrealTargetConfiguration Configuration, List<string> UPLScripts, VersionNumber SdkVersion, bool bCreateStubIPA, string BundleID)
 		{
 			string SubDir = GetTargetPlatformName();
 
@@ -1176,13 +1180,13 @@ namespace UnrealBuildTool
 
 			if (BuildHostPlatform.Current.Platform == UnrealTargetPlatform.Mac && Environment.GetEnvironmentVariable("UBT_NO_POST_DEPLOY") != "true")
 			{
-				return PrepForUATPackageOrDeploy(Configuration, ProjectFile, GameName, ProjectDirectory, BuildPath + "/" + DecoratedGameName, "../../Engine", bForDistribution, "", false, bCreateStubIPA, UPLScripts, SdkVersion);
+				return PrepForUATPackageOrDeploy(Configuration, ProjectFile, GameName, ProjectDirectory, BuildPath + "/" + DecoratedGameName, "../../Engine", bForDistribution, "", false, bCreateStubIPA, UPLScripts, SdkVersion, BundleID);
 			}
 			else
 			{
 				// @todo tvos merge: This used to copy the bundle back - where did that code go? It needs to be fixed up for TVOS directories
 				bool bSupportPortrait, bSupportLandscape, bSkipIcons;
-				GeneratePList(ProjectFile, Configuration, ProjectDirectory, bIsUE4Game, GameName, (ProjectFile == null) ? "" : Path.GetFileNameWithoutExtension(ProjectFile.FullName), "../../Engine", "", UPLScripts, SdkVersion, out bSupportPortrait, out bSupportLandscape, out bSkipIcons);
+				GeneratePList(ProjectFile, Configuration, ProjectDirectory, bIsUE4Game, GameName, (ProjectFile == null) ? "" : Path.GetFileNameWithoutExtension(ProjectFile.FullName), "../../Engine", "", UPLScripts, SdkVersion, BundleID, out bSupportPortrait, out bSupportLandscape, out bSkipIcons);
 			}
 			return true;
 		}
diff --git a/Engine/Source/Programs/UnrealBuildTool/Platform/TVOS/UEBuildTVOS.cs b/Engine/Source/Programs/UnrealBuildTool/Platform/TVOS/UEBuildTVOS.cs
--- a/Engine/Source/Programs/UnrealBuildTool/Platform/TVOS/UEBuildTVOS.cs
+++ b/Engine/Source/Programs/UnrealBuildTool/Platform/TVOS/UEBuildTVOS.cs
@@ -28,8 +28,8 @@ namespace UnrealBuildTool
 			get { return "3"; }
 		}
 
-		public TVOSProjectSettings(FileReference ProjectFile)
-			: base(ProjectFile, UnrealTargetPlatform.TVOS)
+		public TVOSProjectSettings(FileReference ProjectFile, String Bundle)
+			: base(ProjectFile, UnrealTargetPlatform.TVOS, Bundle)
 		{
 		}
 	}
@@ -71,14 +71,14 @@ namespace UnrealBuildTool
 			}
 		}
 
-		public new TVOSProjectSettings ReadProjectSettings(FileReference ProjectFile)
+		public new TVOSProjectSettings ReadProjectSettings(FileReference ProjectFile, string Bundle = "")
 		{
-			return (TVOSProjectSettings)base.ReadProjectSettings(ProjectFile);
+			return (TVOSProjectSettings)base.ReadProjectSettings(ProjectFile, Bundle);
 		}
 
-		protected override IOSProjectSettings CreateProjectSettings(FileReference ProjectFile)
+		protected override IOSProjectSettings CreateProjectSettings(FileReference ProjectFile, string Bundle)
 		{
-			return new TVOSProjectSettings(ProjectFile);
+			return new TVOSProjectSettings(ProjectFile, Bundle);
 		}
 
 		public TVOSProvisioningData ReadProvisioningData(TVOSProjectSettings ProjectSettings, bool bForDistribution = false)
diff --git a/Engine/Source/Programs/UnrealBuildTool/Platform/TVOS/UEDeployTVOS.cs b/Engine/Source/Programs/UnrealBuildTool/Platform/TVOS/UEDeployTVOS.cs
--- a/Engine/Source/Programs/UnrealBuildTool/Platform/TVOS/UEDeployTVOS.cs
+++ b/Engine/Source/Programs/UnrealBuildTool/Platform/TVOS/UEDeployTVOS.cs
@@ -24,7 +24,7 @@ namespace UnrealBuildTool
 			return "TVOS";
 		}
 
-		public static bool GenerateTVOSPList(string ProjectDirectory, bool bIsUE4Game, string GameName, string ProjectName, string InEngineDir, string AppDirectory, UnrealPluginLanguage UPL)
+		public static bool GenerateTVOSPList(string ProjectDirectory, bool bIsUE4Game, string GameName, string ProjectName, string InEngineDir, string AppDirectory, UnrealPluginLanguage UPL, string BundleID)
 		{
 			// @todo tvos: THIS!
 
@@ -61,6 +61,10 @@ namespace UnrealBuildTool
 			// bundle identifier
 			string BundleIdentifier;
 			Ini.GetString("/Script/IOSRuntimeSettings.IOSRuntimeSettings", "BundleIdentifier", out BundleIdentifier);
+			if (!string.IsNullOrEmpty(BundleID))
+			{
+				BundleIdentifier = BundleID;
+			}
 
 			// bundle name
 			string BundleName;
@@ -331,11 +335,11 @@ namespace UnrealBuildTool
 			return bSkipDefaultPNGs;
 		}
 
-		public override bool GeneratePList(FileReference ProjectFile, UnrealTargetConfiguration Config, string ProjectDirectory, bool bIsUE4Game, string GameName, string ProjectName, string InEngineDir, string AppDirectory, List<string> UPLScripts, VersionNumber SdkVersion, out bool bSupportsPortrait, out bool bSupportsLandscape, out bool bSkipIcons)
+		public override bool GeneratePList(FileReference ProjectFile, UnrealTargetConfiguration Config, string ProjectDirectory, bool bIsUE4Game, string GameName, string ProjectName, string InEngineDir, string AppDirectory, List<string> UPLScripts, VersionNumber SdkVersion, string BundleID, out bool bSupportsPortrait, out bool bSupportsLandscape, out bool bSkipIcons)
 		{
 			bSupportsLandscape = bSupportsPortrait = true;
             bSkipIcons = true;
-			return GenerateTVOSPList(ProjectDirectory, bIsUE4Game, GameName, ProjectName, InEngineDir, AppDirectory, null);
+			return GenerateTVOSPList(ProjectDirectory, bIsUE4Game, GameName, ProjectName, InEngineDir, AppDirectory, null, BundleID);
 		}
 
         protected override void CopyGraphicsResources(bool bSkipDefaultPNGs, bool bSkipIcons, string InEngineDir, string AppDirectory, string BuildDirectory, string IntermediateDir, bool bSupportsPortrait, bool bSupportsLandscape)
diff --git a/Engine/Source/Programs/UnrealBuildTool/ProjectFiles/Xcode/XcodeProject.cs b/Engine/Source/Programs/UnrealBuildTool/ProjectFiles/Xcode/XcodeProject.cs
--- a/Engine/Source/Programs/UnrealBuildTool/ProjectFiles/Xcode/XcodeProject.cs
+++ b/Engine/Source/Programs/UnrealBuildTool/ProjectFiles/Xcode/XcodeProject.cs
@@ -134,11 +134,13 @@ namespace UnrealBuildTool
 		/// <param name="InitFilePath">The path to the project file on disk</param>
 		/// <param name="InOnlyGameProject"></param>
 		/// <param name="IsForDistribution">True for distribution builds</param>
-		public XcodeProjectFile(FileReference InitFilePath, FileReference InOnlyGameProject, bool IsForDistribution)
+		/// <param name="BundleID">Override option for bundle identifier</param>
+		public XcodeProjectFile(FileReference InitFilePath, FileReference InOnlyGameProject, bool IsForDistribution, string BundleID)
 			: base(InitFilePath)
 		{
 			OnlyGameProject = InOnlyGameProject;
 			bForDistribution = IsForDistribution;
+			BundleIdentifier = BundleID;
 		}
 
 		public override string ToString()
@@ -151,6 +153,11 @@ namespace UnrealBuildTool
 		/// </summary>
 		bool bForDistribution = false;
 
+		/// <summary>
+		/// Override for bundle identifier
+		/// </summary>
+		string BundleIdentifier = "";
+
 		/// <summary>
 		/// Gets Xcode file category based on its extension
 		/// </summary>
@@ -1245,12 +1252,12 @@ namespace UnrealBuildTool
 							TargetReceipt Receipt;
 							TargetReceipt.TryRead(ReceiptFilename, out Receipt);
 							VersionNumber SdkVersion = UEDeployIOS.GetSdkVersion(Receipt);
-							UEDeployIOS.GenerateIOSPList(ProjectFile, Config.BuildConfig, ProjectPath.FullName, bIsUE4Game, GameName, Config.BuildTarget, EngineDir.FullName, ProjectPath + "/Binaries/IOS/Payload", SdkVersion, null, out bSupportPortrait, out bSupportLandscape, out bSkipIcons);
+							UEDeployIOS.GenerateIOSPList(ProjectFile, Config.BuildConfig, ProjectPath.FullName, bIsUE4Game, GameName, Config.BuildTarget, EngineDir.FullName, ProjectPath + "/Binaries/IOS/Payload", SdkVersion, null, BundleIdentifier, out bSupportPortrait, out bSupportLandscape, out bSkipIcons);
 						}
 						if (bCreateTVOSInfoPlist)
 						{
 							Directory.CreateDirectory(Path.GetDirectoryName(TVOSInfoPlistPath));
-							UEDeployTVOS.GenerateTVOSPList(ProjectPath.FullName, bIsUE4Game, GameName, Config.BuildTarget, EngineDir.FullName, ProjectPath + "/Binaries/TVOS/Payload", null);
+							UEDeployTVOS.GenerateTVOSPList(ProjectPath.FullName, bIsUE4Game, GameName, Config.BuildTarget, EngineDir.FullName, ProjectPath + "/Binaries/TVOS/Payload", null, BundleIdentifier);
 						}
 					}
 				}
diff --git a/Engine/Source/Programs/UnrealBuildTool/ProjectFiles/Xcode/XcodeProjectFileGenerator.cs b/Engine/Source/Programs/UnrealBuildTool/ProjectFiles/Xcode/XcodeProjectFileGenerator.cs
--- a/Engine/Source/Programs/UnrealBuildTool/ProjectFiles/Xcode/XcodeProjectFileGenerator.cs
+++ b/Engine/Source/Programs/UnrealBuildTool/ProjectFiles/Xcode/XcodeProjectFileGenerator.cs
@@ -32,6 +32,11 @@ namespace UnrealBuildTool
 		/// </summary>
 		bool bForDistribution = false;
 
+		/// <summary>
+		/// Override BundleID
+		/// </summary>
+		string BundleIdentifier = "";
+
 		public XcodeProjectFileGenerator(FileReference InOnlyGameProject, CommandLineArguments CommandLine)
 			: base(InOnlyGameProject)
 		{
@@ -39,6 +44,10 @@ namespace UnrealBuildTool
 			{
 				bForDistribution = true;
 			}
+			if (CommandLine.HasValue("-bundleID="))
+			{
+				BundleIdentifier = CommandLine.GetString("-bundleID=");
+			}
 		}
 
 		/// <summary>
@@ -99,7 +108,7 @@ namespace UnrealBuildTool
 		/// <returns>The newly allocated project file object</returns>
 		protected override ProjectFile AllocateProjectFile(FileReference InitFilePath)
 		{
-			return new XcodeProjectFile(InitFilePath, OnlyGameProject, bForDistribution);
+			return new XcodeProjectFile(InitFilePath, OnlyGameProject, bForDistribution, BundleIdentifier);
 		}
 
 		/// ProjectFileGenerator interface
diff --git a/Engine/Source/Programs/UnrealBuildTool/ToolChain/RemoteMac.cs b/Engine/Source/Programs/UnrealBuildTool/ToolChain/RemoteMac.cs
--- a/Engine/Source/Programs/UnrealBuildTool/ToolChain/RemoteMac.cs
+++ b/Engine/Source/Programs/UnrealBuildTool/ToolChain/RemoteMac.cs
@@ -116,6 +116,8 @@ namespace UnrealBuildTool
 		/// </summary>
 		private List<string> CommonRsyncArguments;
 
+		private string IniBundleIdentifier = "";
+
 		/// <summary>
 		/// Constructor
 		/// </summary>
@@ -192,6 +194,8 @@ namespace UnrealBuildTool
 				}
 			}
 
+			Ini.GetString("/Script/IOSRuntimeSettings.IOSRuntimeSettings", "BundleIdentifier", out IniBundleIdentifier);
+
 			// If it's not set, look in the standard locations. If that fails, spawn the batch file to generate one.
 			if (SshPrivateKey == null && !TryGetSshPrivateKey(out SshPrivateKey))
 			{
@@ -430,7 +434,7 @@ namespace UnrealBuildTool
 				RemoteArguments.Add("-NoUBTMakefiles");
 
 				// Get the provisioning data for this project
-				IOSProvisioningData ProvisioningData = ((IOSPlatform)UEBuildPlatform.GetBuildPlatform(TargetDesc.Platform)).ReadProvisioningData(TargetDesc.ProjectFile, TargetDesc.AdditionalArguments.HasOption("-distribution"));
+				IOSProvisioningData ProvisioningData = ((IOSPlatform)UEBuildPlatform.GetBuildPlatform(TargetDesc.Platform)).ReadProvisioningData(TargetDesc.ProjectFile, TargetDesc.AdditionalArguments.HasOption("-distribution"), IniBundleIdentifier);
 				if(ProvisioningData == null || ProvisioningData.MobileProvisionFile == null)
 				{
 					throw new BuildException("Unable to find mobile provision for {0}. See log for more information.", TargetDesc.Name);
