Reinstate the warning about short keys but wrap it in the condition that was lost in the merge from Release-4.22

diff --git a/Engine/Source/Programs/UnrealBuildTool/System/EncryptionAndSigning.cs b/Engine/Source/Programs/UnrealBuildTool/System/EncryptionAndSigning.cs
--- a/Engine/Source/Programs/UnrealBuildTool/System/EncryptionAndSigning.cs
+++ b/Engine/Source/Programs/UnrealBuildTool/System/EncryptionAndSigning.cs
@@ -436,8 +436,10 @@ namespace UnrealBuildTool
 				}
 			}
 
+			// Check if we have a valid signing key that is of the old short form
+			if (Settings.SigningKey != null && Settings.SigningKey.IsValid() && Settings.SigningKey.IsUnsecureLegacyKey())
 			{
-				//Log.TraceWarningOnce("Project signing keys found in '{0}' are of the old insecure short format. Please regenerate them using the project crypto settings panel in the editor!", InProjectDirectory);
+				Log.TraceWarningOnce("Project signing keys found in '{0}' are of the old insecure short format. Please regenerate them using the project crypto settings panel in the editor!", InProjectDirectory);
 			}
 
 			// Validate the settings we have read
