- Fixed dsymutil version check for older builds

diff --git a/Engine/Source/Programs/UnrealBuildTool/ToolChain/AppleToolChain.cs b/Engine/Source/Programs/UnrealBuildTool/ToolChain/AppleToolChain.cs
--- a/Engine/Source/Programs/UnrealBuildTool/ToolChain/AppleToolChain.cs
+++ b/Engine/Source/Programs/UnrealBuildTool/ToolChain/AppleToolChain.cs
@@ -165,22 +165,66 @@ namespace UnrealBuildTool
 		{
 			FileReference DsymutilLocation = new FileReference("/usr/bin/dsymutil");
 
-			// Default dsymutil has a bug that causes issues, it's fixed in autosdks but not everyone has those set up so for the timebeing we have
-			// a version in P4...
-			FileReference PatchedDsymutilLocation = FileReference.Combine(UnrealBuildTool.EngineDirectory, "Binaries/Mac/NotForLicensees/LLVM/bin/dsymutil");
+			// dsymutil before 10.0.1 has a bug that causes issues, it's fixed in autosdks but not everyone has those set up so for the timebeing we have
+			// a version in P4 - first determine if we need it
+			string DsymutilVersionString = Utils.RunLocalProcessAndReturnStdOut(DsymutilLocation.FullName, "-version");
+
+			bool bUseInstalledDsymutil = true;
+
+			// tease out the version number
+			string[] Tokens = DsymutilVersionString.Split(" ".ToCharArray());
 			
-			if (File.Exists(PatchedDsymutilLocation.FullName))
+			// sanity check
+			if (Tokens.Length < 4 || Tokens[3].Contains(".") == false)
 			{
-				DsymutilLocation = PatchedDsymutilLocation;
+				Log.TraceInformationOnce("Unable to parse dsymutil version out of: {0}", DsymutilVersionString);
+
+				bUseInstalledDsymutil = false;
 			}
-			
-			DirectoryReference AutoSdkDir;
-			if (UEBuildPlatformSDK.TryGetHostPlatformAutoSDKDir(out AutoSdkDir))
+			else
+			{
+				string[] Versions = Tokens[3].Split(".".ToCharArray());
+				if (Versions.Length < 3)
+				{
+					Log.TraceInformationOnce("Unable to parse version token: {0}", Tokens[3]);
+				}
+				else
+				{
+					int Major, Minor, Patch;
+					if (!int.TryParse(Versions[0], out Major) || !int.TryParse(Versions[1], out Minor) || !int.TryParse(Versions[2], out Patch))
+					{
+						Log.TraceInformationOnce("Unable to parse version tokens: {0}", Tokens[3]);
+					}
+					else
+					{
+						Log.TraceInformationOnce("Parsed dsymutil version as {0}.{1}.{2}", Major, Minor, Patch);
+
+						if (Major < 10 || (Minor == 0 && Patch == 0))
+						{
+							bUseInstalledDsymutil = false;
+						}
+					}
+				}
+			}
+
+			// if the installed one is too old, use a fixed up one if it can
+			if (bUseInstalledDsymutil == false)
 			{
-				FileReference AutoSdkDsymutilLocation = FileReference.Combine(AutoSdkDir, "Mac", "LLVM", "bin", "dsymutil");
-				if (FileReference.Exists(AutoSdkDsymutilLocation))
+				FileReference PatchedDsymutilLocation = FileReference.Combine(UnrealBuildTool.EngineDirectory, "Binaries/Mac/NotForLicensees/LLVM/bin/dsymutil");
+
+				if (File.Exists(PatchedDsymutilLocation.FullName))
+				{
+					DsymutilLocation = PatchedDsymutilLocation;
+				}
+
+				DirectoryReference AutoSdkDir;
+				if (UEBuildPlatformSDK.TryGetHostPlatformAutoSDKDir(out AutoSdkDir))
 				{
-					DsymutilLocation = AutoSdkDsymutilLocation;
+					FileReference AutoSdkDsymutilLocation = FileReference.Combine(AutoSdkDir, "Mac", "LLVM", "bin", "dsymutil");
+					if (FileReference.Exists(AutoSdkDsymutilLocation))
+					{
+						DsymutilLocation = AutoSdkDsymutilLocation;
+					}
 				}
 			}
 
