Set target sdk properties on subprojects to prevent unwanted permissions being added automatically

#ue4
#android

diff --git a/Engine/Source/Programs/UnrealBuildTool/Platform/Android/UEDeployAndroid.cs b/Engine/Source/Programs/UnrealBuildTool/Platform/Android/UEDeployAndroid.cs
--- a/Engine/Source/Programs/UnrealBuildTool/Platform/Android/UEDeployAndroid.cs
+++ b/Engine/Source/Programs/UnrealBuildTool/Platform/Android/UEDeployAndroid.cs
@@ -3015,18 +3015,11 @@ namespace UnrealBuildTool
 			return true;
 		}
 
-		private void CreateGradlePropertiesFiles(AndroidToolChain ToolChain, string Arch, string CompileSDKVersion, string BuildToolsVersion, string PackageName, string DestApkName, string NDKArch,
-			string UE4BuildFilesPath, string GameBuildFilesPath, string UE4BuildGradleAppPath, string UE4BuildPath, string UE4BuildGradlePath, bool bForDistribution)
+		private void GetMinTargetSDKVersions(AndroidToolChain ToolChain, string Arch, out int MinSDKVersion, out int TargetSDKVersion)
 		{
-			// Create gradle.properties
-			StringBuilder GradleProperties = new StringBuilder();
-
-			int StoreVersion = GetStoreVersion();
-			string VersionDisplayName = GetVersionDisplayName();
-			int MinSDKVersion;
 			ConfigHierarchy Ini = GetConfigCacheIni(ConfigHierarchyType.Engine);
 			Ini.GetInt32("/Script/AndroidRuntimeSettings.AndroidRuntimeSettings", "MinSDKVersion", out MinSDKVersion);
-			int TargetSDKVersion = MinSDKVersion;
+			TargetSDKVersion = MinSDKVersion;
 			Ini.GetInt32("/Script/AndroidRuntimeSettings.AndroidRuntimeSettings", "TargetSDKVersion", out TargetSDKVersion);
 
 			// Make sure minSdkVersion is at least 13 (need this for appcompat-v13 used by AndroidPermissions)
@@ -3059,6 +3052,18 @@ namespace UnrealBuildTool
 					Log.TraceInformation("Fixing minSdkVersion; NDK level above 19 requires minSdkVersion of 21 (arch={0})", Arch.Substring(1));
 				}
 			}
+		}
+		
+		private void CreateGradlePropertiesFiles(string Arch, int MinSDKVersion, int TargetSDKVersion, string CompileSDKVersion, string BuildToolsVersion, string PackageName,
+			string DestApkName, string NDKArch,	string UE4BuildFilesPath, string GameBuildFilesPath, string UE4BuildGradleAppPath, string UE4BuildPath, string UE4BuildGradlePath, bool bForDistribution)
+		{
+			// Create gradle.properties
+			StringBuilder GradleProperties = new StringBuilder();
+
+			int StoreVersion = GetStoreVersion();
+			string VersionDisplayName = GetVersionDisplayName();
+
+			ConfigHierarchy Ini = GetConfigCacheIni(ConfigHierarchyType.Engine);
 
 			GradleProperties.AppendLine("org.gradle.daemon=false");
 			GradleProperties.AppendLine("org.gradle.jvmargs=-XX:MaxHeapSize=4096m -Xmx9216m");
@@ -3317,7 +3322,7 @@ namespace UnrealBuildTool
 			string UE4OBBDataFileName = GetUE4JavaOBBDataFileName(TemplateDestinationBase);
 			string UE4DownloadShimFileName = GetUE4JavaDownloadShimFileName(UE4JavaFilePath);
 
-			// Template generated files           
+			// Template generated files
 			string JavaTemplateSourceDir = GetUE4TemplateJavaSourceDir(EngineDirectory);
 			IEnumerable<TemplateFile> templates = from template in Directory.EnumerateFiles(JavaTemplateSourceDir, "*.template")
 							let RealName = Path.GetFileNameWithoutExtension(template)
@@ -3927,9 +3932,14 @@ namespace UnrealBuildTool
 					// do any plugin requested copies
 					UPL.ProcessPluginNode(NDKArch, "gradleCopies", "");
 
+					// get min and target SDK versions
+					int MinSDKVersion = 0;
+					int TargetSDKVersion = 0;
+					GetMinTargetSDKVersions(ToolChain, Arch, out MinSDKVersion, out TargetSDKVersion);
+					
 					// move JavaLibs into subprojects
 					string JavaLibsDir = Path.Combine(UE4BuildPath, "JavaLibs");
-					PrepareJavaLibsForGradle(JavaLibsDir, UE4BuildGradlePath, CompileSDKVersion, BuildToolsVersion);
+					PrepareJavaLibsForGradle(JavaLibsDir, UE4BuildGradlePath, MinSDKVersion.ToString(), TargetSDKVersion.ToString(), CompileSDKVersion, BuildToolsVersion);
 
 					// Create local.properties
 					String LocalPropertiesFilename = Path.Combine(UE4BuildGradlePath, "local.properties");
@@ -3938,7 +3948,7 @@ namespace UnrealBuildTool
 					LocalProperties.AppendLine(string.Format("sdk.dir={0}", Environment.GetEnvironmentVariable("ANDROID_HOME").Replace("\\", "/")));
 					File.WriteAllText(LocalPropertiesFilename, LocalProperties.ToString());
 
-					CreateGradlePropertiesFiles(ToolChain, Arch, CompileSDKVersion, BuildToolsVersion, PackageName, DestApkName, NDKArch,
+					CreateGradlePropertiesFiles(Arch, MinSDKVersion, TargetSDKVersion, CompileSDKVersion, BuildToolsVersion, PackageName, DestApkName, NDKArch,
 						UE4BuildFilesPath, GameBuildFilesPath, UE4BuildGradleAppPath, UE4BuildPath, UE4BuildGradlePath, bForDistribution);
 
 					if (!bSkipGradleBuild)
@@ -4254,7 +4264,7 @@ namespace UnrealBuildTool
 											"\t\t}\n";
 					break;
 			}
-			
+
 			Dictionary<string, string> Replacements = new Dictionary<string, string>{
 				{ "//$${gameActivityImportAdditions}$$", UPL.ProcessPluginNode(NDKArch, "gameActivityImportAdditions", "")},
 				{ "//$${gameActivityPostImportAdditions}$$", UPL.ProcessPluginNode(NDKArch, "gameActivityPostImportAdditions", "")},
@@ -4467,7 +4477,7 @@ namespace UnrealBuildTool
 			return AARHandler;
 		}
 
-		private void PrepareJavaLibsForGradle(string JavaLibsDir, string UE4BuildGradlePath, string CompileSDKVersion, string BuildToolsVersion)
+		private void PrepareJavaLibsForGradle(string JavaLibsDir, string UE4BuildGradlePath, string InMinSdkVersion, string InTargetSdkVersion, string CompileSDKVersion, string BuildToolsVersion)
 		{
 			StringBuilder SettingsGradleContent = new StringBuilder();
 			StringBuilder ProjectDependencyContent = new StringBuilder();
@@ -4517,8 +4527,8 @@ namespace UnrealBuildTool
 				// Try to get the SDK target from the AndroidManifest.xml
 				string VersionCode = "";
 				string VersionName = "";
-				string MinSdkVersion = "";
-				string TargetSdkVersion = "";
+				string MinSdkVersion = InMinSdkVersion;
+				string TargetSdkVersion = InTargetSdkVersion;
 				XDocument ManifestXML;
 				if (File.Exists(ManifestFilename))
 				{
