Use forward slashes in include paths.

#rnx

diff --git a/Engine/Source/Programs/UnrealBuildTool/Configuration/UEBuildModuleCPP.cs b/Engine/Source/Programs/UnrealBuildTool/Configuration/UEBuildModuleCPP.cs
--- a/Engine/Source/Programs/UnrealBuildTool/Configuration/UEBuildModuleCPP.cs
+++ b/Engine/Source/Programs/UnrealBuildTool/Configuration/UEBuildModuleCPP.cs
@@ -918,7 +918,7 @@ namespace UnrealBuildTool
 			{
 				Writer.WriteLine("// PCH for {0}", IncludedFile.AbsolutePath);
 				WriteDefinitions(Definitions, Writer);
-				Writer.WriteLine("#include \"{0}\"", IncludedFile.AbsolutePath);
+				Writer.WriteLine("#include \"{0}\"", IncludedFile.AbsolutePath.Replace('\\', '/'));
 			}
 
 			// Create the item
diff --git a/Engine/Source/Programs/UnrealBuildTool/System/Unity.cs b/Engine/Source/Programs/UnrealBuildTool/System/Unity.cs
--- a/Engine/Source/Programs/UnrealBuildTool/System/Unity.cs
+++ b/Engine/Source/Programs/UnrealBuildTool/System/Unity.cs
@@ -316,7 +316,7 @@ namespace UnrealBuildTool
 				// Add source files to the unity file
 				foreach (FileItem CPPFile in UnityFile.Files)
 				{
-					OutputUnityCPPWriter.WriteLine("#include \"{0}\"", CPPFile.AbsolutePath);
+					OutputUnityCPPWriter.WriteLine("#include \"{0}\"", CPPFile.AbsolutePath.Replace('\\', '/'));
 				}
 
 				// Determine unity file path name
