Add gameActivityAfterMainViewCreatedAdditions to GameActivity.java.template

#android

diff --git a/Engine/Build/Android/Java/src/com/epicgames/ue4/GameActivity.java.template b/Engine/Build/Android/Java/src/com/epicgames/ue4/GameActivity.java.template
--- a/Engine/Build/Android/Java/src/com/epicgames/ue4/GameActivity.java.template
+++ b/Engine/Build/Android/Java/src/com/epicgames/ue4/GameActivity.java.template
@@ -2866,6 +2866,8 @@ public class GameActivity extends $${gameActivitySuperClass}$$ implements Surfac
         	}
         });
 
+//$${gameActivityAfterMainViewCreatedAdditions}$$
+
 		clipboardManager = (ClipboardManager) getSystemService(CLIPBOARD_SERVICE);
 
 //$${gameActivityOnCreateAdditions}$$
diff --git a/Engine/Source/Programs/UnrealBuildTool/Platform/Android/UEDeployAndroid.cs b/Engine/Source/Programs/UnrealBuildTool/Platform/Android/UEDeployAndroid.cs
--- a/Engine/Source/Programs/UnrealBuildTool/Platform/Android/UEDeployAndroid.cs
+++ b/Engine/Source/Programs/UnrealBuildTool/Platform/Android/UEDeployAndroid.cs
@@ -4283,6 +4283,7 @@ namespace UnrealBuildTool
   				{ "//$${gameActivityPostConfigRulesAdditions}$$", UPL.ProcessPluginNode(NDKArch, "gameActivityPostConfigRulesAdditions", "")},
   				{ "//$${gameActivityFinalizeConfigRulesAdditions}$$", UPL.ProcessPluginNode(NDKArch, "gameActivityFinalizeConfigRulesAdditions", "")},
 				{ "//$${gameActivityBeforeConfigRulesAppliedAdditions}$$", UPL.ProcessPluginNode(NDKArch, "gameActivityBeforeConfigRulesAppliedAdditions", "")},
+				{ "//$${gameActivityAfterMainViewCreatedAdditions}$$", UPL.ProcessPluginNode(NDKArch, "gameActivityAfterMainViewCreatedAdditions", "")},
 				{ "//$${gameActivityLoggerCallbackAdditions}$$", UPL.ProcessPluginNode(NDKArch, "gameActivityLoggerCallbackAdditions", "")},
 				{ "//$${soLoadLibrary}$$", UPL.ProcessPluginNode(NDKArch, "soLoadLibrary", LoadLibraryDefaults)},
 				{ "$${gameActivitySuperClass}$$", SuperClassDefault},
