fix for not properly uploading Crashlytics data

diff --git a/Engine/Source/Programs/UnrealBuildTool/Platform/IOS/IOSToolChain.cs b/Engine/Source/Programs/UnrealBuildTool/Platform/IOS/IOSToolChain.cs
--- a/Engine/Source/Programs/UnrealBuildTool/Platform/IOS/IOSToolChain.cs
+++ b/Engine/Source/Programs/UnrealBuildTool/Platform/IOS/IOSToolChain.cs
@@ -1521,7 +1521,7 @@ namespace UnrealBuildTool
 		private static void GenerateCrashlyticsData(string ExecutableDirectory, string ExecutableName, string ProjectDir, string ProjectName)
         {
 			Log.TraceInformation("Generating and uploading Crashlytics Data");
-            string FabricPath = UnrealBuildTool.EngineDirectory + "/Intermediate/ThirdPartyFrameworks/Fabric.embeddedframework";
+            string FabricPath = UnrealBuildTool.EngineDirectory + "/Intermediate/UnzippedFrameworks/Crashlytics/Fabric.embeddedframework";
             if (Directory.Exists(FabricPath) && Environment.GetEnvironmentVariable("IsBuildMachine") == "1")
             {
 				string PlistFile = ProjectDir + "/Intermediate/IOS/" + ProjectName + "-Info.plist";
