UBT: Expose all the options for configuring the project file generator to the BuildConfiguration.xml files.

diff --git a/Engine/Source/Programs/UnrealBuildTool/ProjectFiles/ProjectFileGenerator.cs b/Engine/Source/Programs/UnrealBuildTool/ProjectFiles/ProjectFileGenerator.cs
--- a/Engine/Source/Programs/UnrealBuildTool/ProjectFiles/ProjectFileGenerator.cs
+++ b/Engine/Source/Programs/UnrealBuildTool/ProjectFiles/ProjectFileGenerator.cs
@@ -203,11 +203,13 @@ namespace UnrealBuildTool
 		/// <summary>
 		/// Whether we should include configurations for "Test" and "Shipping" in generated projects (pass "-NoShippingConfigs" to disable this)
 		/// </summary>
+		[XmlConfigFile]
 		public static bool bIncludeTestAndShippingConfigs = true;
 
 		/// <summary>
 		/// True if intellisense data should be generated (takes a while longer)
 		/// </summary>
+		[XmlConfigFile]
 		bool bGenerateIntelliSenseData = true;
 
 		/// <summary>
@@ -219,67 +221,80 @@ namespace UnrealBuildTool
 		/// <summary>
 		/// True if all documentation languages should be included in generated projects, otherwise only "INT" will be included
 		/// </summary>
+		[XmlConfigFile]
 		bool bAllDocumentationLanguages = false;
 
 		/// <summary>
 		/// True if build targets should pass the -useprecompiled argument
 		/// </summary>
+		[XmlConfigFile]
 		public bool bUsePrecompiled = false;
 
 		/// <summary>
 		/// True if we should include engine source in the generated solution
 		/// </summary>
+		[XmlConfigFile]
 		protected bool bIncludeEngineSource = true;
 
 		/// <summary>
 		/// Whether to include enterprise source in the generated solution
 		/// </summary>
+		[XmlConfigFile]
 		protected bool bIncludeEnterpriseSource = true;
 
 		/// <summary>
 		/// True if shader source files should be included in the generated projects
 		/// </summary>
+		[XmlConfigFile]
 		protected bool bIncludeShaderSource = true;
 
 		/// <summary>
 		/// True if build system files should be included
 		/// </summary>
+		[XmlConfigFile]
 		bool bIncludeBuildSystemFiles = true;
 
 		/// <summary>
 		/// True if we should include config files (.ini files) in the generated project
 		/// </summary>
+		[XmlConfigFile]
 		protected bool bIncludeConfigFiles = true;
 
 		/// <summary>
 		/// True if we should include localization files (.int/.kor/etc files) in the generated project
 		/// </summary>
+		[XmlConfigFile]
 		bool bIncludeLocalizationFiles = false;
 
 		/// <summary>
 		/// True if we should include template files (.template files) in the generated project
 		/// </summary>
+		[XmlConfigFile]
 		protected bool bIncludeTemplateFiles = true;
 
 		/// <summary>
 		/// True if we should include program projects in the generated solution
 		/// </summary>
+		[XmlConfigFile]
 		protected bool IncludeEnginePrograms = true;
 
 		/// <summary>
 		/// Whether to include temporary targets generated by UAT to support content only projects with non-default settings
 		/// </summary>
+		[XmlConfigFile]
 		bool bIncludeTempTargets = false;
 
 		/// <summary>
 		/// True if we should include .NET Core projects in the generated solution
 		/// </summary>
+		[XmlConfigFile]
 		bool bIncludeDotNETCoreProjects = false;
 
 		/// <summary>
 		/// True if we should reflect "Source" sub-directories on disk in the master project as master project directories.
 		/// This arguably adds some visual clutter to the master project, but is truer to the on-disk file organization.
 		/// </summary>
+		[XmlConfigFile]
 		bool bKeepSourceSubDirectories = true;
 
 		/// <summary>
