UE-71953 - Blueprint iOS projects can't package For Distribution

diff --git a/Engine/Source/Programs/UnrealBuildTool/Platform/IOS/IOSToolChain.cs b/Engine/Source/Programs/UnrealBuildTool/Platform/IOS/IOSToolChain.cs
--- a/Engine/Source/Programs/UnrealBuildTool/Platform/IOS/IOSToolChain.cs
+++ b/Engine/Source/Programs/UnrealBuildTool/Platform/IOS/IOSToolChain.cs
@@ -737,9 +737,9 @@ namespace UnrealBuildTool
 				CompileAction.bShouldOutputStatusDescription = true;
 
 				foreach(UEBuildFramework Framework in CompileEnvironment.AdditionalFrameworks)
-		{
-					if(Framework.ZipFile != null)
 			{
+					if(Framework.ZipFile != null)
+				{
 						FileItem ExtractedTokenFile = ExtractFramework(Framework, Actions);
 						CompileAction.PrerequisiteItems.Add(ExtractedTokenFile);
 				}
@@ -1349,10 +1349,10 @@ namespace UnrealBuildTool
         {
             List<FileItem> OutputFiles = new List<FileItem>(base.PostBuild(Executable, BinaryLinkEnvironment, Actions));
 
-			if (BinaryLinkEnvironment.bIsBuildingLibrary)
-			{
-				return OutputFiles;
-			}
+            if (BinaryLinkEnvironment.bIsBuildingLibrary)
+            {
+                return OutputFiles;
+            }
 
             // For IOS/tvOS, generate the dSYM file if the config file is set to do so
 			if (ProjectSettings.bGeneratedSYMFile == true || ProjectSettings.bGeneratedSYMBundle == true || BinaryLinkEnvironment.bUsePDBFiles == true)
@@ -1563,7 +1563,7 @@ namespace UnrealBuildTool
 			{
 				ProjectFileGenerator.bGenerateProjectFiles = false;
 			}
-		}
+			}
 
 		public static FileReference GetStagedExecutablePath(FileReference Executable, string TargetName)
 		{
@@ -1717,11 +1717,11 @@ namespace UnrealBuildTool
 		}
 
         public static void PostBuildSync(IOSPostBuildSyncTarget Target)
-		{
+			{
 			ConfigHierarchy Ini = ConfigCache.ReadHierarchy(ConfigHierarchyType.Engine, DirectoryReference.FromFile(Target.ProjectFile), UnrealTargetPlatform.IOS);
 			string BundleID;
 			Ini.GetString("/Script/IOSRuntimeSettings.IOSRuntimeSettings", "BundleIdentifier", out BundleID);
-			
+
 			IOSProjectSettings ProjectSettings = ((IOSPlatform)UEBuildPlatform.GetBuildPlatform(Target.Platform)).ReadProjectSettings(Target.ProjectFile);
 
 			if (BuildHostPlatform.Current.Platform == UnrealTargetPlatform.Mac)
@@ -1986,7 +1986,7 @@ namespace UnrealBuildTool
 				CleanIntermediateDirectory(LocalFrameworkAssets);
 
 				foreach (KeyValuePair<string, DirectoryReference> Pair in Target.FrameworkNameToSourceDir)
-				{
+					{
 					string UnpackedZipPath = Pair.Value.FullName;
 
 					// For now, this is hard coded, but we need to loop over all modules, and copy bundled assets that need it
@@ -1995,7 +1995,7 @@ namespace UnrealBuildTool
 					Log.TraceInformation("Copying bundled asset... LocalSource: {0}, LocalDest: {1}", Pair.Value, LocalDest);
 
 					string ResultsText;
-					RunExecutableAndWait("cp", String.Format("-R -L \"{0}\" \"{1}\"", Pair.Value, LocalDest), out ResultsText);
+                    RunExecutableAndWait("cp", String.Format("-R -L \"{0}\" \"{1}\"", Pair.Value, LocalDest), out ResultsText);
                 }
             }
 		}
