UE-71919 tvOS banner and App thumbnail do not appear on device

#4.22.1
#tvOS

New structure for the tvOS assets (Build/TvOS/Resources) needs changes in the Unreal Build Tools

Note: This will also solve the errors when uploading shipping tvOS packages to Apple

diff --git a/Engine/Source/Programs/UnrealBuildTool/Platform/IOS/IOSToolChain.cs b/Engine/Source/Programs/UnrealBuildTool/Platform/IOS/IOSToolChain.cs
--- a/Engine/Source/Programs/UnrealBuildTool/Platform/IOS/IOSToolChain.cs
+++ b/Engine/Source/Programs/UnrealBuildTool/Platform/IOS/IOSToolChain.cs
@@ -929,10 +929,10 @@ namespace UnrealBuildTool
 			Arguments.Append(" --notices");
 			Arguments.Append(" --warnings");
             Arguments.AppendFormat(" --output-partial-info-plist '{0}/assetcatalog_generated_info.plist'", InputDir);
-            Arguments.Append(" --app-icon AppIcon");
 			if(Platform == CppPlatform.TVOS)
 			{
-				Arguments.Append(" --launch-image LaunchImage");
+				Arguments.Append(" --app-icon 'App Icon & Top Shelf Image'");
+				Arguments.Append(" --launch-image 'Launch Image'");
 				Arguments.Append(" --filter-for-device-model AppleTV5,3");
 				//Arguments.Append(" --filter-for-device-os-version 10.0");
 				Arguments.Append(" --target-device tv");
@@ -941,6 +941,7 @@ namespace UnrealBuildTool
 			}
 			else
 			{
+				Arguments.Append(" --app-icon AppIcon");
 				Arguments.Append(" --product-type com.apple.product-type.application");
 				Arguments.Append(" --target-device iphone");
 				Arguments.Append(" --target-device ipad");
@@ -1258,10 +1259,14 @@ namespace UnrealBuildTool
 				};
 				Dir = Path.Combine(IntermediateDir, "Resources", "Assets.xcassets");
 
-				string BuildResourcesGraphicsDir = Path.Combine(BuildDir, "Resources", "Graphics");
+				string BuildResourcesGraphicsDir = Path.Combine(BuildDir, "Resources", "Assets.xcassets");
 				for (int Index = 0; Index < Images.Length; ++Index)
 				{
-					string Image = Path.Combine((Directory.Exists(Path.Combine(BuildDir, "Resources", "Graphics")) ? (BuildDir) : (Path.Combine(EngineDir, "Build", "TVOS"))), "Resources", "Graphics", Images[Index][0]);
+					string SourceDir = Path.Combine((Directory.Exists(BuildResourcesGraphicsDir) ? (BuildDir) : (Path.Combine(EngineDir, "Build", "TVOS"))),
+						"Resources",
+						"Assets.xcassets");
+					string Image = Path.Combine(SourceDir, Images[Index][1], Images[Index][0]);
+
 					if (File.Exists(Image))
 					{
 						bUserImagesExist |= Image.StartsWith(BuildResourcesGraphicsDir);
diff --git a/Engine/Source/Programs/UnrealBuildTool/Platform/TVOS/UEDeployTVOS.cs b/Engine/Source/Programs/UnrealBuildTool/Platform/TVOS/UEDeployTVOS.cs
--- a/Engine/Source/Programs/UnrealBuildTool/Platform/TVOS/UEDeployTVOS.cs
+++ b/Engine/Source/Programs/UnrealBuildTool/Platform/TVOS/UEDeployTVOS.cs
@@ -157,7 +157,7 @@ namespace UnrealBuildTool
             Text.AppendLine("\t\t\t<key>UILaunchImageSize</key>");
             Text.AppendLine("\t\t\t<string>{1920, 1080}</string>");
             Text.AppendLine("\t\t\t<key>UILaunchImageName</key>");
-            Text.AppendLine("\t\t\t<string>LaunchImage</string>");
+            Text.AppendLine("\t\t\t<string>Launch Image</string>");
             Text.AppendLine("\t\t\t<key>UILaunchImageMinimumOSVersion</key>");
             Text.AppendLine("\t\t\t<string>9.0</string>");
             Text.AppendLine("\t\t\t<key>UILaunchImageOrientation</key>");
@@ -167,7 +167,7 @@ namespace UnrealBuildTool
             Text.AppendLine("\t<key>CFBundleIcons</key>");
             Text.AppendLine("\t<dict>");
             Text.AppendLine("\t\t<key>CFBundlePrimaryIcon</key>");
-            Text.AppendLine("\t\t<string>App Icon - Small</string>");
+            Text.AppendLine("\t\t<string>App Icon</string>");
             Text.AppendLine("\t</dict>");
 
             /*			Text.AppendLine("\t<key>CFBundleIcons</key>");
@@ -273,7 +273,7 @@ namespace UnrealBuildTool
                         Text.AppendLine("\t\t\t<key>UILaunchImageSize</key>");
                         Text.AppendLine("\t\t\t<string>{768, 1024}</string>");
                         Text.AppendLine("\t\t</dict>");
-                        Text.AppendLine("\t</array>”);
+                        Text.AppendLine("\t</array>");
                         Text.AppendLine("\t<key>CFBundleSupportedPlatforms</key>");
                         Text.AppendLine("\t<array>");
                         Text.AppendLine("\t\t<string>iPhoneOS</string>");
