UE-72058 tvOS: Project Packaged with CloudKit Support Crashes on Launch Due to Malformed Value in the Entitlement

#tvOS
#4.22.1

Manually integrated (robomerge errors) from CL5731997

When building on Mac, the packaging &  signing rely on the entitlements file.

- the values should correspond to those in the entitlements (take into account syntax differences). Fixed problems related to icloud-container-identifiers, ubiquity-container-identifiers and com.apple.developer.icloud-services
- ICloudContainerIdentifier added to the tvOS plist (previously it was added on iOS)
- added com.apple.developer.icloud-container-environment to the entitlements file, needed when signing the package
- shipping build with CloudKit should also have the push notifications enabled
- fixed an issue with com.apple.developer.associated-domains when the value is a string, not an array
- addding CODE_SIGN_ENTITLEMENTS in the project file

diff --git a/Engine/Source/Programs/IOS/iPhonePackager/MobileProvisionUtilities.cs b/Engine/Source/Programs/IOS/iPhonePackager/MobileProvisionUtilities.cs
--- a/Engine/Source/Programs/IOS/iPhonePackager/MobileProvisionUtilities.cs
+++ b/Engine/Source/Programs/IOS/iPhonePackager/MobileProvisionUtilities.cs
@@ -375,14 +375,20 @@ namespace iPhonePackager
 					// remove the wildcards from the developer.associated-domains array or string
 					if (XCentPList.HasKey("com.apple.developer.associated-domains"))
 					{
-						List<string> AssociatedDomainsGroup = XCentPList.GetArray("com.apple.developer.associated-domains", "string");
+						string AssociatedDomainsString;
+						XCentPList.GetString("com.apple.developer.associated-domains", out AssociatedDomainsString);
 
-						if (AssociatedDomainsGroup.Count == 0 || AssociatedDomainsGroup[0].Contains("*"))
+						//check if the value is string
+						if (AssociatedDomainsString != null && AssociatedDomainsString.Contains("*"))
 						{
-							string AssociatedDomainsString;
-							XCentPList.GetString("com.apple.developer.associated-domains", out AssociatedDomainsString);
-
-							if (AssociatedDomainsString.Contains("*"))
+							XCentPList.RemoveKeyValue("com.apple.developer.associated-domains");
+						}
+						else
+						{
+							//check if the value is an array
+							List<string> AssociatedDomainsGroup = XCentPList.GetArray("com.apple.developer.associated-domains", "string");
+							
+							if (AssociatedDomainsGroup.Count == 1 && AssociatedDomainsGroup[0].Contains("*"))
 							{
 								XCentPList.RemoveKeyValue("com.apple.developer.associated-domains");
 							}
diff --git a/Engine/Source/Programs/UnrealBuildTool/Platform/IOS/IOSToolChain.cs b/Engine/Source/Programs/UnrealBuildTool/Platform/IOS/IOSToolChain.cs
--- a/Engine/Source/Programs/UnrealBuildTool/Platform/IOS/IOSToolChain.cs
+++ b/Engine/Source/Programs/UnrealBuildTool/Platform/IOS/IOSToolChain.cs
@@ -1597,18 +1597,30 @@ namespace UnrealBuildTool
 			string iCloudContainerIdentifiersXML = "";
 			string iCloudContainerIdentifier = "";
 			string UbiquityContainerIdentifiersXML = "";
+			string iCloudServicesXML = "";
+			
 			if (MobileProvisionFile!= null && File.Exists(MobileProvisionFile.FullName))
 			{
 				MobileProvisionContents MobileProvisionContent = MobileProvisionContents.Read(MobileProvisionFile);
 
 				iCloudContainerIdentifier = MobileProvisionContent.GetNodeValueByName("com.apple.developer.icloud-container-identifiers");
 				iCloudContainerIdentifiersXML = MobileProvisionContent.GetNodeXMLValueByName("com.apple.developer.icloud-container-identifiers");
+				iCloudServicesXML = MobileProvisionContent.GetNodeXMLValueByName("com.apple.developer.icloud-services");
 				UbiquityContainerIdentifiersXML = MobileProvisionContent.GetNodeXMLValueByName("com.apple.developer.ubiquity-container-identifiers");
 			}
+			else
+			{
+				Log.TraceWarning("Couldn't locate the MobileProvisioningFile {0}", MobileProvisionFile);
+
+				iCloudContainerIdentifiersXML = "<array><string>iCloud.$(CFBundleIdentifier)</string></array>";
+				iCloudServicesXML = "<array><string>iCloud.$(CFBundleIdentifier)</string></array>";
+				UbiquityContainerIdentifiersXML = "<array><string>CloudKit</string><string>CloudDocuments</string></array>";
+			}
+
 			// create the entitlements file
 			string IntermediateDir = (((Target.ProjectFile != null) ? Target.ProjectFile.Directory.ToString() :
 				UnrealBuildTool.EngineDirectory.ToString())) + "/Intermediate/" + (Target.Platform == UnrealTargetPlatform.IOS ? "IOS" : "TVOS");
-			WriteEntitlementsFile(Path.Combine(IntermediateDir, AppName + ".entitlements"), Target.ProjectFile, Target.bForDistribution, iCloudContainerIdentifiersXML, UbiquityContainerIdentifiersXML);
+			WriteEntitlementsFile(Target, Path.Combine(IntermediateDir, AppName + ".entitlements"), iCloudContainerIdentifiersXML, iCloudServicesXML, UbiquityContainerIdentifiersXML);
 
 			// create a pList key named ICloudContainerIdentifier
 			// to be used at run-time when intializing the CloudKit services
@@ -1645,10 +1657,10 @@ namespace UnrealBuildTool
 			}
 		}
 
-		private static void WriteEntitlementsFile(string OutputFilename, FileReference ProjectFile, bool bForDistribution, string iCloudContainerIdentifiersXML, string UbiquityContainerIdentifiersXML)
+		private static void WriteEntitlementsFile(IOSPostBuildSyncTarget Target, string OutputFilename, string iCloudContainerIdentifiersXML, string iCloudServicesXML, string UbiquityContainerIdentifiersXML)
 		{
 			// get the settings from the ini file
-			ConfigHierarchy Ini = ConfigCache.ReadHierarchy(ConfigHierarchyType.Engine, DirectoryReference.FromFile(ProjectFile), UnrealTargetPlatform.IOS);
+			ConfigHierarchy Ini = ConfigCache.ReadHierarchy(ConfigHierarchyType.Engine, DirectoryReference.FromFile(Target.ProjectFile), UnrealTargetPlatform.IOS);
 			bool bCloudKitSupported = false;
 			Ini.GetBool("/Script/IOSRuntimeSettings.IOSRuntimeSettings", "bEnableCloudKitSupport", out bCloudKitSupported);
 			Directory.CreateDirectory(Path.GetDirectoryName(OutputFilename));
@@ -1660,46 +1672,45 @@ namespace UnrealBuildTool
 			Text.AppendLine("<plist version=\"1.0\">");
 			Text.AppendLine("<dict>");
 			Text.AppendLine("\t<key>get-task-allow</key>");
-			Text.AppendLine(string.Format("\t<{0}/>", bForDistribution ? "false" : "true"));
+			Text.AppendLine(string.Format("\t<{0}/>", Target.bForDistribution ? "false" : "true"));
 			if (bCloudKitSupported)
 			{
-				Text.AppendLine("\t<key>com.apple.developer.icloud-container-identifiers</key>");
-				if (iCloudContainerIdentifiersXML == "")
-				{
-					Text.AppendLine("\t<array>");
-					Text.AppendLine("\t\t<string>iCloud.$(CFBundleIdentifier)</string>");
-					Text.AppendLine("\t</array>");
-				}
-				else
+				if (iCloudContainerIdentifiersXML != "")
 				{
+					Text.AppendLine("\t<key>com.apple.developer.icloud-container-identifiers</key>");
 					Text.AppendLine(iCloudContainerIdentifiersXML);
 				}
-				Text.AppendLine("\t<key>com.apple.developer.icloud-services</key>");
-				Text.AppendLine("\t<array>");
-				Text.AppendLine("\t\t<string>CloudKit</string>");
-				Text.AppendLine("\t\t<string>CloudDocuments</string>");
-				Text.AppendLine("\t</array>");
-				Text.AppendLine("\t<key>com.apple.developer.ubiquity-container-identifiers</key>");
-				if (UbiquityContainerIdentifiersXML == "")
+
+				if (iCloudServicesXML != "")
 				{
-					Text.AppendLine("\t<array>");
-					Text.AppendLine("\t\t<string>iCloud.$(CFBundleIdentifier)</string>");
-					Text.AppendLine("\t</array>");
+					Text.AppendLine("\t<key>com.apple.developer.icloud-services</key>");
+					Text.AppendLine(iCloudServicesXML);
 				}
-				else
+
+				if (UbiquityContainerIdentifiersXML != "")
 				{
+					Text.AppendLine("\t<key>com.apple.developer.ubiquity-container-identifiers</key>");
 					Text.AppendLine(UbiquityContainerIdentifiersXML);
 				}
 				Text.AppendLine("\t<key>com.apple.developer.ubiquity-kvstore-identifier</key>");
 				Text.AppendLine("\t<string>$(TeamIdentifierPrefix)$(CFBundleIdentifier)</string>");
+
+				Text.AppendLine("\t<key>com.apple.developer.icloud-container-environment</key>");
+				Text.AppendLine(string.Format("\t<string>{0}</string>", Target.bForDistribution ? "Production" : "Development"));
 			}
 
 			bool bRemoteNotificationsSupported = false;
 			Ini.GetBool("/Script/IOSRuntimeSettings.IOSRuntimeSettings", "bEnableRemoteNotificationsSupport", out bRemoteNotificationsSupported);
+
+			if (bCloudKitSupported && Target.bForDistribution && Target.Platform == UnrealTargetPlatform.TVOS)
+			{
+				bRemoteNotificationsSupported = true;
+			}
+
 			if (bRemoteNotificationsSupported)
 			{
 				Text.AppendLine("\t<key>aps-environment</key>");
-				Text.AppendLine(string.Format("\t<string>{0}</string>", bForDistribution ? "production" : "development"));
+				Text.AppendLine(string.Format("\t<string>{0}</string>", Target.bForDistribution ? "production" : "development"));
 			}
 			Text.AppendLine("</dict>");
 			Text.AppendLine("</plist>");
diff --git a/Engine/Source/Programs/UnrealBuildTool/Platform/TVOS/UEDeployTVOS.cs b/Engine/Source/Programs/UnrealBuildTool/Platform/TVOS/UEDeployTVOS.cs
--- a/Engine/Source/Programs/UnrealBuildTool/Platform/TVOS/UEDeployTVOS.cs
+++ b/Engine/Source/Programs/UnrealBuildTool/Platform/TVOS/UEDeployTVOS.cs
@@ -291,6 +291,20 @@ namespace UnrealBuildTool
                                 }
                             }
                         }*/
+			// write the iCloud container identifier, if present in the old file
+			if (!string.IsNullOrEmpty(OldPListData))
+			{
+				int index = OldPListData.IndexOf("ICloudContainerIdentifier");
+				if (index > 0)
+				{
+					index = OldPListData.IndexOf("<string>", index) + 8;
+					int length = OldPListData.IndexOf("</string>", index) - index;
+					string ICloudContainerIdentifier = OldPListData.Substring(index, length);
+					Text.AppendLine("\t<key>ICloudContainerIdentifier</key>");
+					Text.AppendLine(string.Format("\t<string>{0}</string>", ICloudContainerIdentifier));
+				}
+			}
+
             Text.AppendLine("</dict>");
 			Text.AppendLine("</plist>");
 
diff --git a/Engine/Source/Programs/UnrealBuildTool/ProjectFiles/Xcode/XcodeProject.cs b/Engine/Source/Programs/UnrealBuildTool/ProjectFiles/Xcode/XcodeProject.cs
--- a/Engine/Source/Programs/UnrealBuildTool/ProjectFiles/Xcode/XcodeProject.cs
+++ b/Engine/Source/Programs/UnrealBuildTool/ProjectFiles/Xcode/XcodeProject.cs
@@ -1148,12 +1148,14 @@ namespace UnrealBuildTool
 				string TVOSInfoPlistPath = null;
 				string MacInfoPlistPath = null;
 				string IOSEntitlementPath = null;
+				string TVOSEntitlementPath = null;
 				if (bIsUE4Game)
 				{
 					IOSInfoPlistPath = UE4Dir + "/Engine/Intermediate/IOS/" + Config.BuildTarget + "-Info.plist";
 					TVOSInfoPlistPath = UE4Dir + "/Engine/Intermediate/TVOS/" + Config.BuildTarget + "-Info.plist";
 					MacInfoPlistPath = UE4Dir + "/Engine/Intermediate/Mac/" + MacExecutableFileName + "-Info.plist";
 					IOSEntitlementPath = "";
+					TVOSEntitlementPath = "";
 				}
 				else if (bIsUE4Client)
 				{
@@ -1161,6 +1163,7 @@ namespace UnrealBuildTool
 					TVOSInfoPlistPath = UE4Dir + "/Engine/Intermediate/TVOS/UE4Game-Info.plist";
 					MacInfoPlistPath = UE4Dir + "/Engine/Intermediate/Mac/" + MacExecutableFileName + "-Info.plist";
 					IOSEntitlementPath = "";
+					TVOSEntitlementPath = "";
 					}
 				else if (ProjectFile != null)
 				{
@@ -1168,6 +1171,7 @@ namespace UnrealBuildTool
 					TVOSInfoPlistPath = GamePath + "/Intermediate/TVOS/" + Config.BuildTarget + "-Info.plist";
 					MacInfoPlistPath = GamePath + "/Intermediate/Mac/" + MacExecutableFileName + "-Info.plist";
 					IOSEntitlementPath = GamePath + "/Intermediate/IOS/" + Config.BuildTarget + ".entitlements";
+					TVOSEntitlementPath = GamePath + "/Intermediate/TVOS/" + Config.BuildTarget + ".entitlements";
 				}
 				else
 				{
@@ -1193,6 +1197,7 @@ namespace UnrealBuildTool
 				else if (XcodeProjectFileGenerator.bGeneratingRunTVOSProject)
 				{
 					Content.Append("\t\t\t\tINFOPLIST_FILE = \"" + TVOSInfoPlistPath + "\";" + ProjectFileGenerator.NewLine);
+					Content.Append("\t\t\t\tCODE_SIGN_ENTITLEMENTS = \"" + TVOSEntitlementPath + "\";" + ProjectFileGenerator.NewLine);
 				}
 				else
 				{
@@ -1205,6 +1210,7 @@ namespace UnrealBuildTool
 					if (TVOSRunTimeVersion != null)
 					{
 						Content.Append("\t\t\t\t\"INFOPLIST_FILE[sdk=appletvos*]\" = \"" + TVOSInfoPlistPath + "\";" + ProjectFileGenerator.NewLine);
+						Content.Append("\t\t\t\t\"CODE_SIGN_ENTITLEMENTS[sdk=appletvos*]\" = \"" + TVOSEntitlementPath + "\";" + ProjectFileGenerator.NewLine);
 					}
 				}
 
