UBT: Use the full path to the xgConsole executable when starting a build. Fixes issues with environment variables being screwed up.

#rnx

diff --git a/Engine/Source/Programs/UnrealBuildTool/Executors/XGE.cs b/Engine/Source/Programs/UnrealBuildTool/Executors/XGE.cs
--- a/Engine/Source/Programs/UnrealBuildTool/Executors/XGE.cs
+++ b/Engine/Source/Programs/UnrealBuildTool/Executors/XGE.cs
@@ -457,11 +457,17 @@ namespace UnrealBuildTool
 				}
 			}
 
+			string XgConsolePath;
+			if(!TryGetXgConsoleExecutable(out XgConsolePath))
+			{
+				throw new BuildException("Unable to find xgConsole executable.");
+			}
+
 			bool bSilentCompileOutput = false;
 			string SilentOption = bSilentCompileOutput ? "/Silent" : "";
 
 			ProcessStartInfo XGEStartInfo = new ProcessStartInfo(
-				"xgConsole",
+				XgConsolePath,
 				string.Format("\"{0}\" /Rebuild /NoWait {1} /NoLogo {2} /ShowAgent /ShowTime {3}",
 					TaskFilePath,
 					bStopXGECompilationAfterErrors ? "/StopOnErrors" : "",
