UBT: Sort the installed versions of Visual Studio to weight against Express versions being chosen.

diff --git a/Engine/Source/Programs/UnrealBuildTool/Platform/Windows/UEBuildWindows.cs b/Engine/Source/Programs/UnrealBuildTool/Platform/Windows/UEBuildWindows.cs
--- a/Engine/Source/Programs/UnrealBuildTool/Platform/Windows/UEBuildWindows.cs
+++ b/Engine/Source/Programs/UnrealBuildTool/Platform/Windows/UEBuildWindows.cs
@@ -767,7 +767,7 @@ namespace UnrealBuildTool
 				    }
 				    else if(Compiler == WindowsCompiler.VisualStudio2017 || Compiler == WindowsCompiler.VisualStudio2019)
 				    {
-						List<DirectoryReference> PreReleaseInstallDirs = new List<DirectoryReference>();
+						SortedDictionary<int, DirectoryReference> SortedInstallDirs = new SortedDictionary<int, DirectoryReference>(); 
 						try
 						{
 							SetupConfiguration Setup = new SetupConfiguration();
@@ -803,22 +803,28 @@ namespace UnrealBuildTool
 										}
 									}
 
+									int SortOrder = SortedInstallDirs.Count;
+
 									ISetupInstanceCatalog Catalog = (ISetupInstanceCatalog)Instance as ISetupInstanceCatalog;
 									if (Catalog != null && Catalog.IsPrerelease())
 									{
-										PreReleaseInstallDirs.Add(new DirectoryReference(Instance.GetInstallationPath()));
+										SortOrder |= (1 << 16);
 									}
-									else
+
+									string ProductId = Instance.GetProduct().GetId();
+									if (ProductId.Equals("Microsoft.VisualStudio.Product.WDExpress", StringComparison.Ordinal))
 									{
-										InstallDirs.Add(new DirectoryReference(Instance.GetInstallationPath()));
+										SortOrder |= (1 << 17);
 									}
+
+									SortedInstallDirs.Add(SortOrder, new DirectoryReference(Instance.GetInstallationPath()));
 								}
 							}
 						}
 						catch
 						{
 						}
-						InstallDirs.AddRange(PreReleaseInstallDirs);
+						InstallDirs.AddRange(SortedInstallDirs.Values);
 					}
 					else
 				    {
