Add additional logging for Visual Studio version, to help track down problems with Visual Studio Express.

diff --git a/Engine/Source/Programs/UnrealBuildTool/Platform/Windows/UEBuildWindows.cs b/Engine/Source/Programs/UnrealBuildTool/Platform/Windows/UEBuildWindows.cs
--- a/Engine/Source/Programs/UnrealBuildTool/Platform/Windows/UEBuildWindows.cs
+++ b/Engine/Source/Programs/UnrealBuildTool/Platform/Windows/UEBuildWindows.cs
@@ -817,6 +817,8 @@ namespace UnrealBuildTool
 										SortOrder |= (1 << 17);
 									}
 
+									Log.TraceLog("Found Visual Studio installation: {0} (Product={1}, Version={2}, Sort={3})", Instance.GetInstallationPath(), ProductId, VersionString, SortOrder);
+
 									SortedInstallDirs.Add(SortOrder, new DirectoryReference(Instance.GetInstallationPath()));
 								}
 							}
@@ -922,9 +924,13 @@ namespace UnrealBuildTool
 							    foreach(DirectoryReference ToolChainDir in DirectoryReference.EnumerateDirectories(ToolChainBaseDir))
 							    {
 								    VersionNumber Version;
-								    if(VersionNumber.TryParse(ToolChainDir.GetDirectoryName(), out Version) && IsValidToolChainDir2017or2019(ToolChainDir) && !ToolChainVersionToDir.ContainsKey(Version))
-								    {
-									    ToolChainVersionToDir[Version] = ToolChainDir;
+								    if(VersionNumber.TryParse(ToolChainDir.GetDirectoryName(), out Version) && IsValidToolChainDir2017or2019(ToolChainDir))
+									{
+										Log.TraceLog("Found Visual Studio toolchain: {0} (Version={1})", ToolChainDir, Version);
+										if(!ToolChainVersionToDir.ContainsKey(Version))
+										{
+											ToolChainVersionToDir[Version] = ToolChainDir;
+										}
 								    }
 							    }
 						    }
@@ -940,9 +946,13 @@ namespace UnrealBuildTool
 								foreach(DirectoryReference ToolChainDir in DirectoryReference.EnumerateDirectories(ToolChainBaseDir))
 								{
 									VersionNumber Version;
-									if(VersionNumber.TryParse(ToolChainDir.GetDirectoryName(), out Version) && IsValidToolChainDir2017or2019(ToolChainDir) && !ToolChainVersionToDir.ContainsKey(Version))
+									if(VersionNumber.TryParse(ToolChainDir.GetDirectoryName(), out Version) && IsValidToolChainDir2017or2019(ToolChainDir))
 									{
-										ToolChainVersionToDir[Version] = ToolChainDir;
+										Log.TraceLog("Found Visual Studio toolchain: {0} (Version={1})", ToolChainDir, Version);
+										if (!ToolChainVersionToDir.ContainsKey(Version))
+										{
+											ToolChainVersionToDir[Version] = ToolChainDir;
+										}
 									}
 								}
 							}
