UBT: Always prefer 64-bit toolchains over 32-bit toolchains. Prevents Visual Studio Express versions being selected over Pro versions.

diff --git a/Engine/Source/Programs/UnrealBuildTool/Platform/Windows/UEBuildWindows.cs b/Engine/Source/Programs/UnrealBuildTool/Platform/Windows/UEBuildWindows.cs
--- a/Engine/Source/Programs/UnrealBuildTool/Platform/Windows/UEBuildWindows.cs
+++ b/Engine/Source/Programs/UnrealBuildTool/Platform/Windows/UEBuildWindows.cs
@@ -998,6 +998,17 @@ namespace UnrealBuildTool
 			return FileReference.Exists(FileReference.Combine(ToolChainDir, "bin", "Hostx86", "x64", "cl.exe")) || FileReference.Exists(FileReference.Combine(ToolChainDir, "bin", "Hostx64", "x64", "cl.exe"));
 		}
 
+
+		/// <summary>
+		/// Checks if the given directory contains a 64-bit toolchain. Used to prefer regular Visual Studio versions over express editions.
+		/// </summary>
+		/// <param name="ToolChainDir">Directory to check</param>
+		/// <returns>True if the given directory contains a 64-bit toolchain</returns>
+		static bool Has64BitToolChain(DirectoryReference ToolChainDir)
+		{
+			return FileReference.Exists(FileReference.Combine(ToolChainDir, "bin", "amd64", "cl.exe")) || FileReference.Exists(FileReference.Combine(ToolChainDir, "bin", "Hostx64", "x64", "cl.exe"));
+		}
+
 		/// <summary>
 		/// Determines if an IDE for the given compiler is installed.
 		/// </summary>
@@ -1037,7 +1048,7 @@ namespace UnrealBuildTool
 			{
 				if(String.Compare(CompilerVersion, "Latest", StringComparison.InvariantCultureIgnoreCase) == 0 && ToolChainVersionToDir.Count > 0)
 				{
-					ToolChainVersion = ToolChainVersionToDir.OrderBy(x => x.Key).Last().Key;
+					ToolChainVersion = ToolChainVersionToDir.OrderBy(x => Has64BitToolChain(x.Value)).ThenBy(x => x.Key).Last().Key;
 				}
 				else if(!VersionNumber.TryParse(CompilerVersion, out ToolChainVersion))
 				{
@@ -1062,7 +1073,7 @@ namespace UnrealBuildTool
 				}
 				else if(ToolChainVersionToDir.Count > 0)
 				{
-					ToolChainVersion = ToolChainVersionToDir.OrderBy(x => x.Key).Last().Key;
+					ToolChainVersion = ToolChainVersionToDir.OrderBy(x => Has64BitToolChain(x.Value)).ThenBy(x => x.Key).Last().Key;
 				}
 			}
 
