Don't use a Visual Studio Express installation even if it's our preferred toolchain version.

diff --git a/Engine/Source/Programs/UnrealBuildTool/Platform/Windows/UEBuildWindows.cs b/Engine/Source/Programs/UnrealBuildTool/Platform/Windows/UEBuildWindows.cs
--- a/Engine/Source/Programs/UnrealBuildTool/Platform/Windows/UEBuildWindows.cs
+++ b/Engine/Source/Programs/UnrealBuildTool/Platform/Windows/UEBuildWindows.cs
@@ -1067,7 +1067,8 @@ namespace UnrealBuildTool
 					DefaultToolChainVersion = DefaultVisualStudioToolChainVersion;
 				}
 
-				if(ToolChainVersionToDir.ContainsKey(DefaultToolChainVersion))
+				DirectoryReference DefaultToolChainDir;
+				if(ToolChainVersionToDir.TryGetValue(DefaultToolChainVersion, out DefaultToolChainDir) && Has64BitToolChain(DefaultToolChainDir))
 				{
 					ToolChainVersion = DefaultToolChainVersion;
 				}
