Upgrade to VC 15.9.11 toolchain - includes PGO optimization bug fix.

diff --git a/Engine/Source/Programs/UnrealBuildTool/Platform/Windows/UEBuildWindows.cs b/Engine/Source/Programs/UnrealBuildTool/Platform/Windows/UEBuildWindows.cs
--- a/Engine/Source/Programs/UnrealBuildTool/Platform/Windows/UEBuildWindows.cs
+++ b/Engine/Source/Programs/UnrealBuildTool/Platform/Windows/UEBuildWindows.cs
@@ -452,7 +452,7 @@ namespace UnrealBuildTool
 		/// <summary>
 		/// The default compiler version to be used, if installed. 
 		/// </summary>
-		static readonly VersionNumber DefaultVisualStudioToolChainVersion = VersionNumber.Parse("14.16.27023");
+		static readonly VersionNumber DefaultVisualStudioToolChainVersion = VersionNumber.Parse("14.16.27023.1");
 
 		/// <summary>
 		/// The default Windows SDK version to be used, if installed.
