warning CS1591: Missing XML comment for publicly visible type or member 'IOSExports.WriteEntitlements(UnrealTargetPlatform, ConfigHierarchy, string, FileReference, bool, string)' [D:\Build\++UE4+Release-4.22+Incremental\Sync\Engine\Source\Programs\UnrealBuildTool\UnrealBuildTool.csproj]

#4.22.1

diff --git a/Engine/Source/Programs/UnrealBuildTool/Platform/IOS/IOSExports.cs b/Engine/Source/Programs/UnrealBuildTool/Platform/IOS/IOSExports.cs
--- a/Engine/Source/Programs/UnrealBuildTool/Platform/IOS/IOSExports.cs
+++ b/Engine/Source/Programs/UnrealBuildTool/Platform/IOS/IOSExports.cs
@@ -179,6 +179,15 @@ namespace UnrealBuildTool
 			}
 		}
 
+		/// <summary>
+		/// 
+		/// </summary>
+		/// <param name="Platform"></param>
+		/// <param name="PlatformGameConfig"></param>
+		/// <param name="AppName"></param>
+		/// <param name="MobileProvisionFile"></param>
+		/// <param name="bForDistribution"></param>
+		/// <param name="IntermediateDir"></param>
 		public static void WriteEntitlements(UnrealTargetPlatform Platform, ConfigHierarchy PlatformGameConfig,
 		string AppName, FileReference MobileProvisionFile, bool bForDistribution, string IntermediateDir)
 		{
