- Changed how dsyms and stripping are handled - only generate the dsym if we strip, and only strip in Shipping (unless commandlines override it)

diff --git a/Engine/Source/Programs/UnrealBuildTool/Platform/IOS/IOSToolChain.cs b/Engine/Source/Programs/UnrealBuildTool/Platform/IOS/IOSToolChain.cs
--- a/Engine/Source/Programs/UnrealBuildTool/Platform/IOS/IOSToolChain.cs
+++ b/Engine/Source/Programs/UnrealBuildTool/Platform/IOS/IOSToolChain.cs
@@ -181,7 +181,7 @@ namespace UnrealBuildTool
 					BuildProducts.Add(AssetFile, BuildProductType.RequiredResource);
 				}
 			}
-            if ((ProjectSettings.bGeneratedSYMFile == true || ProjectSettings.bGeneratedSYMBundle == true) && (ProjectSettings.bGenerateCrashReportSymbols || Target.bUseMallocProfiler) && Binary.Type == UEBuildBinaryType.Executable)
+            if (Target.IOSPlatform.bGeneratedSYM && (ProjectSettings.bGenerateCrashReportSymbols || Target.bUseMallocProfiler) && Binary.Type == UEBuildBinaryType.Executable)
             {
                 FileReference DebugFile = FileReference.Combine(Binary.OutputFilePath.Directory, Binary.OutputFilePath.GetFileNameWithoutExtension() + ".udebugsymbols");
                 BuildProducts.Add(DebugFile, BuildProductType.SymbolFile);
@@ -1359,8 +1359,8 @@ namespace UnrealBuildTool
                 return OutputFiles;
             }
 
-            // For IOS/tvOS, generate the dSYM file if the config file is set to do so
-			if (ProjectSettings.bGeneratedSYMFile == true || ProjectSettings.bGeneratedSYMBundle == true || BinaryLinkEnvironment.bUsePDBFiles == true)
+            // For IOS/tvOS, generate the dSYM file if needed or requested
+			if (Target.IOSPlatform.bGeneratedSYM)
             {
                 OutputFiles.Add(GenerateDebugInfo(Executable, Actions));
                 if (ProjectSettings.bGenerateCrashReportSymbols || Target.bUseMallocProfiler)
diff --git a/Engine/Source/Programs/UnrealBuildTool/Platform/IOS/UEBuildIOS.cs b/Engine/Source/Programs/UnrealBuildTool/Platform/IOS/UEBuildIOS.cs
--- a/Engine/Source/Programs/UnrealBuildTool/Platform/IOS/UEBuildIOS.cs
+++ b/Engine/Source/Programs/UnrealBuildTool/Platform/IOS/UEBuildIOS.cs
@@ -18,9 +18,10 @@ namespace UnrealBuildTool
 	public class IOSTargetRules
 	{
 		/// <summary>
-		/// Whether to strip iOS symbols or not (implied by bGeneratedSYMFile).
+		/// Whether to strip iOS symbols or not (implied by Shipping config)
 		/// </summary>
 		[XmlConfigFile(Category = "BuildConfiguration")]
+		[CommandLine("-stripsymbols", Value = "true")]
 		public bool bStripSymbols = false;
 
 		/// <summary>
@@ -29,6 +30,12 @@ namespace UnrealBuildTool
 		[CommandLine("-CreateStub", Value = "true")]
 		public bool bCreateStubIPA = false;
 
+		/// <summary>
+		/// Don't generate crashlytics data
+		/// </summary>
+		[CommandLine("-alwaysgeneratedsym", Value = "true")]
+		public bool bGeneratedSYM = false;
+
 		/// <summary>
 		/// Don't generate crashlytics data
 		/// </summary>
@@ -96,6 +103,11 @@ namespace UnrealBuildTool
 			get { return Inner.bStripSymbols; }
 		}
 			
+		public bool bGeneratedSYM
+		{
+			get { return Inner.bGeneratedSYM; }
+		}
+
 		public bool bCreateStubIPA
 		{
 			get { return Inner.bCreateStubIPA; }
@@ -151,14 +163,13 @@ namespace UnrealBuildTool
 		/// Whether to generate a dSYM file or not.
 		/// </summary>
 		[ConfigFile(ConfigHierarchyType.Engine, "/Script/IOSRuntimeSettings.IOSRuntimeSettings", "bGeneratedSYMFile")]
-		[CommandLine("-skipgeneratedsymfile", Value="false")]
-		public readonly bool bGeneratedSYMFile = true;
-
+		[CommandLine("-generatedsymfile")]
+		public readonly bool bGeneratedSYMFile = false;
 		/// <summary>
-		/// Whether to generate a dSYM bundle or not.
+		/// Whether to generate a dSYM bundle (as opposed to single file dSYM)
 		/// </summary>
 		[ConfigFile(ConfigHierarchyType.Engine, "/Script/IOSRuntimeSettings.IOSRuntimeSettings", "bGeneratedSYMBundle")]
-		[CommandLine("-skipgeneratedsymbundle", Value = "false")]
+		[CommandLine("-generatedsymbundle")]
 		public readonly bool bGeneratedSYMBundle = false;
 
         /// <summary>
@@ -715,6 +726,18 @@ namespace UnrealBuildTool
 			Target.bDeployAfterCompile = true;
 
 			Target.IOSPlatform.ProjectSettings = ((IOSPlatform)GetBuildPlatform(Target.Platform)).ReadProjectSettings(Target.ProjectFile);
+			
+			// always strip in shipping configuration (commandline could have set it also)
+			if (Target.Configuration == UnrealTargetConfiguration.Shipping)
+			{
+				Target.IOSPlatform.bStripSymbols = true;	
+			}
+			
+			// if we are stripping the executable, or if the project requested it, or if it's a buildmachine, generate the dsym
+			if (Target.IOSPlatform.bStripSymbols || Target.IOSPlatform.ProjectSettings.bGeneratedSYMFile || Environment.GetEnvironmentVariable("IsBuildMachine") == "1")
+			{
+				Target.IOSPlatform.bGeneratedSYM = true;
+			}
 		}
 
 		public override void ValidateTarget(TargetRules Target)
@@ -816,16 +839,20 @@ namespace UnrealBuildTool
 
 		public override string[] GetDebugInfoExtensions(ReadOnlyTargetRules InTarget, UEBuildBinaryType InBinaryType)
 		{
-			IOSProjectSettings ProjectSettings = ReadProjectSettings(InTarget.ProjectFile);
-
-			if(ProjectSettings.bGeneratedSYMBundle)
+			if (InTarget.IOSPlatform.bGeneratedSYM)
 			{
-				return new string[] {".dSYM.zip"};
+				IOSProjectSettings ProjectSettings = ReadProjectSettings(InTarget.ProjectFile);
+
+				// which format?
+				if (ProjectSettings.bGeneratedSYMBundle)
+				{
+					return new string[] { ".dSYM.zip" };
+				}
+				else
+				{
+					return new string[] { ".dSYM" };
+				}
 			}
-			else if (ProjectSettings.bGeneratedSYMFile)
-            {
-                return new string[] {".dSYM"};
-            }
 
             return new string [] {};
 		}
