Add some more Gradle warnings to sanitize

#android

diff --git a/Engine/Source/Programs/UnrealBuildTool/Platform/Android/UEDeployAndroid.cs b/Engine/Source/Programs/UnrealBuildTool/Platform/Android/UEDeployAndroid.cs
--- a/Engine/Source/Programs/UnrealBuildTool/Platform/Android/UEDeployAndroid.cs
+++ b/Engine/Source/Programs/UnrealBuildTool/Platform/Android/UEDeployAndroid.cs
@@ -1410,7 +1410,17 @@ namespace UnrealBuildTool
 		{
 			if (e.Data != null)
 			{
-				// apply filtering of the warning we want to ignore
+				// apply filtering of the warnings we want to ignore
+				if (e.Data.Contains("WARNING: The option 'android.enableD8' is deprecated and should not be used anymore."))
+				{
+					Log.TraceInformation("{0}", e.Data.Replace("WARNING: ", ">> "));
+					return;
+				}
+				if (e.Data.Contains("WARNING: The specified Android SDK Build Tools version"))
+				{
+					Log.TraceInformation("{0}", e.Data.Replace("WARNING: ", ">> "));
+					return;
+				}
 				if (e.Data.Contains("Warning: Resigning with jarsigner."))
 				{
 					Log.TraceInformation("{0}", e.Data.Replace("Warning: ", ">> "));
