Catch warnings during cleaning phase with Gradle

#android

diff --git a/Engine/Source/Programs/UnrealBuildTool/Platform/Android/UEDeployAndroid.cs b/Engine/Source/Programs/UnrealBuildTool/Platform/Android/UEDeployAndroid.cs
--- a/Engine/Source/Programs/UnrealBuildTool/Platform/Android/UEDeployAndroid.cs
+++ b/Engine/Source/Programs/UnrealBuildTool/Platform/Android/UEDeployAndroid.cs
@@ -1431,6 +1431,10 @@ namespace UnrealBuildTool
 					Log.TraceInformation("{0}", e.Data.Replace("due to error", ""));
 					return;
 				}
+				if (e.Data.Contains("To suppress this warning,"))
+				{
+					Log.TraceInformation("{0}", e.Data.Replace(" warning,", ","));
+				}
 				Log.TraceInformation("{0}", e.Data);
 			}
 		}
@@ -4006,7 +4010,7 @@ namespace UnrealBuildTool
 								//RunCommandLineProgramWithException(UE4BuildGradlePath, ShellExecutable, ShellParametersBegin + "\"" + GradleScriptPath + "\" " + GradleOptions + ShellParametersEnd, "Listing all tasks...");
 
 								GradleOptions = "clean";
-								RunCommandLineProgramWithException(UE4BuildGradlePath, ShellExecutable, ShellParametersBegin + "\"" + GradleScriptPath + "\" " + GradleOptions + ShellParametersEnd, "Cleaning Gradle intermediates...");
+								RunCommandLineProgramWithExceptionAndFiltering(UE4BuildGradlePath, ShellExecutable, ShellParametersBegin + "\"" + GradleScriptPath + "\" " + GradleOptions + ShellParametersEnd, "Cleaning Gradle intermediates...");
 							}
 						}
 						else
