Fix default Clang toolchain version being 7.1.0 rather than 7.0.1.

diff --git a/Engine/Source/Programs/UnrealBuildTool/Platform/Windows/UEBuildWindows.cs b/Engine/Source/Programs/UnrealBuildTool/Platform/Windows/UEBuildWindows.cs
--- a/Engine/Source/Programs/UnrealBuildTool/Platform/Windows/UEBuildWindows.cs
+++ b/Engine/Source/Programs/UnrealBuildTool/Platform/Windows/UEBuildWindows.cs
@@ -1068,7 +1068,7 @@ namespace UnrealBuildTool
 				}
 
 				DirectoryReference DefaultToolChainDir;
-				if(ToolChainVersionToDir.TryGetValue(DefaultToolChainVersion, out DefaultToolChainDir) && Has64BitToolChain(DefaultToolChainDir))
+				if(ToolChainVersionToDir.TryGetValue(DefaultToolChainVersion, out DefaultToolChainDir) && (Compiler == WindowsCompiler.Clang || Has64BitToolChain(DefaultToolChainDir)))
 				{
 					ToolChainVersion = DefaultToolChainVersion;
 				}
