Removing entitlements file generation from the PostBuildSync step ∩┐╜ only for MAC builds

#4.22.1

diff --git a/Engine/Source/Programs/UnrealBuildTool/Platform/IOS/IOSExports.cs b/Engine/Source/Programs/UnrealBuildTool/Platform/IOS/IOSExports.cs
--- a/Engine/Source/Programs/UnrealBuildTool/Platform/IOS/IOSExports.cs
+++ b/Engine/Source/Programs/UnrealBuildTool/Platform/IOS/IOSExports.cs
@@ -201,6 +201,8 @@ namespace UnrealBuildTool
 
 			if (MobileProvisionFile != null && File.Exists(MobileProvisionFile.FullName))
 			{
+				Console.WriteLine("Write entitlements from provisioning file {0}", MobileProvisionFile);
+				
 				MobileProvisionContents MobileProvisionContent = MobileProvisionContents.Read(MobileProvisionFile);
 
 				iCloudContainerIdentifier = MobileProvisionContent.GetNodeValueByName("com.apple.developer.icloud-container-identifiers");
@@ -210,7 +212,7 @@ namespace UnrealBuildTool
 			}
 			else
 			{
-				Log.TraceWarning("Couldn't locate the mobile provisioning file {0}", MobileProvisionFile);
+				Console.WriteLine("Couldn't locate the mobile provisioning file {0}", MobileProvisionFile);
 
 				iCloudContainerIdentifiersXML = "<array><string>iCloud.$(CFBundleIdentifier)</string></array>";
 				iCloudServicesXML = "<array><string>iCloud.$(CFBundleIdentifier)</string></array>";
diff --git a/Engine/Source/Programs/UnrealBuildTool/Platform/IOS/IOSToolChain.cs b/Engine/Source/Programs/UnrealBuildTool/Platform/IOS/IOSToolChain.cs
--- a/Engine/Source/Programs/UnrealBuildTool/Platform/IOS/IOSToolChain.cs
+++ b/Engine/Source/Programs/UnrealBuildTool/Platform/IOS/IOSToolChain.cs
@@ -1860,9 +1860,6 @@ namespace UnrealBuildTool
 				UEDeployIOS DeployHandler = (Target.Platform == UnrealTargetPlatform.IOS ? new UEDeployIOS() : new UEDeployTVOS());
 				DeployHandler.ForDistribution = Target.bForDistribution;
 				DeployHandler.PrepTargetForDeployment(Target.ProjectFile, Target.TargetName, Target.Platform, Target.Configuration, Target.UPLScripts, Target.SdkVersion, Target.bCreateStubIPA, BundleID);
-
-				// write the entitlements file (building on Mac)
-				WriteEntitlements(Target);
 			}
 			
 
