//UE4/Release-4.22 - Compile UE4Game TVOS / IOS - warning: invalid character '*' in Bundle Identifier at column 0. This string must be a uniform type identifier (UTI) that contains only alphanumeric (A-Z,a-z,0-9), hyphen (-), and period (.) characters

diff --git a/Engine/Source/Programs/UnrealBuildTool/Modes/BuildMode.cs b/Engine/Source/Programs/UnrealBuildTool/Modes/BuildMode.cs
--- a/Engine/Source/Programs/UnrealBuildTool/Modes/BuildMode.cs
+++ b/Engine/Source/Programs/UnrealBuildTool/Modes/BuildMode.cs
@@ -114,12 +114,7 @@ namespace UnrealBuildTool
 			BuildConfiguration BuildConfiguration = new BuildConfiguration();
 			XmlConfig.ApplyTo(BuildConfiguration);
 			Arguments.ApplyTo(BuildConfiguration);
-
-			// Parse the remote INI setting
-			string RemoteIniPath;
-			Arguments.TryGetValue("-RemoteIni=", out RemoteIniPath);
-			UnrealBuildTool.SetRemoteIniPath(RemoteIniPath);
-
+			
 			// now that we know the available platforms, we can delete other platforms' junk. if we're only building specific modules from the editor, don't touch anything else (it may be in use).
 			if (!bIgnoreJunk && !UnrealBuildTool.IsEngineInstalled())
 			{
diff --git a/Engine/Source/Programs/UnrealBuildTool/Platform/IOS/IOSToolChain.cs b/Engine/Source/Programs/UnrealBuildTool/Platform/IOS/IOSToolChain.cs
--- a/Engine/Source/Programs/UnrealBuildTool/Platform/IOS/IOSToolChain.cs
+++ b/Engine/Source/Programs/UnrealBuildTool/Platform/IOS/IOSToolChain.cs
@@ -1444,7 +1444,7 @@ namespace UnrealBuildTool
 				FileReference PostBuildSyncFile = FileReference.Combine(BinaryLinkEnvironment.IntermediateDirectory, "PostBuildSync.dat");
 				BinaryFormatterUtils.Save(PostBuildSyncFile, PostBuildSyncTarget);
 
-				string PostBuildSyncArguments = String.Format("-Input=\"{0}\" -XmlConfigCache=\"{1}\"", PostBuildSyncFile, XmlConfig.CacheFile);
+				string PostBuildSyncArguments = String.Format("-Input=\"{0}\" -XmlConfigCache=\"{1}\" -remoteini={2}", PostBuildSyncFile, XmlConfig.CacheFile, UnrealBuildTool.GetRemoteIniPath());
 
 				Action PostBuildSyncAction = Action.CreateRecursiveAction<IOSPostBuildSyncMode>(ActionType.CreateAppBundle, PostBuildSyncArguments);
 				PostBuildSyncAction.WorkingDirectory = UnrealBuildTool.EngineSourceDirectory;
@@ -1606,7 +1606,7 @@ namespace UnrealBuildTool
 	
         public static void PostBuildSync(IOSPostBuildSyncTarget Target)
 			{
-			ConfigHierarchy Ini = ConfigCache.ReadHierarchy(ConfigHierarchyType.Engine, DirectoryReference.FromFile(Target.ProjectFile), UnrealTargetPlatform.IOS);
+			ConfigHierarchy Ini = ConfigCache.ReadHierarchy(ConfigHierarchyType.Engine, Target.ProjectFile == null ? DirectoryReference.FromString(UnrealBuildTool.GetRemoteIniPath()): DirectoryReference.FromFile(Target.ProjectFile), UnrealTargetPlatform.IOS);
 			string BundleID;
 			Ini.GetString("/Script/IOSRuntimeSettings.IOSRuntimeSettings", "BundleIdentifier", out BundleID);
 
@@ -1733,7 +1733,6 @@ namespace UnrealBuildTool
 						MobileProvisionFile = ProvisioningData.MobileProvisionFile;
 						MobileProvisionUUID = ProvisioningData.MobileProvisionUUID;
 						TeamUUID = ProvisioningData.TeamUUID;
-						BundleID = ProvisioningData.BundleIdentifier;
 					}
 					else
 					{
@@ -1742,7 +1741,6 @@ namespace UnrealBuildTool
 						MobileProvisionContents MobileProvision = MobileProvisionContents.Read(MobileProvisionFile);
 						MobileProvisionUUID = MobileProvision.GetUniqueId();
 						MobileProvision.TryGetTeamUniqueId(out TeamUUID);
-						BundleID = MobileProvision.GetBundleIdentifier();
 					}
 
 					string ConfigName = Target.Configuration.ToString();
diff --git a/Engine/Source/Programs/UnrealBuildTool/ToolChain/RemoteMac.cs b/Engine/Source/Programs/UnrealBuildTool/ToolChain/RemoteMac.cs
--- a/Engine/Source/Programs/UnrealBuildTool/ToolChain/RemoteMac.cs
+++ b/Engine/Source/Programs/UnrealBuildTool/ToolChain/RemoteMac.cs
@@ -583,8 +583,9 @@ namespace UnrealBuildTool
 			RemoteArguments.Add(TargetDesc.Configuration.ToString());
 			RemoteArguments.Add("-SkipRulesCompile"); // Use the rules assembly built locally
 			RemoteArguments.Add(String.Format("-XmlConfigCache={0}", GetRemotePath(XmlConfig.CacheFile))); // Use the XML config cache built locally, since the remote won't have it
+			RemoteArguments.Add(String.Format("-remoteini={0}", GetRemotePath(UnrealBuildTool.GetRemoteIniPath())));
 
-			if(TargetDesc.ProjectFile != null)
+			if (TargetDesc.ProjectFile != null)
 			{
 				RemoteArguments.Add(String.Format("-Project={0}", GetRemotePath(TargetDesc.ProjectFile)));
 			}
@@ -874,11 +875,24 @@ namespace UnrealBuildTool
 			UploadDirectory(UnrealBuildTool.EngineDirectory, GetRemotePath(UnrealBuildTool.EngineDirectory), EngineFilters);
 
 			// Upload the project files
-			if(ProjectFile != null && !ProjectFile.IsUnderDirectory(UnrealBuildTool.EngineDirectory))
+			DirectoryReference ProjectDir = null;
+			if (ProjectFile != null && !ProjectFile.IsUnderDirectory(UnrealBuildTool.EngineDirectory))
+			{
+				ProjectDir = ProjectFile.Directory;
+			}
+			else if (!string.IsNullOrEmpty(UnrealBuildTool.GetRemoteIniPath()))
+			{
+				ProjectDir = new DirectoryReference(UnrealBuildTool.GetRemoteIniPath());
+				if (ProjectDir.IsUnderDirectory(UnrealBuildTool.EngineDirectory))
+				{
+					ProjectDir = null;
+				}
+			}
+			if (ProjectDir != null)
 			{
 				List<FileReference> ProjectFilters = new List<FileReference>();
 
-				FileReference CustomProjectFilter = FileReference.Combine(ProjectFile.Directory, "Build", "Rsync", "RsyncProject.txt");
+				FileReference CustomProjectFilter = FileReference.Combine(ProjectDir, "Build", "Rsync", "RsyncProject.txt");
 				if(FileReference.Exists(CustomProjectFilter))
 				{
 					ProjectFilters.Add(CustomProjectFilter);
@@ -886,7 +900,7 @@ namespace UnrealBuildTool
 				ProjectFilters.Add(FileReference.Combine(UnrealBuildTool.EngineDirectory, "Build", "Rsync", "RsyncProject.txt"));
 
 				Log.TraceInformation("[Remote] Uploading project files...");
-				UploadDirectory(ProjectFile.Directory, GetRemotePath(ProjectFile.Directory), ProjectFilters);
+				UploadDirectory(ProjectDir, GetRemotePath(ProjectDir), ProjectFilters);
 			}
 
 			// Fixup permissions on any shell scripts
diff --git a/Engine/Source/Programs/UnrealBuildTool/UnrealBuildTool.cs b/Engine/Source/Programs/UnrealBuildTool/UnrealBuildTool.cs
--- a/Engine/Source/Programs/UnrealBuildTool/UnrealBuildTool.cs
+++ b/Engine/Source/Programs/UnrealBuildTool/UnrealBuildTool.cs
@@ -277,6 +277,12 @@ namespace UnrealBuildTool
 			[CommandLine(Prefix = "-WaitMutex")]
 			public bool bWaitMutex = false;
 
+			/// <summary>
+			/// Whether to wait for the mutex rather than aborting immediately
+			/// </summary>
+			[CommandLine(Prefix = "-RemoteIni")]
+			public string RemoteIni = "";
+
 			/// <summary>
 			/// The mode to execute
 			/// </summary>
@@ -303,6 +309,7 @@ namespace UnrealBuildTool
 			public GlobalOptions(CommandLineArguments Arguments)
 			{
 				Arguments.ApplyTo(this);
+				UnrealBuildTool.SetRemoteIniPath(RemoteIni);
 			}
 		}
 
