UE-72666 - //UE4/Release-4.22 - Compile UE4Game ****** - The path is not of a legal form

diff --git a/Engine/Source/Programs/UnrealBuildTool/UnrealBuildTool.cs b/Engine/Source/Programs/UnrealBuildTool/UnrealBuildTool.cs
--- a/Engine/Source/Programs/UnrealBuildTool/UnrealBuildTool.cs
+++ b/Engine/Source/Programs/UnrealBuildTool/UnrealBuildTool.cs
@@ -309,7 +309,10 @@ namespace UnrealBuildTool
 			public GlobalOptions(CommandLineArguments Arguments)
 			{
 				Arguments.ApplyTo(this);
-				UnrealBuildTool.SetRemoteIniPath(RemoteIni);
+				if (!string.IsNullOrEmpty(RemoteIni))
+				{
+					UnrealBuildTool.SetRemoteIniPath(RemoteIni);
+				}
 			}
 		}
 
