fix for failure to read mobile provision

diff --git a/Engine/Source/Programs/UnrealBuildTool/Platform/IOS/MobileProvisionContents.cs b/Engine/Source/Programs/UnrealBuildTool/Platform/IOS/MobileProvisionContents.cs
--- a/Engine/Source/Programs/UnrealBuildTool/Platform/IOS/MobileProvisionContents.cs
+++ b/Engine/Source/Programs/UnrealBuildTool/Platform/IOS/MobileProvisionContents.cs
@@ -138,7 +138,7 @@ namespace UnrealBuildTool
 		public static XmlDocument ReadXml(FileReference Location)
 		{
 			// Provision data is stored as PKCS7-signed file in ASN.1 BER format
-			using(BinaryReader Reader = new BinaryReader(File.Open(Location.FullName, FileMode.Open)))
+			using(BinaryReader Reader = new BinaryReader(File.Open(Location.FullName, FileMode.Open, FileAccess.Read)))
 			{
 				long Length = Reader.BaseStream.Length;
 				while(Reader.BaseStream.Position < Length)
