Create directory for packageInfo.txt if doesn't exist

#ue4
#android

diff --git a/Engine/Source/Programs/UnrealBuildTool/Platform/Android/UEDeployAndroid.cs b/Engine/Source/Programs/UnrealBuildTool/Platform/Android/UEDeployAndroid.cs
--- a/Engine/Source/Programs/UnrealBuildTool/Platform/Android/UEDeployAndroid.cs
+++ b/Engine/Source/Programs/UnrealBuildTool/Platform/Android/UEDeployAndroid.cs
@@ -4146,6 +4146,12 @@ namespace UnrealBuildTool
 
 			Log.TraceInformation("Writing packageInfo pkgName:{0} storeVersion:{1} versionDisplayName:{2} to {3}", PackageInfoSource[0], PackageInfoSource[1], PackageInfoSource[2], DestPackageNameFileName);
 
+			string DestDirectory = Path.GetDirectoryName(DestPackageNameFileName);
+			if (!Directory.Exists(DestDirectory))
+			{
+				Directory.CreateDirectory(DestDirectory);
+			}
+
 			File.WriteAllLines(DestPackageNameFileName, PackageInfoSource);
 
 			return true;
