Fixed ASAN on iOS

diff --git a/Engine/Source/Programs/UnrealBuildTool/Platform/IOS/IOSToolChain.cs b/Engine/Source/Programs/UnrealBuildTool/Platform/IOS/IOSToolChain.cs
--- a/Engine/Source/Programs/UnrealBuildTool/Platform/IOS/IOSToolChain.cs
+++ b/Engine/Source/Programs/UnrealBuildTool/Platform/IOS/IOSToolChain.cs
@@ -516,7 +516,7 @@ namespace UnrealBuildTool
 			string SanitizerMode = Environment.GetEnvironmentVariable("ENABLE_ADDRESS_SANITIZER");
 			if(SanitizerMode != null && SanitizerMode == "YES")
 			{
-				Result += " -rpath \"@executable_path/Frameworks/libclang_rt.asan_ios_dynamic.dylib\"";
+				Result += " -rpath \"@executable_path/Frameworks\"";
 				Result += " -fsanitize=address";
 			}
 			
