UBT: Fix crash when doing a remote build without the -RemoteINI argument.

diff --git a/Engine/Source/Programs/UnrealBuildTool/ToolChain/RemoteMac.cs b/Engine/Source/Programs/UnrealBuildTool/ToolChain/RemoteMac.cs
--- a/Engine/Source/Programs/UnrealBuildTool/ToolChain/RemoteMac.cs
+++ b/Engine/Source/Programs/UnrealBuildTool/ToolChain/RemoteMac.cs
@@ -583,7 +583,12 @@ namespace UnrealBuildTool
 			RemoteArguments.Add(TargetDesc.Configuration.ToString());
 			RemoteArguments.Add("-SkipRulesCompile"); // Use the rules assembly built locally
 			RemoteArguments.Add(String.Format("-XmlConfigCache={0}", GetRemotePath(XmlConfig.CacheFile))); // Use the XML config cache built locally, since the remote won't have it
-			RemoteArguments.Add(String.Format("-remoteini={0}", GetRemotePath(UnrealBuildTool.GetRemoteIniPath())));
+
+			string RemoteIniPath = UnrealBuildTool.GetRemoteIniPath();
+			if(!String.IsNullOrEmpty(RemoteIniPath))
+			{
+				RemoteArguments.Add(String.Format("-remoteini={0}", GetRemotePath(RemoteIniPath)));
+			}
 
 			if (TargetDesc.ProjectFile != null)
 			{
