Fix errors trying to compile TVOS for installed engine build using remote Mac.

diff --git a/Engine/Source/Programs/UnrealBuildTool/ToolChain/RemoteMac.cs b/Engine/Source/Programs/UnrealBuildTool/ToolChain/RemoteMac.cs
--- a/Engine/Source/Programs/UnrealBuildTool/ToolChain/RemoteMac.cs
+++ b/Engine/Source/Programs/UnrealBuildTool/ToolChain/RemoteMac.cs
@@ -542,9 +542,11 @@ namespace UnrealBuildTool
 			// Copy remote FrameworkAssets directory as it could contain resource bundles that must be packaged locally.
 			DirectoryReference BaseDir = DirectoryReference.FromFile(TargetDesc.ProjectFile) ?? UnrealBuildTool.EngineDirectory;
 			DirectoryReference FrameworkAssetsDir = DirectoryReference.Combine(BaseDir, "Intermediate", TargetDesc.Platform == UnrealTargetPlatform.IOS ? "IOS" : "TVOS", "FrameworkAssets");
-
-			Log.TraceInformation("[Remote] Downloading {0}", FrameworkAssetsDir);
-			DownloadDirectory(FrameworkAssetsDir);
+			if(RemoteDirectoryExists(FrameworkAssetsDir))
+			{
+				Log.TraceInformation("[Remote] Downloading {0}", FrameworkAssetsDir);
+				DownloadDirectory(FrameworkAssetsDir);
+			}
 
 			// Write out all the local manifests
 			foreach(FileReference LocalManifestFile in LocalManifestFiles)
@@ -978,6 +980,17 @@ namespace UnrealBuildTool
 			}
 		}
 
+		/// <summary>
+		/// Checks whether a directory exists on the remote machine
+		/// </summary>
+		/// <param name="LocalDirectory">Path to the directory on the local machine</param>
+		/// <returns>True if the remote directory exists</returns>
+		private bool RemoteDirectoryExists(DirectoryReference LocalDirectory)
+		{
+			string RemoteDirectory = GetRemotePath(LocalDirectory);
+			return Execute(UnrealBuildTool.RootDirectory, String.Format("[ -d {0} ]", EscapeShellArgument(RemoteDirectory))) == 0;
+		}
+
 		/// <summary>
 		/// Download a directory from the remote Mac
 		/// </summary>
@@ -989,7 +1002,6 @@ namespace UnrealBuildTool
 			string RemoteDirectory = GetRemotePath(LocalDirectory);
 
 			List<string> Arguments = new List<string>(CommonRsyncArguments);
-			Arguments.Add(String.Format("--rsync-path=\"[ ! -d {0} ] || rsync\"", EscapeShellArgument(RemoteDirectory)));
 			Arguments.Add(String.Format("\"{0}@{1}\":'{2}/'", UserName, ServerName, RemoteDirectory));
 			Arguments.Add(String.Format("\"{0}/\"", GetLocalCygwinPath(LocalDirectory)));
 
