UE-72638 iOS iCloud app fails to upload to App Store -- entitlements that are not supported; value '*'
#4.22.1
#tvOS
#iOS

applying the same logic in WriteEntitlements as in iPhonePackager.MobileProvision.GetEntitlementsString (used for signing when building on PC) - avoid generic values (*) for the relevant CloudKit fields
fixing a bug with existing plist element (previous build, different settings)

in Automation fixing a bug with caching when building on Mac, legacy from a previous revision

diff --git a/Engine/Source/Programs/AutomationTool/IOS/IOSPlatform.Automation.cs b/Engine/Source/Programs/AutomationTool/IOS/IOSPlatform.Automation.cs
--- a/Engine/Source/Programs/AutomationTool/IOS/IOSPlatform.Automation.cs
+++ b/Engine/Source/Programs/AutomationTool/IOS/IOSPlatform.Automation.cs
@@ -1698,8 +1698,7 @@ public class IOSPlatform : Platform
 
 		// ios configuration from the ini file
 		ConfigHierarchy PlatformGameConfig;
-		if (Params.EngineConfigs.TryGetValue(SC.StageTargetPlatform.PlatformType, out PlatformGameConfig) &&
-			!File.Exists(OutputFilename))
+		if (Params.EngineConfigs.TryGetValue(SC.StageTargetPlatform.PlatformType, out PlatformGameConfig))
 		{
 			IOSExports.WriteEntitlements(Platform, PlatformGameConfig, AppName, MobileProvisionFile, bForDistribution, IntermediateDir);
 		}
diff --git a/Engine/Source/Programs/UnrealBuildTool/Platform/IOS/IOSExports.cs b/Engine/Source/Programs/UnrealBuildTool/Platform/IOS/IOSExports.cs
--- a/Engine/Source/Programs/UnrealBuildTool/Platform/IOS/IOSExports.cs
+++ b/Engine/Source/Programs/UnrealBuildTool/Platform/IOS/IOSExports.cs
@@ -193,10 +193,12 @@ namespace UnrealBuildTool
 		{
 			// get some info from the mobileprovisioning file
 			// the iCloud identifier and the bundle id may differ
-			string iCloudContainerIdentifiersXML = "";
 			string iCloudContainerIdentifier = "";
-			string UbiquityContainerIdentifiersXML = "";
-			string iCloudServicesXML = "";
+			string iCloudContainerIdentifiersXML = "<array><string>iCloud.$(CFBundleIdentifier)</string></array>";
+			string UbiquityContainerIdentifiersXML = "<array><string>iCloud.$(CFBundleIdentifier)</string></array>";
+			string iCloudServicesXML = "<array><string>CloudKit</string><string>CloudDocuments</string></array>";
+			string UbiquityKVStoreIdentifiersXML = "\t<string>$(TeamIdentifierPrefix)$(CFBundleIdentifier)</string>";
+
 			string OutputFileName = Path.Combine(IntermediateDir, AppName + ".entitlements");
 
 			if (MobileProvisionFile != null && File.Exists(MobileProvisionFile.FullName))
@@ -207,16 +209,32 @@ namespace UnrealBuildTool
 
 				iCloudContainerIdentifier = MobileProvisionContent.GetNodeValueByName("com.apple.developer.icloud-container-identifiers");
 				iCloudContainerIdentifiersXML = MobileProvisionContent.GetNodeXMLValueByName("com.apple.developer.icloud-container-identifiers");
-				iCloudServicesXML = MobileProvisionContent.GetNodeXMLValueByName("com.apple.developer.icloud-services");
-				UbiquityContainerIdentifiersXML = MobileProvisionContent.GetNodeXMLValueByName("com.apple.developer.ubiquity-container-identifiers");
+
+				string entitlementXML = MobileProvisionContent.GetNodeXMLValueByName("com.apple.developer.icloud-services");
+
+				if (!entitlementXML.Contains("*") || Platform == UnrealTargetPlatform.TVOS)
+				{
+					// for iOS, replace the generic value (*) with the default
+					iCloudServicesXML = entitlementXML;
+				}
+
+				entitlementXML = MobileProvisionContent.GetNodeXMLValueByName("com.apple.developer.ubiquity-container-identifiers");
+				if (!entitlementXML.Contains("*") || !bForDistribution)
+				{
+					// for distribution, replace the generic value (*) with the default
+					UbiquityContainerIdentifiersXML = entitlementXML;
+				}
+
+				entitlementXML = MobileProvisionContent.GetNodeXMLValueByName("com.apple.developer.ubiquity-kvstore-identifier");
+				if (!entitlementXML.Contains("*") || !bForDistribution)
+				{
+					// for distribution, replace the generic value (*) with the default
+					UbiquityKVStoreIdentifiersXML = entitlementXML;
+				}
 			}
 			else
 			{
 				Console.WriteLine("Couldn't locate the mobile provisioning file {0}", MobileProvisionFile);
-
-				iCloudContainerIdentifiersXML = "<array><string>iCloud.$(CFBundleIdentifier)</string></array>";
-				iCloudServicesXML = "<array><string>iCloud.$(CFBundleIdentifier)</string></array>";
-				UbiquityContainerIdentifiersXML = "<array><string>CloudKit</string><string>CloudDocuments</string></array>";
 			}
 
 			// write the entitlements file
@@ -252,8 +270,12 @@ namespace UnrealBuildTool
 						Text.AppendLine("\t<key>com.apple.developer.ubiquity-container-identifiers</key>");
 						Text.AppendLine(UbiquityContainerIdentifiersXML);
 					}
-					Text.AppendLine("\t<key>com.apple.developer.ubiquity-kvstore-identifier</key>");
-					Text.AppendLine("\t<string>$(TeamIdentifierPrefix)$(CFBundleIdentifier)</string>");
+
+					if (UbiquityKVStoreIdentifiersXML != "")
+					{
+						Text.AppendLine("\t<key>com.apple.developer.ubiquity-kvstore-identifier</key>");
+						Text.AppendLine(UbiquityKVStoreIdentifiersXML);
+					}
 
 					Text.AppendLine("\t<key>com.apple.developer.icloud-container-environment</key>");
 					Text.AppendLine(string.Format("\t<string>{0}</string>", bForDistribution ? "Production" : "Development"));
@@ -262,6 +284,7 @@ namespace UnrealBuildTool
 				bool bRemoteNotificationsSupported = false;
 				PlatformGameConfig.GetBool("/Script/IOSRuntimeSettings.IOSRuntimeSettings", "bEnableRemoteNotificationsSupport", out bRemoteNotificationsSupported);
 
+				// for TVOS we need push notifications when building for distribution with CloudKit
 				if (bCloudKitSupported && bForDistribution && Platform == UnrealTargetPlatform.TVOS)
 				{
 					bRemoteNotificationsSupported = true;
@@ -312,8 +335,29 @@ namespace UnrealBuildTool
 						XElement dictElement = XDoc.Root.Element("dict");
 						if (dictElement != null)
 						{
-							dictElement.Add(new XElement("key", "ICloudContainerIdentifier"));
-							dictElement.Add(new XElement("string", iCloudContainerIdentifier));
+							XElement containerIdKeyNew = new XElement("key", "ICloudContainerIdentifier");
+							XElement containerIdValueNew = new XElement("string", iCloudContainerIdentifier);
+
+							XElement containerIdKey = dictElement.Elements("key").FirstOrDefault(x => x.Value == "ICloudContainerIdentifier");
+							if (containerIdKey != null)
+							{
+								// if ICloudContainerIdentifier already exists in the pList file, update its value
+								XElement containerIdValue = containerIdKey.ElementsAfterSelf("string").FirstOrDefault();
+								if (containerIdValue != null)
+								{
+									containerIdValue.Value = iCloudContainerIdentifier;
+								}
+								else
+								{
+									containerIdKey.AddAfterSelf(containerIdValueNew);
+								}
+							}
+							else
+							{
+								// add ICloudContainerIdentifier to the pList
+								dictElement.Add(containerIdKeyNew);
+								dictElement.Add(containerIdValueNew);
+							}
 
 							XDoc.Save(PListFile);
 						}
