UE-72100 - tvOS For Distribution packages fail with Code Signing Error if made after a Development Build

diff --git a/Engine/Source/Programs/UnrealBuildTool/ProjectFiles/Xcode/XcodeProject.cs b/Engine/Source/Programs/UnrealBuildTool/ProjectFiles/Xcode/XcodeProject.cs
--- a/Engine/Source/Programs/UnrealBuildTool/ProjectFiles/Xcode/XcodeProject.cs
+++ b/Engine/Source/Programs/UnrealBuildTool/ProjectFiles/Xcode/XcodeProject.cs
@@ -1218,10 +1218,9 @@ namespace UnrealBuildTool
 				// This is needed for the target to pass the settings validation before code signing. UBT will overwrite this plist file later, with proper contents.
 				if (BuildHostPlatform.Current.Platform == UnrealTargetPlatform.Mac)
 				{
-					bool bCreateMacInfoPlist = !File.Exists(MacInfoPlistPath);
-					bool bCreateIOSInfoPlist = !File.Exists(IOSInfoPlistPath) && IOSRunTimeVersion != null;
-					bool bCreateTVOSInfoPlist = !File.Exists(TVOSInfoPlistPath) && TVOSRunTimeVersion != null;
-					if (bCreateMacInfoPlist || bCreateIOSInfoPlist || bCreateTVOSInfoPlist)
+					bool bCreateIOSInfoPlist = IOSRunTimeVersion != null;
+					bool bCreateTVOSInfoPlist = TVOSRunTimeVersion != null;
+					if (bCreateIOSInfoPlist || bCreateTVOSInfoPlist)
 					{
 						DirectoryReference ProjectPath = GameDir;
 						DirectoryReference EngineDir = DirectoryReference.Combine(new DirectoryReference(UE4Dir), "Engine");
@@ -1236,11 +1235,9 @@ namespace UnrealBuildTool
 							GameName = "UE4Game";
 						}
 
-						if (bCreateMacInfoPlist)
-						{
-							Directory.CreateDirectory(Path.GetDirectoryName(MacInfoPlistPath));
-							UEDeployMac.GeneratePList(ProjectPath.FullName, bIsUE4Game, GameName, Config.BuildTarget, EngineDir.FullName, MacExecutableFileName);
-						}
+						Directory.CreateDirectory(Path.GetDirectoryName(MacInfoPlistPath));
+						UEDeployMac.GeneratePList(ProjectPath.FullName, bIsUE4Game, GameName, Config.BuildTarget, EngineDir.FullName, MacExecutableFileName);
+
 						if (bCreateIOSInfoPlist)
 						{
 							// get the receipt
