Fix exception when exporting JSON for SwarmInterface. Fixes CIS error with IncludeTool.

diff --git a/Engine/Source/Programs/UnrealBuildTool/Configuration/UEBuildModule.cs b/Engine/Source/Programs/UnrealBuildTool/Configuration/UEBuildModule.cs
--- a/Engine/Source/Programs/UnrealBuildTool/Configuration/UEBuildModule.cs
+++ b/Engine/Source/Programs/UnrealBuildTool/Configuration/UEBuildModule.cs
@@ -823,8 +823,10 @@ namespace UnrealBuildTool
 		/// <summary>
 		/// Write information about this binary to a JSON file
 		/// </summary>
+		/// <param name="BinaryOutputDir">The output directory for the binary containing this module</param>
+		/// <param name="TargetOutputDir">The output directory for the target executable</param>
 		/// <param name="Writer">Writer for this binary's data</param>
-		public virtual void ExportJson(JsonWriter Writer)
+		public virtual void ExportJson(DirectoryReference BinaryOutputDir, DirectoryReference TargetOutputDir, JsonWriter Writer)
 		{
 			Writer.WriteValue("Name", Name);
 			Writer.WriteValue("Directory", ModuleDirectory.FullName);
@@ -868,10 +870,10 @@ namespace UnrealBuildTool
 			foreach(ModuleRules.RuntimeDependency RuntimeDependency in Rules.RuntimeDependencies.Inner)
 			{
 				Writer.WriteObjectStart();
-				Writer.WriteValue("Path", ExpandPathVariables(RuntimeDependency.Path, null, null));
+				Writer.WriteValue("Path", ExpandPathVariables(RuntimeDependency.Path, BinaryOutputDir, TargetOutputDir));
 				if(RuntimeDependency.SourcePath != null)
 				{
-					Writer.WriteValue("SourcePath", ExpandPathVariables(RuntimeDependency.SourcePath, null, null));
+					Writer.WriteValue("SourcePath", ExpandPathVariables(RuntimeDependency.SourcePath, BinaryOutputDir, TargetOutputDir));
 				}
 				Writer.WriteValue("Type", RuntimeDependency.Type.ToString());
 				Writer.WriteObjectEnd();
diff --git a/Engine/Source/Programs/UnrealBuildTool/Configuration/UEBuildTarget.cs b/Engine/Source/Programs/UnrealBuildTool/Configuration/UEBuildTarget.cs
--- a/Engine/Source/Programs/UnrealBuildTool/Configuration/UEBuildTarget.cs
+++ b/Engine/Source/Programs/UnrealBuildTool/Configuration/UEBuildTarget.cs
@@ -1304,11 +1304,7 @@ namespace UnrealBuildTool
 			}
 
 			// Build the target's binaries.
-			DirectoryReference ExeDir = Binaries[0].OutputDir;
-			if (CppPlatform == CppPlatform.Mac && ExeDir.FullName.EndsWith(".app/Contents/MacOS"))
-			{
-				ExeDir = ExeDir.ParentDirectory.ParentDirectory.ParentDirectory;
-			}
+			DirectoryReference ExeDir = GetExecutableDir();
 			using(Timeline.ScopeEvent("UEBuildBinary.Build()"))
 			{
 				foreach (UEBuildBinary Binary in Binaries)
@@ -1532,6 +1528,20 @@ namespace UnrealBuildTool
 			return Makefile;
 		}
 
+		/// <summary>
+		/// Gets the output directory for the main executable
+		/// </summary>
+		/// <returns>The executable directory</returns>
+		DirectoryReference GetExecutableDir()
+		{
+			DirectoryReference ExeDir = Binaries[0].OutputDir;
+			if (Platform == UnrealTargetPlatform.Mac && ExeDir.FullName.EndsWith(".app/Contents/MacOS"))
+			{
+				ExeDir = ExeDir.ParentDirectory.ParentDirectory.ParentDirectory;
+			}
+			return ExeDir;
+		}
+
 		/// <summary>
 		/// Check that copying a file from one location to another does not violate rules regarding restricted folders
 		/// </summary>
@@ -1732,7 +1742,7 @@ namespace UnrealBuildTool
 				foreach(UEBuildModule Module in Modules.Values)
 				{
 					Writer.WriteObjectStart(Module.Name);
-					Module.ExportJson(Writer);
+					Module.ExportJson(Module.Binary?.OutputDir, GetExecutableDir(), Writer);
 					Writer.WriteObjectEnd();
 				}
 				Writer.WriteObjectEnd();
