LiveCoding: Fix asserts/crashes due to variables being reconstructed when files are split out of a unity blob.

These variables (and their dynamic initializers) were being treated as new, rather than being linked back to their original instances. UBT now outputs a mapping of object files to their original unity object file during a live coding compile, which allows Live++ to disambiguate between new and reconstructed instances.

diff --git a/Engine/Source/Developer/Windows/LiveCodingServer/Private/External/LC_AppSettings.cpp b/Engine/Source/Developer/Windows/LiveCodingServer/Private/External/LC_AppSettings.cpp
--- a/Engine/Source/Developer/Windows/LiveCodingServer/Private/External/LC_AppSettings.cpp
+++ b/Engine/Source/Developer/Windows/LiveCodingServer/Private/External/LC_AppSettings.cpp
@@ -443,7 +443,7 @@ void appSettings::Startup(const wchar_t* group)
 		L"amalgamation_split_into_single_parts",
 		L"Split into single parts",
 		L"Specifies whether amalgamated/unity files are automatically split into single files",
-		false
+		true
 	);
 
 	g_amalgamationSplitMinCppCount = new SettingInt
diff --git a/Engine/Source/Developer/Windows/LiveCodingServer/Private/External/LC_ServerCommandThread.cpp b/Engine/Source/Developer/Windows/LiveCodingServer/Private/External/LC_ServerCommandThread.cpp
--- a/Engine/Source/Developer/Windows/LiveCodingServer/Private/External/LC_ServerCommandThread.cpp
+++ b/Engine/Source/Developer/Windows/LiveCodingServer/Private/External/LC_ServerCommandThread.cpp
@@ -708,7 +708,8 @@ void ServerCommandThread::CompileChanges(bool didAllProcessesMakeProgress)
 		GLiveCodingServer->GetStatusChangeDelegate().ExecuteIfBound(L"Compiling changes for live coding...");
 
 		TMap<FString, TArray<FString>> ModuleToObjectFiles;
-		if (!CompileDelegate.Execute(Targets, ModuleToObjectFiles))
+		TMap<FString, FString> ObjectFileToUnityObjectFile;
+		if (!CompileDelegate.Execute(Targets, ModuleToObjectFiles, ObjectFileToUnityObjectFile))
 		{
 			GLiveCodingServer->GetCompileFinishedDelegate().ExecuteIfBound(ELiveCodingResult::Error, L"Compilation error.");
 			return;
@@ -776,6 +777,13 @@ void ServerCommandThread::CompileChanges(bool didAllProcessesMakeProgress)
 			{
 				LiveModule::ModifiedObjFile ModifiedObjFile;
 				ModifiedObjFile.objPath = file::NormalizePath(*ObjectFile);
+
+				const FString* UnityObjectFile = ObjectFileToUnityObjectFile.Find(ObjectFile);
+				if (UnityObjectFile != nullptr)
+				{
+					ModifiedObjFile.amalgamatedObjPath = file::NormalizePath(**UnityObjectFile);
+				}
+
 				ObjectFiles.push_back(std::move(ModifiedObjFile));
 			}
 
diff --git a/Engine/Source/Developer/Windows/LiveCodingServer/Private/External/LC_SymbolResolution.cpp b/Engine/Source/Developer/Windows/LiveCodingServer/Private/External/LC_SymbolResolution.cpp
--- a/Engine/Source/Developer/Windows/LiveCodingServer/Private/External/LC_SymbolResolution.cpp
+++ b/Engine/Source/Developer/Windows/LiveCodingServer/Private/External/LC_SymbolResolution.cpp
@@ -4,7 +4,13 @@
 #include "LC_PointerUtil.h"
 #include "LC_Logging.h"
 #include "Windows/WindowsHWrapper.h"
+
+#pragma push_macro("OPTIONAL")
+#ifndef OPTIONAL
+#define OPTIONAL
+#endif
 #include <dbghelp.h>
+#pragma pop_macro("OPTIONAL")
 
 namespace symbolResolution
 {
diff --git a/Engine/Source/Developer/Windows/LiveCodingServer/Public/ILiveCodingServer.h b/Engine/Source/Developer/Windows/LiveCodingServer/Public/ILiveCodingServer.h
--- a/Engine/Source/Developer/Windows/LiveCodingServer/Public/ILiveCodingServer.h
+++ b/Engine/Source/Developer/Windows/LiveCodingServer/Public/ILiveCodingServer.h
@@ -44,7 +44,8 @@ public:
 	virtual FLogOutputDelegate& GetLogOutputDelegate() = 0;
 
 	typedef TMap<FString, TArray<FString>> FModuleToObjectFiles;
-	DECLARE_DELEGATE_RetVal_TwoParams(bool, FCompileDelegate, const TArray<FString>&, FModuleToObjectFiles&)
+	typedef TMap<FString, FString> FObjectFileToUnityObjectFile;
+	DECLARE_DELEGATE_RetVal_ThreeParams(bool, FCompileDelegate, const TArray<FString>&, FModuleToObjectFiles&, FObjectFileToUnityObjectFile&)
 	virtual FCompileDelegate& GetCompileDelegate() = 0;
 
 	DECLARE_DELEGATE(FCompileStartedDelegate);
diff --git a/Engine/Source/Programs/LiveCodingConsole/Private/LiveCodingConsole.cpp b/Engine/Source/Programs/LiveCodingConsole/Private/LiveCodingConsole.cpp
--- a/Engine/Source/Programs/LiveCodingConsole/Private/LiveCodingConsole.cpp
+++ b/Engine/Source/Programs/LiveCodingConsole/Private/LiveCodingConsole.cpp
@@ -179,7 +179,7 @@ private:
 		LogWidget->AppendLine(GetLogColor(Verbosity), MoveTemp(Text));
 	}
 
-	bool CompilePatch(const TArray<FString>& Targets, TMap<FString, TArray<FString>>& ModuleToObjectFiles)
+	bool CompilePatch(const TArray<FString>& Targets, TMap<FString, TArray<FString>>& ModuleToObjectFiles, TMap<FString, FString>& ObjectFileToUnityObjectFile)
 	{
 		// Update the compile start time. This gets copied into the last patch time once a patch has been confirmed to have been applied.
 		NextPatchStartTime = FDateTime::UtcNow();
@@ -250,6 +250,9 @@ private:
 				}
 			}
 		}
+
+		// Save off the mapping to unity file
+		ObjectFileToUnityObjectFile = MoveTemp(Manifest.ObjectFileToUnityObjectFile);
 		return true;
 	}
 
diff --git a/Engine/Source/Programs/LiveCodingConsole/Private/LiveCodingManifest.cpp b/Engine/Source/Programs/LiveCodingConsole/Private/LiveCodingManifest.cpp
--- a/Engine/Source/Programs/LiveCodingConsole/Private/LiveCodingManifest.cpp
+++ b/Engine/Source/Programs/LiveCodingConsole/Private/LiveCodingManifest.cpp
@@ -95,5 +95,37 @@ bool FLiveCodingManifest::Parse(FJsonObject& Object, FString& OutFailReason)
 			InputFiles.Add(Input->AsString());
 		}
 	}
+
+	const TArray<TSharedPtr<FJsonValue>>* AdaptiveUnityFilesArray;
+	if (Object.TryGetArrayField(TEXT("AdaptiveUnityFiles"), AdaptiveUnityFilesArray))
+	{
+		for (const TSharedPtr<FJsonValue>& AdaptiveUnityFileValue : *AdaptiveUnityFilesArray)
+		{
+			if (AdaptiveUnityFileValue->Type != EJson::Object)
+			{
+				OutFailReason = TEXT("invalid module object");
+				return false;
+			}
+
+			const FJsonObject& AdaptiveUnityFileObject = *AdaptiveUnityFileValue->AsObject();
+
+			FString ObjectFile;
+			if (!AdaptiveUnityFileObject.TryGetStringField(TEXT("ObjectFile"), ObjectFile))
+			{
+				OutFailReason = TEXT("missing module 'ObjectFile' field");
+				return false;
+			}
+
+			FString UnityObjectFile;
+			if (!AdaptiveUnityFileObject.TryGetStringField(TEXT("UnityObjectFile"), UnityObjectFile))
+			{
+				OutFailReason = TEXT("missing module 'UnityObjectFile' field");
+				return false;
+			}
+
+			ObjectFileToUnityObjectFile.Emplace(MoveTemp(ObjectFile), MoveTemp(UnityObjectFile));
+		}
+	}
+
 	return true;
 }
diff --git a/Engine/Source/Programs/LiveCodingConsole/Private/LiveCodingManifest.h b/Engine/Source/Programs/LiveCodingConsole/Private/LiveCodingManifest.h
--- a/Engine/Source/Programs/LiveCodingConsole/Private/LiveCodingManifest.h
+++ b/Engine/Source/Programs/LiveCodingConsole/Private/LiveCodingManifest.h
@@ -11,6 +11,7 @@ struct FLiveCodingManifest
 	FString LinkerPath;
 	TMap<FString, FString> LinkerEnvironment;
 	TMap<FString, TArray<FString>> BinaryToObjectFiles;
+	TMap<FString, FString> ObjectFileToUnityObjectFile;
 
 	bool Read(const TCHAR* FileName, FString& OutFailReason);
 	bool Parse(FJsonObject& Object, FString& OutFailReason);
diff --git a/Engine/Source/Programs/UnrealBuildTool/Modes/BuildMode.cs b/Engine/Source/Programs/UnrealBuildTool/Modes/BuildMode.cs
--- a/Engine/Source/Programs/UnrealBuildTool/Modes/BuildMode.cs
+++ b/Engine/Source/Programs/UnrealBuildTool/Modes/BuildMode.cs
@@ -221,9 +221,11 @@ namespace UnrealBuildTool
 			// Output the manifest
 			for(int TargetIdx = 0; TargetIdx < TargetDescriptors.Count; TargetIdx++)
 			{
-				if(TargetDescriptors[TargetIdx].LiveCodingManifest != null)
+				TargetDescriptor TargetDescriptor = TargetDescriptors[TargetIdx];
+				if(TargetDescriptor.LiveCodingManifest != null)
 				{
-					HotReload.WriteLiveCodeManifest(TargetDescriptors[TargetIdx].LiveCodingManifest, Makefiles[TargetIdx].Actions);
+					TargetMakefile Makefile = Makefiles[TargetIdx];
+					HotReload.WriteLiveCodeManifest(TargetDescriptor.LiveCodingManifest, Makefile.SourceFileToUnityFile, Makefile.Actions);
 				}
 			}
 
diff --git a/Engine/Source/Programs/UnrealBuildTool/System/HotReload.cs b/Engine/Source/Programs/UnrealBuildTool/System/HotReload.cs
--- a/Engine/Source/Programs/UnrealBuildTool/System/HotReload.cs
+++ b/Engine/Source/Programs/UnrealBuildTool/System/HotReload.cs
@@ -727,8 +727,9 @@ namespace UnrealBuildTool
 		/// Writes a manifest containing all the information needed to create a live coding patch
 		/// </summary>
 		/// <param name="ManifestFile">File to write to</param>
+		/// <param name="SourceFileToUnityFile">Mapping of source file to the unity file that would contain it</param>
 		/// <param name="Actions">List of actions that are part of the graph</param>
-		public static void WriteLiveCodeManifest(FileReference ManifestFile, List<Action> Actions)
+		public static void WriteLiveCodeManifest(FileReference ManifestFile, Dictionary<FileItem, FileItem> SourceFileToUnityFile, List<Action> Actions)
 		{
 			// Find all the output object files
 			HashSet<FileItem> ObjectFiles = new HashSet<FileItem>();
@@ -740,8 +741,32 @@ namespace UnrealBuildTool
 				}
 			}
 
+			// Find any action that has a non-unity cpp prerequisite, and map it to the appropriate 
+			Dictionary<FileItem, FileItem> ObjectFileToUnityObjectFile = new Dictionary<FileItem, FileItem>();
+			foreach(Action Action in Actions)
+			{
+				if (Action.ActionType == ActionType.Compile)
+				{
+					foreach(FileItem PrerequisiteItem in Action.PrerequisiteItems)
+					{
+						FileItem UnityFile;
+						if(SourceFileToUnityFile.TryGetValue(PrerequisiteItem, out UnityFile))
+						{
+							foreach(FileItem ProducedItem in Action.ProducedItems)
+							{
+								if(ProducedItem.HasExtension(".obj"))
+								{
+									FileItem UnityObjFile = FileItem.GetItemByFileReference(new FileReference(UnityFile.Location.FullName + ".obj"));
+									ObjectFileToUnityObjectFile[ProducedItem] = UnityObjFile;
+								}
+							}
+						}
+					}
+				}
+			}
+
 			// Write the output manifest
-			using(JsonWriter Writer = new JsonWriter(ManifestFile))
+			using (JsonWriter Writer = new JsonWriter(ManifestFile))
 			{
 				Writer.WriteObjectStart();
 
@@ -785,6 +810,26 @@ namespace UnrealBuildTool
 				}
 				Writer.WriteArrayEnd();
 
+				Writer.WriteArrayStart("AdaptiveUnityFiles");
+				foreach (Action Action in Actions)
+				{
+					if (Action.ActionType == ActionType.Link && Action.ProducedItems.Any(x => x.HasExtension(".exe") || x.HasExtension(".dll")))
+					{
+						foreach (FileItem InputFile in Action.PrerequisiteItems)
+						{
+							FileItem UnityFile;
+							if (ObjectFileToUnityObjectFile.TryGetValue(InputFile, out UnityFile))
+							{
+								Writer.WriteObjectStart();
+								Writer.WriteValue("ObjectFile", InputFile.AbsolutePath);
+								Writer.WriteValue("UnityObjectFile", UnityFile.AbsolutePath);
+								Writer.WriteObjectEnd();
+							}
+						}
+					}
+				}
+				Writer.WriteArrayEnd();
+
 				Writer.WriteObjectEnd();
 			}
 		}
diff --git a/Engine/Source/Programs/UnrealBuildTool/System/TargetMakefile.cs b/Engine/Source/Programs/UnrealBuildTool/System/TargetMakefile.cs
--- a/Engine/Source/Programs/UnrealBuildTool/System/TargetMakefile.cs
+++ b/Engine/Source/Programs/UnrealBuildTool/System/TargetMakefile.cs
@@ -20,7 +20,7 @@ namespace UnrealBuildTool
 		/// <summary>
 		/// The version number to write
 		/// </summary>
-		public const int CurrentVersion = 11;
+		public const int CurrentVersion = 12;
 
 		/// <summary>
 		/// The time at which the makefile was created
@@ -113,6 +113,11 @@ namespace UnrealBuildTool
 		/// </summary>
 		public HashSet<FileItem> CandidatesForWorkingSet = new HashSet<FileItem>();
 
+		/// <summary>
+		/// Mapping from source file to the unity file that contains it. Used for Live Coding. Note that this mapping includes files that are part of the working set, even if they are now a separate compiland.
+		/// </summary>
+		public Dictionary<FileItem, FileItem> SourceFileToUnityFile = new Dictionary<FileItem, FileItem>();
+
 		/// <summary>
 		/// Maps each target to a list of UObject module info structures
 		/// </summary>
@@ -190,6 +195,7 @@ namespace UnrealBuildTool
 			DirectoryToSourceFiles = Reader.ReadDictionary(() => Reader.ReadDirectoryItem(), () => Reader.ReadArray(() => Reader.ReadFileItem()));
 			WorkingSet = Reader.ReadHashSet(() => Reader.ReadFileItem());
 			CandidatesForWorkingSet = Reader.ReadHashSet(() => Reader.ReadFileItem());
+			SourceFileToUnityFile = Reader.ReadDictionary(() => Reader.ReadFileItem(), () => Reader.ReadFileItem());
 			UObjectModules = Reader.ReadList(() => new UHTModuleInfo(Reader));
 			UObjectModuleHeaders = Reader.ReadList(() => new UHTModuleHeaderInfo(Reader));
 			PluginFiles = Reader.ReadHashSet(() => Reader.ReadFileItem());
@@ -220,6 +226,7 @@ namespace UnrealBuildTool
 			Writer.WriteDictionary(DirectoryToSourceFiles, k => Writer.WriteDirectoryItem(k), v => Writer.WriteArray(v, e => Writer.WriteFileItem(e)));
 			Writer.WriteHashSet(WorkingSet, x => Writer.WriteFileItem(x));
 			Writer.WriteHashSet(CandidatesForWorkingSet, x => Writer.WriteFileItem(x));
+			Writer.WriteDictionary(SourceFileToUnityFile, k => Writer.WriteFileItem(k), v => Writer.WriteFileItem(v));
 			Writer.WriteList(UObjectModules, e => e.Write(Writer));
 			Writer.WriteList(UObjectModuleHeaders, x => x.Write(Writer));
 			Writer.WriteHashSet(PluginFiles, x => Writer.WriteFileItem(x));
diff --git a/Engine/Source/Programs/UnrealBuildTool/System/Unity.cs b/Engine/Source/Programs/UnrealBuildTool/System/Unity.cs
--- a/Engine/Source/Programs/UnrealBuildTool/System/Unity.cs
+++ b/Engine/Source/Programs/UnrealBuildTool/System/Unity.cs
@@ -28,6 +28,7 @@ namespace UnrealBuildTool
 		public class FileCollection
 		{
 			public List<FileItem> Files { get; private set; }
+			public List<FileItem> VirtualFiles { get; private set; }
 			public long TotalLength { get; private set; }
 
 			/// The length of this file collection, plus any additional virtual space needed for bUseAdapativeUnityBuild.
@@ -37,6 +38,7 @@ namespace UnrealBuildTool
 			public FileCollection()
 			{
 				Files = new List<FileItem>();
+				VirtualFiles = new List<FileItem>();
 				TotalLength = 0;
 				VirtualLength = 0;
 			}
@@ -56,10 +58,11 @@ namespace UnrealBuildTool
 			/// one of that module's unity blobs.  Basically, it can prevent dozens of files from being recompiled after the first
 			/// time building after your working set of source files changes
 			/// </summary>
-			/// <param name="VirtualFileLength">Length of the virtual file to add to this file collection</param>
-			public void AddVirtualFile(long VirtualFileLength)
+			/// <param name="File">The virtual file to add to the collection</param>
+			public void AddVirtualFile(FileItem File)
 			{
-				VirtualLength += VirtualFileLength;
+				VirtualFiles.Add(File);
+				VirtualLength += File.Length;
 			}
 		}
 
@@ -105,7 +108,7 @@ namespace UnrealBuildTool
 			/// <param name="File">The file to add virtually.  Only the size of the file is tracked.</param>
 			public void AddVirtualFile(FileItem File)
 			{
-				CurrentCollection.AddVirtualFile(File.Length);
+				CurrentCollection.AddVirtualFile(File);
 				if (SplitLength != -1 && CurrentCollection.VirtualLength > SplitLength)
 				{
 					EndCurrentUnityFile();
@@ -334,6 +337,16 @@ namespace UnrealBuildTool
 				// Write the unity file to the intermediate folder.
 				FileItem UnityCPPFile = FileItem.CreateIntermediateTextFile(UnityCPPFilePath, OutputUnityCPPWriter.ToString());
 				NewCPPFiles.Add(UnityCPPFile);
+
+				// Store the mapping of source files to unity files in the makefile
+				foreach(FileItem SourceFile in UnityFile.Files)
+				{
+					Makefile.SourceFileToUnityFile[SourceFile] = UnityCPPFile;
+				}
+				foreach (FileItem SourceFile in UnityFile.VirtualFiles)
+				{
+					Makefile.SourceFileToUnityFile[SourceFile] = UnityCPPFile;
+				}
 			}
 
 			return NewCPPFiles;
