UBT: Ensure that we always write definitions to a force-included header file rather than including them on the command line. Definitions specified on the command line have different restrictions (use of quotes, etc...).

#rnx

diff --git a/Engine/Source/Programs/UnrealBuildTool/Configuration/UEBuildModuleCPP.cs b/Engine/Source/Programs/UnrealBuildTool/Configuration/UEBuildModuleCPP.cs
--- a/Engine/Source/Programs/UnrealBuildTool/Configuration/UEBuildModuleCPP.cs
+++ b/Engine/Source/Programs/UnrealBuildTool/Configuration/UEBuildModuleCPP.cs
@@ -805,6 +805,9 @@ namespace UnrealBuildTool
 
 		static CPPOutput CompileAdaptiveNonUnityFiles(UEToolChain ToolChain, CppCompileEnvironment CompileEnvironment, List<FileItem> Files, DirectoryReference IntermediateDirectory, string ModuleName, List<Action> Actions)
 		{
+			// Write all the definitions out to a separate file
+			CreateHeaderForDefinitions(CompileEnvironment, IntermediateDirectory, "Adaptive");
+
 			// Compile the files
 			return ToolChain.CompileCPPFiles(CompileEnvironment, Files, IntermediateDirectory, ModuleName, Actions);
 		}
