UBT: Adding support for grabbing and parsing compile timing data for VCToolChain and aggregate it.

diff --git a/Engine/Extras/cl-filter/cl-filter.cpp b/Engine/Extras/cl-filter/cl-filter.cpp
--- a/Engine/Extras/cl-filter/cl-filter.cpp
+++ b/Engine/Extras/cl-filter/cl-filter.cpp
@@ -27,19 +27,21 @@ int wmain(int ArgC, const wchar_t* ArgV[])
 	ChildCmdLine += 4;
 
 	// Make sure we've got an output file and compiler path
-	if (ArgC <= 4 || wcscmp(ArgV[2], L"--") != 0)
+	if (ArgC <= 4 || (wcscmp(ArgV[2], L"--") != 0 && wcscmp(ArgV[3], L"--") != 0))
 	{
-		wprintf(L"ERROR: Syntax: cl-filter <dependencies-file> -- <child command line>\n");
+		wprintf(L"ERROR: Syntax: cl-filter <dependencies-file> [timing-file] -- <child command line>\n");
 		return -1;
 	}
 
-	// Get the arguments we care about
+	// Get the arguments we care about.
+	bool ParseTimingInfo = wcscmp(ArgV[2], L"--") != 0 && wcscmp(ArgV[3], L"--") == 0;
 	const wchar_t* OutputFileName = ArgV[1];
-	const wchar_t* CompilerFileName = ArgV[3];
+	const wchar_t* TimingFileName = ParseTimingInfo ? ArgV[2] : nullptr;
+	const wchar_t* CompilerFileName = ParseTimingInfo ? ArgV[4] : ArgV[3];
 
 	// Get all the possible localized string prefixes for /showIncludes output
-	std::vector<std::vector<char>> LocalizedPrefixes;
-	GetLocalizedIncludePrefixes(CompilerFileName, LocalizedPrefixes);
+	std::vector<std::vector<char>> LocalizedIncludePrefixes;
+	GetLocalizedIncludePrefixes(CompilerFileName, LocalizedIncludePrefixes);
 
 	// Create the child process
 	PROCESS_INFORMATION ProcessInfo;
@@ -84,14 +86,24 @@ int wmain(int ArgC, const wchar_t* ArgV[])
 	CloseHandle(StdOutWriteHandle);
 	CloseHandle(StdErrWriteHandle);
 
-	// Delete the output file
+	// Delete the output file.
 	DeleteFileW(OutputFileName);
 
-	// Get the path to a temporary output filename
+	// Delete the timing file if needed.
+	if (ParseTimingInfo)
+	{
+		DeleteFileW(TimingFileName);
+	}
+
+	// Get the path to a temporary output filename.
 	std::wstring TempOutputFileName(OutputFileName);
 	TempOutputFileName += L".tmp";
 
-	// Create a file to contain the dependency list
+	// Get the path to a temporary timing filename.
+	std::wstring TempTimingFileName(TimingFileName == nullptr ? L"" : TimingFileName);
+	TempTimingFileName += L".tmp";
+
+	// Create a file to contain the dependency list.
 	HANDLE OutputFile = CreateFileW(TempOutputFileName.c_str(), GENERIC_WRITE, FILE_SHARE_READ, NULL, CREATE_ALWAYS, FILE_ATTRIBUTE_NORMAL, NULL);
 	if(OutputFile == INVALID_HANDLE_VALUE)
 	{
@@ -99,10 +111,24 @@ int wmain(int ArgC, const wchar_t* ArgV[])
 		return -1;
 	}
 
+	// If needed, create a file to contain unparsed timing info.
+	HANDLE TimingFile = INVALID_HANDLE_VALUE;
+	if (ParseTimingInfo)
+	{
+		TimingFile = CreateFileW(TempTimingFileName.c_str(), GENERIC_WRITE, FILE_SHARE_READ, NULL, CREATE_ALWAYS, FILE_ATTRIBUTE_NORMAL, NULL);
+		if (TimingFile == INVALID_HANDLE_VALUE)
+		{
+			wprintf(L"ERROR: Unable to open %s for output", TempTimingFileName.c_str());
+			return -1;
+		}
+	}
+
 	// Pipe the output to stdout
 	char Buffer[1024];
 	size_t BufferSize = 0;
 
+	bool InTimingInfo = false;
+	int TimingLinesToCapture = 0;
 	for (;;)
 	{
 		// Read the next chunk of data from the output stream
@@ -147,7 +173,7 @@ int wmain(int ArgC, const wchar_t* ArgV[])
 			}
 
 			// Filter out any lines that have the "Note: including file: " prefix.
-			for (const std::vector<char>& LocalizedPrefix : LocalizedPrefixes)
+			for (const std::vector<char>& LocalizedPrefix : LocalizedIncludePrefixes)
 			{
 				if (memcmp(Buffer + LineStart, LocalizedPrefix.data(), LocalizedPrefix.size() - 1) == 0)
 				{
@@ -164,6 +190,43 @@ int wmain(int ArgC, const wchar_t* ArgV[])
 				}
 			}
 
+			// Handling timing info parsing if needed.
+			if (ParseTimingInfo && LineStart < LineEnd)
+			{
+				// See if we've moved into a block of timing information. Order is always Headers -> Classes -> Functions, so to keep for loops to a minimum,
+				//   handle by what our current state is.
+				const char StartTimingInfoText[] = "Include Headers:";
+				const char EndOfTimingInfoText[] = "\tLeast Hits:";
+				if (!InTimingInfo)
+				{
+					if (memcmp(Buffer + LineStart, StartTimingInfoText, 16) == 0)
+					{
+						InTimingInfo = true;
+						TimingLinesToCapture = 1;
+					}
+				}
+				else
+				{
+					if (memcmp(Buffer + LineStart, EndOfTimingInfoText, 12) == 0)
+					{
+						InTimingInfo = false;
+						TimingLinesToCapture = 3;
+					}
+					else
+					{
+						TimingLinesToCapture = 1;
+					}
+				}
+
+				if (TimingLinesToCapture > 0)
+				{
+					--TimingLinesToCapture;
+					DWORD BytesWritten;
+					WriteFile(TimingFile, Buffer + LineStart, (DWORD)(LineEnd - LineStart), &BytesWritten, NULL);
+					LineStart = LineEnd;
+				}
+			}
+
 			// If we didn't write anything out, write it to stdout
 			if(LineStart < LineEnd)
 			{
@@ -196,6 +259,16 @@ int wmain(int ArgC, const wchar_t* ArgV[])
 		ExitCode = 1;
 	}
 
+	if (ParseTimingInfo)
+	{
+		CloseHandle(TimingFile);
+		if (ExitCode == 0 && !MoveFileW(TempTimingFileName.c_str(), TimingFileName))
+		{
+			wprintf(L"ERROR: Unable to rename %s to %s\n", TempTimingFileName.c_str(), TimingFileName);
+			ExitCode = 1;
+		}
+	}
+
 	return ExitCode;
 }
 
diff --git a/Engine/Source/Programs/AutomationTool/AutomationUtils/UE4Build.cs b/Engine/Source/Programs/AutomationTool/AutomationUtils/UE4Build.cs
--- a/Engine/Source/Programs/AutomationTool/AutomationUtils/UE4Build.cs
+++ b/Engine/Source/Programs/AutomationTool/AutomationUtils/UE4Build.cs
@@ -32,6 +32,7 @@ namespace AutomationTool
 	[Help("ForceNonUnity", "Toggle to disable the unity build system")]
 	[Help("ForceUnity", "Toggle to force enable the unity build system")]
 	[Help("Licensee", "If set, this build is being compiled by a licensee")]
+	[Help("Timing", "If set, this build will attempt to generate compile timing data")]
 	public class UE4Build
 	{
 		private BuildCommand OwnerCommand;
diff --git a/Engine/Source/Programs/AutomationTool/BuildGraph/BuildGraph.cs b/Engine/Source/Programs/AutomationTool/BuildGraph/BuildGraph.cs
--- a/Engine/Source/Programs/AutomationTool/BuildGraph/BuildGraph.cs
+++ b/Engine/Source/Programs/AutomationTool/BuildGraph/BuildGraph.cs
@@ -68,6 +68,7 @@ namespace AutomationTool
 	[Help("SharedStorageDir=<DirName>", "Sets the directory to use to transfer build products between agents in a build farm")]
 	[Help("SingleNode=<Name>", "Run only the given node. Intended for use on a build system after running with -Export.")]
 	[Help("WriteToSharedStorage", "Allow writing to shared storage. If not set, but -SharedStorageDir is specified, build products will read but not written")]
+	[Help("Timing", "If set, this build will attempt to generate compile timing data")]
 	public class BuildGraph : BuildCommand
 	{
 		/// <summary>
@@ -94,7 +95,8 @@ namespace AutomationTool
 			bool bShowDiagnostics = ParseParam("ShowDiagnostics");
 			bool bWriteToSharedStorage = ParseParam("WriteToSharedStorage") || CommandUtils.IsBuildMachine;
 			bool bPublicTasksOnly = ParseParam("PublicTasksOnly");
-			string ReportName = ParseParamValue("ReportName", null); 
+			string ReportName = ParseParamValue("ReportName", null);
+			bool bTiming = ParseParam("Timing");
 
 			GraphPrintOptions PrintOptions = GraphPrintOptions.ShowCommandLineOptions;
 			if(ParseParam("ShowDeps"))
@@ -128,6 +130,7 @@ namespace AutomationTool
 			DefaultProperties["HostPlatform"] = HostPlatform.Current.HostEditorPlatform.ToString();
 			DefaultProperties["RestrictedFolderNames"] = String.Join(";", RestrictedFolders.Names);
 			DefaultProperties["RestrictedFolderFilter"] = String.Join(";", RestrictedFolders.Names.Select(x => String.Format(".../{0}/...", x)));
+			DefaultProperties["Timing"] = bTiming.ToString();
 
 			// Attempt to read existing Build Version information
 			BuildVersion Version;
diff --git a/Engine/Source/Programs/UnrealBuildTool/Configuration/TargetRules.cs b/Engine/Source/Programs/UnrealBuildTool/Configuration/TargetRules.cs
--- a/Engine/Source/Programs/UnrealBuildTool/Configuration/TargetRules.cs
+++ b/Engine/Source/Programs/UnrealBuildTool/Configuration/TargetRules.cs
@@ -981,6 +981,13 @@ namespace UnrealBuildTool
 		[XmlConfigFile(Category = "BuildConfiguration")]
 		public bool bPrintToolChainTimingInfo = false;
 
+		/// <summary>
+		/// Whether to parse timing data into a tracing file compatible with chrome://tracing.
+		/// </summary>
+		[CommandLine("-Tracing")]
+		[XmlConfigFile(Category = "BuildConfiguration")]
+		public bool bParseTimingInfoForTracing = false;
+
 		/// <summary>
 		/// Whether to hide symbols by default on POSIX platforms
 		/// </summary>
@@ -2167,6 +2174,11 @@ namespace UnrealBuildTool
 			get { return Inner.bPrintToolChainTimingInfo; }
 		}
 
+		public bool bParseTimingInfoForTracing
+		{
+			get { return Inner.bParseTimingInfoForTracing; }
+		}
+
 		public bool bHideSymbolsByDefault
 		{
 			get { return Inner.bHideSymbolsByDefault; }
diff --git a/Engine/Source/Programs/UnrealBuildTool/Configuration/UEBuildTarget.cs b/Engine/Source/Programs/UnrealBuildTool/Configuration/UEBuildTarget.cs
--- a/Engine/Source/Programs/UnrealBuildTool/Configuration/UEBuildTarget.cs
+++ b/Engine/Source/Programs/UnrealBuildTool/Configuration/UEBuildTarget.cs
@@ -2929,6 +2929,7 @@ namespace UnrealBuildTool
 			GlobalCompileEnvironment.bAllowRemotelyCompiledPCHs = Rules.bAllowRemotelyCompiledPCHs;
 			GlobalCompileEnvironment.bCheckSystemHeadersForModification = Rules.bCheckSystemHeadersForModification;
 			GlobalCompileEnvironment.bPrintTimingInfo = Rules.bPrintToolChainTimingInfo;
+			GlobalCompileEnvironment.bParseTimingInfoForTracing = Rules.bParseTimingInfoForTracing;
 			GlobalCompileEnvironment.bUseRTTI = Rules.bForceEnableRTTI;
 			GlobalCompileEnvironment.bUseInlining = Rules.bUseInlining;
 			GlobalCompileEnvironment.bHideSymbolsByDefault = Rules.bHideSymbolsByDefault;
diff --git a/Engine/Source/Programs/UnrealBuildTool/Platform/Windows/VCToolChain.cs b/Engine/Source/Programs/UnrealBuildTool/Platform/Windows/VCToolChain.cs
--- a/Engine/Source/Programs/UnrealBuildTool/Platform/Windows/VCToolChain.cs
+++ b/Engine/Source/Programs/UnrealBuildTool/Platform/Windows/VCToolChain.cs
@@ -1116,10 +1116,7 @@ namespace UnrealBuildTool
 
 				if(CompileEnvironment.bGenerateDependenciesFile)
 				{
-					CompileAction.DependencyListFile = FileItem.GetItemByFileReference(FileReference.Combine(OutputDir, SourceFile.Location.GetFileName() + ".txt"));
-					CompileAction.CommandArguments = String.Format("{0} -- {1} {2} /showIncludes", Utils.MakePathSafeToUseWithCommandLine(CompileAction.DependencyListFile.Location), Utils.MakePathSafeToUseWithCommandLine(CompileAction.CommandPath), CompileAction.CommandArguments);
-					CompileAction.CommandPath = FileReference.Combine(UnrealBuildTool.EngineDirectory, "Build", "Windows", "cl-filter", "cl-filter.exe");
-					CompileAction.ProducedItems.Add(CompileAction.DependencyListFile);
+					GenerateDependenciesFile(CompileEnvironment, OutputDir, Result, SourceFile, Actions, CompileAction);
 				}
 
 				if (CompileEnvironment.PrecompiledHeaderAction == PrecompiledHeaderAction.Create)
@@ -1150,9 +1147,83 @@ namespace UnrealBuildTool
 
 				Actions.Add(CompileAction);
 			}
+
 			return Result;
 		}
 
+		private void GenerateDependenciesFile(CppCompileEnvironment CompileEnvironment, DirectoryReference OutputDir, CPPOutput Result, FileItem SourceFile, List<Action> Actions, Action CompileAction)
+		{
+			var commandArguments = new List<string>();
+			CompileAction.DependencyListFile = FileItem.GetItemByFileReference(FileReference.Combine(OutputDir, $"{SourceFile.Location.GetFileName()}.txt"));
+			CompileAction.ProducedItems.Add(CompileAction.DependencyListFile);
+			commandArguments.Add(Utils.MakePathSafeToUseWithCommandLine(CompileAction.DependencyListFile.Location));
+
+			if (CompileEnvironment.bPrintTimingInfo)
+			{
+				CompileAction.TimingFile = FileItem.GetItemByFileReference(FileReference.Combine(OutputDir, $"{SourceFile.Location.GetFileName()}.timing.txt"));
+				CompileAction.ProducedItems.Add(CompileAction.TimingFile);
+				commandArguments.Add(Utils.MakePathSafeToUseWithCommandLine(CompileAction.TimingFile.Location));
+
+				var TimingJsonFile = FileItem.GetItemByPath(Path.ChangeExtension(CompileAction.TimingFile.AbsolutePath, ".bin"));
+				Result.DebugDataFiles.Add(TimingJsonFile);
+				Result.DebugDataFiles.Add(CompileAction.TimingFile);
+
+				var ParseTimingArguments = new List<string>() { $"-TimingFile=\"{CompileAction.TimingFile}\"" };
+				if (CompileEnvironment.bParseTimingInfoForTracing)
+				{
+					ParseTimingArguments.Add("-Tracing");
+				}
+
+				var ParseTimingInfoAction = Action.CreateRecursiveAction<ParseMsvcTimingInfoMode>(ActionType.ParseTimingInfo, string.Join(" ", ParseTimingArguments));
+				ParseTimingInfoAction.WorkingDirectory = UnrealBuildTool.EngineSourceDirectory;
+				ParseTimingInfoAction.StatusDescription = $"Parsing Timing File '{CompileAction.TimingFile}'";
+				ParseTimingInfoAction.bCanExecuteRemotely = true;
+				ParseTimingInfoAction.PrerequisiteItems.Add(SourceFile);
+				ParseTimingInfoAction.PrerequisiteItems.Add(CompileAction.TimingFile);
+				ParseTimingInfoAction.ProducedItems.Add(TimingJsonFile);
+				Actions.Add(ParseTimingInfoAction);
+			}
+
+			commandArguments.Add("--");
+			commandArguments.Add(Utils.MakePathSafeToUseWithCommandLine(CompileAction.CommandPath));
+			commandArguments.Add(CompileAction.CommandArguments);
+			commandArguments.Add("/showIncludes");
+			CompileAction.CommandArguments = string.Join(" ", commandArguments);
+			CompileAction.CommandPath = FileReference.Combine(UnrealBuildTool.EngineDirectory, "Build", "Windows", "cl-filter", "cl-filter.exe");	
+		}
+
+		public override void FinalizeOutput(ReadOnlyTargetRules Target, TargetMakefile Makefile)
+		{
+			if (Target.bPrintToolChainTimingInfo)
+			{
+				var ParseTimingActions = Makefile.Actions.Where(x => x.ActionType == ActionType.ParseTimingInfo);
+				Makefile.OutputItems.AddRange(ParseTimingActions.SelectMany(x => x.ProducedItems));
+
+				// Handing generating aggregate timing information if we compiled more than one file.
+				if (ParseTimingActions.Count() > 1)
+				{
+					// Find all .timing.json files.
+					var TimingJsonFiles = ParseTimingActions.Select(a => a.ProducedItems.First());
+
+					// Generate the file manifest for the aggregator.
+					var ManifestFile = FileReference.Combine(Makefile.ProjectIntermediateDirectory, $"{Target.Name}TimingManifest.txt");
+					File.WriteAllLines(ManifestFile.FullName, TimingJsonFiles.Select(f => f.AbsolutePath));
+
+					var AggregateTimingInfoAction = Action.CreateRecursiveAction<AggregateParsedTimingInfo>(ActionType.ParseTimingInfo, $"-Name={Target.Name} -ManifestFile=\"{ManifestFile.FullName}\"");
+					AggregateTimingInfoAction.WorkingDirectory = UnrealBuildTool.EngineSourceDirectory;
+					AggregateTimingInfoAction.StatusDescription = $"Aggregating {TimingJsonFiles.Count()} Timing File(s)";
+					AggregateTimingInfoAction.bCanExecuteRemotely = false;
+					AggregateTimingInfoAction.PrerequisiteItems.AddRange(ParseTimingActions.SelectMany(a => a.PrerequisiteItems));
+					AggregateTimingInfoAction.PrerequisiteItems.AddRange(TimingJsonFiles);
+
+					var AggregateOutputFile = FileItem.GetItemByFileReference(FileReference.Combine(Makefile.ProjectIntermediateDirectory, $"{Target.Name}.timing.bin"));
+					AggregateTimingInfoAction.ProducedItems.Add(AggregateOutputFile);
+					Makefile.OutputItems.Add(AggregateOutputFile);
+					Makefile.Actions.Add(AggregateTimingInfoAction);
+				}
+			}
+		}
+
 		public override CPPOutput CompileRCFiles(CppCompileEnvironment CompileEnvironment, List<FileItem> InputFiles, DirectoryReference OutputDir, List<Action> Actions)
 		{
 			CPPOutput Result = new CPPOutput();
diff --git a/Engine/Source/Programs/UnrealBuildTool/System/Action.cs b/Engine/Source/Programs/UnrealBuildTool/System/Action.cs
--- a/Engine/Source/Programs/UnrealBuildTool/System/Action.cs
+++ b/Engine/Source/Programs/UnrealBuildTool/System/Action.cs
@@ -29,6 +29,8 @@ namespace UnrealBuildTool
 		WriteMetadata,
 
 		PostBuildStep,
+
+		ParseTimingInfo,
 	}
 
 	/// <summary>
@@ -65,6 +67,11 @@ namespace UnrealBuildTool
 		/// </summary>
 		public FileItem DependencyListFile;
 
+		/// <summary>
+		/// For C++ source files, specifies a timing file used to track timing information.
+		/// </summary>
+		public FileItem TimingFile;
+
 		/// <summary>
 		/// Set of other actions that this action depends on. This set is built when the action graph is linked.
 		/// </summary>
diff --git a/Engine/Source/Programs/UnrealBuildTool/System/AggregateParsedTimingInfo.cs b/Engine/Source/Programs/UnrealBuildTool/System/AggregateParsedTimingInfo.cs
new file mode 100644
--- /dev/null
+++ b/Engine/Source/Programs/UnrealBuildTool/System/AggregateParsedTimingInfo.cs
@@ -0,0 +1,210 @@
+// Copyright 1998-2019 Epic Games, Inc. All Rights Reserved.
+
+using System;
+using System.Collections.Concurrent;
+using System.Collections.Generic;
+using System.IO;
+using System.IO.Compression;
+using System.Linq;
+using System.Text.RegularExpressions;
+using System.Threading;
+using System.Threading.Tasks;
+using Tools.DotNETCommon;
+
+namespace UnrealBuildTool
+{
+	/// <summary>
+	/// Aggregates parsed Visual C++ timing information files together into one monolithic breakdown file.
+	/// </summary>
+	[ToolMode("AggregateParsedTimingInfo", ToolModeOptions.None)]
+	class AggregateParsedTimingInfo : ToolMode
+	{
+		private TimingData FileTimingData;
+		private ConcurrentDictionary<string, int> DecompressedFileSizes = new ConcurrentDictionary<string, int>();
+		private ConcurrentDictionary<string, byte[]> CompressedFiles = new ConcurrentDictionary<string, byte[]>();
+		private object AddFileLock = new object();
+
+		public override int Execute(CommandLineArguments Arguments)
+		{
+			var ManifestFile = Arguments.GetFileReference("-ManifestFile=");
+			string[] ParsedFileNames = FileReference.ReadAllLines(ManifestFile);
+
+			// Load all the file timing data and summarize them for the aggregate.
+			FileTimingData = new TimingData() { Name = "Files", Type = TimingDataType.Summary };
+			Parallel.ForEach(ParsedFileNames, new ParallelOptions() { MaxDegreeOfParallelism = 4 }, ParseTimingDataFile);
+
+			// Create aggregate summary. Duration is the duration of the files in the aggregate.
+			var AggregateName = Arguments.GetString("-Name=");
+			var AggregateData = new TimingData() { Name = AggregateName, Type = TimingDataType.Aggregate };
+			AggregateData.AddChild(FileTimingData);
+
+			// Write out aggregate summary.
+			var OutputFile = Path.Combine(ManifestFile.Directory.FullName, $"{AggregateName}.timing.bin");
+			using (var Writer = new BinaryWriter(File.Open(OutputFile, FileMode.Create)))
+			{
+				// Write out the aggregate data.
+				Writer.Write(AggregateData);
+
+				// Write the look up table for the compressed binary blobs.
+				int Offset = 0;
+				Writer.Write(CompressedFiles.Count);
+				foreach (var CompressedFile in CompressedFiles)
+				{
+					Writer.Write(CompressedFile.Key);
+					Writer.Write(Offset);
+					Writer.Write(CompressedFile.Value.Length);
+					Writer.Write(DecompressedFileSizes[CompressedFile.Key]);
+					Offset += CompressedFile.Value.Length;
+				}
+
+				// Write the compressed binary blobs.
+				foreach (var CompressedFile in CompressedFiles)
+				{
+					Writer.Write(CompressedFile.Value);
+				}
+			}
+
+			return 0;
+		}
+
+		private void ParseTimingDataFile(string ParsedFileName)
+		{
+			// Convert input file back into summary objects.
+			using (BinaryReader Reader = new BinaryReader(File.Open(ParsedFileName, FileMode.Open, FileAccess.Read)))
+			{
+				var ParsedTimingData = new TimingData(Reader);
+
+				var CompressDataTask = Task.Run(() =>
+				{
+					Reader.BaseStream.Seek(0, SeekOrigin.Begin);
+					using (var CompressedMemoryStream = new MemoryStream())
+					{
+						using (var CompressionStream = new GZipStream(CompressedMemoryStream, CompressionMode.Compress))
+						{
+							Reader.BaseStream.CopyTo(CompressionStream);
+						}
+
+						CompressedFiles.TryAdd(ParsedTimingData.Name, CompressedMemoryStream.ToArray());
+						DecompressedFileSizes.TryAdd(ParsedTimingData.Name, (int)Reader.BaseStream.Length);
+					}
+				});
+
+				var SummarizedTimingData = new TimingData() { Name = ParsedTimingData.Name, Parent = FileTimingData, Type = TimingDataType.Summary };
+
+				// Gather and update child timing data.
+				var SummarizedIncludes = new TimingData() { Parent = SummarizedTimingData, Type = TimingDataType.Summary };
+				var QueueIncludesTask = Task.Run(() =>
+				{
+					SummarizeTimingData(ParsedTimingData.Children["IncludeTimings"].Children.Values, SummarizedIncludes, "Includes", 10);
+				});
+
+				var SummarizedClasses = new TimingData() { Parent = SummarizedTimingData, Type = TimingDataType.Summary };
+				var QueueClassesTask = Task.Run(() =>
+				{
+					// Collapse templates into single entries.
+					var CollapsedClasses = GroupChildren(ParsedTimingData.Children["ClassTimings"].Children.Values, TimingDataType.Class);
+					SummarizeTimingData(CollapsedClasses, SummarizedClasses, "Classes", 10);
+				});
+
+				var SummarizedFunctions = new TimingData() { Parent = SummarizedTimingData, Type = TimingDataType.Summary };
+				var QueueFunctionsTask = Task.Run(() =>
+				{
+					// Collapse templates into single entries.
+					var CollapsedFunctions = GroupChildren(ParsedTimingData.Children["FunctionTimings"].Children.Values, TimingDataType.Function);
+					SummarizeTimingData(CollapsedFunctions, SummarizedFunctions, "Functions", 10);
+				});
+
+				while (!CompressDataTask.IsCompleted || !QueueIncludesTask.IsCompleted || !QueueClassesTask.IsCompleted || !QueueFunctionsTask.IsCompleted)
+				{
+					Thread.Sleep(TimeSpan.FromMilliseconds(50));
+				}
+
+				SummarizedTimingData.AddChild(SummarizedIncludes);
+				SummarizedTimingData.AddChild(SummarizedClasses);
+				SummarizedTimingData.AddChild(SummarizedFunctions);
+
+				lock (AddFileLock)
+				{
+					FileTimingData.AddChild(SummarizedTimingData);
+				}
+			}
+		}
+
+		private IEnumerable<TimingData> GroupChildren(IEnumerable<TimingData> Children, TimingDataType Type)
+		{
+			var GroupedChildren = new List<TimingData>();
+			var ChildGroups = new Dictionary<string, List<TimingData>>();
+			foreach (var Child in Children)
+			{
+				// See if this is a templated class. If not, add it as is.
+				var Match = Regex.Match(Child.Name, @"^([^<]*)(?<Template><.*>)");
+				if (!Match.Success)
+				{
+					GroupedChildren.Add(Child);
+				}
+				else
+				{
+					// Generate group name from template.
+					var TemplateParamCount = Match.Groups["Template"].Value.Count(c => c == ',') + 1;
+					var TemplateParamSig = new List<string>(TemplateParamCount);
+					for (int i = 0; i < TemplateParamCount; ++i)
+					{
+						TemplateParamSig.Add("...");
+					}
+
+					var GroupName = Child.Name.Replace(Match.Groups["Template"].Value, $"<{string.Join(", ", TemplateParamSig)}>");
+
+					// See if we have a group for this template already. If not, add it.
+					if (!ChildGroups.ContainsKey(GroupName))
+					{
+						ChildGroups.Add(GroupName, new List<TimingData>());
+					}
+
+					ChildGroups[GroupName].Add(Child);
+				}
+			}
+
+			// Add grouped children.
+			foreach (var Group in ChildGroups)
+			{
+				if (Group.Value.Count == 1)
+				{
+					GroupedChildren.Add(Group.Value.First());
+					continue;
+				}
+
+				var NewViewModel = new TimingData()
+				{
+					Name = Group.Key,
+					Type = Type,
+				};
+
+				foreach (var Child in Group.Value)
+				{
+					var NewChild = new TimingData() { Name = Child.Name, Type = Child.Type, Parent = NewViewModel, ExclusiveDuration = Child.ExclusiveDuration };
+					NewViewModel.AddChild(NewChild);
+				}
+
+				GroupedChildren.Add(NewViewModel);
+			}
+
+			return GroupedChildren;
+		}
+
+		private void SummarizeTimingData(IEnumerable<TimingData> Data, TimingData Summary, string TypeName, int NumberToKeep)
+		{
+			// Sort the data by inclusive duration, highest first, so we can grab the top entries to keep.
+			var Sorted = Data.OrderByDescending(v => v.InclusiveDuration).ToList();
+			Summary.Name = $"Total {TypeName}: {Sorted.Count}";
+			var TopData = Sorted.Take(NumberToKeep);
+			foreach (var Child in TopData)
+			{
+				var SummarizedChild = new TimingData() { Name = Child.Name, Parent = Summary, Type = Child.Type, ExclusiveDuration = Child.InclusiveDuration };
+				Summary.AddChild(SummarizedChild);
+			}
+
+			// Set the summary to the duration of all children, not just the top ten.
+			Summary.ExclusiveDuration = Sorted.Except(TopData).Sum(c => c.InclusiveDuration);
+		}
+	}
+}
diff --git a/Engine/Source/Programs/UnrealBuildTool/System/CppCompileEnvironment.cs b/Engine/Source/Programs/UnrealBuildTool/System/CppCompileEnvironment.cs
--- a/Engine/Source/Programs/UnrealBuildTool/System/CppCompileEnvironment.cs
+++ b/Engine/Source/Programs/UnrealBuildTool/System/CppCompileEnvironment.cs
@@ -302,6 +302,11 @@ namespace UnrealBuildTool
 		/// </summary>
 		public bool bPrintTimingInfo;
 
+		/// <summary>
+		/// Whether to parse timing information for use with chrome://tracing
+		/// </summary>
+		public bool bParseTimingInfoForTracing;
+
 		/// <summary>
 		/// Whether to output a dependencies file along with the output build products
 		/// </summary>
@@ -427,6 +432,7 @@ namespace UnrealBuildTool
 			PGOFilenamePrefix = Other.PGOFilenamePrefix;
 			PGODirectory = Other.PGODirectory;
 			bPrintTimingInfo = Other.bPrintTimingInfo;
+			bParseTimingInfoForTracing = Other.bParseTimingInfoForTracing;
 			bGenerateDependenciesFile = Other.bGenerateDependenciesFile;
 			bAllowRemotelyCompiledPCHs = Other.bAllowRemotelyCompiledPCHs;
 			UserIncludePaths = new HashSet<DirectoryReference>(Other.UserIncludePaths);
diff --git a/Engine/Source/Programs/UnrealBuildTool/System/TimingData.cs b/Engine/Source/Programs/UnrealBuildTool/System/TimingData.cs
new file mode 100644
--- /dev/null
+++ b/Engine/Source/Programs/UnrealBuildTool/System/TimingData.cs
@@ -0,0 +1,173 @@
+// Copyright 1998-2019 Epic Games, Inc. All Rights Reserved.
+
+using System;
+using System.Collections.Generic;
+using System.IO;
+using System.Linq;
+using Tools.DotNETCommon;
+
+namespace UnrealBuildTool
+{
+	/// <summary>
+	/// Types of timing data events.
+	/// </summary>
+	public enum TimingDataType
+	{
+		/// <summary>
+		/// No type assigned.
+		/// </summary>
+		None,
+
+		/// <summary>
+		/// This event represents an aggregate of summary events.
+		/// </summary>
+		Aggregate,
+
+		/// <summary>
+		/// This event represents a summary of related include, class, and function events, usually
+		/// summarizing a single compilation unit or part of a compilation unit.
+		/// </summary>
+		Summary,
+
+		/// <summary>
+		/// Timing related to compiling include files.
+		/// </summary>
+		Include,
+
+		/// <summary>
+		/// Timing related to compiling classes.
+		/// </summary>
+		Class,
+
+		/// <summary>
+		/// Timing related to compiling functions.
+		/// </summary>
+		Function,
+	}
+
+	/// <summary>
+	/// Class used to capture timing data from events like those generated by MSVC when timing data
+	/// is being generated.
+	/// </summary>
+	[Serializable]
+	public class TimingData : IBinarySerializable
+	{
+		/// <summary>
+		/// Constructs a new instance of <see cref="TimingData"/>.
+		/// </summary>
+		public TimingData()
+		{
+		}
+
+		/// <summary>
+		/// Constructs a new instance of <see cref="TimingData"/> using data from a <see cref="BinaryReader"/>.
+		/// </summary>
+		/// <param name="Reader">The <see cref="BinaryReader"/> to get the data from.</param>
+		public TimingData(BinaryReader Reader)
+		{
+			Name = Reader.ReadString();
+			Type = (TimingDataType)Reader.ReadByte();
+			ExclusiveDuration = (float)Reader.ReadDouble();
+			int ChildCount = Reader.ReadInt32();
+			for (int i = 0; i < ChildCount; ++i)
+			{
+				var NewChild = new TimingData(Reader);
+				Children.Add(NewChild.Name, NewChild);
+			}
+		}
+
+		/// <summary>
+		/// The name of the event being represented, such as a file or category name.
+		/// </summary>
+		public string Name { get; set; }
+
+		/// <summary>
+		/// The type of timing data this event represents.
+		/// </summary>
+		public TimingDataType Type { get; set; }
+
+		/// <summary>
+		/// Gets the amount of time, in milliseconds, that the event took to complete disregarding its sub-events.
+		/// </summary>
+		public double ExclusiveDuration { get; set; }
+
+		/// <summary>
+		/// Gets the amount of time, in milliseconds, that the event took including its sub-events.
+		/// </summary>
+		public double InclusiveDuration
+		{
+			get
+			{
+				return ExclusiveDuration + Children.Sum(c => c.Value.InclusiveDuration);
+			}
+		}
+
+		/// <summary>
+		/// If this event is a sub-event, the parent is the event it is the sub-event to.
+		/// </summary>
+		public TimingData Parent { get; set; }
+
+		/// <summary>
+		/// Any sub-events of this event.
+		/// </summary>
+		public Dictionary<string, TimingData> Children { get; set; } = new Dictionary<string, TimingData>();
+
+		/// <summary>
+		/// Adds a sub-event (child) to this event if it doesn't exist, or increase the duration by the
+		/// duration of the provided event.
+		/// </summary>
+		/// <param name="ChildData">The sub-event to add to this event.</param>
+		public void AddChild(TimingData ChildData)
+		{
+			if (Children.TryGetValue(ChildData.Name, out var MatchingData))
+			{
+				MatchingData.ExclusiveDuration += ChildData.ExclusiveDuration;
+				foreach (var ChildChildData in ChildData.Children.Values)
+				{
+					MatchingData.AddChild(ChildChildData);
+				}
+			}
+			else
+			{
+				Children.Add(ChildData.Name, ChildData);
+			}
+		}
+
+		/// <summary>
+		/// Creates a deep clone of this event.
+		/// </summary>
+		/// <returns>A deep clone of this event.</returns>
+		public TimingData Clone()
+		{
+			var ClonedTimingData = new TimingData()
+			{
+				Name = Name,
+				Type = Type,
+				ExclusiveDuration = ExclusiveDuration,
+			};
+
+			foreach (var Child in Children)
+			{
+				ClonedTimingData.Children.Add(Child.Key, Child.Value.Clone());
+			}
+
+			return ClonedTimingData;
+		}
+
+		/// <summary>
+		/// Writes a binary representation of this event and its sub-events.
+		/// </summary>
+		/// <param name="Writer">The <see cref="BinaryWriter"/> to write the binary representation to.</param>
+		public void Write(BinaryWriter Writer)
+		{
+			Writer.Write(Name);
+			Writer.Write((byte)Type);
+			Writer.Write(ExclusiveDuration);
+			Writer.Write(Children.Count);
+			foreach (var Child in Children.Values)
+			{
+				Writer.Write(Child);
+			}
+		}
+	}
+}
diff --git a/Engine/Source/Programs/UnrealBuildTool/UnrealBuildTool.csproj b/Engine/Source/Programs/UnrealBuildTool/UnrealBuildTool.csproj
--- a/Engine/Source/Programs/UnrealBuildTool/UnrealBuildTool.csproj
+++ b/Engine/Source/Programs/UnrealBuildTool/UnrealBuildTool.csproj
@@ -185,12 +185,14 @@
     <Compile Include="ProjectFiles\Xcode\XcodeProjectFileGenerator.cs" />
     <Compile Include="System\Action.cs" />
     <Compile Include="System\ActionExecutor.cs" />
+    <Compile Include="System\AggregateParsedTimingInfo.cs" />
     <Compile Include="System\Asn.cs" />
     <Compile Include="System\BuildVersion.cs" />
     <Compile Include="System\CppCompileEnvironment.cs" />
     <Compile Include="System\CppDependencyCache.cs" />
     <Compile Include="System\CppIncludeParser.cs" />
     <Compile Include="System\FileMetadataPrefetch.cs" />
+    <Compile Include="System\ParseMsvcTimingInfo.cs" />
     <Compile Include="System\SourceFileMetadataCache.cs" />
     <Compile Include="System\DataDrivenPlatformInfo.cs" />
     <Compile Include="System\DirectoryItem.cs" />
@@ -229,6 +231,7 @@
     <Compile Include="System\TargetMakefile.cs" />
     <Compile Include="System\TargetReceipt.cs" />
     <Compile Include="System\Timeline.cs" />
+    <Compile Include="System\TimingData.cs" />
     <Compile Include="System\ToolMode.cs" />
     <Compile Include="System\VersionNumber.cs" />
     <Compile Include="System\XmlConfig.cs" />
