Fix compile errors on build farm (older compiler version?)

diff --git a/Engine/Source/Programs/UnrealBuildTool/System/ParseMsvcTimingInfo.cs b/Engine/Source/Programs/UnrealBuildTool/System/ParseMsvcTimingInfo.cs
--- a/Engine/Source/Programs/UnrealBuildTool/System/ParseMsvcTimingInfo.cs
+++ b/Engine/Source/Programs/UnrealBuildTool/System/ParseMsvcTimingInfo.cs
@@ -120,7 +120,8 @@ namespace UnrealBuildTool
 			foreach (var Data in ParsedTimingData)
 			{
 				// See if we've already added a child that matches this data's name. If so, just add to the duration.
-				if (Summary.Children.TryGetValue(Data.Name, out var MatchedData))
+				TimingData MatchedData;
+				if (Summary.Children.TryGetValue(Data.Name, out MatchedData))
 				{
 					MatchedData.ExclusiveDuration += Data.ExclusiveDuration;
 				}
diff --git a/Engine/Source/Programs/UnrealBuildTool/System/TimingData.cs b/Engine/Source/Programs/UnrealBuildTool/System/TimingData.cs
--- a/Engine/Source/Programs/UnrealBuildTool/System/TimingData.cs
+++ b/Engine/Source/Programs/UnrealBuildTool/System/TimingData.cs
@@ -119,7 +119,8 @@ namespace UnrealBuildTool
 		/// <param name="ChildData">The sub-event to add to this event.</param>
 		public void AddChild(TimingData ChildData)
 		{
-			if (Children.TryGetValue(ChildData.Name, out var MatchingData))
+			TimingData MatchingData;
+			if (Children.TryGetValue(ChildData.Name, out MatchingData))
 			{
 				MatchingData.ExclusiveDuration += ChildData.ExclusiveDuration;
 				foreach (var ChildChildData in ChildData.Children.Values)
