Another compile error.

diff --git a/Engine/Source/Programs/UnrealBuildTool/System/ParseMsvcTimingInfo.cs b/Engine/Source/Programs/UnrealBuildTool/System/ParseMsvcTimingInfo.cs
--- a/Engine/Source/Programs/UnrealBuildTool/System/ParseMsvcTimingInfo.cs
+++ b/Engine/Source/Programs/UnrealBuildTool/System/ParseMsvcTimingInfo.cs
@@ -144,7 +144,8 @@ namespace UnrealBuildTool
 			TimingData LastTimingData = null;
 			foreach (var Line in Lines)
 			{
-				var CurrentTimingData = ParseTimingDataFromLine(TimingType, Line, out int LineDepth);
+				int LineDepth;
+				var CurrentTimingData = ParseTimingDataFromLine(TimingType, Line, out LineDepth);
 				if (LineDepth == 0)
 				{
 					ParsedTimingData.Add(CurrentTimingData);
