Add optional handling of activity resize events for Android

#android

diff --git a/Engine/Source/Programs/UnrealBuildTool/Platform/Android/UEDeployAndroid.cs b/Engine/Source/Programs/UnrealBuildTool/Platform/Android/UEDeployAndroid.cs
--- a/Engine/Source/Programs/UnrealBuildTool/Platform/Android/UEDeployAndroid.cs
+++ b/Engine/Source/Programs/UnrealBuildTool/Platform/Android/UEDeployAndroid.cs
@@ -2417,16 +2417,21 @@ namespace UnrealBuildTool
 				Text.AppendLine("\t\t<activity android:name=\"com.epicgames.ue4.GameActivity\"");
 				Text.AppendLine("\t\t          android:label=\"@string/app_name\"");
 				Text.AppendLine("\t\t          android:theme=\"@style/UE4SplashTheme\"");
-				Text.AppendLine(bAddDensity ? "\t\t          android:configChanges=\"mcc|mnc|uiMode|density|screenSize|orientation|keyboardHidden|keyboard\""
-											: "\t\t          android:configChanges=\"mcc|mnc|uiMode|screenSize|orientation|keyboardHidden|keyboard\"");
+				Text.AppendLine(bAddDensity ? "\t\t          android:configChanges=\"mcc|mnc|uiMode|density|screenSize|smallestScreenSize|screenLayout|orientation|keyboardHidden|keyboard\""
+											: "\t\t          android:configChanges=\"mcc|mnc|uiMode|screenSize|smallestScreenSize|screenLayout|orientation|keyboardHidden|keyboard\"");
 			}
 			else
 			{
 				Text.AppendLine("\t\t<activity android:name=\"com.epicgames.ue4.GameActivity\"");
 				Text.AppendLine("\t\t          android:label=\"@string/app_name\"");
 				Text.AppendLine("\t\t          android:theme=\"@android:style/Theme.Black.NoTitleBar.Fullscreen\"");
-				Text.AppendLine(bAddDensity ? "\t\t          android:configChanges=\"mcc|mnc|uiMode|density|screenSize|orientation|keyboardHidden|keyboard\""
-											: "\t\t          android:configChanges=\"mcc|mnc|uiMode|screenSize|orientation|keyboardHidden|keyboard\"");
+				Text.AppendLine(bAddDensity ? "\t\t          android:configChanges=\"mcc|mnc|uiMode|density|screenSize|smallestScreenSize|screenLayout|orientation|keyboardHidden|keyboard\""
+											: "\t\t          android:configChanges=\"mcc|mnc|uiMode|screenSize|smallestScreenSize|screenLayout|orientation|keyboardHidden|keyboard\"");
+
+			}
+			if (SDKLevelInt >= 24)
+			{
+				Text.AppendLine("\t\t          android:resizeableActivity=\"false\"");
 			}
 			Text.AppendLine("\t\t          android:launchMode=\"singleTask\"");
 			Text.AppendLine(string.Format("\t\t          android:screenOrientation=\"{0}\"", Orientation));
diff --git a/Engine/Source/Runtime/Launch/Private/Android/LaunchAndroid.cpp b/Engine/Source/Runtime/Launch/Private/Android/LaunchAndroid.cpp
--- a/Engine/Source/Runtime/Launch/Private/Android/LaunchAndroid.cpp
+++ b/Engine/Source/Runtime/Launch/Private/Android/LaunchAndroid.cpp
@@ -124,6 +124,15 @@ extern "C"
 	void *__dso_handle;
 }
 
+int32 GAndroidEnableNativeResizeEvent = 0;
+static FAutoConsoleVariableRef CVarEnableResizeNativeEvent(
+	TEXT("Android.EnableNativeResizeEvent"),
+	GAndroidEnableNativeResizeEvent,
+	TEXT("Whether native resize event is enabled on Android.\n")
+	TEXT(" 0: disabled (default)\n")
+	TEXT(" 1: enabled"),
+	ECVF_ReadOnly);
+
 int32 GAndroidEnableMouse = 0;
 static FAutoConsoleVariableRef CVarEnableMouse(
 	TEXT("Android.EnableMouse"),
@@ -314,6 +323,14 @@ static void InitCommandLine()
 
 extern void AndroidThunkCpp_DismissSplashScreen();
 
+//Called in main thread for native window resizing
+static void OnNativeWindowResized(ANativeActivity* activity, ANativeWindow* window)
+{
+	static int8_t cmd = APP_CMD_WINDOW_RESIZED;
+	struct android_app* app = (struct android_app *)activity->instance;
+	write(app->msgwrite, &cmd, sizeof(cmd));
+}
+
 //Main function called from the android entry point
 int32 AndroidMain(struct android_app* state)
 {
@@ -433,6 +450,12 @@ int32 AndroidMain(struct android_app* state)
 		return PreInitResult;
 	}
 
+	// register callback for native window resize
+	if (GAndroidEnableNativeResizeEvent)
+	{
+		state->activity->callbacks->onNativeWindowResized = OnNativeWindowResized;
+	}
+
 	// initialize HMDs
 	{
 		SCOPED_BOOT_TIMING("InitHMDs");
