UBT: Implementing the recommended code fixes from Ben for the compile timing code.

diff --git a/Engine/Source/Programs/AutomationTool/AutomationUtils/UE4Build.cs b/Engine/Source/Programs/AutomationTool/AutomationUtils/UE4Build.cs
--- a/Engine/Source/Programs/AutomationTool/AutomationUtils/UE4Build.cs
+++ b/Engine/Source/Programs/AutomationTool/AutomationUtils/UE4Build.cs
@@ -32,7 +32,6 @@ namespace AutomationTool
 	[Help("ForceNonUnity", "Toggle to disable the unity build system")]
 	[Help("ForceUnity", "Toggle to force enable the unity build system")]
 	[Help("Licensee", "If set, this build is being compiled by a licensee")]
-	[Help("Timing", "If set, this build will attempt to generate compile timing data")]
 	public class UE4Build
 	{
 		private BuildCommand OwnerCommand;
diff --git a/Engine/Source/Programs/AutomationTool/BuildGraph/BuildGraph.cs b/Engine/Source/Programs/AutomationTool/BuildGraph/BuildGraph.cs
--- a/Engine/Source/Programs/AutomationTool/BuildGraph/BuildGraph.cs
+++ b/Engine/Source/Programs/AutomationTool/BuildGraph/BuildGraph.cs
@@ -68,7 +68,6 @@ namespace AutomationTool
 	[Help("SharedStorageDir=<DirName>", "Sets the directory to use to transfer build products between agents in a build farm")]
 	[Help("SingleNode=<Name>", "Run only the given node. Intended for use on a build system after running with -Export.")]
 	[Help("WriteToSharedStorage", "Allow writing to shared storage. If not set, but -SharedStorageDir is specified, build products will read but not written")]
-	[Help("Timing", "If set, this build will attempt to generate compile timing data")]
 	public class BuildGraph : BuildCommand
 	{
 		/// <summary>
@@ -96,7 +95,6 @@ namespace AutomationTool
 			bool bWriteToSharedStorage = ParseParam("WriteToSharedStorage") || CommandUtils.IsBuildMachine;
 			bool bPublicTasksOnly = ParseParam("PublicTasksOnly");
 			string ReportName = ParseParamValue("ReportName", null);
-			bool bTiming = ParseParam("Timing");
 
 			GraphPrintOptions PrintOptions = GraphPrintOptions.ShowCommandLineOptions;
 			if(ParseParam("ShowDeps"))
@@ -130,7 +128,6 @@ namespace AutomationTool
 			DefaultProperties["HostPlatform"] = HostPlatform.Current.HostEditorPlatform.ToString();
 			DefaultProperties["RestrictedFolderNames"] = String.Join(";", RestrictedFolders.Names);
 			DefaultProperties["RestrictedFolderFilter"] = String.Join(";", RestrictedFolders.Names.Select(x => String.Format(".../{0}/...", x)));
-			DefaultProperties["Timing"] = bTiming.ToString();
 
 			// Attempt to read existing Build Version information
 			BuildVersion Version;
diff --git a/Engine/Source/Programs/UnrealBuildTool/Configuration/UEBuildTarget.cs b/Engine/Source/Programs/UnrealBuildTool/Configuration/UEBuildTarget.cs
--- a/Engine/Source/Programs/UnrealBuildTool/Configuration/UEBuildTarget.cs
+++ b/Engine/Source/Programs/UnrealBuildTool/Configuration/UEBuildTarget.cs
@@ -2929,7 +2929,6 @@ namespace UnrealBuildTool
 			GlobalCompileEnvironment.bAllowRemotelyCompiledPCHs = Rules.bAllowRemotelyCompiledPCHs;
 			GlobalCompileEnvironment.bCheckSystemHeadersForModification = Rules.bCheckSystemHeadersForModification;
 			GlobalCompileEnvironment.bPrintTimingInfo = Rules.bPrintToolChainTimingInfo;
-			GlobalCompileEnvironment.bParseTimingInfoForTracing = Rules.bParseTimingInfoForTracing;
 			GlobalCompileEnvironment.bUseRTTI = Rules.bForceEnableRTTI;
 			GlobalCompileEnvironment.bUseInlining = Rules.bUseInlining;
 			GlobalCompileEnvironment.bHideSymbolsByDefault = Rules.bHideSymbolsByDefault;
diff --git a/Engine/Source/Programs/UnrealBuildTool/Platform/Windows/VCToolChain.cs b/Engine/Source/Programs/UnrealBuildTool/Platform/Windows/VCToolChain.cs
--- a/Engine/Source/Programs/UnrealBuildTool/Platform/Windows/VCToolChain.cs
+++ b/Engine/Source/Programs/UnrealBuildTool/Platform/Windows/VCToolChain.cs
@@ -1116,7 +1116,7 @@ namespace UnrealBuildTool
 
 				if(CompileEnvironment.bGenerateDependenciesFile)
 				{
-					GenerateDependenciesFile(CompileEnvironment, OutputDir, Result, SourceFile, Actions, CompileAction);
+					GenerateDependenciesFile(Target, OutputDir, Result, SourceFile, Actions, CompileAction);
 				}
 
 				if (CompileEnvironment.PrecompiledHeaderAction == PrecompiledHeaderAction.Create)
@@ -1151,32 +1151,32 @@ namespace UnrealBuildTool
 			return Result;
 		}
 
-		private void GenerateDependenciesFile(CppCompileEnvironment CompileEnvironment, DirectoryReference OutputDir, CPPOutput Result, FileItem SourceFile, List<Action> Actions, Action CompileAction)
+		private void GenerateDependenciesFile(ReadOnlyTargetRules Target, DirectoryReference OutputDir, CPPOutput Result, FileItem SourceFile, List<Action> Actions, Action CompileAction)
 		{
-			var commandArguments = new List<string>();
-			CompileAction.DependencyListFile = FileItem.GetItemByFileReference(FileReference.Combine(OutputDir, $"{SourceFile.Location.GetFileName()}.txt"));
+			List<string> commandArguments = new List<string>();
+			CompileAction.DependencyListFile = FileItem.GetItemByFileReference(FileReference.Combine(OutputDir, String.Format("{0}.txt", SourceFile.Location.GetFileName())));
 			CompileAction.ProducedItems.Add(CompileAction.DependencyListFile);
 			commandArguments.Add(Utils.MakePathSafeToUseWithCommandLine(CompileAction.DependencyListFile.Location));
 
-			if (CompileEnvironment.bPrintTimingInfo)
+			if (Target.bPrintToolChainTimingInfo)
 			{
-				CompileAction.TimingFile = FileItem.GetItemByFileReference(FileReference.Combine(OutputDir, $"{SourceFile.Location.GetFileName()}.timing.txt"));
+				CompileAction.TimingFile = FileItem.GetItemByFileReference(FileReference.Combine(OutputDir, String.Format("{0}.timing.txt", SourceFile.Location.GetFileName())));
 				CompileAction.ProducedItems.Add(CompileAction.TimingFile);
 				commandArguments.Add(Utils.MakePathSafeToUseWithCommandLine(CompileAction.TimingFile.Location));
 
-				var TimingJsonFile = FileItem.GetItemByPath(Path.ChangeExtension(CompileAction.TimingFile.AbsolutePath, ".bin"));
+				FileItem TimingJsonFile = FileItem.GetItemByPath(Path.ChangeExtension(CompileAction.TimingFile.AbsolutePath, ".timing.bin"));
 				Result.DebugDataFiles.Add(TimingJsonFile);
 				Result.DebugDataFiles.Add(CompileAction.TimingFile);
 
-				var ParseTimingArguments = new List<string>() { $"-TimingFile=\"{CompileAction.TimingFile}\"" };
-				if (CompileEnvironment.bParseTimingInfoForTracing)
+				string ParseTimingArguments = String.Format("-TimingFile=\"{0}\"", CompileAction.TimingFile);
+				if (Target.bParseTimingInfoForTracing)
 				{
-					ParseTimingArguments.Add("-Tracing");
+					ParseTimingArguments += " -Tracing";
 				}
 
-				var ParseTimingInfoAction = Action.CreateRecursiveAction<ParseMsvcTimingInfoMode>(ActionType.ParseTimingInfo, string.Join(" ", ParseTimingArguments));
+				Action ParseTimingInfoAction = Action.CreateRecursiveAction<ParseMsvcTimingInfoMode>(ActionType.ParseTimingInfo, ParseTimingArguments);
 				ParseTimingInfoAction.WorkingDirectory = UnrealBuildTool.EngineSourceDirectory;
-				ParseTimingInfoAction.StatusDescription = $"Parsing Timing File '{CompileAction.TimingFile}'";
+				ParseTimingInfoAction.StatusDescription = String.Format("Parsing Timing File '{0}'", CompileAction.TimingFile);
 				ParseTimingInfoAction.bCanExecuteRemotely = true;
 				ParseTimingInfoAction.PrerequisiteItems.Add(SourceFile);
 				ParseTimingInfoAction.PrerequisiteItems.Add(CompileAction.TimingFile);
@@ -1196,27 +1196,25 @@ namespace UnrealBuildTool
 		{
 			if (Target.bPrintToolChainTimingInfo)
 			{
-				var ParseTimingActions = Makefile.Actions.Where(x => x.ActionType == ActionType.ParseTimingInfo);
-				Makefile.OutputItems.AddRange(ParseTimingActions.SelectMany(x => x.ProducedItems));
+				List<Action> ParseTimingActions = Makefile.Actions.Where(x => x.ActionType == ActionType.ParseTimingInfo).ToList();
+				List<FileItem> TimingJsonFiles = ParseTimingActions.SelectMany(a => a.ProducedItems.Where(i => i.HasExtension(".timing.bin"))).ToList();
+				Makefile.OutputItems.AddRange(TimingJsonFiles);
 
 				// Handing generating aggregate timing information if we compiled more than one file.
-				if (ParseTimingActions.Count() > 1)
+				if (TimingJsonFiles.Count > 1)
 				{
-					// Find all .timing.json files.
-					var TimingJsonFiles = ParseTimingActions.Select(a => a.ProducedItems.First());
-
 					// Generate the file manifest for the aggregator.
-					var ManifestFile = FileReference.Combine(Makefile.ProjectIntermediateDirectory, $"{Target.Name}TimingManifest.txt");
+					FileReference ManifestFile = FileReference.Combine(Makefile.ProjectIntermediateDirectory, $"{Target.Name}TimingManifest.txt");
 					File.WriteAllLines(ManifestFile.FullName, TimingJsonFiles.Select(f => f.AbsolutePath));
 
-					var AggregateTimingInfoAction = Action.CreateRecursiveAction<AggregateParsedTimingInfo>(ActionType.ParseTimingInfo, $"-Name={Target.Name} -ManifestFile=\"{ManifestFile.FullName}\"");
+					Action AggregateTimingInfoAction = Action.CreateRecursiveAction<AggregateParsedTimingInfo>(ActionType.ParseTimingInfo, $"-Name={Target.Name} -ManifestFile=\"{ManifestFile.FullName}\"");
 					AggregateTimingInfoAction.WorkingDirectory = UnrealBuildTool.EngineSourceDirectory;
-					AggregateTimingInfoAction.StatusDescription = $"Aggregating {TimingJsonFiles.Count()} Timing File(s)";
+					AggregateTimingInfoAction.StatusDescription = $"Aggregating {TimingJsonFiles.Count} Timing File(s)";
 					AggregateTimingInfoAction.bCanExecuteRemotely = false;
 					AggregateTimingInfoAction.PrerequisiteItems.AddRange(ParseTimingActions.SelectMany(a => a.PrerequisiteItems));
 					AggregateTimingInfoAction.PrerequisiteItems.AddRange(TimingJsonFiles);
 
-					var AggregateOutputFile = FileItem.GetItemByFileReference(FileReference.Combine(Makefile.ProjectIntermediateDirectory, $"{Target.Name}.timing.bin"));
+					FileItem AggregateOutputFile = FileItem.GetItemByFileReference(FileReference.Combine(Makefile.ProjectIntermediateDirectory, $"{Target.Name}.timing.bin"));
 					AggregateTimingInfoAction.ProducedItems.Add(AggregateOutputFile);
 					Makefile.OutputItems.Add(AggregateOutputFile);
 					Makefile.Actions.Add(AggregateTimingInfoAction);
diff --git a/Engine/Source/Programs/UnrealBuildTool/System/AggregateParsedTimingInfo.cs b/Engine/Source/Programs/UnrealBuildTool/System/AggregateParsedTimingInfo.cs
--- a/Engine/Source/Programs/UnrealBuildTool/System/AggregateParsedTimingInfo.cs
+++ b/Engine/Source/Programs/UnrealBuildTool/System/AggregateParsedTimingInfo.cs
@@ -26,7 +26,7 @@ namespace UnrealBuildTool
 
 		public override int Execute(CommandLineArguments Arguments)
 		{
-			var ManifestFile = Arguments.GetFileReference("-ManifestFile=");
+			FileReference ManifestFile = Arguments.GetFileReference("-ManifestFile=");
 			string[] ParsedFileNames = FileReference.ReadAllLines(ManifestFile);
 
 			// Load all the file timing data and summarize them for the aggregate.
@@ -34,13 +34,13 @@ namespace UnrealBuildTool
 			Parallel.ForEach(ParsedFileNames, new ParallelOptions() { MaxDegreeOfParallelism = 4 }, ParseTimingDataFile);
 
 			// Create aggregate summary. Duration is the duration of the files in the aggregate.
-			var AggregateName = Arguments.GetString("-Name=");
-			var AggregateData = new TimingData() { Name = AggregateName, Type = TimingDataType.Aggregate };
+			string AggregateName = Arguments.GetString("-Name=");
+			TimingData AggregateData = new TimingData() { Name = AggregateName, Type = TimingDataType.Aggregate };
 			AggregateData.AddChild(FileTimingData);
 
 			// Write out aggregate summary.
-			var OutputFile = Path.Combine(ManifestFile.Directory.FullName, $"{AggregateName}.timing.bin");
-			using (var Writer = new BinaryWriter(File.Open(OutputFile, FileMode.Create)))
+			string OutputFile = Path.Combine(ManifestFile.Directory.FullName, $"{AggregateName}.timing.bin");
+			using (BinaryWriter Writer = new BinaryWriter(File.Open(OutputFile, FileMode.Create)))
 			{
 				// Write out the aggregate data.
 				Writer.Write(AggregateData);
@@ -48,7 +48,7 @@ namespace UnrealBuildTool
 				// Write the look up table for the compressed binary blobs.
 				int Offset = 0;
 				Writer.Write(CompressedFiles.Count);
-				foreach (var CompressedFile in CompressedFiles)
+				foreach (KeyValuePair<string, byte[]> CompressedFile in CompressedFiles)
 				{
 					Writer.Write(CompressedFile.Key);
 					Writer.Write(Offset);
@@ -58,7 +58,7 @@ namespace UnrealBuildTool
 				}
 
 				// Write the compressed binary blobs.
-				foreach (var CompressedFile in CompressedFiles)
+				foreach (KeyValuePair<string, byte[]> CompressedFile in CompressedFiles)
 				{
 					Writer.Write(CompressedFile.Value);
 				}
@@ -72,14 +72,14 @@ namespace UnrealBuildTool
 			// Convert input file back into summary objects.
 			using (BinaryReader Reader = new BinaryReader(File.Open(ParsedFileName, FileMode.Open, FileAccess.Read)))
 			{
-				var ParsedTimingData = new TimingData(Reader);
+				TimingData ParsedTimingData = new TimingData(Reader);
 
-				var CompressDataTask = Task.Run(() =>
+				Task CompressDataTask = Task.Run(() =>
 				{
 					Reader.BaseStream.Seek(0, SeekOrigin.Begin);
-					using (var CompressedMemoryStream = new MemoryStream())
+					using (MemoryStream CompressedMemoryStream = new MemoryStream())
 					{
-						using (var CompressionStream = new GZipStream(CompressedMemoryStream, CompressionMode.Compress))
+						using (GZipStream CompressionStream = new GZipStream(CompressedMemoryStream, CompressionMode.Compress))
 						{
 							Reader.BaseStream.CopyTo(CompressionStream);
 						}
@@ -89,28 +89,28 @@ namespace UnrealBuildTool
 					}
 				});
 
-				var SummarizedTimingData = new TimingData() { Name = ParsedTimingData.Name, Parent = FileTimingData, Type = TimingDataType.Summary };
+				TimingData SummarizedTimingData = new TimingData() { Name = ParsedTimingData.Name, Parent = FileTimingData, Type = TimingDataType.Summary };
 
 				// Gather and update child timing data.
-				var SummarizedIncludes = new TimingData() { Parent = SummarizedTimingData, Type = TimingDataType.Summary };
-				var QueueIncludesTask = Task.Run(() =>
+				TimingData SummarizedIncludes = new TimingData() { Parent = SummarizedTimingData, Type = TimingDataType.Summary };
+				Task QueueIncludesTask = Task.Run(() =>
 				{
 					SummarizeTimingData(ParsedTimingData.Children["IncludeTimings"].Children.Values, SummarizedIncludes, "Includes", 10);
 				});
 
-				var SummarizedClasses = new TimingData() { Parent = SummarizedTimingData, Type = TimingDataType.Summary };
-				var QueueClassesTask = Task.Run(() =>
+				TimingData SummarizedClasses = new TimingData() { Parent = SummarizedTimingData, Type = TimingDataType.Summary };
+				Task QueueClassesTask = Task.Run(() =>
 				{
 					// Collapse templates into single entries.
-					var CollapsedClasses = GroupChildren(ParsedTimingData.Children["ClassTimings"].Children.Values, TimingDataType.Class);
+					IEnumerable<TimingData> CollapsedClasses = GroupChildren(ParsedTimingData.Children["ClassTimings"].Children.Values, TimingDataType.Class);
 					SummarizeTimingData(CollapsedClasses, SummarizedClasses, "Classes", 10);
 				});
 
-				var SummarizedFunctions = new TimingData() { Parent = SummarizedTimingData, Type = TimingDataType.Summary };
-				var QueueFunctionsTask = Task.Run(() =>
+				TimingData SummarizedFunctions = new TimingData() { Parent = SummarizedTimingData, Type = TimingDataType.Summary };
+				Task QueueFunctionsTask = Task.Run(() =>
 				{
 					// Collapse templates into single entries.
-					var CollapsedFunctions = GroupChildren(ParsedTimingData.Children["FunctionTimings"].Children.Values, TimingDataType.Function);
+					IEnumerable<TimingData> CollapsedFunctions = GroupChildren(ParsedTimingData.Children["FunctionTimings"].Children.Values, TimingDataType.Function);
 					SummarizeTimingData(CollapsedFunctions, SummarizedFunctions, "Functions", 10);
 				});
 
@@ -132,12 +132,12 @@ namespace UnrealBuildTool
 
 		private IEnumerable<TimingData> GroupChildren(IEnumerable<TimingData> Children, TimingDataType Type)
 		{
-			var GroupedChildren = new List<TimingData>();
-			var ChildGroups = new Dictionary<string, List<TimingData>>();
-			foreach (var Child in Children)
+			List<TimingData> GroupedChildren = new List<TimingData>();
+			Dictionary<string, List<TimingData>> ChildGroups = new Dictionary<string, List<TimingData>>();
+			foreach (TimingData Child in Children)
 			{
 				// See if this is a templated class. If not, add it as is.
-				var Match = Regex.Match(Child.Name, @"^([^<]*)(?<Template><.*>)");
+				Match Match = Regex.Match(Child.Name, @"^([^<]*)(?<Template><.*>)");
 				if (!Match.Success)
 				{
 					GroupedChildren.Add(Child);
@@ -145,14 +145,13 @@ namespace UnrealBuildTool
 				else
 				{
 					// Generate group name from template.
-					var TemplateParamCount = Match.Groups["Template"].Value.Count(c => c == ',') + 1;
-					var TemplateParamSig = new List<string>(TemplateParamCount);
+					int TemplateParamCount = Match.Groups["Template"].Value.Count(c => c == ',') + 1;
+					List<string> TemplateParamSig = new List<string>(TemplateParamCount);
 					for (int i = 0; i < TemplateParamCount; ++i)
 					{
 						TemplateParamSig.Add("...");
 					}
-
-					var GroupName = Child.Name.Replace(Match.Groups["Template"].Value, $"<{string.Join(", ", TemplateParamSig)}>");
+					string GroupName = Child.Name.Replace(Match.Groups["Template"].Value, $"<{string.Join(", ", TemplateParamSig)}>");
 
 					// See if we have a group for this template already. If not, add it.
 					if (!ChildGroups.ContainsKey(GroupName))
@@ -165,7 +164,7 @@ namespace UnrealBuildTool
 			}
 
 			// Add grouped children.
-			foreach (var Group in ChildGroups)
+			foreach (KeyValuePair<string, List<TimingData>> Group in ChildGroups)
 			{
 				if (Group.Value.Count == 1)
 				{
@@ -173,15 +172,15 @@ namespace UnrealBuildTool
 					continue;
 				}
 
-				var NewViewModel = new TimingData()
+				TimingData NewViewModel = new TimingData()
 				{
 					Name = Group.Key,
 					Type = Type,
 				};
 
-				foreach (var Child in Group.Value)
+				foreach (TimingData Child in Group.Value)
 				{
-					var NewChild = new TimingData() { Name = Child.Name, Type = Child.Type, Parent = NewViewModel, ExclusiveDuration = Child.ExclusiveDuration };
+					TimingData NewChild = new TimingData() { Name = Child.Name, Type = Child.Type, Parent = NewViewModel, ExclusiveDuration = Child.ExclusiveDuration };
 					NewViewModel.AddChild(NewChild);
 				}
 
@@ -194,12 +193,12 @@ namespace UnrealBuildTool
 		private void SummarizeTimingData(IEnumerable<TimingData> Data, TimingData Summary, string TypeName, int NumberToKeep)
 		{
 			// Sort the data by inclusive duration, highest first, so we can grab the top entries to keep.
-			var Sorted = Data.OrderByDescending(v => v.InclusiveDuration).ToList();
+			List<TimingData> Sorted = Data.OrderByDescending(v => v.InclusiveDuration).ToList();
 			Summary.Name = $"Total {TypeName}: {Sorted.Count}";
-			var TopData = Sorted.Take(NumberToKeep);
-			foreach (var Child in TopData)
+			IEnumerable<TimingData> TopData = Sorted.Take(NumberToKeep);
+			foreach (TimingData Child in TopData)
 			{
-				var SummarizedChild = new TimingData() { Name = Child.Name, Parent = Summary, Type = Child.Type, ExclusiveDuration = Child.InclusiveDuration };
+				TimingData SummarizedChild = new TimingData() { Name = Child.Name, Parent = Summary, Type = Child.Type, ExclusiveDuration = Child.InclusiveDuration };
 				Summary.AddChild(SummarizedChild);
 			}
 
diff --git a/Engine/Source/Programs/UnrealBuildTool/System/CppCompileEnvironment.cs b/Engine/Source/Programs/UnrealBuildTool/System/CppCompileEnvironment.cs
--- a/Engine/Source/Programs/UnrealBuildTool/System/CppCompileEnvironment.cs
+++ b/Engine/Source/Programs/UnrealBuildTool/System/CppCompileEnvironment.cs
@@ -302,11 +302,6 @@ namespace UnrealBuildTool
 		/// </summary>
 		public bool bPrintTimingInfo;
 
-		/// <summary>
-		/// Whether to parse timing information for use with chrome://tracing
-		/// </summary>
-		public bool bParseTimingInfoForTracing;
-
 		/// <summary>
 		/// Whether to output a dependencies file along with the output build products
 		/// </summary>
@@ -432,7 +427,6 @@ namespace UnrealBuildTool
 			PGOFilenamePrefix = Other.PGOFilenamePrefix;
 			PGODirectory = Other.PGODirectory;
 			bPrintTimingInfo = Other.bPrintTimingInfo;
-			bParseTimingInfoForTracing = Other.bParseTimingInfoForTracing;
 			bGenerateDependenciesFile = Other.bGenerateDependenciesFile;
 			bAllowRemotelyCompiledPCHs = Other.bAllowRemotelyCompiledPCHs;
 			UserIncludePaths = new HashSet<DirectoryReference>(Other.UserIncludePaths);
diff --git a/Engine/Source/Programs/UnrealBuildTool/System/ParseMsvcTimingInfo.cs b/Engine/Source/Programs/UnrealBuildTool/System/ParseMsvcTimingInfo.cs
--- a/Engine/Source/Programs/UnrealBuildTool/System/ParseMsvcTimingInfo.cs
+++ b/Engine/Source/Programs/UnrealBuildTool/System/ParseMsvcTimingInfo.cs
@@ -20,7 +20,7 @@ namespace UnrealBuildTool
 
 		public override int Execute(CommandLineArguments Arguments)
 		{
-			var InputFile = Arguments.GetFileReference("-TimingFile=");
+			FileReference InputFile = Arguments.GetFileReference("-TimingFile=");
 
 			// If the tracing argument was passed, hand off to the logic to generate a JSON file compatible with 
 			// chrome://tracing
@@ -31,12 +31,12 @@ namespace UnrealBuildTool
 			}
 
 			// Break the input file into the various sections for processing.
-			var AllLines = FileReference.ReadAllLines(InputFile);
-			var Includes = new List<string>();
-			var Classes = new List<string>();
-			var Functions = new List<string>();
-			var CurrentType = TimingDataType.None;
-			foreach (var Line in AllLines)
+			string[] AllLines = FileReference.ReadAllLines(InputFile);
+			List<string> Includes = new List<string>();
+			List<string> Classes = new List<string>();
+			List<string> Functions = new List<string>();
+			TimingDataType CurrentType = TimingDataType.None;
+			foreach (string Line in AllLines)
 			{
 				if (string.IsNullOrWhiteSpace(Line))
 				{
@@ -68,7 +68,7 @@ namespace UnrealBuildTool
 
 				// If we didn't change types and this isn't the count line and it doesn't match the expected output,
 				//  clear the current type and move on.
-				var TimingDataMatch = Regex.Match(Line, TimingDataRegex);
+				Match TimingDataMatch = Regex.Match(Line, TimingDataRegex);
 				if (!TimingDataMatch.Success)
 				{
 					CurrentType = TimingDataType.None;
@@ -99,13 +99,13 @@ namespace UnrealBuildTool
 			}
 
 			// Build the summary.
-			var Summary = new TimingData() { Name = InputFile.FullName.Replace(".timing.txt", string.Empty), Type = TimingDataType.Summary };
+			TimingData Summary = new TimingData() { Name = InputFile.FullName.Replace(".timing.txt", string.Empty), Type = TimingDataType.Summary };
 			Summary.AddChild(SummarizeParsedTimingData("IncludeTimings", TimingDataType.Include, Includes));
 			Summary.AddChild(SummarizeParsedTimingData("ClassTimings", TimingDataType.Class, Classes));
 			Summary.AddChild(SummarizeParsedTimingData("FunctionTimings", TimingDataType.Function, Functions));
 
 			// Write out the timing binary file.
-			using (var Writer = new BinaryWriter(File.Open(InputFile.ChangeExtension(".bin").FullName, FileMode.Create)))
+			using (BinaryWriter Writer = new BinaryWriter(File.Open(InputFile.ChangeExtension(".timing.bin").FullName, FileMode.Create)))
 			{
 				Writer.Write(Summary);
 			}
@@ -115,9 +115,9 @@ namespace UnrealBuildTool
 
 		TimingData SummarizeParsedTimingData(string SummaryName, TimingDataType TimingType,  IEnumerable<string> Lines)
 		{
-			var Summary = new TimingData() { Name = SummaryName, Type = TimingDataType.Summary };
-			var ParsedTimingData = ParseTimingDataFromLines(TimingType, Lines);
-			foreach (var Data in ParsedTimingData)
+			TimingData Summary = new TimingData() { Name = SummaryName, Type = TimingDataType.Summary };
+			List<TimingData> ParsedTimingData = ParseTimingDataFromLines(TimingType, Lines);
+			foreach (TimingData Data in ParsedTimingData)
 			{
 				// See if we've already added a child that matches this data's name. If so, just add to the duration.
 				TimingData MatchedData;
@@ -139,13 +139,13 @@ namespace UnrealBuildTool
 
 		List<TimingData> ParseTimingDataFromLines(TimingDataType TimingType, IEnumerable<string> Lines)
 		{
-			var ParsedTimingData = new List<TimingData>();
+			List<TimingData> ParsedTimingData = new List<TimingData>();
 			int LastDepth = 0;
 			TimingData LastTimingData = null;
-			foreach (var Line in Lines)
+			foreach (string Line in Lines)
 			{
 				int LineDepth;
-				var CurrentTimingData = ParseTimingDataFromLine(TimingType, Line, out LineDepth);
+				TimingData CurrentTimingData = ParseTimingDataFromLine(TimingType, Line, out LineDepth);
 				if (LineDepth == 0)
 				{
 					ParsedTimingData.Add(CurrentTimingData);
@@ -189,7 +189,7 @@ namespace UnrealBuildTool
 
 		TimingData ParseTimingDataFromLine(TimingDataType TimingType, string Line, out int LineDepth)
 		{
-			var TimingDataMatch = Regex.Match(Line, TimingDataRegex);
+			Match TimingDataMatch = Regex.Match(Line, TimingDataRegex);
 			if (!TimingDataMatch.Success)
 			{
 				LineDepth = -1;
@@ -198,7 +198,7 @@ namespace UnrealBuildTool
 
 			LineDepth = TimingDataMatch.Groups["Indent"].Success ? TimingDataMatch.Groups["Indent"].Value.Count() : 0;
 
-			var ParsedTimingData = new TimingData()
+			TimingData ParsedTimingData = new TimingData()
 			{
 				Name = TimingDataMatch.Groups["Name"].Value,
 				Type = TimingType,
diff --git a/Engine/Source/Programs/UnrealBuildTool/System/TimingData.cs b/Engine/Source/Programs/UnrealBuildTool/System/TimingData.cs
--- a/Engine/Source/Programs/UnrealBuildTool/System/TimingData.cs
+++ b/Engine/Source/Programs/UnrealBuildTool/System/TimingData.cs
@@ -71,7 +71,7 @@ namespace UnrealBuildTool
 			int ChildCount = Reader.ReadInt32();
 			for (int i = 0; i < ChildCount; ++i)
 			{
-				var NewChild = new TimingData(Reader);
+				TimingData NewChild = new TimingData(Reader);
 				Children.Add(NewChild.Name, NewChild);
 			}
 		}
@@ -123,7 +123,7 @@ namespace UnrealBuildTool
 			if (Children.TryGetValue(ChildData.Name, out MatchingData))
 			{
 				MatchingData.ExclusiveDuration += ChildData.ExclusiveDuration;
-				foreach (var ChildChildData in ChildData.Children.Values)
+				foreach (TimingData ChildChildData in ChildData.Children.Values)
 				{
 					MatchingData.AddChild(ChildChildData);
 				}
@@ -140,14 +140,14 @@ namespace UnrealBuildTool
 		/// <returns>A deep clone of this event.</returns>
 		public TimingData Clone()
 		{
-			var ClonedTimingData = new TimingData()
+			TimingData ClonedTimingData = new TimingData()
 			{
 				Name = Name,
 				Type = Type,
 				ExclusiveDuration = ExclusiveDuration,
 			};
 
-			foreach (var Child in Children)
+			foreach (KeyValuePair<string, TimingData> Child in Children)
 			{
 				ClonedTimingData.Children.Add(Child.Key, Child.Value.Clone());
 			}
@@ -165,7 +165,7 @@ namespace UnrealBuildTool
 			Writer.Write((byte)Type);
 			Writer.Write(ExclusiveDuration);
 			Writer.Write(Children.Count);
-			foreach (var Child in Children.Values)
+			foreach (TimingData Child in Children.Values)
 			{
 				Writer.Write(Child);
 			}
