Removing some string interpolations I missed in AggregateParsedTimingInfo.

diff --git a/Engine/Source/Programs/UnrealBuildTool/System/AggregateParsedTimingInfo.cs b/Engine/Source/Programs/UnrealBuildTool/System/AggregateParsedTimingInfo.cs
--- a/Engine/Source/Programs/UnrealBuildTool/System/AggregateParsedTimingInfo.cs
+++ b/Engine/Source/Programs/UnrealBuildTool/System/AggregateParsedTimingInfo.cs
@@ -39,7 +39,7 @@ namespace UnrealBuildTool
 			AggregateData.AddChild(FileTimingData);
 
 			// Write out aggregate summary.
-			string OutputFile = Path.Combine(ManifestFile.Directory.FullName, $"{AggregateName}.timing.bin");
+			string OutputFile = Path.Combine(ManifestFile.Directory.FullName, String.Format("{0}.timing.bin", AggregateName));
 			using (BinaryWriter Writer = new BinaryWriter(File.Open(OutputFile, FileMode.Create)))
 			{
 				// Write out the aggregate data.
@@ -151,7 +151,7 @@ namespace UnrealBuildTool
 					{
 						TemplateParamSig.Add("...");
 					}
-					string GroupName = Child.Name.Replace(Match.Groups["Template"].Value, $"<{string.Join(", ", TemplateParamSig)}>");
+					string GroupName = Child.Name.Replace(Match.Groups["Template"].Value, String.Format("<{0}>", string.Join(", ", TemplateParamSig)));
 
 					// See if we have a group for this template already. If not, add it.
 					if (!ChildGroups.ContainsKey(GroupName))
@@ -194,7 +194,7 @@ namespace UnrealBuildTool
 		{
 			// Sort the data by inclusive duration, highest first, so we can grab the top entries to keep.
 			List<TimingData> Sorted = Data.OrderByDescending(v => v.InclusiveDuration).ToList();
-			Summary.Name = $"Total {TypeName}: {Sorted.Count}";
+			Summary.Name = String.Format("Total {0}: {1}", TypeName, Sorted.Count);
 			IEnumerable<TimingData> TopData = Sorted.Take(NumberToKeep);
 			foreach (TimingData Child in TopData)
 			{
