UAT: Fix incorrect logic causing plugin references to override whether a plugin supports a target platform or not.

diff --git a/Engine/Source/Programs/UnrealBuildTool/System/Plugins.cs b/Engine/Source/Programs/UnrealBuildTool/System/Plugins.cs
--- a/Engine/Source/Programs/UnrealBuildTool/System/Plugins.cs
+++ b/Engine/Source/Programs/UnrealBuildTool/System/Plugins.cs
@@ -385,7 +385,12 @@ namespace UnrealBuildTool
 		/// <returns>True if the plugin should be enabled for this project</returns>
 		public static bool IsPluginEnabledForProject(PluginInfo Plugin, ProjectDescriptor Project, UnrealTargetPlatform Platform, UnrealTargetConfiguration TargetConfiguration, TargetType Target)
 		{
-			bool bEnabled = Plugin.EnabledByDefault && Plugin.Descriptor.SupportsTargetPlatform(Platform);
+			if (!Plugin.Descriptor.SupportsTargetPlatform(Platform))
+			{
+				return false;
+			}
+
+			bool bEnabled = Plugin.EnabledByDefault;
 			if (Project != null && Project.Plugins != null)
 			{
 				foreach (PluginReferenceDescriptor PluginReference in Project.Plugins)
