UE-71914 - Packaging code projects in Shipping configuration from Mac fails if there is a space in the project path

diff --git a/Engine/Source/Programs/UnrealBuildTool/Platform/IOS/IOSToolChain.cs b/Engine/Source/Programs/UnrealBuildTool/Platform/IOS/IOSToolChain.cs
--- a/Engine/Source/Programs/UnrealBuildTool/Platform/IOS/IOSToolChain.cs
+++ b/Engine/Source/Programs/UnrealBuildTool/Platform/IOS/IOSToolChain.cs
@@ -499,7 +499,7 @@ namespace UnrealBuildTool
 				FileItem OutputFile = FileItem.GetItemByFileReference(LinkEnvironment.OutputFilePath);
 
 				Result += " -fembed-bitcode -Xlinker -bitcode_verify -Xlinker -bitcode_hide_symbols -Xlinker -bitcode_symbol_map ";
-				Result += " -Xlinker " + Path.GetDirectoryName(OutputFile.AbsolutePath);
+				Result += " -Xlinker \"" + Path.GetDirectoryName(OutputFile.AbsolutePath) + "\"";
 			}
 
 			Result += " -dead_strip";
@@ -877,7 +877,7 @@ namespace UnrealBuildTool
 
 				linkCommandArguments += string.Format("rm -f \"{0}\";", OutputFile.AbsolutePath);
 				linkCommandArguments += string.Format("rm -f \"{0}\\*.bcsymbolmap\";", Path.GetDirectoryName(OutputFile.AbsolutePath));
-				linkCommandArguments += LinkerPath + " " + LinkCommandArguments + ";";
+				linkCommandArguments += "\"" + LinkerPath + "\" " + LinkCommandArguments + ";";
 
 				linkCommandArguments += "'";
 
@@ -886,7 +886,7 @@ namespace UnrealBuildTool
 			else
 			{
 				// This is not a shipping build so no need to delete the output file since symbols will not have been stripped from it.
-				LinkAction.CommandArguments = string.Format("-c '{0} {1}'", LinkerPath, LinkCommandArguments);
+				LinkAction.CommandArguments = string.Format("-c '\"{0}\" {1}'", LinkerPath, LinkCommandArguments);
 			}
 			Actions.Add(LinkAction);
 
