Add -writeactions= and -writeoutdatedactions= arguments to export the action graph to a JSON file, and an "Execute" tool mode to read in an action graph and execute it.

diff --git a/Engine/Source/Programs/DotNETCommon/DotNETUtilities/JsonWriter.cs b/Engine/Source/Programs/DotNETCommon/DotNETUtilities/JsonWriter.cs
--- a/Engine/Source/Programs/DotNETCommon/DotNETUtilities/JsonWriter.cs
+++ b/Engine/Source/Programs/DotNETCommon/DotNETUtilities/JsonWriter.cs
@@ -215,6 +215,17 @@ namespace Tools.DotNETCommon
 			WriteValueInternal(Name, Value ? "true" : "false");
 		}
 
+		/// <summary>
+		/// Write a field name and enum value
+		/// </summary>
+		/// <typeparam name="T">The enum type</typeparam>
+		/// <param name="Name">Name of the field</param>
+		/// <param name="Value">Value for the field</param>
+		public void WriteEnumValue<T>(string Name, T Value) where T : struct
+		{
+			WriteValue(Name, Value.ToString());
+		}
+
 		void WriteCommaNewline()
 		{
 			if (bRequiresComma)
diff --git a/Engine/Source/Programs/UnrealBuildTool/Configuration/TargetDescriptor.cs b/Engine/Source/Programs/UnrealBuildTool/Configuration/TargetDescriptor.cs
--- a/Engine/Source/Programs/UnrealBuildTool/Configuration/TargetDescriptor.cs
+++ b/Engine/Source/Programs/UnrealBuildTool/Configuration/TargetDescriptor.cs
@@ -52,6 +52,12 @@ namespace UnrealBuildTool
 		/// </summary>
 		public Dictionary<string, int> HotReloadModuleNameToSuffix = new Dictionary<string, int>(StringComparer.OrdinalIgnoreCase);
 
+		/// <summary>
+		/// Export the actions for the target to a file
+		/// </summary>
+		[CommandLine("-WriteActions=")]
+		public List<FileReference> WriteActionFiles = new List<FileReference>();
+
 		/// <summary>
 		/// Path to the manifest for passing info about the output to live coding
 		/// </summary>
diff --git a/Engine/Source/Programs/UnrealBuildTool/Modes/BuildMode.cs b/Engine/Source/Programs/UnrealBuildTool/Modes/BuildMode.cs
--- a/Engine/Source/Programs/UnrealBuildTool/Modes/BuildMode.cs
+++ b/Engine/Source/Programs/UnrealBuildTool/Modes/BuildMode.cs
@@ -70,6 +70,12 @@ namespace UnrealBuildTool
 		[CommandLine("-XGEExport")]
 		public bool bXGEExport = false;
 
+		/// <summary>
+		/// Whether we should just export the outdated actions list
+		/// </summary>
+		[CommandLine("-WriteOutdatedActions=")]
+		public FileReference WriteOutdatedActionsFile = null;
+
 		/// <summary>
 		/// Main entry point
 		/// </summary>
@@ -188,7 +194,7 @@ namespace UnrealBuildTool
 					// Create the working set provider
 					using (ISourceFileWorkingSet WorkingSet = SourceFileWorkingSet.Create(UnrealBuildTool.RootDirectory, ProjectDirs))
 					{
-						Build(TargetDescriptors, BuildConfiguration, WorkingSet, Options);
+						Build(TargetDescriptors, BuildConfiguration, WorkingSet, Options, WriteOutdatedActionsFile);
 					}
 				}
 			}
@@ -208,8 +214,9 @@ namespace UnrealBuildTool
 		/// <param name="BuildConfiguration">Current build configuration</param>
 		/// <param name="WorkingSet">The source file working set</param>
 		/// <param name="Options">Additional options for the build</param>
+		/// <param name="WriteOutdatedActionsFile">Files to write the list of outdated actions to (rather than building them)</param>
 		/// <returns>Result from the compilation</returns>
-		public static void Build(List<TargetDescriptor> TargetDescriptors, BuildConfiguration BuildConfiguration, ISourceFileWorkingSet WorkingSet, BuildOptions Options)
+		public static void Build(List<TargetDescriptor> TargetDescriptors, BuildConfiguration BuildConfiguration, ISourceFileWorkingSet WorkingSet, BuildOptions Options, FileReference WriteOutdatedActionsFile)
 		{
 			// Create a makefile for each target
 			TargetMakefile[] Makefiles = new TargetMakefile[TargetDescriptors.Count];
@@ -229,8 +236,19 @@ namespace UnrealBuildTool
 				}
 			}
 
+			// Export the actions for each target
+			for (int TargetIdx = 0; TargetIdx < TargetDescriptors.Count; TargetIdx++)
+			{
+				TargetDescriptor TargetDescriptor = TargetDescriptors[TargetIdx];
+				foreach(FileReference WriteActionFile in TargetDescriptor.WriteActionFiles)
+				{
+					Log.TraceInformation("Writing actions to {0}", WriteActionFile);
+					ActionGraph.ExportJson(Makefiles[TargetIdx].Actions, WriteActionFile);
+				}
+			}
+
 			// Execute the build
-			if((Options & BuildOptions.SkipBuild) == 0)
+			if ((Options & BuildOptions.SkipBuild) == 0)
 			{
 				// Make sure that none of the actions conflict with any other (producing output files differently, etc...)
 				ActionGraph.CheckForConflicts(Makefiles.SelectMany(x => x.Actions));
@@ -284,6 +302,14 @@ namespace UnrealBuildTool
 						XGE.ExportActions(MergedActionsToExecute);
 					}
 				}
+				else if(WriteOutdatedActionsFile != null)
+				{
+					// Write actions to an output file
+					using(Timeline.ScopeEvent("ActionGraph.WriteActions"))
+					{
+						ActionGraph.ExportJson(MergedActionsToExecute, WriteOutdatedActionsFile);
+					}
+				}
 				else
 				{
 					// Execute the actions
diff --git a/Engine/Source/Programs/UnrealBuildTool/Modes/ExecuteMode.cs b/Engine/Source/Programs/UnrealBuildTool/Modes/ExecuteMode.cs
new file mode 100644
--- /dev/null
+++ b/Engine/Source/Programs/UnrealBuildTool/Modes/ExecuteMode.cs
@@ -0,0 +1,67 @@
+// Copyright 1998-2019 Epic Games, Inc. All Rights Reserved.
+
+using System;
+using System.Collections.Generic;
+using System.Diagnostics;
+using System.IO;
+using System.Linq;
+using System.Text;
+using System.Threading;
+using System.Threading.Tasks;
+using Tools.DotNETCommon;
+
+namespace UnrealBuildTool
+{
+	/// <summary>
+	/// Builds a target
+	/// </summary>
+	[ToolMode("Execute", ToolModeOptions.XmlConfig | ToolModeOptions.BuildPlatforms | ToolModeOptions.SingleInstance | ToolModeOptions.StartPrefetchingEngine | ToolModeOptions.ShowExecutionTime)]
+	class ExecuteMode : ToolMode
+	{
+		/// <summary>
+		/// Whether we should just export the outdated actions list
+		/// </summary>
+		[CommandLine("-Actions=", Required = true)]
+		public FileReference ActionsFile = null;
+
+		/// <summary>
+		/// Main entry point
+		/// </summary>
+		/// <param name="Arguments">Command-line arguments</param>
+		/// <returns>One of the values of ECompilationResult</returns>
+		public override int Execute(CommandLineArguments Arguments)
+		{
+			Arguments.ApplyTo(this);
+
+			// Read the XML configuration files
+			XmlConfig.ApplyTo(this);
+
+			// Create the build configuration object, and read the settings
+			BuildConfiguration BuildConfiguration = new BuildConfiguration();
+			XmlConfig.ApplyTo(BuildConfiguration);
+			Arguments.ApplyTo(BuildConfiguration);
+
+			// Read the actions file
+			List<Action> Actions;
+			using(Timeline.ScopeEvent("ActionGraph.ReadActions()"))
+			{
+				Actions = ActionGraph.ImportJson(ActionsFile);
+			}
+
+			// Link the action graph
+			using(Timeline.ScopeEvent("ActionGraph.Link()"))
+			{
+				ActionGraph.Link(Actions);
+			}
+
+			// Execute the actions
+			using (Timeline.ScopeEvent("ActionGraph.ExecuteActions()"))
+			{
+				ActionGraph.ExecuteActions(BuildConfiguration, Actions);
+			}
+
+			return 0;
+		}
+	}
+}
+
diff --git a/Engine/Source/Programs/UnrealBuildTool/System/Action.cs b/Engine/Source/Programs/UnrealBuildTool/System/Action.cs
--- a/Engine/Source/Programs/UnrealBuildTool/System/Action.cs
+++ b/Engine/Source/Programs/UnrealBuildTool/System/Action.cs
@@ -213,6 +213,153 @@ namespace UnrealBuildTool
 			Writer.WriteFileItem(DependencyListFile);
 		}
 
+		/// <summary>
+		/// Writes an action to a json file
+		/// </summary>
+		/// <param name="Object">The object to parse</param>
+		public static Action ImportJson(JsonObject Object)
+		{
+			Action Action = new Action(Object.GetEnumField<ActionType>("Type"));
+
+			string WorkingDirectory;
+			if(Object.TryGetStringField("WorkingDirectory", out WorkingDirectory))
+			{
+				Action.WorkingDirectory = new DirectoryReference(WorkingDirectory);
+			}
+
+			string CommandPath;
+			if(Object.TryGetStringField("CommandPath", out CommandPath))
+			{
+				Action.CommandPath = new FileReference(CommandPath);
+			}
+			
+			string CommandArguments;
+			if(Object.TryGetStringField("CommandArguments", out CommandArguments))
+			{
+				Action.CommandArguments = CommandArguments;
+			}
+
+			string CommandDescription;
+			if(Object.TryGetStringField("CommandDescription", out CommandDescription))
+			{
+				Action.CommandDescription = CommandDescription;
+			}
+			
+			string StatusDescription;
+			if(Object.TryGetStringField("StatusDescription", out StatusDescription))
+			{
+				Action.StatusDescription = StatusDescription;
+			}
+
+			bool bPrintDebugInfo;
+			if(Object.TryGetBoolField("bPrintDebugInfo", out bPrintDebugInfo))
+			{
+				Action.bPrintDebugInfo = bPrintDebugInfo;
+			}
+			
+			bool bCanExecuteRemotely;
+			if(Object.TryGetBoolField("bCanExecuteRemotely", out bCanExecuteRemotely))
+			{
+				Action.bCanExecuteRemotely = bCanExecuteRemotely;
+			}
+
+			bool bCanExecuteRemotelyWithSNDBS;
+			if(Object.TryGetBoolField("bCanExecuteRemotelyWithSNDBS", out bCanExecuteRemotelyWithSNDBS))
+			{
+				Action.bCanExecuteRemotelyWithSNDBS = bCanExecuteRemotelyWithSNDBS;
+			}
+
+			bool bIsGCCCompiler;
+			if(Object.TryGetBoolField("bIsGCCCompiler", out bIsGCCCompiler))
+			{
+				Action.bIsGCCCompiler = bIsGCCCompiler;
+			}
+
+			bool bShouldOutputStatusDescription;
+			if(Object.TryGetBoolField("bShouldOutputStatusDescription", out bShouldOutputStatusDescription))
+			{
+				Action.bShouldOutputStatusDescription = bShouldOutputStatusDescription;
+			}
+
+			bool bProducesImportLibrary;
+			if(Object.TryGetBoolField("bProducesImportLibrary", out bProducesImportLibrary))
+			{
+				Action.bProducesImportLibrary = bProducesImportLibrary;
+			}
+
+			string[] PrerequisiteItems;
+			if (Object.TryGetStringArrayField("PrerequisiteItems", out PrerequisiteItems))
+			{
+				Action.PrerequisiteItems.AddRange(PrerequisiteItems.Select(x => FileItem.GetItemByPath(x)));
+			}
+
+			string[] ProducedItems;
+			if (Object.TryGetStringArrayField("ProducedItems", out ProducedItems))
+			{
+				Action.ProducedItems.AddRange(ProducedItems.Select(x => FileItem.GetItemByPath(x)));
+			}
+
+			string[] DeleteItems;
+			if (Object.TryGetStringArrayField("DeleteItems", out DeleteItems))
+			{
+				Action.DeleteItems.AddRange(DeleteItems.Select(x => FileItem.GetItemByPath(x)));
+			}
+
+			string DependencyListFile;
+			if (Object.TryGetStringField("DependencyListFile", out DependencyListFile))
+			{
+				Action.DependencyListFile = FileItem.GetItemByPath(DependencyListFile);
+			}
+
+			return Action;
+		}
+
+		/// <summary>
+		/// Writes an action to a json file
+		/// </summary>
+		/// <param name="Writer">Writer to receive the output</param>
+		public void ExportJson(JsonWriter Writer)
+		{
+			Writer.WriteEnumValue("Type", ActionType);
+			Writer.WriteValue("WorkingDirectory", WorkingDirectory.FullName);
+			Writer.WriteValue("CommandPath", CommandPath.FullName);
+			Writer.WriteValue("CommandArguments", CommandArguments);
+			Writer.WriteValue("CommandDescription", CommandDescription);
+			Writer.WriteValue("StatusDescription", StatusDescription);
+			Writer.WriteValue("bPrintDebugInfo", bPrintDebugInfo);
+			Writer.WriteValue("bCanExecuteRemotely", bCanExecuteRemotely);
+			Writer.WriteValue("bCanExecuteRemotelyWithSNDBS", bCanExecuteRemotelyWithSNDBS);
+			Writer.WriteValue("bIsGCCCompiler", bIsGCCCompiler);
+			Writer.WriteValue("bShouldOutputStatusDescription", bShouldOutputStatusDescription);
+			Writer.WriteValue("bProducesImportLibrary", bProducesImportLibrary);
+
+			Writer.WriteArrayStart("PrerequisiteItems");
+			foreach(FileItem PrerequisiteItem in PrerequisiteItems)
+			{
+				Writer.WriteValue(PrerequisiteItem.AbsolutePath);
+			}
+			Writer.WriteArrayEnd();
+
+			Writer.WriteArrayStart("ProducedItems");
+			foreach(FileItem ProducedItem in ProducedItems)
+			{
+				Writer.WriteValue(ProducedItem.AbsolutePath);
+			}
+			Writer.WriteArrayEnd();
+
+			Writer.WriteArrayStart("DeleteItems");
+			foreach(FileItem DeleteItem in DeleteItems)
+			{
+				Writer.WriteValue(DeleteItem.AbsolutePath);
+			}
+			Writer.WriteArrayEnd();
+
+			if (DependencyListFile != null)
+			{
+				Writer.WriteValue("DependencyListFile", DependencyListFile.AbsolutePath);
+			}
+		}
+
 		/// <summary>
 		/// Creates an action which calls UBT recursively
 		/// </summary>
diff --git a/Engine/Source/Programs/UnrealBuildTool/System/ActionGraph.cs b/Engine/Source/Programs/UnrealBuildTool/System/ActionGraph.cs
--- a/Engine/Source/Programs/UnrealBuildTool/System/ActionGraph.cs
+++ b/Engine/Source/Programs/UnrealBuildTool/System/ActionGraph.cs
@@ -615,5 +615,62 @@ namespace UnrealBuildTool
 				}
 			}
 		}
+
+		/// <summary>
+		/// Imports an action graph from a JSON file
+		/// </summary>
+		/// <param name="InputFile">The file to read from</param>
+		/// <returns>List of actions</returns>
+		public static List<Action> ImportJson(FileReference InputFile)
+		{
+			JsonObject Object = JsonObject.Read(InputFile);
+
+			JsonObject EnvironmentObject = Object.GetObjectField("Environment");
+			foreach(string KeyName in EnvironmentObject.KeyNames)
+			{
+				Environment.SetEnvironmentVariable(KeyName, EnvironmentObject.GetStringField(KeyName));
+			}
+
+			List<Action> Actions = new List<Action>();
+			foreach (JsonObject ActionObject in Object.GetObjectArrayField("Actions"))
+			{
+				Actions.Add(Action.ImportJson(ActionObject));
+			}
+			return Actions;
+		}
+
+		/// <summary>
+		/// Exports an action graph to a JSON file
+		/// </summary>
+		/// <param name="Actions">The actions to write</param>
+		/// <param name="OutputFile">Output file to write the actions to</param>
+		public static void ExportJson(List<Action> Actions, FileReference OutputFile)
+		{
+			DirectoryReference.CreateDirectory(OutputFile.Directory);
+			using (JsonWriter Writer = new JsonWriter(OutputFile))
+			{
+				Writer.WriteObjectStart();
+
+				Writer.WriteObjectStart("Environment");
+				foreach (System.Collections.DictionaryEntry Pair in Environment.GetEnvironmentVariables())
+				{
+					if (!UnrealBuildTool.InitialEnvironment.Contains(Pair.Key) || (string)(UnrealBuildTool.InitialEnvironment[Pair.Key]) != (string)(Pair.Value))
+					{
+						Writer.WriteValue((string)Pair.Key, (string)Pair.Value);
+					}
+				}
+				Writer.WriteObjectEnd();
+
+				Writer.WriteArrayStart("Actions");
+				foreach (Action Action in Actions)
+				{
+					Writer.WriteObjectStart();
+					Action.ExportJson(Writer);
+					Writer.WriteObjectEnd();
+				}
+				Writer.WriteArrayEnd();
+				Writer.WriteObjectEnd();
+			}
+		}
 	}
 }
diff --git a/Engine/Source/Programs/UnrealBuildTool/System/ExternalExecution.cs b/Engine/Source/Programs/UnrealBuildTool/System/ExternalExecution.cs
--- a/Engine/Source/Programs/UnrealBuildTool/System/ExternalExecution.cs
+++ b/Engine/Source/Programs/UnrealBuildTool/System/ExternalExecution.cs
@@ -1160,7 +1160,7 @@ namespace UnrealBuildTool
 
 						using(Timeline.ScopeEvent("Buildng UnrealHeaderTool"))
 						{
-							BuildMode.Build(TargetDescriptors, BuildConfiguration, WorkingSet, BuildOptions.Quiet);
+							BuildMode.Build(TargetDescriptors, BuildConfiguration, WorkingSet, BuildOptions.Quiet, null);
 						}
 					}
 
diff --git a/Engine/Source/Programs/UnrealBuildTool/UnrealBuildTool.csproj b/Engine/Source/Programs/UnrealBuildTool/UnrealBuildTool.csproj
--- a/Engine/Source/Programs/UnrealBuildTool/UnrealBuildTool.csproj
+++ b/Engine/Source/Programs/UnrealBuildTool/UnrealBuildTool.csproj
@@ -143,6 +143,7 @@
     <Compile Include="Modes\BuildMode.cs" />
     <Compile Include="Modes\CleanMode.cs" />
     <Compile Include="Modes\DeployMode.cs" />
+    <Compile Include="Modes\ExecuteMode.cs" />
     <Compile Include="Modes\JsonExportMode.cs" />
     <Compile Include="Modes\GenerateProjectFilesMode.cs" />
     <Compile Include="Modes\SetupPlatformsMode.cs" />
