UBT: Fix UnrealHeaderTool always needing to run whenever the makefile was invalidated. Order of files in the timestamp file is not guaranteed, so we can't do an element-wise comparison against current list of files.

diff --git a/Engine/Source/Programs/UnrealBuildTool/System/ExternalExecution.cs b/Engine/Source/Programs/UnrealBuildTool/System/ExternalExecution.cs
--- a/Engine/Source/Programs/UnrealBuildTool/System/ExternalExecution.cs
+++ b/Engine/Source/Programs/UnrealBuildTool/System/ExternalExecution.cs
@@ -884,11 +884,14 @@ namespace UnrealBuildTool
 						Log.TraceLog("UnrealHeaderTool needs to run because there are a different number of UObject source files in module {0}", Module.ModuleName);
 						return true;
 					}
-					for (int FileIndex = 0; FileIndex < AllUObjectHeaders.Count; ++FileIndex)
+
+					// Iterate over our UObjects headers and figure out if any of them have changed
+					HashSet<string> ObjectHeadersSet = new HashSet<string>(AllUObjectHeaders.Select(x => x.AbsolutePath), FileReference.Comparer);
+					foreach (string FileName in UObjectFilesFromPreviousRun)
 					{
-						if (!UObjectFilesFromPreviousRun[FileIndex].Equals(AllUObjectHeaders[FileIndex].AbsolutePath, StringComparison.InvariantCultureIgnoreCase))
+						if(!ObjectHeadersSet.Contains(FileName))
 						{
-							Log.TraceLog("UnrealHeaderTool needs to run because the set of UObject source files in module {0} has changed", Module.ModuleName);
+							Log.TraceLog("UnrealHeaderTool needs to run because the set of UObject source files in module {0} has changed ({1})", Module.ModuleName, FileName);
 							return true;
 						}
 					}
