Back out CL5935374 for UE-72100 as it causes 3 other issues

diff --git a/Engine/Source/Programs/UnrealBuildTool/ProjectFiles/Xcode/XcodeProject.cs b/Engine/Source/Programs/UnrealBuildTool/ProjectFiles/Xcode/XcodeProject.cs
--- a/Engine/Source/Programs/UnrealBuildTool/ProjectFiles/Xcode/XcodeProject.cs
+++ b/Engine/Source/Programs/UnrealBuildTool/ProjectFiles/Xcode/XcodeProject.cs
@@ -1218,9 +1218,10 @@ namespace UnrealBuildTool
 				// This is needed for the target to pass the settings validation before code signing. UBT will overwrite this plist file later, with proper contents.
 				if (BuildHostPlatform.Current.Platform == UnrealTargetPlatform.Mac)
 				{
-					bool bCreateIOSInfoPlist = IOSRunTimeVersion != null;
-					bool bCreateTVOSInfoPlist = TVOSRunTimeVersion != null;
-					if (bCreateIOSInfoPlist || bCreateTVOSInfoPlist)
+					bool bCreateMacInfoPlist = !File.Exists(MacInfoPlistPath);
+					bool bCreateIOSInfoPlist = !File.Exists(IOSInfoPlistPath) && IOSRunTimeVersion != null;
+					bool bCreateTVOSInfoPlist = !File.Exists(TVOSInfoPlistPath) && TVOSRunTimeVersion != null;
+					if (bCreateMacInfoPlist || bCreateIOSInfoPlist || bCreateTVOSInfoPlist)
 					{
 						DirectoryReference ProjectPath = GameDir;
 						DirectoryReference EngineDir = DirectoryReference.Combine(new DirectoryReference(UE4Dir), "Engine");
@@ -1235,9 +1236,11 @@ namespace UnrealBuildTool
 							GameName = "UE4Game";
 						}
 
-						Directory.CreateDirectory(Path.GetDirectoryName(MacInfoPlistPath));
-						UEDeployMac.GeneratePList(ProjectPath.FullName, bIsUE4Game, GameName, Config.BuildTarget, EngineDir.FullName, MacExecutableFileName);
-
+						if (bCreateMacInfoPlist)
+						{
+							Directory.CreateDirectory(Path.GetDirectoryName(MacInfoPlistPath));
+							UEDeployMac.GeneratePList(ProjectPath.FullName, bIsUE4Game, GameName, Config.BuildTarget, EngineDir.FullName, MacExecutableFileName);
+						}
 						if (bCreateIOSInfoPlist)
 						{
 							// get the receipt
