UBT: Remove global function for determining whether a file is under an engine directory. This is now done via a trait that is propagated through construction of individual modules.

#rnx

diff --git a/Engine/Source/Programs/UnrealBuildTool/Configuration/ModuleRules.cs b/Engine/Source/Programs/UnrealBuildTool/Configuration/ModuleRules.cs
--- a/Engine/Source/Programs/UnrealBuildTool/Configuration/ModuleRules.cs
+++ b/Engine/Source/Programs/UnrealBuildTool/Configuration/ModuleRules.cs
@@ -413,6 +413,11 @@ namespace UnrealBuildTool
 		/// </summary>
 		internal PluginInfo Plugin;
 
+		/// <summary>
+		/// The rules context for this instance
+		/// </summary>
+		internal ModuleRulesContext Context;
+
 		/// <summary>
 		/// Rules for the target that this module belongs to
 		/// </summary>
diff --git a/Engine/Source/Programs/UnrealBuildTool/Configuration/ModuleRulesContext.cs b/Engine/Source/Programs/UnrealBuildTool/Configuration/ModuleRulesContext.cs
new file mode 100644
--- /dev/null
+++ b/Engine/Source/Programs/UnrealBuildTool/Configuration/ModuleRulesContext.cs
@@ -0,0 +1,74 @@
+using System;
+using System.Collections.Generic;
+using System.Linq;
+using System.Text;
+using System.Threading.Tasks;
+using Tools.DotNETCommon;
+
+namespace UnrealBuildTool
+{
+	/// <summary>
+	/// Stores information about where a module rules object came from, and how it can be used. 
+	/// </summary>
+	class ModuleRulesContext
+	{
+		/// <summary>
+		/// The scope for this module. Used to validate references to other modules.
+		/// </summary>
+		public RulesScope Scope;
+
+		/// <summary>
+		/// The default directory for output files
+		/// </summary>
+		public DirectoryReference DefaultOutputBaseDir;
+
+		/// <summary>
+		/// The plugin that this module belongs to
+		/// </summary>
+		public PluginInfo Plugin;
+
+		/// <summary>
+		/// Whether this module should be included in the default hot reload set
+		/// </summary>
+		public bool bCanHotReload;
+
+		/// <summary>
+		/// Whether this module should be compiled with optimization disabled in DebugGame configurations (ie. whether it's a game module).
+		/// </summary>
+		public bool bCanBuildDebugGame;
+
+		/// <summary>
+		/// Whether to treat this module as a game module for UHT ordering
+		/// </summary>
+		public bool bClassifyAsGameModuleForUHT;
+
+		/// <summary>
+		/// The default module type for UnrealHeaderTool. Do not use this for inferring other things about the module.
+		/// </summary>
+		public UHTModuleType? DefaultUHTModuleType;
+
+		/// <summary>
+		/// Constructor
+		/// </summary>
+		public ModuleRulesContext(RulesScope Scope, DirectoryReference DefaultOutputBaseDir)
+		{
+			this.Scope = Scope;
+			this.DefaultOutputBaseDir = DefaultOutputBaseDir;
+		}
+
+		/// <summary>
+		/// Copy constructor
+		/// </summary>
+		/// <param name="Other">The context to copy from</param>
+		public ModuleRulesContext(ModuleRulesContext Other)
+		{
+			this.Scope = Other.Scope;
+			this.DefaultOutputBaseDir = Other.DefaultOutputBaseDir;
+			this.Plugin = Other.Plugin;
+			this.bCanHotReload = Other.bCanHotReload;
+			this.bCanBuildDebugGame = Other.bCanBuildDebugGame;
+			this.bClassifyAsGameModuleForUHT = Other.bClassifyAsGameModuleForUHT;
+			this.DefaultUHTModuleType = Other.DefaultUHTModuleType;
+		}
+	}
+}
diff --git a/Engine/Source/Programs/UnrealBuildTool/Configuration/UEBuildBinary.cs b/Engine/Source/Programs/UnrealBuildTool/Configuration/UEBuildBinary.cs
--- a/Engine/Source/Programs/UnrealBuildTool/Configuration/UEBuildBinary.cs
+++ b/Engine/Source/Programs/UnrealBuildTool/Configuration/UEBuildBinary.cs
@@ -328,12 +328,12 @@ namespace UnrealBuildTool
 		/// Called to allow the binary to find game modules.
 		/// </summary>
 		/// <returns>The OnlyModule if found, null if not</returns>
-		public List<UEBuildModule> FindGameModules()
+		public List<UEBuildModule> FindHotReloadModules()
 		{
 			List<UEBuildModule> GameModules = new List<UEBuildModule>();
 			foreach (UEBuildModule Module in Modules)
 			{
-				if (!UnrealBuildTool.IsUnderAnEngineDirectory(Module.ModuleDirectory))
+				if(Module.Rules.Context.bCanHotReload)
 				{
 					GameModules.Add(Module);
 				}
diff --git a/Engine/Source/Programs/UnrealBuildTool/Configuration/UEBuildTarget.cs b/Engine/Source/Programs/UnrealBuildTool/Configuration/UEBuildTarget.cs
--- a/Engine/Source/Programs/UnrealBuildTool/Configuration/UEBuildTarget.cs
+++ b/Engine/Source/Programs/UnrealBuildTool/Configuration/UEBuildTarget.cs
@@ -932,7 +932,7 @@ namespace UnrealBuildTool
 					if(Binary.Type == UEBuildBinaryType.DynamicLinkLibrary)
 					{
 						DirectoryReference DirectoryName = Binary.OutputFilePath.Directory;
-						bool bIsGameBinary = RulesAssembly.IsGameModule(Binary.PrimaryModule.Name);
+						bool bIsGameBinary = Binary.PrimaryModule.Rules.Context.bCanBuildDebugGame;
 						FileReference ManifestFileName = FileReference.Combine(DirectoryName, ModuleManifest.GetStandardFileName(AppName, Platform, Configuration, Architecture, bIsGameBinary));
 
 						ModuleManifest Manifest;
@@ -1152,39 +1152,14 @@ namespace UnrealBuildTool
 			// Check there aren't any engine binaries with dependencies on game modules. This can happen when game-specific plugins override engine plugins.
 			foreach(UEBuildModule Module in Modules.Values)
 			{
-				if(Module.Binary != null && UnrealBuildTool.IsUnderAnEngineDirectory(Module.RulesFile.Directory))
+				if(Module.Binary != null)
 				{
-					DirectoryReference RootDirectory = UnrealBuildTool.EngineDirectory;
-
-					if (Module.RulesFile.IsUnderDirectory(UnrealBuildTool.EnterpriseDirectory))
-					{
-						RootDirectory = UnrealBuildTool.EnterpriseDirectory;
-					}
-
 					HashSet<UEBuildModule> ReferencedModules = Module.GetDependencies(bWithIncludePathModules: true, bWithDynamicallyLoadedModules: true);
-
-					// Make sure engine modules don't depend on enterprise or game modules and that enterprise modules don't depend on game modules
 					foreach(UEBuildModule ReferencedModule in ReferencedModules)
 					{
-						if(ReferencedModule.RulesFile != null && !ReferencedModule.RulesFile.IsUnderDirectory(UnrealBuildTool.EngineDirectory) && !ReferencedModule.RulesFile.IsUnderDirectory(RootDirectory))
-						{
-							string EngineModuleRelativePath = Module.RulesFile.MakeRelativeTo(UnrealBuildTool.EngineDirectory.ParentDirectory);
-							string ReferencedModuleRelativePath = (ProjectFile != null && ReferencedModule.RulesFile.IsUnderDirectory(ProjectFile.Directory)) ? ReferencedModule.RulesFile.MakeRelativeTo(ProjectFile.Directory.ParentDirectory) : ReferencedModule.RulesFile.FullName;
-							throw new BuildException("Engine module '{0}' should not depend on game module '{1}'", EngineModuleRelativePath, ReferencedModuleRelativePath);
-						}
-					}
-
-					// Make sure engine modules don't directly reference engine plugins
-					if(Module.RulesFile.IsUnderDirectory(UnrealBuildTool.EngineSourceDirectory) && !Module.RulesFile.IsUnderDirectory(TargetRulesFile.Directory))
-					{
-						foreach(UEBuildModule ReferencedModule in ReferencedModules)
+						if(!Module.Rules.Context.Scope.Contains(ReferencedModule.Rules.Context.Scope) && !IsWhitelistedEnginePluginReference(Module.Name, ReferencedModule.Name))
 						{
-							if(ReferencedModule.RulesFile != null && ModuleToPlugin.ContainsKey(ReferencedModule) && !IsWhitelistedEnginePluginReference(Module.Name, ReferencedModule.Name))
-							{
-								string EngineModuleRelativePath = Module.RulesFile.MakeRelativeTo(UnrealBuildTool.EngineDirectory.ParentDirectory);
-								string ReferencedModuleRelativePath = ReferencedModule.RulesFile.MakeRelativeTo(UnrealBuildTool.EngineDirectory.ParentDirectory);
-								Log.TraceWarning("Warning: Engine module '{0}' should not depend on plugin module '{1}'", EngineModuleRelativePath, ReferencedModuleRelativePath);
-							}
+							throw new BuildException("Module '{0}' ({1}) should not reference module '{2}' ({3}). Hierarchy is {4}.", Module.Name, Module.Rules.Context.Scope.Name, ReferencedModule.Name, ReferencedModule.Rules.Context.Scope.Name, ReferencedModule.Rules.Context.Scope.FormatHierarchy());
 						}
 					}
 				}
@@ -2098,7 +2073,7 @@ namespace UnrealBuildTool
 			{
 				foreach (UEBuildBinary Binary in Binaries)
 				{
-					List<UEBuildModule> GameModules = Binary.FindGameModules();
+					List<UEBuildModule> GameModules = Binary.FindHotReloadModules();
 					if (GameModules != null && GameModules.Count > 0)
 					{
 						if (!UnrealBuildTool.IsProjectInstalled() || EnabledPlugins.Where(x => x.Type == PluginType.Mod).Any(x => Binary.OutputFilePaths[0].IsUnderDirectory(x.Directory)))
@@ -2388,24 +2363,13 @@ namespace UnrealBuildTool
 		{
 			// Get the root output directory and base name (target name/app name) for this binary
 			DirectoryReference BaseOutputDirectory;
-			if (ModuleRules.Plugin != null)
+			if(bUseSharedBuildEnvironment)
 			{
-				BaseOutputDirectory = ModuleRules.Plugin.Directory;
-			}
-			else if (RulesAssembly.IsGameModule(ModuleRules.Name) || !bUseSharedBuildEnvironment)
-			{
-				BaseOutputDirectory = ProjectDirectory;
+				BaseOutputDirectory = ModuleRules.Context.DefaultOutputBaseDir;
 			}
 			else
 			{
-				if (RulesAssembly.GetModuleFileName(ModuleRules.Name).IsUnderDirectory(UnrealBuildTool.EnterpriseDirectory))
-				{
-					BaseOutputDirectory = UnrealBuildTool.EnterpriseDirectory;
-				}
-				else
-				{
-					BaseOutputDirectory = UnrealBuildTool.EngineDirectory;
-				}
+				BaseOutputDirectory = ProjectDirectory;
 			}
 			return BaseOutputDirectory;
 		}
@@ -2422,7 +2386,7 @@ namespace UnrealBuildTool
 
 			// Get the configuration that this module will be built in. Engine modules compiled in DebugGame will use Development.
 			UnrealTargetConfiguration ModuleConfiguration = Configuration;
-			if (Configuration == UnrealTargetConfiguration.DebugGame && !RulesAssembly.IsGameModule(ModuleRules.Name) && !ModuleRules.Name.Equals(Rules.LaunchModuleName, StringComparison.InvariantCultureIgnoreCase))
+			if (Configuration == UnrealTargetConfiguration.DebugGame && !ModuleRules.Context.bCanBuildDebugGame && !ModuleRules.Name.Equals(Rules.LaunchModuleName, StringComparison.InvariantCultureIgnoreCase))
 			{
 				ModuleConfiguration = UnrealTargetConfiguration.Development;
 			}
@@ -2458,7 +2422,7 @@ namespace UnrealBuildTool
 
 			// Get the configuration that this module will be built in. Engine modules compiled in DebugGame will use Development.
 			UnrealTargetConfiguration ModuleConfiguration = Configuration;
-			if (Configuration == UnrealTargetConfiguration.DebugGame && !RulesAssembly.IsGameModule(Module.Name))
+			if (Configuration == UnrealTargetConfiguration.DebugGame && !Module.Rules.Context.bCanBuildDebugGame)
 			{
 				ModuleConfiguration = UnrealTargetConfiguration.Development;
 			}
@@ -3379,17 +3343,9 @@ namespace UnrealBuildTool
 				if (RulesObject.Type != ModuleRules.ModuleType.External)
 				{
 					// Get the base directory
-					if (RulesObject.Plugin != null)
-					{
-						GeneratedCodeDirectory = RulesObject.Plugin.Directory;
-					}
-					else if (bUseSharedBuildEnvironment && RulesObject.File.IsUnderDirectory(UnrealBuildTool.EngineDirectory))
-					{
-						GeneratedCodeDirectory = UnrealBuildTool.EngineDirectory;
-					}
-					else if (bUseSharedBuildEnvironment && UnrealBuildTool.IsUnderAnEngineDirectory(RulesObject.File.Directory))
+					if (bUseSharedBuildEnvironment)
 					{
-						GeneratedCodeDirectory = UnrealBuildTool.EnterpriseDirectory;
+						GeneratedCodeDirectory = RulesObject.Context.DefaultOutputBaseDir;
 					}
 					else
 					{
diff --git a/Engine/Source/Programs/UnrealBuildTool/ProjectFiles/ProjectFileGenerator.cs b/Engine/Source/Programs/UnrealBuildTool/ProjectFiles/ProjectFileGenerator.cs
--- a/Engine/Source/Programs/UnrealBuildTool/ProjectFiles/ProjectFileGenerator.cs
+++ b/Engine/Source/Programs/UnrealBuildTool/ProjectFiles/ProjectFileGenerator.cs
@@ -1893,10 +1893,9 @@ namespace UnrealBuildTool
 				bool WantProjectFileForModule = true;
 
 				// We'll keep track of whether this is an "engine" or "external" module.  This is determined below while loading module rules.
-				bool IsEngineModule = UnrealBuildTool.IsUnderAnEngineDirectory(CurModuleFile.Directory);
 				bool IsThirdPartyModule = CurModuleFile.ContainsName("ThirdParty", UnrealBuildTool.RootDirectory);
 
-				if( IsEngineModule && !bIncludeEngineSource )
+				if( CurModuleFile.IsUnderDirectory(UnrealBuildTool.EngineDirectory) && !bIncludeEngineSource )
 				{
 					// We were asked to exclude engine modules from the generated projects
 					WantProjectFileForModule = false;
diff --git a/Engine/Source/Programs/UnrealBuildTool/System/ExternalExecution.cs b/Engine/Source/Programs/UnrealBuildTool/System/ExternalExecution.cs
--- a/Engine/Source/Programs/UnrealBuildTool/System/ExternalExecution.cs
+++ b/Engine/Source/Programs/UnrealBuildTool/System/ExternalExecution.cs
@@ -114,6 +114,8 @@ namespace UnrealBuildTool
 		{
 			switch (ModuleType)
 			{
+				case ModuleHostType.Program:
+					return UHTModuleType.Program;
 				case ModuleHostType.Runtime:
 				case ModuleHostType.RuntimeNoCommandlet:
                 case ModuleHostType.RuntimeAndProgram:
@@ -134,6 +136,8 @@ namespace UnrealBuildTool
 		{
 			switch (ModuleType)
 			{
+				case ModuleHostType.Program:
+					return UHTModuleType.Program;
 				case ModuleHostType.Runtime:
 				case ModuleHostType.RuntimeNoCommandlet:
                 case ModuleHostType.RuntimeAndProgram:
@@ -376,7 +380,7 @@ namespace UnrealBuildTool
             UHTModuleType? Type = UHTModuleTypeExtensions.EngineModuleTypeFromHostType(Module.Type);
             if (Type == null)
             {
-                throw new BuildException("Unhandled engine module type {0}", Module.Type.ToString());
+                throw new BuildException("Unhandled engine module type {0} for {1}", Module.Type.ToString(), Module.Name);
             }
             return Type.GetValueOrDefault();
         }
@@ -391,36 +395,6 @@ namespace UnrealBuildTool
             return Type.GetValueOrDefault();
         }
 
-		static UHTModuleType? GetEngineModuleTypeBasedOnLocation(DirectoryReference SourceDirectory, FileReference ModuleFileName)
-		{
-			if (ModuleFileName.IsUnderDirectory(DirectoryReference.Combine(SourceDirectory, "Runtime")))
-			{
-				return UHTModuleType.EngineRuntime;
-			}
-
-			if (ModuleFileName.IsUnderDirectory(DirectoryReference.Combine(SourceDirectory, "Developer")))
-			{
-				return UHTModuleType.EngineDeveloper;
-			}
-
-			if (ModuleFileName.IsUnderDirectory(DirectoryReference.Combine(SourceDirectory, "Editor")))
-			{
-				return UHTModuleType.EngineEditor;
-			}
-
-			if (ModuleFileName.IsUnderDirectory(DirectoryReference.Combine(SourceDirectory, "Programs")))
-			{
-				return UHTModuleType.Program;
-			}
-
-			if (ModuleFileName.IsUnderDirectory(DirectoryReference.Combine(SourceDirectory, "ThirdParty")))
-			{
-				return UHTModuleType.EngineThirdParty;
-			}
-
-			return null;
-		}
-
 		/// <summary>
 		/// Returns a copy of Nodes sorted by dependency.  Independent or circularly-dependent nodes should
 		/// remain in their same relative order within the original Nodes sequence.
@@ -485,84 +459,55 @@ namespace UnrealBuildTool
 		/// <returns>The module type</returns>
 		static UHTModuleType GetModuleType(ModuleRules RulesObject, ProjectDescriptor ProjectDescriptor)
 		{
-			// Get the type of module we're creating
-			UHTModuleType? ModuleType = null;
-
-			// Get the module descriptor for this module if it's a plugin
-			ModuleDescriptor PluginModuleDesc = null;
-			if (RulesObject.Plugin != null)
+			ModuleRulesContext Context = RulesObject.Context;
+			if (Context.bClassifyAsGameModuleForUHT)
 			{
-				PluginModuleDesc = RulesObject.Plugin.Descriptor.Modules.FirstOrDefault(x => x.Name == RulesObject.Name);
-				if (PluginModuleDesc != null && PluginModuleDesc.Type == ModuleHostType.Program)
+				if (RulesObject.Type == ModuleRules.ModuleType.External)
 				{
-					ModuleType = UHTModuleType.Program;
+					return UHTModuleType.GameThirdParty;
 				}
-			}
-
-			if (UnrealBuildTool.IsUnderAnEngineDirectory(RulesObject.File.Directory))
-			{
-				if (RulesObject.Type == ModuleRules.ModuleType.External)
+				if (Context.DefaultUHTModuleType.HasValue)
 				{
-					ModuleType = UHTModuleType.EngineThirdParty;
+					return Context.DefaultUHTModuleType.Value;
 				}
-				else
+				if (RulesObject.Plugin != null)
 				{
-					if (!ModuleType.HasValue && PluginModuleDesc != null)
+					ModuleDescriptor Module = RulesObject.Plugin.Descriptor.Modules.FirstOrDefault(x => x.Name == RulesObject.Name);
+					if(Module != null)
 					{
-						ModuleType = ExternalExecution.GetEngineModuleTypeFromDescriptor(PluginModuleDesc);
+						return GetGameModuleTypeFromDescriptor(Module);
 					}
-
-					if (!ModuleType.HasValue)
+				}
+				if(ProjectDescriptor != null && ProjectDescriptor.Modules != null)
+				{
+					ModuleDescriptor Module = ProjectDescriptor.Modules.FirstOrDefault(x => x.Name == RulesObject.Name);
+					if(Module != null)
 					{
-						if (RulesObject.File.IsUnderDirectory(UnrealBuildTool.EngineDirectory))
-						{
-							ModuleType = ExternalExecution.GetEngineModuleTypeBasedOnLocation(UnrealBuildTool.EngineSourceDirectory, RulesObject.File);
-						}
-						else if (RulesObject.File.IsUnderDirectory(UnrealBuildTool.EnterpriseSourceDirectory))
-						{
-							ModuleType = ExternalExecution.GetEngineModuleTypeBasedOnLocation(UnrealBuildTool.EnterpriseSourceDirectory, RulesObject.File);
-						}
+						return UHTModuleTypeExtensions.GameModuleTypeFromHostType(Module.Type) ?? UHTModuleType.GameRuntime;
 					}
 				}
+				return UHTModuleType.GameRuntime;
 			}
 			else
 			{
 				if (RulesObject.Type == ModuleRules.ModuleType.External)
 				{
-					ModuleType = UHTModuleType.GameThirdParty;
+					return UHTModuleType.EngineThirdParty;
 				}
-				else
+				if (Context.DefaultUHTModuleType.HasValue)
 				{
-					if (!ModuleType.HasValue && PluginModuleDesc != null)
-					{
-						ModuleType = ExternalExecution.GetGameModuleTypeFromDescriptor(PluginModuleDesc);
-					}
-
-					if (!ModuleType.HasValue)
+					return Context.DefaultUHTModuleType.Value;
+				}
+				if (RulesObject.Plugin != null)
+				{
+					ModuleDescriptor Module = RulesObject.Plugin.Descriptor.Modules.FirstOrDefault(x => x.Name == RulesObject.Name);
+					if (Module != null)
 					{
-						if (ProjectDescriptor != null)
-						{
-							ModuleDescriptor ProjectModule = (ProjectDescriptor.Modules == null)? null : ProjectDescriptor.Modules.FirstOrDefault(x => x.Name == RulesObject.Name);
-							if (ProjectModule != null)
-							{
-								ModuleType = UHTModuleTypeExtensions.GameModuleTypeFromHostType(ProjectModule.Type) ?? UHTModuleType.GameRuntime;
-							}
-							else
-							{
-								// No descriptor file or module was not on the list
-								ModuleType = UHTModuleType.GameRuntime;
-							}
-						}
+						return GetEngineModuleTypeFromDescriptor(Module);
 					}
 				}
-			}
-
-			if (!ModuleType.HasValue)
-			{
 				throw new BuildException("Unable to determine UHT module type for {0}", RulesObject.File);
 			}
-
-			return ModuleType.Value;
 		}
 
 		/// <summary>
diff --git a/Engine/Source/Programs/UnrealBuildTool/System/RulesAssembly.cs b/Engine/Source/Programs/UnrealBuildTool/System/RulesAssembly.cs
--- a/Engine/Source/Programs/UnrealBuildTool/System/RulesAssembly.cs
+++ b/Engine/Source/Programs/UnrealBuildTool/System/RulesAssembly.cs
@@ -17,6 +17,11 @@ namespace UnrealBuildTool
 	/// </summary>
 	public class RulesAssembly
 	{
+		/// <summary>
+		/// Outers scope for items created by this assembly. Used for chaining assemblies together.
+		/// </summary>
+		internal readonly RulesScope Scope;
+
 		/// <summary>
 		/// The compiled assembly
 		/// </summary>
@@ -43,9 +48,9 @@ namespace UnrealBuildTool
 		private Dictionary<string, FileReference> TargetNameToTargetFile = new Dictionary<string, FileReference>(StringComparer.InvariantCultureIgnoreCase);
 
 		/// <summary>
-		/// Mapping from module file to its plugin info.
+		/// Mapping from module file to its context.
 		/// </summary>
-		private Dictionary<FileReference, PluginInfo> ModuleFileToPluginInfo;
+		private Dictionary<FileReference, ModuleRulesContext> ModuleFileToContext;
 
 		/// <summary>
 		/// Cache for whether a module has source code
@@ -76,22 +81,23 @@ namespace UnrealBuildTool
 		/// <summary>
 		/// Constructor. Compiles a rules assembly from the given source files.
 		/// </summary>
+		/// <param name="Scope">The scope of items created by this assembly</param>
 		/// <param name="BaseDir">The base directory for this assembly</param>
 		/// <param name="Plugins">All the plugins included in this assembly</param>
-		/// <param name="ModuleFiles">List of module files to compile</param>
+		/// <param name="ModuleFileToContext">List of module files to compile</param>
 		/// <param name="TargetFiles">List of target files to compile</param>
-		/// <param name="ModuleFileToPluginInfo">Mapping of module file to the plugin that contains it</param>
 		/// <param name="AssemblyFileName">The output path for the compiled assembly</param>
 		/// <param name="bContainsEngineModules">Whether this assembly contains engine modules. Used to initialize the default value for ModuleRules.bTreatAsEngineModule.</param>
 		/// <param name="bUseBackwardsCompatibleDefaults">Whether modules in this assembly should use backwards-compatible defaults.</param>
 		/// <param name="bReadOnly">Whether the modules and targets in this assembly are installed, and should be created with the bUsePrecompiled flag set</param> 
 		/// <param name="bSkipCompile">Whether to skip compiling this assembly</param>
 		/// <param name="Parent">The parent rules assembly</param>
-		public RulesAssembly(DirectoryReference BaseDir, IReadOnlyList<PluginInfo> Plugins, List<FileReference> ModuleFiles, List<FileReference> TargetFiles, Dictionary<FileReference, PluginInfo> ModuleFileToPluginInfo, FileReference AssemblyFileName, bool bContainsEngineModules, bool bUseBackwardsCompatibleDefaults, bool bReadOnly, bool bSkipCompile, RulesAssembly Parent)
+		internal RulesAssembly(RulesScope Scope, DirectoryReference BaseDir, IReadOnlyList<PluginInfo> Plugins, Dictionary<FileReference, ModuleRulesContext> ModuleFileToContext, List<FileReference> TargetFiles, FileReference AssemblyFileName, bool bContainsEngineModules, bool bUseBackwardsCompatibleDefaults, bool bReadOnly, bool bSkipCompile, RulesAssembly Parent)
 		{
+			this.Scope = Scope;
 			this.BaseDir = BaseDir;
 			this.Plugins = Plugins;
-			this.ModuleFileToPluginInfo = ModuleFileToPluginInfo;
+			this.ModuleFileToContext = ModuleFileToContext;
 			this.bContainsEngineModules = bContainsEngineModules;
 			this.bUseBackwardsCompatibleDefaults = bUseBackwardsCompatibleDefaults;
 			this.bReadOnly = bReadOnly;
@@ -99,7 +105,7 @@ namespace UnrealBuildTool
 
 			// Find all the source files
 			HashSet<FileReference> AssemblySourceFiles = new HashSet<FileReference>();
-			AssemblySourceFiles.UnionWith(ModuleFiles);
+			AssemblySourceFiles.UnionWith(ModuleFileToContext.Keys);
 			AssemblySourceFiles.UnionWith(TargetFiles);
 
 			// Compile the assembly
@@ -110,7 +116,7 @@ namespace UnrealBuildTool
 			}
 
 			// Setup the module map
-			foreach (FileReference ModuleFile in ModuleFiles)
+			foreach (FileReference ModuleFile in ModuleFileToContext.Keys)
 			{
 				string ModuleName = ModuleFile.GetFileNameWithoutAnyExtensions();
 				if (!ModuleNameToModuleFile.ContainsKey(ModuleName))
@@ -371,7 +377,8 @@ namespace UnrealBuildTool
 				RulesObject.Name = ModuleName;
 				RulesObject.File = ModuleFileName;
 				RulesObject.Directory = ModuleFileName.Directory;
-				ModuleFileToPluginInfo.TryGetValue(RulesObject.File, out RulesObject.Plugin);
+				RulesObject.Context = ModuleFileToContext[RulesObject.File];
+				RulesObject.Plugin = RulesObject.Context.Plugin;
 				RulesObject.bTreatAsEngineModule = bContainsEngineModules;
 				RulesObject.bUseBackwardsCompatibleDefaults = bUseBackwardsCompatibleDefaults && Target.bUseBackwardsCompatibleDefaults;
 				RulesObject.bPrecompile = (RulesObject.bTreatAsEngineModule || ModuleName.Equals("UE4Game", StringComparison.OrdinalIgnoreCase)) && Target.bPrecompile;
@@ -394,15 +401,6 @@ namespace UnrealBuildTool
 			}
 		}
 
-		/// <summary>
-		/// Determines whether the given module name is a game module (as opposed to an engine module)
-		/// </summary>
-		public bool IsGameModule(string InModuleName)
-		{
-			FileReference ModuleFileName = GetModuleFileName(InModuleName);
-			return (ModuleFileName != null && !UnrealBuildTool.IsUnderAnEngineDirectory(ModuleFileName.Directory));
-		}
-
 		/// <summary>
 		/// Construct an instance of the given target rules
 		/// </summary>
@@ -667,14 +665,18 @@ namespace UnrealBuildTool
 		/// <returns>True if the module belongs to a plugin</returns>
 		public bool TryGetPluginForModule(FileReference ModuleFile, out PluginInfo Plugin)
 		{
-			if (ModuleFileToPluginInfo.TryGetValue(ModuleFile, out Plugin))
+			ModuleRulesContext Context;
+			if (ModuleFileToContext.TryGetValue(ModuleFile, out Context))
 			{
-				return true;
+				Plugin = Context.Plugin;
+				return Plugin != null;
 			}
-			else
+			if(Parent == null)
 			{
-				return (Parent == null) ? false : Parent.TryGetPluginForModule(ModuleFile, out Plugin);
+				Plugin = null;
+				return false;
 			}
+			return Parent.TryGetPluginForModule(ModuleFile, out Plugin);
 		}
 	}
 }
diff --git a/Engine/Source/Programs/UnrealBuildTool/System/RulesCompiler.cs b/Engine/Source/Programs/UnrealBuildTool/System/RulesCompiler.cs
--- a/Engine/Source/Programs/UnrealBuildTool/System/RulesCompiler.cs
+++ b/Engine/Source/Programs/UnrealBuildTool/System/RulesCompiler.cs
@@ -1,14 +1,10 @@
 // Copyright 1998-2019 Epic Games, Inc. All Rights Reserved.
 
 using System;
-using System.IO;
 using System.Collections.Generic;
+using System.IO;
 using System.Linq;
-using System.Reflection;
-using System.Diagnostics;
-using System.Runtime.Serialization;
 using Tools.DotNETCommon;
-using System.Collections.Concurrent;
 
 namespace UnrealBuildTool
 {
@@ -307,6 +303,39 @@ namespace UnrealBuildTool
 			}
 		}
 
+		/// <summary>
+		/// Find all the module rules files under a given directory
+		/// </summary>
+		/// <param name="BaseDirectory">The directory to search under</param>
+		/// <param name="ModuleContext">The module context for each found rules instance</param>
+		/// <param name="ModuleFileToContext">Map of module files to their context</param>
+		private static void AddModuleRulesWithContext(DirectoryReference BaseDirectory, ModuleRulesContext ModuleContext, Dictionary<FileReference, ModuleRulesContext> ModuleFileToContext)
+		{
+			IReadOnlyList<FileReference> RulesFiles = FindAllRulesFiles(BaseDirectory, RulesFileType.Module);
+			foreach (FileReference RulesFile in RulesFiles)
+			{
+				ModuleFileToContext[RulesFile] = ModuleContext;
+			}
+		}
+
+		/// <summary>
+		/// Find all the module rules files under a given directory
+		/// </summary>
+		/// <param name="BaseDirectory">The directory to search under</param>
+		/// <param name="SubDirectoryName">Name of the subdirectory to look under</param>
+		/// <param name="BaseModuleContext">The module context for each found rules instance</param>
+		/// <param name="DefaultUHTModuleType">The UHT module type</param>
+		/// <param name="ModuleFileToContext">Map of module files to their context</param>
+		private static void AddEngineModuleRulesWithContext(DirectoryReference BaseDirectory, string SubDirectoryName, ModuleRulesContext BaseModuleContext, UHTModuleType? DefaultUHTModuleType, Dictionary<FileReference, ModuleRulesContext> ModuleFileToContext)
+		{
+			DirectoryReference Directory = DirectoryReference.Combine(BaseDirectory, SubDirectoryName);
+			if (DirectoryLookupCache.DirectoryExists(Directory))
+			{
+				ModuleRulesContext ModuleContext = new ModuleRulesContext(BaseModuleContext) { DefaultUHTModuleType = DefaultUHTModuleType };
+				AddModuleRulesWithContext(Directory, ModuleContext, ModuleFileToContext);
+			}
+		}
+
 		/// <summary>
 		/// The cached rules assembly for engine modules and targets.
 		/// </summary>
@@ -334,7 +363,8 @@ namespace UnrealBuildTool
 			if (EngineRulesAssembly == null)
 			{
 				IReadOnlyList<PluginInfo> IncludedPlugins = Plugins.ReadEnginePlugins(UnrealBuildTool.EngineDirectory);
-				EngineRulesAssembly = CreateEngineOrEnterpriseRulesAssembly(UnrealBuildTool.EngineDirectory, ProjectFileGenerator.EngineProjectFileNameBase, IncludedPlugins, UnrealBuildTool.IsEngineInstalled() || bUsePrecompiled, bSkipCompile, null);
+				RulesScope EngineScope = new RulesScope("Engine", null);
+				EngineRulesAssembly = CreateEngineOrEnterpriseRulesAssembly(EngineScope, UnrealBuildTool.EngineDirectory, ProjectFileGenerator.EngineProjectFileNameBase, IncludedPlugins, UnrealBuildTool.IsEngineInstalled() || bUsePrecompiled, bSkipCompile, null);
 			}
 			return EngineRulesAssembly;
 		}
@@ -349,17 +379,18 @@ namespace UnrealBuildTool
 		{
 			if (EnterpriseRulesAssembly == null)
 			{
+				RulesAssembly EngineAssembly = CreateEngineRulesAssembly(bUsePrecompiled, bSkipCompile);
 				if (DirectoryReference.Exists(UnrealBuildTool.EnterpriseDirectory))
 				{
+					RulesScope EnterpriseScope = new RulesScope("Enterprise", EngineAssembly.Scope);
 					IReadOnlyList<PluginInfo> IncludedPlugins = Plugins.ReadEnterprisePlugins(UnrealBuildTool.EnterpriseDirectory);
-					EnterpriseRulesAssembly = CreateEngineOrEnterpriseRulesAssembly(UnrealBuildTool.EnterpriseDirectory, ProjectFileGenerator.EnterpriseProjectFileNameBase, IncludedPlugins, UnrealBuildTool.IsEnterpriseInstalled() || bUsePrecompiled, bSkipCompile, CreateEngineRulesAssembly(bUsePrecompiled, bSkipCompile));
+					EnterpriseRulesAssembly = CreateEngineOrEnterpriseRulesAssembly(EnterpriseScope, UnrealBuildTool.EnterpriseDirectory, ProjectFileGenerator.EnterpriseProjectFileNameBase, IncludedPlugins, UnrealBuildTool.IsEnterpriseInstalled() || bUsePrecompiled, bSkipCompile, EngineAssembly);
 				}
 				else
 				{
-					Log.TraceWarning("Trying to build an enterprise target but the enterprise directory is missing. Falling back on engine components only.");
-
 					// If we're asked for the enterprise rules assembly but the enterprise directory is missing, fallback on the engine rules assembly
-					return CreateEngineRulesAssembly(bUsePrecompiled, bSkipCompile);
+					Log.TraceWarning("Trying to build an enterprise target but the enterprise directory is missing. Falling back on engine components only.");
+					return EngineAssembly;
 				}
 			}
 
@@ -369,6 +400,7 @@ namespace UnrealBuildTool
 		/// <summary>
 		/// Creates a rules assembly
 		/// </summary>
+		/// <param name="Scope">Scope for items created from this assembly</param>
 		/// <param name="RootDirectory">The root directory to create rules for</param>
 		/// <param name="AssemblyPrefix">A prefix for the assembly file name</param>
 		/// <param name="Plugins">List of plugins to include in this assembly</param>
@@ -376,40 +408,46 @@ namespace UnrealBuildTool
 		/// <param name="bSkipCompile">Whether to skip compilation for this assembly</param>
 		/// <param name="Parent">The parent rules assembly</param>
 		/// <returns>New rules assembly</returns>
-		private static RulesAssembly CreateEngineOrEnterpriseRulesAssembly(DirectoryReference RootDirectory, string AssemblyPrefix, IReadOnlyList<PluginInfo> Plugins, bool bReadOnly, bool bSkipCompile, RulesAssembly Parent)
+		private static RulesAssembly CreateEngineOrEnterpriseRulesAssembly(RulesScope Scope, DirectoryReference RootDirectory, string AssemblyPrefix, IReadOnlyList<PluginInfo> Plugins, bool bReadOnly, bool bSkipCompile, RulesAssembly Parent)
 		{
 			DirectoryReference SourceDirectory = DirectoryReference.Combine(RootDirectory, "Source");
 			DirectoryReference ProgramsDirectory = DirectoryReference.Combine(SourceDirectory, "Programs");
 
 			// Find the shared modules, excluding the programs directory. These are used to create an assembly with the bContainsEngineModules flag set to true.
-			List<FileReference> EngineModuleFiles = new List<FileReference>();
+			Dictionary<FileReference, ModuleRulesContext> ModuleFileToContext = new Dictionary<FileReference, ModuleRulesContext>();
 			using(Timeline.ScopeEvent("Finding engine modules"))
 			{
-				foreach (DirectoryReference SubDirectory in DirectoryLookupCache.EnumerateDirectories(SourceDirectory))
-				{
-					if(SubDirectory != ProgramsDirectory)
-					{
-						EngineModuleFiles.AddRange(FindAllRulesFiles(SubDirectory, RulesFileType.Module));
-					}
-				}
+				ModuleRulesContext DefaultModuleContext = new ModuleRulesContext(Scope, RootDirectory);
+				AddEngineModuleRulesWithContext(SourceDirectory, "Runtime", DefaultModuleContext, UHTModuleType.EngineRuntime, ModuleFileToContext);
+				AddEngineModuleRulesWithContext(SourceDirectory, "Developer", DefaultModuleContext, UHTModuleType.EngineDeveloper, ModuleFileToContext);
+				AddEngineModuleRulesWithContext(SourceDirectory, "Editor", DefaultModuleContext, UHTModuleType.EngineEditor, ModuleFileToContext);
+				AddEngineModuleRulesWithContext(SourceDirectory, "ThirdParty", DefaultModuleContext, UHTModuleType.EngineThirdParty, ModuleFileToContext);
 			}
 
 			// Add all the plugin modules too
-			Dictionary<FileReference, PluginInfo> ModuleFileToPluginInfo = new Dictionary<FileReference, PluginInfo>();
+			RulesScope PluginsScope = new RulesScope(Scope.Name + " Plugins", Scope);
 			using(Timeline.ScopeEvent("Finding plugin modules"))
 			{
-				FindModuleRulesForPlugins(Plugins, EngineModuleFiles, ModuleFileToPluginInfo);
+				ModuleRulesContext PluginsModuleContext = new ModuleRulesContext(PluginsScope, RootDirectory);
+				FindModuleRulesForPlugins(Plugins, PluginsModuleContext, ModuleFileToContext);
 			}
 
 			// Create the assembly
 			FileReference EngineAssemblyFileName = FileReference.Combine(RootDirectory, "Intermediate", "Build", "BuildRules", AssemblyPrefix + "Rules" + FrameworkAssemblyExtension);
-			RulesAssembly EngineAssembly = new RulesAssembly(RootDirectory, Plugins, EngineModuleFiles, new List<FileReference>(), ModuleFileToPluginInfo, EngineAssemblyFileName, bContainsEngineModules: true, bUseBackwardsCompatibleDefaults: false, bReadOnly: bReadOnly, bSkipCompile: bSkipCompile, Parent: Parent);
+			RulesAssembly EngineAssembly = new RulesAssembly(Scope, RootDirectory, Plugins, ModuleFileToContext, new List<FileReference>(), EngineAssemblyFileName, bContainsEngineModules: true, bUseBackwardsCompatibleDefaults: false, bReadOnly: bReadOnly, bSkipCompile: bSkipCompile, Parent: Parent);
+
+			// Create a new scope for programs
+			RulesScope ProgramsScope = new RulesScope(Scope.Name + " Programs", PluginsScope);
+
+			// Also create a scope for them, and update the UHT module type
+			ModuleRulesContext ProgramsModuleContext = new ModuleRulesContext(ProgramsScope, RootDirectory);
+			ProgramsModuleContext.DefaultUHTModuleType = UHTModuleType.Program;
 
 			// Find all the rules files
-			List<FileReference> ProgramModuleFiles;
+			Dictionary<FileReference, ModuleRulesContext> ProgramModuleFiles = new Dictionary<FileReference, ModuleRulesContext>();
 			using(Timeline.ScopeEvent("Finding program modules"))
 			{
-				ProgramModuleFiles = new List<FileReference>(FindAllRulesFiles(ProgramsDirectory, RulesFileType.Module));
+				AddModuleRulesWithContext(ProgramsDirectory, ProgramsModuleContext, ProgramModuleFiles);
 			}
 			List<FileReference> ProgramTargetFiles;
 			using(Timeline.ScopeEvent("Finding program targets"))
@@ -419,7 +457,7 @@ namespace UnrealBuildTool
 
 			// Create a path to the assembly that we'll either load or compile
 			FileReference ProgramAssemblyFileName = FileReference.Combine(RootDirectory, "Intermediate", "Build", "BuildRules", AssemblyPrefix + "ProgramRules" + FrameworkAssemblyExtension);
-			RulesAssembly ProgramAssembly = new RulesAssembly(RootDirectory, new List<PluginInfo>().AsReadOnly(), ProgramModuleFiles, ProgramTargetFiles, new Dictionary<FileReference, PluginInfo>(), ProgramAssemblyFileName, bContainsEngineModules: false, bUseBackwardsCompatibleDefaults: false, bReadOnly: bReadOnly, bSkipCompile: bSkipCompile, Parent: EngineAssembly);
+			RulesAssembly ProgramAssembly = new RulesAssembly(ProgramsScope, RootDirectory, new List<PluginInfo>().AsReadOnly(), ProgramModuleFiles, ProgramTargetFiles, ProgramAssemblyFileName, bContainsEngineModules: false, bUseBackwardsCompatibleDefaults: false, bReadOnly: bReadOnly, bSkipCompile: bSkipCompile, Parent: EngineAssembly);
 
 			// Return the combined assembly
 			return ProgramAssembly;
@@ -451,10 +489,22 @@ namespace UnrealBuildTool
 					Parent = CreateEngineRulesAssembly(bUsePrecompiled, bSkipCompile);
 				}
 
+				// Create a scope for things in this assembly
+				RulesScope Scope = new RulesScope("Project", Parent.Scope);
+
 				// Find all the rules under the project source directory
 				DirectoryReference ProjectDirectory = ProjectFileName.Directory;
 				DirectoryReference ProjectSourceDirectory = DirectoryReference.Combine(ProjectDirectory, "Source");
-				List<FileReference> ModuleFiles = new List<FileReference>(FindAllRulesFiles(ProjectSourceDirectory, RulesFileType.Module));
+
+				// Create a new context for modules created by this assembly
+				ModuleRulesContext DefaultModuleContext = new ModuleRulesContext(Scope, ProjectDirectory);
+				DefaultModuleContext.bCanBuildDebugGame = true;
+				DefaultModuleContext.bCanHotReload = true;
+				DefaultModuleContext.bClassifyAsGameModuleForUHT = true;
+				
+				Dictionary<FileReference, ModuleRulesContext> ModuleFiles = new Dictionary<FileReference, ModuleRulesContext>();
+				AddModuleRulesWithContext(ProjectSourceDirectory, DefaultModuleContext, ModuleFiles);
+				
 				List<FileReference> TargetFiles = new List<FileReference>(FindAllRulesFiles(ProjectSourceDirectory, RulesFileType.Target));
 
 				// Find all the project plugins
@@ -470,14 +520,13 @@ namespace UnrealBuildTool
 				}
 
 				// Find all the plugin module rules
-                Dictionary<FileReference, PluginInfo> ModuleFileToPluginInfo = new Dictionary<FileReference, PluginInfo>();
-				FindModuleRulesForPlugins(ProjectPlugins, ModuleFiles, ModuleFileToPluginInfo);
+				FindModuleRulesForPlugins(ProjectPlugins, DefaultModuleContext, ModuleFiles);
 
 				// Add the games project's intermediate source folder
 				DirectoryReference ProjectIntermediateSourceDirectory = DirectoryReference.Combine(ProjectDirectory, "Intermediate", "Source");
 				if (DirectoryReference.Exists(ProjectIntermediateSourceDirectory))
 				{
-					ModuleFiles.AddRange(FindAllRulesFiles(ProjectIntermediateSourceDirectory, RulesFileType.Module));
+					AddModuleRulesWithContext(ProjectIntermediateSourceDirectory, DefaultModuleContext, ModuleFiles);
 					TargetFiles.AddRange(FindAllRulesFiles(ProjectIntermediateSourceDirectory, RulesFileType.Target));
 				}
 
@@ -489,7 +538,7 @@ namespace UnrealBuildTool
 				}
 				else
 				{
-					ProjectRulesAssembly = new RulesAssembly(ProjectDirectory, ProjectPlugins, ModuleFiles, TargetFiles, ModuleFileToPluginInfo, AssemblyFileName, bContainsEngineModules: false, bUseBackwardsCompatibleDefaults: true, bReadOnly: UnrealBuildTool.IsProjectInstalled(), bSkipCompile: bSkipCompile, Parent: Parent);
+					ProjectRulesAssembly = new RulesAssembly(Scope, ProjectDirectory, ProjectPlugins, ModuleFiles, TargetFiles, AssemblyFileName, bContainsEngineModules: false, bUseBackwardsCompatibleDefaults: true, bReadOnly: UnrealBuildTool.IsProjectInstalled(), bSkipCompile: bSkipCompile, Parent: Parent);
 				}
 				LoadedAssemblyMap.Add(ProjectFileName, ProjectRulesAssembly);
 			}
@@ -511,7 +560,7 @@ namespace UnrealBuildTool
 			if (!LoadedAssemblyMap.TryGetValue(PluginFileName, out PluginRulesAssembly))
 			{
 				// Find all the rules source files
-				List<FileReference> ModuleFiles = new List<FileReference>();
+				Dictionary<FileReference, ModuleRulesContext> ModuleFiles = new Dictionary<FileReference, ModuleRulesContext>();
 				List<FileReference> TargetFiles = new List<FileReference>();
 
 				// Create a list of plugins for this assembly. If it already exists in the parent assembly, just create an empty assembly.
@@ -521,13 +570,16 @@ namespace UnrealBuildTool
 					ForeignPlugins.Add(new PluginInfo(PluginFileName, PluginType.External));
 				}
 
+				// Create a new scope for the plugin. It should not reference anything else.
+				RulesScope Scope = new RulesScope("Plugin", Parent.Scope);
+
 				// Find all the modules
-				Dictionary<FileReference, PluginInfo> ModuleFileToPluginInfo = new Dictionary<FileReference, PluginInfo>();
-				FindModuleRulesForPlugins(ForeignPlugins, ModuleFiles, ModuleFileToPluginInfo);
+				ModuleRulesContext PluginModuleContext = new ModuleRulesContext(Scope, PluginFileName.Directory);
+				FindModuleRulesForPlugins(ForeignPlugins, PluginModuleContext, ModuleFiles);
 
 				// Compile the assembly
 				FileReference AssemblyFileName = FileReference.Combine(PluginFileName.Directory, "Intermediate", "Build", "BuildRules", Path.GetFileNameWithoutExtension(PluginFileName.FullName) + "ModuleRules" + FrameworkAssemblyExtension);
-				PluginRulesAssembly = new RulesAssembly(PluginFileName.Directory, ForeignPlugins, ModuleFiles, TargetFiles, ModuleFileToPluginInfo, AssemblyFileName, bContainsEngineModules, bUseBackwardsCompatibleDefaults: false, bReadOnly: false, bSkipCompile: bSkipCompile, Parent: Parent);
+				PluginRulesAssembly = new RulesAssembly(Scope, PluginFileName.Directory, ForeignPlugins, ModuleFiles, TargetFiles, AssemblyFileName, bContainsEngineModules, bUseBackwardsCompatibleDefaults: false, bReadOnly: false, bSkipCompile: bSkipCompile, Parent: Parent);
 				LoadedAssemblyMap.Add(PluginFileName, PluginRulesAssembly);
 			}
 			return PluginRulesAssembly;
@@ -570,9 +622,9 @@ namespace UnrealBuildTool
 		/// Finds all the module rules for plugins under the given directory.
 		/// </summary>
 		/// <param name="Plugins">The directory to search</param>
-		/// <param name="ModuleFiles">List of module files to be populated</param>
-		/// <param name="ModuleFileToPluginInfo">Dictionary which is filled with mappings from the module file to its corresponding plugin file</param>
-		private static void FindModuleRulesForPlugins(IReadOnlyList<PluginInfo> Plugins, List<FileReference> ModuleFiles, Dictionary<FileReference, PluginInfo> ModuleFileToPluginInfo)
+		/// <param name="DefaultContext">The default context for any files that are enumerated</param>
+		/// <param name="ModuleFileToContext">Dictionary which is filled with mappings from the module file to its corresponding context</param>
+		private static void FindModuleRulesForPlugins(IReadOnlyList<PluginInfo> Plugins, ModuleRulesContext DefaultContext, Dictionary<FileReference, ModuleRulesContext> ModuleFileToContext)
 		{
 			PrefetchRulesFiles(Plugins.Select(x => DirectoryReference.Combine(x.Directory, "Source")));
 
@@ -581,8 +633,10 @@ namespace UnrealBuildTool
 				IReadOnlyList<FileReference> PluginModuleFiles = FindAllRulesFiles(DirectoryReference.Combine(Plugin.Directory, "Source"), RulesFileType.Module);
 				foreach (FileReference ModuleFile in PluginModuleFiles)
 				{
-					ModuleFiles.Add(ModuleFile);
-					ModuleFileToPluginInfo[ModuleFile] = Plugin;
+					ModuleRulesContext PluginContext = new ModuleRulesContext(DefaultContext);
+					PluginContext.DefaultOutputBaseDir = Plugin.Directory;
+					PluginContext.Plugin = Plugin;
+					ModuleFileToContext[ModuleFile] = PluginContext;
 				}
 			}
 		}
diff --git a/Engine/Source/Programs/UnrealBuildTool/System/RulesScope.cs b/Engine/Source/Programs/UnrealBuildTool/System/RulesScope.cs
new file mode 100644
--- /dev/null
+++ b/Engine/Source/Programs/UnrealBuildTool/System/RulesScope.cs
@@ -0,0 +1,68 @@
+using System;
+using System.Collections.Generic;
+using System.Linq;
+using System.Text;
+using System.Threading.Tasks;
+
+namespace UnrealBuildTool
+{
+	/// <summary>
+	/// Represents a layer within the rules hierarchy. Any module is created within a certain scope, and may only reference modules in an equal or parent scope (eg. engine modules cannot reference project modules).
+	/// </summary>
+	class RulesScope
+	{
+		/// <summary>
+		/// Name of this scope
+		/// </summary>
+		public string Name;
+
+		/// <summary>
+		/// The parent scope
+		/// </summary>
+		public RulesScope Parent;
+
+		/// <summary>
+		/// Constructor
+		/// </summary>
+		/// <param name="Name">Name of this scope</param>
+		/// <param name="Parent">The parent scope</param>
+		public RulesScope(string Name, RulesScope Parent)
+		{
+			this.Name = Name;
+			this.Parent = Parent;
+		}
+
+		/// <summary>
+		/// Checks whether this scope contains another scope
+		/// </summary>
+		/// <param name="Other">The other scope to check</param>
+		/// <returns>True if this scope contains the other scope</returns>
+		public bool Contains(RulesScope Other)
+		{
+			for(RulesScope Scope = this; Scope != null; Scope = Scope.Parent)
+			{
+				if(Scope == Other)
+				{
+					return true;
+				}
+			}
+			return false;
+		}
+
+		/// <summary>
+		/// Formats the hierarchy of scopes
+		/// </summary>
+		/// <returns>String representing the hierarchy of scopes</returns>
+		public string FormatHierarchy()
+		{
+			if(Parent == null)
+			{
+				return Name;
+			}
+			else
+			{
+				return String.Format("{0} -> {1}", Name, Parent.FormatHierarchy());
+			}
+		}
+	}
+}
diff --git a/Engine/Source/Programs/UnrealBuildTool/UnrealBuildTool.cs b/Engine/Source/Programs/UnrealBuildTool/UnrealBuildTool.cs
--- a/Engine/Source/Programs/UnrealBuildTool/UnrealBuildTool.cs
+++ b/Engine/Source/Programs/UnrealBuildTool/UnrealBuildTool.cs
@@ -215,20 +215,6 @@ namespace UnrealBuildTool
 			RemoteIniPath = Path;
 		}
 
-		/// <summary>
-		/// Determines whether a directory is part of the engine
-		/// </summary>
-		/// <param name="InDirectory"></param>
-		/// <returns>true if the directory is under of the engine directories, false if not</returns>
-		static public bool IsUnderAnEngineDirectory(DirectoryReference InDirectory)
-		{
-			// Enterprise modules are considered as engine modules
-			return InDirectory.IsUnderDirectory(UnrealBuildTool.EngineDirectory) 
-				|| InDirectory.IsUnderDirectory(UnrealBuildTool.EnterpriseSourceDirectory) 
-				|| InDirectory.IsUnderDirectory(UnrealBuildTool.EnterprisePluginsDirectory) 
-				|| InDirectory.IsUnderDirectory(UnrealBuildTool.EnterpriseIntermediateDirectory);
-		}
-
 		/// <summary>
 		/// Global options for UBT (any modes)
 		/// </summary>
diff --git a/Engine/Source/Programs/UnrealBuildTool/UnrealBuildTool.csproj b/Engine/Source/Programs/UnrealBuildTool/UnrealBuildTool.csproj
--- a/Engine/Source/Programs/UnrealBuildTool/UnrealBuildTool.csproj
+++ b/Engine/Source/Programs/UnrealBuildTool/UnrealBuildTool.csproj
@@ -139,6 +139,7 @@
     </When>
   </Choose>
   <ItemGroup>
+    <Compile Include="Configuration\ModuleRulesContext.cs" />
     <Compile Include="Executors\HybridExecutor.cs" />
     <Compile Include="Modes\BuildMode.cs" />
     <Compile Include="Modes\CleanMode.cs" />
@@ -194,6 +195,7 @@
     <Compile Include="System\CppIncludeParser.cs" />
     <Compile Include="System\FileMetadataPrefetch.cs" />
     <Compile Include="System\ParseMsvcTimingInfo.cs" />
+    <Compile Include="System\RulesScope.cs" />
     <Compile Include="System\SourceFileMetadataCache.cs" />
     <Compile Include="System\DataDrivenPlatformInfo.cs" />
     <Compile Include="System\DirectoryItem.cs" />
