UBT: Add a flag for whether a module can be used for a shared PCH, rather than checking against the engine directory.

#rnx

diff --git a/Engine/Source/Programs/UnrealBuildTool/Configuration/ModuleRulesContext.cs b/Engine/Source/Programs/UnrealBuildTool/Configuration/ModuleRulesContext.cs
--- a/Engine/Source/Programs/UnrealBuildTool/Configuration/ModuleRulesContext.cs
+++ b/Engine/Source/Programs/UnrealBuildTool/Configuration/ModuleRulesContext.cs
@@ -37,6 +37,11 @@ namespace UnrealBuildTool
 		/// </summary>
 		public bool bCanBuildDebugGame;
 
+		/// <summary>
+		/// Whether this module can be used for generating shared PCHs
+		/// </summary>
+		public bool bCanUseForSharedPCH;
+
 		/// <summary>
 		/// Whether to treat this module as a game module for UHT ordering
 		/// </summary>
@@ -54,6 +59,7 @@ namespace UnrealBuildTool
 		{
 			this.Scope = Scope;
 			this.DefaultOutputBaseDir = DefaultOutputBaseDir;
+			this.bCanUseForSharedPCH = true;
 		}
 
 		/// <summary>
@@ -67,6 +73,7 @@ namespace UnrealBuildTool
 			this.Plugin = Other.Plugin;
 			this.bCanHotReload = Other.bCanHotReload;
 			this.bCanBuildDebugGame = Other.bCanBuildDebugGame;
+			this.bCanUseForSharedPCH = Other.bCanUseForSharedPCH;
 			this.bClassifyAsGameModuleForUHT = Other.bClassifyAsGameModuleForUHT;
 			this.DefaultUHTModuleType = Other.DefaultUHTModuleType;
 		}
diff --git a/Engine/Source/Programs/UnrealBuildTool/Configuration/UEBuildTarget.cs b/Engine/Source/Programs/UnrealBuildTool/Configuration/UEBuildTarget.cs
--- a/Engine/Source/Programs/UnrealBuildTool/Configuration/UEBuildTarget.cs
+++ b/Engine/Source/Programs/UnrealBuildTool/Configuration/UEBuildTarget.cs
@@ -2107,7 +2107,7 @@ namespace UnrealBuildTool
 			}
 
 			// Shared PCHs are only supported for engine modules at the moment. Check there are no game modules in the list.
-			List<UEBuildModuleCPP> NonEngineSharedPCHs = SharedPCHModules.Where(x => !x.RulesFile.IsUnderDirectory(UnrealBuildTool.EngineDirectory)).ToList();
+			List<UEBuildModuleCPP> NonEngineSharedPCHs = SharedPCHModules.Where(x => !x.Rules.Context.bCanUseForSharedPCH).ToList();
 			if(NonEngineSharedPCHs.Count > 0)
 			{
 				throw new BuildException("Shared PCHs are only supported for engine modules (found {0}).", String.Join(", ", NonEngineSharedPCHs.Select(x => x.Name)));
diff --git a/Engine/Source/Programs/UnrealBuildTool/System/RulesCompiler.cs b/Engine/Source/Programs/UnrealBuildTool/System/RulesCompiler.cs
--- a/Engine/Source/Programs/UnrealBuildTool/System/RulesCompiler.cs
+++ b/Engine/Source/Programs/UnrealBuildTool/System/RulesCompiler.cs
@@ -501,6 +501,7 @@ namespace UnrealBuildTool
 				DefaultModuleContext.bCanBuildDebugGame = true;
 				DefaultModuleContext.bCanHotReload = true;
 				DefaultModuleContext.bClassifyAsGameModuleForUHT = true;
+				DefaultModuleContext.bCanUseForSharedPCH = false;
 				
 				Dictionary<FileReference, ModuleRulesContext> ModuleFiles = new Dictionary<FileReference, ModuleRulesContext>();
 				AddModuleRulesWithContext(ProjectSourceDirectory, DefaultModuleContext, ModuleFiles);
