Disable GameApplication for Android projects still using Ant

#android

diff --git a/Engine/Source/Programs/UnrealBuildTool/Platform/Android/UEDeployAndroid.cs b/Engine/Source/Programs/UnrealBuildTool/Platform/Android/UEDeployAndroid.cs
--- a/Engine/Source/Programs/UnrealBuildTool/Platform/Android/UEDeployAndroid.cs
+++ b/Engine/Source/Programs/UnrealBuildTool/Platform/Android/UEDeployAndroid.cs
@@ -2398,7 +2398,10 @@ namespace UnrealBuildTool
 				}
 			}
 			Text.AppendLine("\t             android:hardwareAccelerated=\"true\"");
-			Text.AppendLine("\t				android:name=\"com.epicgames.ue4.GameApplication\"");
+			if (bGradleEnabled)
+			{
+				Text.AppendLine("\t				android:name=\"com.epicgames.ue4.GameApplication\"");
+			}
 			Text.AppendLine("\t             android:hasCode=\"true\">");
 			if (bShowLaunchImage)
 			{
@@ -3720,7 +3723,7 @@ namespace UnrealBuildTool
 
 				// update GameActivity.java and GameApplication.java if out of date
 				UpdateGameActivity(Arch, NDKArch, EngineDirectory, UE4BuildPath);
-				UpdateGameApplication(Arch, NDKArch, EngineDirectory, UE4BuildPath);
+				UpdateGameApplication(Arch, NDKArch, EngineDirectory, UE4BuildPath, bGradleEnabled);
 
 				// we don't actually need the SO for the bSkipGradleBuild case
 				string FinalSOName = null;
@@ -4351,11 +4354,18 @@ namespace UnrealBuildTool
 			}
 		}
 
-		private void UpdateGameApplication(string UE4Arch, string NDKArch, string EngineDir, string UE4BuildPath)
+		private void UpdateGameApplication(string UE4Arch, string NDKArch, string EngineDir, string UE4BuildPath, bool bGradleEnabled)
 		{
 			string SourceFilename = Path.Combine(EngineDir, "Build", "Android", "Java", "src", "com", "epicgames", "ue4", "GameApplication.java.template");
 			string DestFilename = Path.Combine(UE4BuildPath, "src", "com", "epicgames", "ue4", "GameApplication.java");
 
+			if (!bGradleEnabled)
+			{
+				// do not use GameApplication for Ant
+				SafeDeleteFile(DestFilename);
+				return;
+			}
+			
 			Dictionary<string, string> Replacements = new Dictionary<string, string>{
 				{ "//$${gameApplicationImportAdditions}$$", UPL.ProcessPluginNode(NDKArch, "gameApplicationImportAdditions", "")},
 				{ "//$${gameApplicationOnCreateAdditions}$$", UPL.ProcessPluginNode(NDKArch, "gameApplicationOnCreateAdditions", "")},
