UBT: Add FileItem.FullName and DirectoryItem.FullName properties.

#rnx

diff --git a/Engine/Source/Programs/UnrealBuildTool/System/DirectoryItem.cs b/Engine/Source/Programs/UnrealBuildTool/System/DirectoryItem.cs
--- a/Engine/Source/Programs/UnrealBuildTool/System/DirectoryItem.cs
+++ b/Engine/Source/Programs/UnrealBuildTool/System/DirectoryItem.cs
@@ -60,6 +60,14 @@ namespace UnrealBuildTool
 			get { return Info.Name; }
 		}
 
+		/// <summary>
+		/// The full name of this directory
+		/// </summary>
+		public string FullName
+		{
+			get { return Location.FullName; }
+		}
+
 		/// <summary>
 		/// Whether the directory exists or not
 		/// </summary>
diff --git a/Engine/Source/Programs/UnrealBuildTool/System/FileItem.cs b/Engine/Source/Programs/UnrealBuildTool/System/FileItem.cs
--- a/Engine/Source/Programs/UnrealBuildTool/System/FileItem.cs
+++ b/Engine/Source/Programs/UnrealBuildTool/System/FileItem.cs
@@ -57,6 +57,14 @@ namespace UnrealBuildTool
 			get { return Info.Name; }
 		}
 
+		/// <summary>
+		/// Full name of this file
+		/// </summary>
+		public string FullName
+		{
+			get { return Location.FullName; }
+		}
+
 		/// <summary>
 		/// Accessor for the absolute path to the file
 		/// </summary>
