UBT: Use a cached test for whether an include path exists. Small perf improvement when building a target.

#rnx

diff --git a/Engine/Source/Programs/UnrealBuildTool/Configuration/UEBuildModule.cs b/Engine/Source/Programs/UnrealBuildTool/Configuration/UEBuildModule.cs
--- a/Engine/Source/Programs/UnrealBuildTool/Configuration/UEBuildModule.cs
+++ b/Engine/Source/Programs/UnrealBuildTool/Configuration/UEBuildModule.cs
@@ -252,7 +252,7 @@ namespace UnrealBuildTool
 				foreach(string InputString in InEnumerableStrings)
 				{
 					DirectoryReference Dir = new DirectoryReference(ExpandPathVariables(InputString, null, null));
-					if(DirectoryReference.Exists(Dir))
+					if(DirectoryLookupCache.DirectoryExists(Dir))
 					{
 						Directories.Add(Dir);
 					}
