PR #5777: UBT: Clang 7/8 clang-cl supports PCH using /Yc and shared PCHs just fine. Remove legacy workaround slowing the build. (Contributed by [REDACTED_PERSON])

diff --git a/Engine/Source/Programs/UnrealBuildTool/Platform/Windows/UEBuildWindows.cs b/Engine/Source/Programs/UnrealBuildTool/Platform/Windows/UEBuildWindows.cs
--- a/Engine/Source/Programs/UnrealBuildTool/Platform/Windows/UEBuildWindows.cs
+++ b/Engine/Source/Programs/UnrealBuildTool/Platform/Windows/UEBuildWindows.cs
@@ -559,15 +559,6 @@ namespace UnrealBuildTool
 			}
 
 			// Override PCH settings
-			if (Target.WindowsPlatform.Compiler == WindowsCompiler.Clang)
-			{
-				// @todo clang: Shared PCHs don't work on clang yet because the PCH will have definitions assigned to different values
-				// than the consuming translation unit.  Unlike the warning in MSVC, this is a compile in Clang error which cannot be suppressed
-				Target.bUseSharedPCHs = false;
-
-				// @todo clang: PCH files aren't supported by "clang-cl" yet (no /Yc support, and "-x c++-header" cannot be specified)
-				Target.bUsePCHFiles = false;
-			}
 			if (Target.WindowsPlatform.Compiler == WindowsCompiler.Intel)
 			{
 				Target.NumIncludedBytesPerUnityCPP = Math.Min(Target.NumIncludedBytesPerUnityCPP, 256 * 1024);
