UBT: Remove the "default" value for the build environment. The property now returns the appropriate default value depending on the target being built.

diff --git a/Engine/Source/Programs/DotNETCommon/DotNETUtilities/CommandLineArguments.cs b/Engine/Source/Programs/DotNETCommon/DotNETUtilities/CommandLineArguments.cs
--- a/Engine/Source/Programs/DotNETCommon/DotNETUtilities/CommandLineArguments.cs
+++ b/Engine/Source/Programs/DotNETCommon/DotNETUtilities/CommandLineArguments.cs
@@ -903,7 +903,12 @@ namespace Tools.DotNETCommon
 		/// <returns>True if the text could be parsed, false otherwise</returns>
 		private static bool TryParseValue(Type FieldType, string Text, out object Value)
 		{
-			if(FieldType.IsEnum)
+			if(FieldType.IsGenericType && FieldType.GetGenericTypeDefinition() == typeof(Nullable<>))
+			{
+				// Try to parse the inner type instead
+				return TryParseValue(FieldType.GetGenericArguments()[0], Text, out Value);
+			}
+			else if (FieldType.IsEnum)
 			{
 				// Special handling for enums; parse the value ignoring case.
 				try
diff --git a/Engine/Source/Programs/UnrealBuildTool/Configuration/TargetRules.cs b/Engine/Source/Programs/UnrealBuildTool/Configuration/TargetRules.cs
--- a/Engine/Source/Programs/UnrealBuildTool/Configuration/TargetRules.cs
+++ b/Engine/Source/Programs/UnrealBuildTool/Configuration/TargetRules.cs
@@ -76,11 +76,6 @@ namespace UnrealBuildTool
 	[Serializable]
 	public enum TargetBuildEnvironment
 	{
-		/// <summary>
-		/// Use the default build environment for this target type (and whether the engine is installed)
-		/// </summary>
-		Default,
-
 		/// <summary>
 		/// Engine binaries and intermediates are output to the engine folder. Target-specific modifications to the engine build environment will be ignored.
 		/// </summary>
@@ -1118,11 +1113,41 @@ namespace UnrealBuildTool
 		public List<FileReference> DependencyListFileNames = new List<FileReference>();
 
 		/// <summary>
-		/// Specifies the build environment for this target. See TargetBuildEnvironment for more information on the available options.
+		/// Backing storage for the BuildEnvironment property
 		/// </summary>
 		[CommandLine("-SharedBuildEnvironment", Value = "Shared")]
 		[CommandLine("-UniqueBuildEnvironment", Value = "Unique")]
-		public TargetBuildEnvironment BuildEnvironment = TargetBuildEnvironment.Default;
+		private TargetBuildEnvironment? BuildEnvironmentOverride;
+
+		/// <summary>
+		/// Specifies the build environment for this target. See TargetBuildEnvironment for more information on the available options.
+		/// </summary>
+		public TargetBuildEnvironment BuildEnvironment
+		{
+			get
+			{
+				if(BuildEnvironmentOverride.HasValue)
+				{
+					return BuildEnvironmentOverride.Value;
+				}
+				if (Type == TargetType.Program && ProjectFile != null && File.IsUnderDirectory(ProjectFile.Directory))
+				{
+					return TargetBuildEnvironment.Unique;
+				}
+				else if (UnrealBuildTool.IsEngineInstalled() || LinkType != TargetLinkType.Monolithic)
+				{
+					return TargetBuildEnvironment.Shared;
+				}
+				else
+				{
+					return TargetBuildEnvironment.Unique;
+				}
+			}
+			set 
+			{ 
+				BuildEnvironmentOverride = value; 
+			}
+		}
 
 		/// <summary>
 		/// Whether to ignore violations to the shared build environment (eg. editor targets modifying definitions)
diff --git a/Engine/Source/Programs/UnrealBuildTool/System/RulesAssembly.cs b/Engine/Source/Programs/UnrealBuildTool/System/RulesAssembly.cs
--- a/Engine/Source/Programs/UnrealBuildTool/System/RulesAssembly.cs
+++ b/Engine/Source/Programs/UnrealBuildTool/System/RulesAssembly.cs
@@ -504,27 +504,9 @@ namespace UnrealBuildTool
 			}
 
 			// Set the default value for whether to use the shared build environment
-			if(Rules.BuildEnvironment == TargetBuildEnvironment.Default)
+			if(Rules.BuildEnvironment == TargetBuildEnvironment.Unique && UnrealBuildTool.IsEngineInstalled())
 			{
-				if(Rules.Type == TargetType.Program && TargetInfo.ProjectFile != null && TargetNameToTargetFile[Rules.Name].IsUnderDirectory(TargetInfo.ProjectFile.Directory))
-				{
-					Rules.BuildEnvironment = TargetBuildEnvironment.Unique;
-				}
-				else if(UnrealBuildTool.IsEngineInstalled() || Rules.LinkType != TargetLinkType.Monolithic)
-				{
-					Rules.BuildEnvironment = TargetBuildEnvironment.Shared;
-				}
-				else
-				{
-					Rules.BuildEnvironment = TargetBuildEnvironment.Unique;
-				}
-			}
-			else
-			{
-				if(Rules.BuildEnvironment == TargetBuildEnvironment.Unique && UnrealBuildTool.IsEngineInstalled())
-				{
-					throw new BuildException("Targets with a unique build environment cannot be built an installed engine.");
-				}
+				throw new BuildException("Targets with a unique build environment cannot be built an installed engine.");
 			}
 
 			// Automatically include CoreUObject
