#ocn - Hotfix from 6258862 to fix iOS link failing with argument list too long.

diff --git a/Engine/Source/Programs/UnrealBuildTool/ProjectFiles/Xcode/XcodeProject.cs b/Engine/Source/Programs/UnrealBuildTool/ProjectFiles/Xcode/XcodeProject.cs
--- a/Engine/Source/Programs/UnrealBuildTool/ProjectFiles/Xcode/XcodeProject.cs
+++ b/Engine/Source/Programs/UnrealBuildTool/ProjectFiles/Xcode/XcodeProject.cs
@@ -1332,6 +1332,9 @@ namespace UnrealBuildTool
 				Content.Append("\t\t\t\tVALID_ARCHS = \"" + ValidArchs + "\";" + ProjectFileGenerator.NewLine);
 				Content.Append("\t\t\t\tSUPPORTED_PLATFORMS = \"" + SupportedPlatforms + "\";" + ProjectFileGenerator.NewLine);
 			}
+			Content.Append("\t\t\t\tGCC_PREPROCESSOR_DEFINITIONS = ();" + ProjectFileGenerator.NewLine);
+			Content.Append("\t\t\t\tHEADER_SEARCH_PATHS = ();" + ProjectFileGenerator.NewLine);
+			Content.Append("\t\t\t\tUSER_HEADER_SEARCH_PATHS = ();" + ProjectFileGenerator.NewLine);
 			Content.Append("\t\t\t\tUE_BUILD_TARGET_NAME = \"" + Config.BuildTarget + "\";" + ProjectFileGenerator.NewLine);
 			Content.Append("\t\t\t\tUE_BUILD_TARGET_CONFIG = \"" + Config.BuildConfig + "\";" + ProjectFileGenerator.NewLine);
 			Content.Append("\t\t\t};" + ProjectFileGenerator.NewLine);
