Add support for UPL to insert a finalizing section at end of onCreate (gameActivityOnCreateFinalAdditions)
Optionally append to version name from an environment variable (UEAndroidDisplayNameExtra)

#android

diff --git a/Engine/Build/Android/Java/src/com/epicgames/ue4/GameActivity.java.template b/Engine/Build/Android/Java/src/com/epicgames/ue4/GameActivity.java.template
--- a/Engine/Build/Android/Java/src/com/epicgames/ue4/GameActivity.java.template
+++ b/Engine/Build/Android/Java/src/com/epicgames/ue4/GameActivity.java.template
@@ -2950,6 +2950,7 @@ public class GameActivity extends $${gameActivitySuperClass}$$ implements Surfac
 		clipboardManager = (ClipboardManager) getSystemService(CLIPBOARD_SERVICE);
 
 //$${gameActivityOnCreateAdditions}$$
+//$${gameActivityOnCreateFinalAdditions}$$
 		
 		Log.debug("==============> GameActive.onCreate complete!");
 	}
diff --git a/Engine/Source/Programs/UnrealBuildTool/Platform/Android/UEDeployAndroid.cs b/Engine/Source/Programs/UnrealBuildTool/Platform/Android/UEDeployAndroid.cs
--- a/Engine/Source/Programs/UnrealBuildTool/Platform/Android/UEDeployAndroid.cs
+++ b/Engine/Source/Programs/UnrealBuildTool/Platform/Android/UEDeployAndroid.cs
@@ -909,6 +909,13 @@ namespace UnrealBuildTool
 					if (bAppendPlatformToVersionDisplayName)
 					{
 						VersionDisplayName = string.Format("{0}-Android", VersionDisplayName);
+
+						// add optional value from environment variable if defined
+						string OptionalAppendDisplayName = Environment.GetEnvironmentVariable("UEAndroidDisplayNameExtra");
+						if (OptionalAppendDisplayName != null)
+						{
+							VersionDisplayName = VersionDisplayName + OptionalAppendDisplayName;
+						}
 					}
 				}
 
@@ -4304,6 +4311,7 @@ namespace UnrealBuildTool
 				{ "//$${gameActivityClassAdditions}$$", UPL.ProcessPluginNode(NDKArch, "gameActivityClassAdditions", "")},
 				{ "//$${gameActivityReadMetadataAdditions}$$", UPL.ProcessPluginNode(NDKArch, "gameActivityReadMetadataAdditions", "")},
 				{ "//$${gameActivityOnCreateAdditions}$$", UPL.ProcessPluginNode(NDKArch, "gameActivityOnCreateAdditions", "")},
+				{ "//$${gameActivityOnCreateFinalAdditions}$$", UPL.ProcessPluginNode(NDKArch, "gameActivityOnCreateFinalAdditions", "")},
 				{ "//$${gameActivityOnDestroyAdditions}$$", UPL.ProcessPluginNode(NDKArch, "gameActivityOnDestroyAdditions", "")},
 				{ "//$${gameActivityOnStartAdditions}$$", UPL.ProcessPluginNode(NDKArch, "gameActivityOnStartAdditions", "")},
 				{ "//$${gameActivityOnStopAdditions}$$", UPL.ProcessPluginNode(NDKArch, "gameActivityOnStopAdditions", "")},
