PR #5796: Fix bShadowVariableWarningsAsErrors being ignored (Contributed by [REDACTED_PERSON])

diff --git a/Engine/Source/Programs/UnrealBuildTool/Platform/Windows/VCToolChain.cs b/Engine/Source/Programs/UnrealBuildTool/Platform/Windows/VCToolChain.cs
--- a/Engine/Source/Programs/UnrealBuildTool/Platform/Windows/VCToolChain.cs
+++ b/Engine/Source/Programs/UnrealBuildTool/Platform/Windows/VCToolChain.cs
@@ -471,12 +471,18 @@ namespace UnrealBuildTool
 			if (Target.WindowsPlatform.Compiler >= WindowsCompiler.VisualStudio2015_DEPRECATED)
 			{
 				// Disable shadow variable warnings
-				if (CompileEnvironment.bEnableShadowVariableWarnings == false)
+				if (!CompileEnvironment.bEnableShadowVariableWarnings)
 				{
 					Arguments.Add("/wd4456"); // 4456 - declaration of 'LocalVariable' hides previous local declaration
 					Arguments.Add("/wd4458"); // 4458 - declaration of 'parameter' hides class member
 					Arguments.Add("/wd4459"); // 4459 - declaration of 'LocalVariable' hides global declaration
 				}
+				else if (CompileEnvironment.bShadowVariableWarningsAsErrors)
+				{
+					Arguments.Add("/we4456"); // 4456 - declaration of 'LocalVariable' hides previous local declaration
+					Arguments.Add("/we4458"); // 4458 - declaration of 'parameter' hides class member
+					Arguments.Add("/we4459"); // 4459 - declaration of 'LocalVariable' hides global declaration
+				}
 
 				Arguments.Add("/wd4463"); // 4463 - overflow; assigning 1 to bit-field that can only hold values from -1 to 0
 
