PR #5786: UE-73480: search for MSBuild in VS2019 paths (Contributed by [REDACTED_PERSON])

diff --git a/Engine/Source/Programs/UnrealBuildTool/Platform/Windows/UEBuildWindows.cs b/Engine/Source/Programs/UnrealBuildTool/Platform/Windows/UEBuildWindows.cs
--- a/Engine/Source/Programs/UnrealBuildTool/Platform/Windows/UEBuildWindows.cs
+++ b/Engine/Source/Programs/UnrealBuildTool/Platform/Windows/UEBuildWindows.cs
@@ -1139,9 +1139,21 @@ namespace UnrealBuildTool
 		/// <returns>True on success.</returns>
 		public static bool TryGetMsBuildPath(out FileReference OutLocation)
 		{
+			// Get the Visual Studio 2019 install directory
+			List<DirectoryReference> InstallDirs2019 = WindowsPlatform.FindVSInstallDirs(WindowsCompiler.VisualStudio2019);
+			foreach (DirectoryReference InstallDir in InstallDirs2019)
+			{
+				FileReference MsBuildLocation = FileReference.Combine(InstallDir, "MSBuild", "Current", "Bin", "MSBuild.exe");
+				if (FileReference.Exists(MsBuildLocation))
+				{
+					OutLocation = MsBuildLocation;
+					return true;
+				}
+			}
+
 			// Get the Visual Studio 2017 install directory
-			List<DirectoryReference> InstallDirs = WindowsPlatform.FindVSInstallDirs(WindowsCompiler.VisualStudio2017);
-			foreach(DirectoryReference InstallDir in InstallDirs)
+			List<DirectoryReference> InstallDirs2017 = WindowsPlatform.FindVSInstallDirs(WindowsCompiler.VisualStudio2017);
+			foreach (DirectoryReference InstallDir in InstallDirs2017)
 			{
 				FileReference MsBuildLocation = FileReference.Combine(InstallDir, "MSBuild", "15.0", "Bin", "MSBuild.exe");
 				if(FileReference.Exists(MsBuildLocation))
