UBT: Fix incorrect output directory being used for target version files.

diff --git a/Engine/Source/Programs/UnrealBuildTool/System/BuildVersion.cs b/Engine/Source/Programs/UnrealBuildTool/System/BuildVersion.cs
--- a/Engine/Source/Programs/UnrealBuildTool/System/BuildVersion.cs
+++ b/Engine/Source/Programs/UnrealBuildTool/System/BuildVersion.cs
@@ -104,6 +104,11 @@ namespace UnrealBuildTool
 		/// <summary>
 		/// Get the default path for a target's version file.
 		/// </summary>
+		/// <param name="OutputDirectory">The output directory for the executable. For MacOS, this is the directory containing the app bundle.</param>
+		/// <param name="TargetName">Name of the target being built</param>
+		/// <param name="Platform">Platform the target is being built for</param>
+		/// <param name="Configuration">The configuration being built</param>
+		/// <param name="Architecture">Architecture of the target being built</param>
 		/// <returns>Path to the target's version file</returns>
 		public static FileReference GetFileNameForTarget(DirectoryReference OutputDirectory, string TargetName, UnrealTargetPlatform Platform, UnrealTargetConfiguration Configuration, string Architecture)
 		{
@@ -117,11 +122,11 @@ namespace UnrealBuildTool
 			// Build the output filename
 			if (String.IsNullOrEmpty(ArchitectureSuffix) && Configuration == UnrealTargetConfiguration.Development)
 			{
-				return FileReference.Combine(OutputDirectory, "Binaries", Platform.ToString(), String.Format("{0}.version", TargetName));
+				return FileReference.Combine(OutputDirectory, String.Format("{0}.version", TargetName));
 			}
 			else
 			{
-				return FileReference.Combine(OutputDirectory, "Binaries", Platform.ToString(), String.Format("{0}-{1}-{2}{3}.version", TargetName, Platform.ToString(), Configuration.ToString(), ArchitectureSuffix));
+				return FileReference.Combine(OutputDirectory, String.Format("{0}-{1}-{2}{3}.version", TargetName, Platform.ToString(), Configuration.ToString(), ArchitectureSuffix));
 			}
 		}
 
