UE-72225

Adding support for generating an SNDBS template for cl-filter.exe so it can be run remotely. Couldn't repro the "not working" part of the bug.

diff --git a/Engine/Build/SNDBSTemplates/cl-filter.exe.sn-dbs-tool.ini b/Engine/Build/SNDBSTemplates/cl-filter.exe.sn-dbs-tool.ini
new file mode 100644
--- /dev/null
+++ b/Engine/Build/SNDBSTemplates/cl-filter.exe.sn-dbs-tool.ini
@@ -0,0 +1,51 @@
+[tool]
+family=msvc
+vc_major_version=14
+use_surrogate=true
+force_synchronous_pdb_writes=true
+error_report_mode=prompt
+
+[group]
+server={VC_COMPILER_DIR}\mspdbsrv.exe
+
+[files]
+main=cl-filter.exe
+file01={VC_COMPILER_DIR}\c1.dll
+file01={VC_COMPILER_DIR}\c1ui.dll
+file02={VC_COMPILER_DIR}\c1xx.dll
+file03={VC_COMPILER_DIR}\c2.dll
+file04={VC_COMPILER_DIR}\mspdb140.dll
+file05={VC_COMPILER_DIR}\mspdbcore.dll
+file06={VC_COMPILER_DIR}\mspdbsrv.exe
+file07={VC_COMPILER_DIR}\mspft140.dll
+file08={VC_COMPILER_DIR}\vcmeta.dll
+file09={VC_COMPILER_DIR}\*\clui.dll
+file10={VC_COMPILER_DIR}\*\mspft140ui.dll
+file11={VC_COMPILER_DIR}\localespc.dll
+file12={VC_COMPILER_DIR}\cppcorecheck.dll
+file13={VC_COMPILER_DIR}\experimentalcppcorecheck.dll
+file14={VC_COMPILER_DIR}\espxengine.dll
+file15={VC_COMPILER_DIR}\c1.exe
+
+[output-file-patterns]
+outputfile01=\s*"([^ ",]+\.cpp\.txt)\"
+
+[output-file-rules]
+rule01=*\sqmcpp*.log|discard=true
+rule02=*\vctoolstelemetry*.dat|discard=true
+rule03=*\Microsoft\Windows\Temporary Internet Files\*|discard=true
+rule04=*\Microsoft\Windows\INetCache\*|discard=true
+
+[input-file-rules]
+rule01=*\sqmcpp*.log|ignore_transient_errors=true;ignore_unexpected_input=true
+rule02=*\vctoolstelemetry*.dat|ignore_transient_errors=true;ignore_unexpected_input=true
+rule03=*\Microsoft\Windows\Temporary Internet Files\*|ignore_transient_errors=true;ignore_unexpected_input=true
+rule04=*\Microsoft\Windows\INetCache\*|ignore_transient_errors=true;ignore_unexpected_input=true
+
+[system-file-filters]
+filter01=msvcr*.dll
+filter02=msvcp*.dll
+filter03=vcruntime140*.dll
+filter04=appcrt140*.dll
+filter05=desktopcrt140*.dll
+filter06=concrt140*.dll
diff --git a/Engine/Source/Programs/UnrealBuildTool/Configuration/UEBuildTarget.cs b/Engine/Source/Programs/UnrealBuildTool/Configuration/UEBuildTarget.cs
--- a/Engine/Source/Programs/UnrealBuildTool/Configuration/UEBuildTarget.cs
+++ b/Engine/Source/Programs/UnrealBuildTool/Configuration/UEBuildTarget.cs
@@ -1404,7 +1404,7 @@ namespace UnrealBuildTool
 				List<UEBuildBinary> FilteredBinaries = new List<UEBuildBinary>();
 				foreach (UEBuildBinary DLLBinary in Binaries)
 				{
-					if(ModDirectories.Any(x => DLLBinary.OutputFilePath.IsUnderDirectory(x)))
+					if (ModDirectories.Any(x => DLLBinary.OutputFilePath.IsUnderDirectory(x)))
 					{
 						FilteredBinaries.Add(DLLBinary);
 					}
@@ -1419,9 +1419,9 @@ namespace UnrealBuildTool
 
 			// Build a mapping from module to its plugin
 			Dictionary<UEBuildModule, UEBuildPlugin> ModuleToPlugin = new Dictionary<UEBuildModule, UEBuildPlugin>();
-			foreach(UEBuildPlugin Plugin in BuildPlugins)
+			foreach (UEBuildPlugin Plugin in BuildPlugins)
 			{
-				foreach(UEBuildModule Module in Plugin.Modules)
+				foreach (UEBuildModule Module in Plugin.Modules)
 				{
 					if (!ModuleToPlugin.ContainsKey(Module))
 					{
@@ -1431,14 +1431,14 @@ namespace UnrealBuildTool
 			}
 
 			// Check there aren't any engine binaries with dependencies on game modules. This can happen when game-specific plugins override engine plugins.
-			foreach(UEBuildModule Module in Modules.Values)
+			foreach (UEBuildModule Module in Modules.Values)
 			{
-				if(Module.Binary != null)
+				if (Module.Binary != null)
 				{
 					HashSet<UEBuildModule> ReferencedModules = Module.GetDependencies(bWithIncludePathModules: true, bWithDynamicallyLoadedModules: true);
-					foreach(UEBuildModule ReferencedModule in ReferencedModules)
+					foreach (UEBuildModule ReferencedModule in ReferencedModules)
 					{
-						if(!Module.Rules.Context.Scope.Contains(ReferencedModule.Rules.Context.Scope) && !IsWhitelistedEnginePluginReference(Module.Name, ReferencedModule.Name))
+						if (!Module.Rules.Context.Scope.Contains(ReferencedModule.Rules.Context.Scope) && !IsWhitelistedEnginePluginReference(Module.Name, ReferencedModule.Name))
 						{
 							throw new BuildException("Module '{0}' ({1}) should not reference module '{2}' ({3}). Hierarchy is {4}.", Module.Name, Module.Rules.Context.Scope.Name, ReferencedModule.Name, ReferencedModule.Rules.Context.Scope.Name, ReferencedModule.Rules.Context.Scope.FormatHierarchy());
 						}
@@ -1447,15 +1447,15 @@ namespace UnrealBuildTool
 			}
 
 			// Check that each plugin declares its dependencies explicitly
-			foreach(UEBuildPlugin Plugin in BuildPlugins)
+			foreach (UEBuildPlugin Plugin in BuildPlugins)
 			{
-				foreach(UEBuildModule Module in Plugin.Modules)
+				foreach (UEBuildModule Module in Plugin.Modules)
 				{
 					HashSet<UEBuildModule> DependencyModules = Module.GetDependencies(bWithIncludePathModules: true, bWithDynamicallyLoadedModules: true);
-					foreach(UEBuildModule DependencyModule in DependencyModules)
+					foreach (UEBuildModule DependencyModule in DependencyModules)
 					{
 						UEBuildPlugin DependencyPlugin;
-						if(ModuleToPlugin.TryGetValue(DependencyModule, out DependencyPlugin) && DependencyPlugin != Plugin && !Plugin.Dependencies.Contains(DependencyPlugin))
+						if (ModuleToPlugin.TryGetValue(DependencyModule, out DependencyPlugin) && DependencyPlugin != Plugin && !Plugin.Dependencies.Contains(DependencyPlugin))
 						{
 							Log.TraceWarning("Warning: Plugin '{0}' does not list plugin '{1}' as a dependency, but module '{2}' depends on '{3}'.", Plugin.Name, DependencyPlugin.Name, Module.Name, DependencyModule.Name);
 						}
@@ -1467,6 +1467,9 @@ namespace UnrealBuildTool
 			string ExternalMetadata = UEBuildPlatform.GetBuildPlatform(Platform).GetExternalBuildMetadata(ProjectFile);
 			TargetMakefile Makefile = new TargetMakefile(TargetToolChain.GetVersionInfo(), ExternalMetadata, ReceiptFileName, ProjectIntermediateDirectory, TargetType, bDeployAfterCompile, bHasProjectScriptPlugin);
 
+			// Generate the SNDBS include-rewrite-rules.ini in case we'll need it.
+			GenerateSNDBSIncludeRewriteRules();
+
 			// Setup the hot reload module list
 			Makefile.HotReloadModuleNames = GetHotReloadModuleNames();
 
@@ -1487,7 +1490,7 @@ namespace UnrealBuildTool
 			}
 
 			// Add global definitions for project-specific binaries. HACK: Also defining for monolithic builds in binary releases. Might be better to set this via command line instead?
-			if(!bUseSharedBuildEnvironment || bCompileMonolithic)
+			if (!bUseSharedBuildEnvironment || bCompileMonolithic)
 			{
 				UEBuildBinary ExecutableBinary = Binaries[0];
 
@@ -1523,7 +1526,7 @@ namespace UnrealBuildTool
 
 			// Generate headers
 			HashSet<UEBuildModuleCPP> ModulesToGenerateHeadersFor = GatherDependencyModules(OriginalBinaries.ToList());
-			using(Timeline.ScopeEvent("ExternalExecution.SetupUObjectModules()"))
+			using (Timeline.ScopeEvent("ExternalExecution.SetupUObjectModules()"))
 			{
 				ExternalExecution.SetupUObjectModules(ModulesToGenerateHeadersFor, Rules.Platform, ProjectDescriptor, Makefile.UObjectModules, Makefile.UObjectModuleHeaders, Rules.GeneratedCodeVersion, bIsAssemblingBuild, MetadataCache);
 			}
@@ -1547,7 +1550,7 @@ namespace UnrealBuildTool
 			{
 				if (Platform == UnrealTargetPlatform.Win32 || Platform == UnrealTargetPlatform.Win64)
 				{
-					if(!Rules.bFormalBuild)
+					if (!Rules.bFormalBuild)
 					{
 						CppCompileEnvironment DefaultResourceCompileEnvironment = new CppCompileEnvironment(GlobalCompileEnvironment);
 
@@ -1561,7 +1564,7 @@ namespace UnrealBuildTool
 
 			// Build the target's binaries.
 			DirectoryReference ExeDir = GetExecutableDir();
-			using(Timeline.ScopeEvent("UEBuildBinary.Build()"))
+			using (Timeline.ScopeEvent("UEBuildBinary.Build()"))
 			{
 				foreach (UEBuildBinary Binary in Binaries)
 				{
@@ -1573,20 +1576,20 @@ namespace UnrealBuildTool
 			// Prepare all the runtime dependencies, copying them from their source folders if necessary
 			List<RuntimeDependency> RuntimeDependencies = new List<RuntimeDependency>();
 			Dictionary<FileReference, FileReference> RuntimeDependencyTargetFileToSourceFile = new Dictionary<FileReference, FileReference>();
-			foreach(UEBuildBinary Binary in Binaries)
+			foreach (UEBuildBinary Binary in Binaries)
 			{
 				Binary.PrepareRuntimeDependencies(RuntimeDependencies, RuntimeDependencyTargetFileToSourceFile, ExeDir);
 			}
-			foreach(KeyValuePair<FileReference, FileReference> Pair in RuntimeDependencyTargetFileToSourceFile)
+			foreach (KeyValuePair<FileReference, FileReference> Pair in RuntimeDependencyTargetFileToSourceFile)
 			{
-				if(!UnrealBuildTool.IsFileInstalled(Pair.Key))
+				if (!UnrealBuildTool.IsFileInstalled(Pair.Key))
 				{
 					Makefile.OutputItems.Add(CreateCopyAction(Pair.Value, Pair.Key, Makefile.Actions));
 				}
 			}
 
 			// If we're just precompiling a plugin, only include output items which are under that directory
-			if(ForeignPlugin != null)
+			if (ForeignPlugin != null)
 			{
 				Makefile.OutputItems.RemoveAll(x => !x.Location.IsUnderDirectory(ForeignPlugin.Directory));
 			}
@@ -1605,7 +1608,7 @@ namespace UnrealBuildTool
 			BuildProducts.AddRange(RuntimeDependencyTargetFileToSourceFile.Select(x => new KeyValuePair<FileReference, BuildProductType>(x.Key, BuildProductType.RequiredResource)));
 
 			// Remove any installed build products that don't exist. They may be part of an optional install.
-			if(UnrealBuildTool.IsEngineInstalled())
+			if (UnrealBuildTool.IsEngineInstalled())
 			{
 				BuildProducts.RemoveAll(x => UnrealBuildTool.IsFileInstalled(x.Key) && !FileReference.Exists(x.Key));
 			}
@@ -1621,13 +1624,13 @@ namespace UnrealBuildTool
 			}
 
 			// Finalize and generate metadata for this target
-			if(!Rules.bDisableLinking)
+			if (!Rules.bDisableLinking)
 			{
 				// Also add any explicitly specified build products
-				if(Rules.AdditionalBuildProducts.Count > 0)
+				if (Rules.AdditionalBuildProducts.Count > 0)
 				{
 					Dictionary<string, string> Variables = GetTargetVariables(null);
-					foreach(string AdditionalBuildProduct in Rules.AdditionalBuildProducts)
+					foreach (string AdditionalBuildProduct in Rules.AdditionalBuildProducts)
 					{
 						FileReference BuildProductFile = new FileReference(Utils.ExpandVariables(AdditionalBuildProduct, Variables));
 						BuildProducts.Add(new KeyValuePair<FileReference, BuildProductType>(BuildProductFile, BuildProductType.RequiredResource));
@@ -1636,18 +1639,18 @@ namespace UnrealBuildTool
 
 				// Get the path to the version file unless this is a formal build (where it will be compiled in)
 				FileReference VersionFile = null;
-				if(Rules.LinkType != TargetLinkType.Monolithic && Binaries[0].Type == UEBuildBinaryType.Executable)
+				if (Rules.LinkType != TargetLinkType.Monolithic && Binaries[0].Type == UEBuildBinaryType.Executable)
 				{
 					UnrealTargetConfiguration VersionConfig = Configuration;
-					if(VersionConfig == UnrealTargetConfiguration.DebugGame && !bCompileMonolithic && TargetType != TargetType.Program && bUseSharedBuildEnvironment)
+					if (VersionConfig == UnrealTargetConfiguration.DebugGame && !bCompileMonolithic && TargetType != TargetType.Program && bUseSharedBuildEnvironment)
 					{
 						VersionConfig = UnrealTargetConfiguration.Development;
 					}
-					VersionFile = BuildVersion.GetFileNameForTarget(ExeDir, bCompileMonolithic? TargetName : AppName, Platform, VersionConfig, Architecture);
+					VersionFile = BuildVersion.GetFileNameForTarget(ExeDir, bCompileMonolithic ? TargetName : AppName, Platform, VersionConfig, Architecture);
 				}
 
 				// Also add the version file as a build product
-				if(VersionFile != null)
+				if (VersionFile != null)
 				{
 					BuildProducts.Add(new KeyValuePair<FileReference, BuildProductType>(VersionFile, BuildProductType.RequiredResource));
 				}
@@ -1667,7 +1670,7 @@ namespace UnrealBuildTool
 				StringBuilder WriteMetadataArguments = new StringBuilder();
 				WriteMetadataArguments.AppendFormat("-Input={0}", Utils.MakePathSafeToUseWithCommandLine(MetadataTargetFile));
 				WriteMetadataArguments.AppendFormat(" -Version={0}", WriteMetadataMode.CurrentVersionNumber);
-				if(Rules.bNoManifestChanges)
+				if (Rules.bNoManifestChanges)
 				{
 					WriteMetadataArguments.Append(" -NoManifestChanges");
 				}
@@ -1685,13 +1688,13 @@ namespace UnrealBuildTool
 
 				// Create actions to run the post build steps
 				FileReference[] PostBuildScripts = CreatePostBuildScripts();
-				foreach(FileReference PostBuildScript in PostBuildScripts)
+				foreach (FileReference PostBuildScript in PostBuildScripts)
 				{
 					FileReference OutputFile = new FileReference(PostBuildScript.FullName + ".ran");
 
 					Action PostBuildStepAction = new Action(ActionType.PostBuildStep);
 					PostBuildStepAction.CommandPath = BuildHostPlatform.Current.Shell;
-					if(BuildHostPlatform.Current.ShellType == ShellType.Cmd)
+					if (BuildHostPlatform.Current.ShellType == ShellType.Cmd)
 					{
 						PostBuildStepAction.CommandArguments = String.Format("/C \"call \"{0}\" && type NUL >\"{1}\"\"", PostBuildScript, OutputFile);
 					}
@@ -1711,22 +1714,22 @@ namespace UnrealBuildTool
 			}
 
 			// Build a list of all the files required to build
-			foreach(FileReference DependencyListFileName in Rules.DependencyListFileNames)
+			foreach (FileReference DependencyListFileName in Rules.DependencyListFileNames)
 			{
 				WriteDependencyList(DependencyListFileName, RuntimeDependencies, RuntimeDependencyTargetFileToSourceFile);
 			}
 
 			// If we're only generating the manifest, return now
-			foreach(FileReference ManifestFileName in Rules.ManifestFileNames)
+			foreach (FileReference ManifestFileName in Rules.ManifestFileNames)
 			{
 				GenerateManifest(ManifestFileName, BuildProducts);
 			}
 
 			// Check there are no EULA or restricted folder violations
-			if(!Rules.bDisableLinking)
+			if (!Rules.bDisableLinking)
 			{
 				// Check the distribution level of all binaries based on the dependencies they have
-				if(ProjectFile == null && !Rules.bOutputPubliclyDistributable)
+				if (ProjectFile == null && !Rules.bOutputPubliclyDistributable)
 				{
 					Dictionary<UEBuildModule, Dictionary<RestrictedFolder, DirectoryReference>> ModuleRestrictedFolderCache = new Dictionary<UEBuildModule, Dictionary<RestrictedFolder, DirectoryReference>>();
 
@@ -1735,12 +1738,12 @@ namespace UnrealBuildTool
 					{
 						bResult &= Binary.CheckRestrictedFolders(DirectoryReference.FromFile(ProjectFile), ModuleRestrictedFolderCache);
 					}
-					foreach(KeyValuePair<FileReference, FileReference> Pair in RuntimeDependencyTargetFileToSourceFile)
+					foreach (KeyValuePair<FileReference, FileReference> Pair in RuntimeDependencyTargetFileToSourceFile)
 					{
 						bResult &= CheckRestrictedFolders(Pair.Key, Pair.Value);
 					}
 
-					if(!bResult)
+					if (!bResult)
 					{
 						throw new BuildException("Unable to create binaries in less restricted locations than their input files.");
 					}
@@ -1751,7 +1754,7 @@ namespace UnrealBuildTool
 			}
 
 			// Add all the plugins to be tracked
-			foreach(FileReference PluginFile in global::UnrealBuildTool.Plugins.EnumeratePlugins(ProjectFile))
+			foreach (FileReference PluginFile in global::UnrealBuildTool.Plugins.EnumeratePlugins(ProjectFile))
 			{
 				FileItem PluginFileItem = FileItem.GetItemByFileReference(PluginFile);
 				Makefile.PluginFiles.Add(PluginFileItem);
@@ -1759,10 +1762,10 @@ namespace UnrealBuildTool
 
 			// Add all the input files to the predicate store
 			Makefile.AdditionalDependencies.Add(FileItem.GetItemByFileReference(TargetRulesFile));
-			foreach(UEBuildModule Module in Modules.Values)
+			foreach (UEBuildModule Module in Modules.Values)
 			{
 				Makefile.AdditionalDependencies.Add(FileItem.GetItemByFileReference(Module.RulesFile));
-				foreach(string ExternalDependency in Module.Rules.ExternalDependencies)
+				foreach (string ExternalDependency in Module.Rules.ExternalDependencies)
 				{
 					FileReference Location = FileReference.Combine(Module.RulesFile.Directory, ExternalDependency);
 					Makefile.AdditionalDependencies.Add(FileItem.GetItemByFileReference(Location));
@@ -1771,7 +1774,7 @@ namespace UnrealBuildTool
 			Makefile.AdditionalDependencies.UnionWith(Makefile.PluginFiles);
 
 			// Write a header containing public definitions for this target
-			if(Rules.ExportPublicHeader != null)
+			if (Rules.ExportPublicHeader != null)
 			{
 				UEBuildBinary Binary = Binaries[0];
 				FileReference Header = FileReference.Combine(Binary.OutputDir, Rules.ExportPublicHeader);
@@ -1784,6 +1787,21 @@ namespace UnrealBuildTool
 			return Makefile;
 		}
 
+		private void GenerateSNDBSIncludeRewriteRules()
+		{
+			FileReference IncludeRewriteRulesFile = FileReference.Combine(UnrealBuildTool.EngineDirectory, "Build", "SNDBS", "include-rewrite-rules.ini");
+			if (!Directory.Exists(IncludeRewriteRulesFile.Directory.FullName))
+			{
+				Directory.CreateDirectory(IncludeRewriteRulesFile.Directory.FullName);
+			}
+			List<string> IncludeRewriteRulesText = new List<string>();
+			IncludeRewriteRulesText.Add("[computed-include-rules]");
+			IncludeRewriteRulesText.Add(@"pattern1=^COMPILED_PLATFORM_HEADER\(\s*([^ ,]+)\)");
+			string PlatformString = Platform.ToString().StartsWith("Win") ? "Windows" : Platform.ToString();
+			IncludeRewriteRulesText.Add(String.Format("expansions1={0}/{0}$1", PlatformString));
+			File.WriteAllText(IncludeRewriteRulesFile.FullName, String.Join(Environment.NewLine, IncludeRewriteRulesText));
+		}
+
 		/// <summary>
 		/// Gets the output directory for the main executable
 		/// </summary>
diff --git a/Engine/Source/Programs/UnrealBuildTool/Executors/SNDBS.cs b/Engine/Source/Programs/UnrealBuildTool/Executors/SNDBS.cs
--- a/Engine/Source/Programs/UnrealBuildTool/Executors/SNDBS.cs
+++ b/Engine/Source/Programs/UnrealBuildTool/Executors/SNDBS.cs
@@ -1,17 +1,14 @@
 // Copyright 1998-2019 Epic Games, Inc. All Rights Reserved.
 
 using System;
+using System.Collections;
 using System.Collections.Generic;
-using System.IO;
 using System.Diagnostics;
-using System.Xml;
-using System.Text.RegularExpressions;
+using System.IO;
 using System.Linq;
-using System.Reflection;
-using System.Threading;
 using System.ServiceProcess;
-using Tools.DotNETCommon;
 using System.Text;
+using Tools.DotNETCommon;
 
 namespace UnrealBuildTool
 {
@@ -182,6 +179,7 @@ namespace UnrealBuildTool
 			int NumScriptedActions = 0;
 			List<Action> LocalActions = new List<Action>();
 			ActionThread DummyActionThread = new ActionThread(null, 1, 1);
+			bool PrintDebugInfo = false;
 			foreach (Action Action in InActions)
 			{
 				ActionThread ActionProcess = null;
@@ -267,6 +265,7 @@ namespace UnrealBuildTool
 							Log.TraceInformation("[{0}/{1}] {2} {3}", JobNumber, InActions.Count, Action.CommandDescription, Action.StatusDescription);
 							JobNumber++;
 							NumScriptedActions++;
+							PrintDebugInfo |= Action.bPrintDebugInfo;
 						}
 					}
 				}
@@ -283,12 +282,15 @@ namespace UnrealBuildTool
                 string SCERoot = Environment.GetEnvironmentVariable("SCE_ROOT_DIR");
                 string SNDBSExecutable = Path.Combine(SCERoot, "Common/SN-DBS/bin/dbsbuild.exe");
 				DirectoryReference TemplatesDir = DirectoryReference.Combine(UnrealBuildTool.EngineDirectory, "Programs", "UnrealBuildTool", "SndbsTemplates");
-                ProcessStartInfo PSI = new ProcessStartInfo(SNDBSExecutable, String.Format("-q -p UE4 -s \"{0}\" -templates \"{1}\"", FileReference.Combine(UnrealBuildTool.EngineDirectory, "Intermediate", "Build", "sndbs.bat").FullName, TemplatesDir.FullName));
+				FileReference IncludeRewriteRules = FileReference.Combine(UnrealBuildTool.EngineDirectory, "Build", "SNDBS", "include-rewrite-rules.ini");
+				string IncludeRewriteRulesArg = String.Format("--include-rewrite-rules \"{0}\"", IncludeRewriteRules.FullName);
+				string VerbosityLevel = PrintDebugInfo ? "-v" : "-q";
+                ProcessStartInfo PSI = new ProcessStartInfo(SNDBSExecutable, String.Format("{0} -p UE4 -s \"{1}\" -templates \"{2}\" {3}", VerbosityLevel, FileReference.Combine(UnrealBuildTool.EngineDirectory, "Intermediate", "Build", "sndbs.bat").FullName, TemplatesDir.FullName, IncludeRewriteRulesArg));
 				PSI.RedirectStandardOutput = true;
 				PSI.RedirectStandardError = true;
 				PSI.UseShellExecute = false;
 				PSI.CreateNoWindow = true;
-				PSI.WorkingDirectory = Path.GetFullPath("."); ;
+				PSI.WorkingDirectory = Path.GetFullPath(".");
 				Process NewProcess = new Process();
 				NewProcess.StartInfo = PSI;
 				NewProcess.OutputDataReceived += new DataReceivedEventHandler(ActionDebugOutput);
@@ -356,6 +358,13 @@ namespace UnrealBuildTool
 			bool SNDBSResult = true;
 			if (Actions.Count > 0)
 			{
+				// Generate any needed templates. Can only generate one per executable, so just use the first Action for reference
+				IEnumerable<Action> CommandPaths = Actions.GroupBy(a => a.CommandPath).Select(g => g.FirstOrDefault()).ToList();
+				foreach (Action CommandPath in CommandPaths)
+				{
+					PrepareToolTemplate(CommandPath);
+				}
+
 				// Use WMI to figure out physical cores, excluding hyper threading.
 				int NumCores = Utils.GetPhysicalProcessorCount();
 				
@@ -450,5 +459,32 @@ namespace UnrealBuildTool
 			}
 			return SNDBSResult;
 		}
+
+		private void PrepareToolTemplate(Action Action)
+		{
+			string TemplateFileName = String.Format("{0}.sn-dbs-tool.ini", Action.CommandPath.GetFileName());
+			FileReference TemplateInput = FileReference.Combine(UnrealBuildTool.EngineDirectory, "Build", "SNDBSTemplates", TemplateFileName);
+			
+			// If no base template exists, don't try to generate one.
+			if (!File.Exists(TemplateInput.FullName))
+			{
+				return;
+			}
+
+			FileReference TemplateOutput = FileReference.Combine(UnrealBuildTool.EngineDirectory, "Programs", "UnrealBuildTool", "SNDBSTemplates", TemplateFileName);
+			if (!Directory.Exists(TemplateOutput.Directory.FullName))
+			{
+				Directory.CreateDirectory(TemplateOutput.Directory.FullName);
+			}
+
+			string TemplateText = File.ReadAllText(TemplateInput.FullName);
+			TemplateText = TemplateText.Replace("{COMMAND_PATH}", Action.CommandPath.Directory.FullName);
+			foreach (DictionaryEntry Variable in Environment.GetEnvironmentVariables(EnvironmentVariableTarget.Process))
+			{
+				string VariableName = String.Format("{{{0}}}", Variable.Key);
+				TemplateText = TemplateText.Replace(VariableName, Variable.Value.ToString());
+			}
+			File.WriteAllText(TemplateOutput.FullName, TemplateText);
+		}
 	}
 }
diff --git a/Engine/Source/Programs/UnrealBuildTool/Platform/Windows/VCEnvironment.cs b/Engine/Source/Programs/UnrealBuildTool/Platform/Windows/VCEnvironment.cs
--- a/Engine/Source/Programs/UnrealBuildTool/Platform/Windows/VCEnvironment.cs
+++ b/Engine/Source/Programs/UnrealBuildTool/Platform/Windows/VCEnvironment.cs
@@ -117,6 +117,10 @@ namespace UnrealBuildTool
             // Compile using 64 bit tools for 64 bit targets, and 32 for 32.
 			CompilerPath = GetCompilerToolPath(Platform, Compiler, CompilerDir);
 
+			// Add the compiler path and directory as environment variables for the process so they may be used elsewhere.
+			Environment.SetEnvironmentVariable("VC_COMPILER_PATH", CompilerPath.FullName, EnvironmentVariableTarget.Process);
+			Environment.SetEnvironmentVariable("VC_COMPILER_DIR", CompilerPath.Directory.FullName, EnvironmentVariableTarget.Process);
+
 			// Regardless of the target, if we're linking on a 64 bit machine, we want to use the 64 bit linker (it's faster than the 32 bit linker and can handle large linking jobs)
 			DirectoryReference DefaultLinkerDir = VCToolPath64;
 			LinkerPath = GetLinkerToolPath(Platform, Compiler, DefaultLinkerDir);
diff --git a/Engine/Source/Programs/UnrealBuildTool/Platform/Windows/VCToolChain.cs b/Engine/Source/Programs/UnrealBuildTool/Platform/Windows/VCToolChain.cs
--- a/Engine/Source/Programs/UnrealBuildTool/Platform/Windows/VCToolChain.cs
+++ b/Engine/Source/Programs/UnrealBuildTool/Platform/Windows/VCToolChain.cs
@@ -1183,6 +1183,7 @@ namespace UnrealBuildTool
 				ParseTimingInfoAction.WorkingDirectory = UnrealBuildTool.EngineSourceDirectory;
 				ParseTimingInfoAction.StatusDescription = String.Format("Parsing Timing File '{0}'", CompileAction.TimingFile);
 				ParseTimingInfoAction.bCanExecuteRemotely = true;
+				ParseTimingInfoAction.bCanExecuteRemotelyWithSNDBS = true;
 				ParseTimingInfoAction.PrerequisiteItems.Add(SourceFile);
 				ParseTimingInfoAction.PrerequisiteItems.Add(CompileAction.TimingFile);
 				ParseTimingInfoAction.ProducedItems.Add(TimingJsonFile);
