Fix for case of engine directory being derived from launch path, and potentially not matching case of files on disk.

diff --git a/Engine/Source/Programs/UnrealBuildTool/UnrealBuildTool.cs b/Engine/Source/Programs/UnrealBuildTool/UnrealBuildTool.cs
--- a/Engine/Source/Programs/UnrealBuildTool/UnrealBuildTool.cs
+++ b/Engine/Source/Programs/UnrealBuildTool/UnrealBuildTool.cs
@@ -39,10 +39,15 @@ namespace UnrealBuildTool
 		/// </summary>
 		static FileReference InstalledProjectFile;
 
+		/// <summary>
+		/// The path to UBT
+		/// </summary>
+		public static readonly FileReference UnrealBuildToolPath = FileReference.FindCorrectCase(new FileReference(Assembly.GetExecutingAssembly().GetOriginalLocation()));
+
 		/// <summary>
 		/// The full name of the Root UE4 directory
 		/// </summary>
-		public static readonly DirectoryReference RootDirectory = new DirectoryReference(Path.Combine(Path.GetDirectoryName(Assembly.GetExecutingAssembly().GetOriginalLocation()), "..", "..", ".."));
+		public static readonly DirectoryReference RootDirectory = DirectoryReference.Combine(UnrealBuildToolPath.Directory, "..", "..", "..");
 
 		/// <summary>
 		/// The full name of the Engine directory
@@ -202,7 +207,6 @@ namespace UnrealBuildTool
 		/// <returns>A string containing the path to the UBT assembly.</returns>
 		static public FileReference GetUBTPath()
 		{
-			FileReference UnrealBuildToolPath = new FileReference(Assembly.GetExecutingAssembly().GetOriginalLocation());
 			return UnrealBuildToolPath;
 		}
 
