- Fixed Mac build and project gneeration errors

diff --git a/Engine/Source/Programs/UnrealBuildTool/ProjectFiles/Xcode/XcodeProject.cs b/Engine/Source/Programs/UnrealBuildTool/ProjectFiles/Xcode/XcodeProject.cs
--- a/Engine/Source/Programs/UnrealBuildTool/ProjectFiles/Xcode/XcodeProject.cs
+++ b/Engine/Source/Programs/UnrealBuildTool/ProjectFiles/Xcode/XcodeProject.cs
@@ -219,7 +219,7 @@ namespace UnrealBuildTool
 
 			if (IsSourceCode(FileExtension))
 			{
-				foreach (string PlatformName in Enum.GetNames(typeof(UnrealTargetPlatform)))
+				foreach (string PlatformName in UnrealTargetPlatform.GetValidPlatformNames())
 				{
 					string AltName = PlatformName == "Win32" || PlatformName == "Win64" ? "windows" : PlatformName.ToLower();
 					if ((SourceFile.Reference.FullName.ToLower().Contains("/" + PlatformName.ToLower() + "/") || SourceFile.Reference.FullName.ToLower().Contains("/" + AltName + "/"))
diff --git a/Engine/Source/Runtime/Core/Public/IOS/IOSPlatformAtomics.h b/Engine/Source/Runtime/Core/Public/IOS/IOSPlatformAtomics.h
new file mode 100644
--- /dev/null
+++ b/Engine/Source/Runtime/Core/Public/IOS/IOSPlatformAtomics.h
@@ -0,0 +1,5 @@
+// Copyright 1998-2019 Epic Games, Inc. All Rights Reserved.
+
+#pragma once
+
+#include "Apple/ApplePlatformAtomics.h"
diff --git a/Engine/Source/Runtime/Core/Public/IOS/IOSPlatformStackWalk.h b/Engine/Source/Runtime/Core/Public/IOS/IOSPlatformStackWalk.h
new file mode 100644
--- /dev/null
+++ b/Engine/Source/Runtime/Core/Public/IOS/IOSPlatformStackWalk.h
@@ -0,0 +1,5 @@
+// Copyright 1998-2019 Epic Games, Inc. All Rights Reserved.
+
+#pragma once
+
+#include "Apple/ApplePlatformStackWalk.h"
diff --git a/Engine/Source/Runtime/Core/Public/Mac/MacPlatformAtomics.h b/Engine/Source/Runtime/Core/Public/Mac/MacPlatformAtomics.h
new file mode 100644
--- /dev/null
+++ b/Engine/Source/Runtime/Core/Public/Mac/MacPlatformAtomics.h
@@ -0,0 +1,5 @@
+// Copyright 1998-2019 Epic Games, Inc. All Rights Reserved.
+
+#pragma once
+
+#include "Apple/ApplePlatformAtomics.h"
diff --git a/Engine/Source/Runtime/Core/Public/Mac/MacPlatformFile.h b/Engine/Source/Runtime/Core/Public/Mac/MacPlatformFile.h
--- a/Engine/Source/Runtime/Core/Public/Mac/MacPlatformFile.h
+++ b/Engine/Source/Runtime/Core/Public/Mac/MacPlatformFile.h
@@ -2,3 +2,4 @@
 
 #pragma once
 
+#include "Apple/ApplePlatformFile.h"
diff --git a/Engine/Source/Runtime/Core/Public/Mac/MacPlatformStackWalk.h b/Engine/Source/Runtime/Core/Public/Mac/MacPlatformStackWalk.h
new file mode 100644
--- /dev/null
+++ b/Engine/Source/Runtime/Core/Public/Mac/MacPlatformStackWalk.h
@@ -0,0 +1,5 @@
+// Copyright 1998-2019 Epic Games, Inc. All Rights Reserved.
+
+#pragma once
+
+#include "Apple/ApplePlatformStackWalk.h"
