UBT: Allow the project file generator to correctly handle multiple modules in programs.

diff --git a/Engine/Source/Programs/UnrealBuildTool/ProjectFiles/ProjectFileGenerator.cs b/Engine/Source/Programs/UnrealBuildTool/ProjectFiles/ProjectFileGenerator.cs
--- a/Engine/Source/Programs/UnrealBuildTool/ProjectFiles/ProjectFileGenerator.cs
+++ b/Engine/Source/Programs/UnrealBuildTool/ProjectFiles/ProjectFileGenerator.cs
@@ -1978,18 +1978,25 @@ namespace UnrealBuildTool
 
 		private ProjectFile FindProjectForModule(FileReference CurModuleFile, List<FileReference> AllGames, Dictionary<FileReference, ProjectFile> ProgramProjects, List<ProjectFile> ModProjects, out DirectoryReference BaseFolder)
 		{
+			// Check if this module is under any program project base folder
+			if (ProgramProjects != null)
+			{
+				foreach (ProjectFile ProgramProject in ProgramProjects.Values)
+				{
+					if (CurModuleFile.IsUnderDirectory(ProgramProject.BaseDir))
+					{
+						BaseFolder = ProgramProject.BaseDir;
+						return ProgramProject;
+					}
+				}
+			}
+
 			string ProjectFileNameBase = null;
 
 			string PossibleProgramTargetName = CurModuleFile.GetFileNameWithoutAnyExtensions();
 
-			// @todo projectfiles: This works fine for now, but is pretty busted.  It assumes only one module per program and that it matches the program target file name. (see TTP 307091)
-			if( ProgramProjects != null && ProgramProjects.Any( ProgramProject => PossibleProgramTargetName.Equals( ProgramProject.Key.GetFileNameWithoutAnyExtensions() ) ) )	// @todo projectfiles: When building (in mem projects), ProgramProjects will be null so we are using the UE4 project instead
-			{
-				ProjectFileNameBase = PossibleProgramTargetName;
-				BaseFolder = CurModuleFile.Directory;
-			}
 			// check for engine, or platform extension engine folders
-			else if( CurModuleFile.IsUnderDirectory(UnrealBuildTool.EngineDirectory) ||
+			if( CurModuleFile.IsUnderDirectory(UnrealBuildTool.EngineDirectory) ||
 				CurModuleFile.ContainsName("Engine", UnrealBuildTool.PlatformExtensionsDirectory) )
 			{
 				ProjectFileNameBase = EngineProjectFileNameBase;
