UE-72225

Moving where the include-rewrite-rules.ini for SNDBS gets generated into the executor and adding the ability for a UEBuildPlatform to specify a platform name to use instead of the default.

#rnx

diff --git a/Engine/Source/Programs/UnrealBuildTool/Configuration/UEBuildPlatform.cs b/Engine/Source/Programs/UnrealBuildTool/Configuration/UEBuildPlatform.cs
--- a/Engine/Source/Programs/UnrealBuildTool/Configuration/UEBuildPlatform.cs
+++ b/Engine/Source/Programs/UnrealBuildTool/Configuration/UEBuildPlatform.cs
@@ -337,7 +337,6 @@ namespace UnrealBuildTool
 			return null;
 		}
 
-
 		/// <summary>
 		/// Determines whether a given platform is available
 		/// </summary>
@@ -454,6 +453,14 @@ namespace UnrealBuildTool
 			return ";";
 		}
 
+		/// <summary>
+		/// Gets the platform name that should be used.
+		/// </summary>
+		public virtual string GetPlatformName()
+		{
+			return Platform.ToString();
+		}
+
 		/// <summary>
 		/// If this platform can be compiled with XGE
 		/// </summary>
diff --git a/Engine/Source/Programs/UnrealBuildTool/Configuration/UEBuildTarget.cs b/Engine/Source/Programs/UnrealBuildTool/Configuration/UEBuildTarget.cs
--- a/Engine/Source/Programs/UnrealBuildTool/Configuration/UEBuildTarget.cs
+++ b/Engine/Source/Programs/UnrealBuildTool/Configuration/UEBuildTarget.cs
@@ -1404,7 +1404,7 @@ namespace UnrealBuildTool
 				List<UEBuildBinary> FilteredBinaries = new List<UEBuildBinary>();
 				foreach (UEBuildBinary DLLBinary in Binaries)
 				{
-					if (ModDirectories.Any(x => DLLBinary.OutputFilePath.IsUnderDirectory(x)))
+					if(ModDirectories.Any(x => DLLBinary.OutputFilePath.IsUnderDirectory(x)))
 					{
 						FilteredBinaries.Add(DLLBinary);
 					}
@@ -1419,9 +1419,9 @@ namespace UnrealBuildTool
 
 			// Build a mapping from module to its plugin
 			Dictionary<UEBuildModule, UEBuildPlugin> ModuleToPlugin = new Dictionary<UEBuildModule, UEBuildPlugin>();
-			foreach (UEBuildPlugin Plugin in BuildPlugins)
+			foreach(UEBuildPlugin Plugin in BuildPlugins)
 			{
-				foreach (UEBuildModule Module in Plugin.Modules)
+				foreach(UEBuildModule Module in Plugin.Modules)
 				{
 					if (!ModuleToPlugin.ContainsKey(Module))
 					{
@@ -1431,14 +1431,14 @@ namespace UnrealBuildTool
 			}
 
 			// Check there aren't any engine binaries with dependencies on game modules. This can happen when game-specific plugins override engine plugins.
-			foreach (UEBuildModule Module in Modules.Values)
+			foreach(UEBuildModule Module in Modules.Values)
 			{
-				if (Module.Binary != null)
+				if(Module.Binary != null)
 				{
 					HashSet<UEBuildModule> ReferencedModules = Module.GetDependencies(bWithIncludePathModules: true, bWithDynamicallyLoadedModules: true);
-					foreach (UEBuildModule ReferencedModule in ReferencedModules)
+					foreach(UEBuildModule ReferencedModule in ReferencedModules)
 					{
-						if (!Module.Rules.Context.Scope.Contains(ReferencedModule.Rules.Context.Scope) && !IsWhitelistedEnginePluginReference(Module.Name, ReferencedModule.Name))
+						if(!Module.Rules.Context.Scope.Contains(ReferencedModule.Rules.Context.Scope) && !IsWhitelistedEnginePluginReference(Module.Name, ReferencedModule.Name))
 						{
 							throw new BuildException("Module '{0}' ({1}) should not reference module '{2}' ({3}). Hierarchy is {4}.", Module.Name, Module.Rules.Context.Scope.Name, ReferencedModule.Name, ReferencedModule.Rules.Context.Scope.Name, ReferencedModule.Rules.Context.Scope.FormatHierarchy());
 						}
@@ -1447,15 +1447,15 @@ namespace UnrealBuildTool
 			}
 
 			// Check that each plugin declares its dependencies explicitly
-			foreach (UEBuildPlugin Plugin in BuildPlugins)
+			foreach(UEBuildPlugin Plugin in BuildPlugins)
 			{
-				foreach (UEBuildModule Module in Plugin.Modules)
+				foreach(UEBuildModule Module in Plugin.Modules)
 				{
 					HashSet<UEBuildModule> DependencyModules = Module.GetDependencies(bWithIncludePathModules: true, bWithDynamicallyLoadedModules: true);
-					foreach (UEBuildModule DependencyModule in DependencyModules)
+					foreach(UEBuildModule DependencyModule in DependencyModules)
 					{
 						UEBuildPlugin DependencyPlugin;
-						if (ModuleToPlugin.TryGetValue(DependencyModule, out DependencyPlugin) && DependencyPlugin != Plugin && !Plugin.Dependencies.Contains(DependencyPlugin))
+						if(ModuleToPlugin.TryGetValue(DependencyModule, out DependencyPlugin) && DependencyPlugin != Plugin && !Plugin.Dependencies.Contains(DependencyPlugin))
 						{
 							Log.TraceWarning("Warning: Plugin '{0}' does not list plugin '{1}' as a dependency, but module '{2}' depends on '{3}'.", Plugin.Name, DependencyPlugin.Name, Module.Name, DependencyModule.Name);
 						}
@@ -1467,9 +1467,6 @@ namespace UnrealBuildTool
 			string ExternalMetadata = UEBuildPlatform.GetBuildPlatform(Platform).GetExternalBuildMetadata(ProjectFile);
 			TargetMakefile Makefile = new TargetMakefile(TargetToolChain.GetVersionInfo(), ExternalMetadata, ReceiptFileName, ProjectIntermediateDirectory, TargetType, bDeployAfterCompile, bHasProjectScriptPlugin);
 
-			// Generate the SNDBS include-rewrite-rules.ini in case we'll need it.
-			GenerateSNDBSIncludeRewriteRules();
-
 			// Setup the hot reload module list
 			Makefile.HotReloadModuleNames = GetHotReloadModuleNames();
 
@@ -1490,7 +1487,7 @@ namespace UnrealBuildTool
 			}
 
 			// Add global definitions for project-specific binaries. HACK: Also defining for monolithic builds in binary releases. Might be better to set this via command line instead?
-			if (!bUseSharedBuildEnvironment || bCompileMonolithic)
+			if(!bUseSharedBuildEnvironment || bCompileMonolithic)
 			{
 				UEBuildBinary ExecutableBinary = Binaries[0];
 
@@ -1526,7 +1523,7 @@ namespace UnrealBuildTool
 
 			// Generate headers
 			HashSet<UEBuildModuleCPP> ModulesToGenerateHeadersFor = GatherDependencyModules(OriginalBinaries.ToList());
-			using (Timeline.ScopeEvent("ExternalExecution.SetupUObjectModules()"))
+			using(Timeline.ScopeEvent("ExternalExecution.SetupUObjectModules()"))
 			{
 				ExternalExecution.SetupUObjectModules(ModulesToGenerateHeadersFor, Rules.Platform, ProjectDescriptor, Makefile.UObjectModules, Makefile.UObjectModuleHeaders, Rules.GeneratedCodeVersion, bIsAssemblingBuild, MetadataCache);
 			}
@@ -1550,7 +1547,7 @@ namespace UnrealBuildTool
 			{
 				if (Platform == UnrealTargetPlatform.Win32 || Platform == UnrealTargetPlatform.Win64)
 				{
-					if (!Rules.bFormalBuild)
+					if(!Rules.bFormalBuild)
 					{
 						CppCompileEnvironment DefaultResourceCompileEnvironment = new CppCompileEnvironment(GlobalCompileEnvironment);
 
@@ -1564,7 +1561,7 @@ namespace UnrealBuildTool
 
 			// Build the target's binaries.
 			DirectoryReference ExeDir = GetExecutableDir();
-			using (Timeline.ScopeEvent("UEBuildBinary.Build()"))
+			using(Timeline.ScopeEvent("UEBuildBinary.Build()"))
 			{
 				foreach (UEBuildBinary Binary in Binaries)
 				{
@@ -1576,20 +1573,20 @@ namespace UnrealBuildTool
 			// Prepare all the runtime dependencies, copying them from their source folders if necessary
 			List<RuntimeDependency> RuntimeDependencies = new List<RuntimeDependency>();
 			Dictionary<FileReference, FileReference> RuntimeDependencyTargetFileToSourceFile = new Dictionary<FileReference, FileReference>();
-			foreach (UEBuildBinary Binary in Binaries)
+			foreach(UEBuildBinary Binary in Binaries)
 			{
 				Binary.PrepareRuntimeDependencies(RuntimeDependencies, RuntimeDependencyTargetFileToSourceFile, ExeDir);
 			}
-			foreach (KeyValuePair<FileReference, FileReference> Pair in RuntimeDependencyTargetFileToSourceFile)
+			foreach(KeyValuePair<FileReference, FileReference> Pair in RuntimeDependencyTargetFileToSourceFile)
 			{
-				if (!UnrealBuildTool.IsFileInstalled(Pair.Key))
+				if(!UnrealBuildTool.IsFileInstalled(Pair.Key))
 				{
 					Makefile.OutputItems.Add(CreateCopyAction(Pair.Value, Pair.Key, Makefile.Actions));
 				}
 			}
 
 			// If we're just precompiling a plugin, only include output items which are under that directory
-			if (ForeignPlugin != null)
+			if(ForeignPlugin != null)
 			{
 				Makefile.OutputItems.RemoveAll(x => !x.Location.IsUnderDirectory(ForeignPlugin.Directory));
 			}
@@ -1608,7 +1605,7 @@ namespace UnrealBuildTool
 			BuildProducts.AddRange(RuntimeDependencyTargetFileToSourceFile.Select(x => new KeyValuePair<FileReference, BuildProductType>(x.Key, BuildProductType.RequiredResource)));
 
 			// Remove any installed build products that don't exist. They may be part of an optional install.
-			if (UnrealBuildTool.IsEngineInstalled())
+			if(UnrealBuildTool.IsEngineInstalled())
 			{
 				BuildProducts.RemoveAll(x => UnrealBuildTool.IsFileInstalled(x.Key) && !FileReference.Exists(x.Key));
 			}
@@ -1624,13 +1621,13 @@ namespace UnrealBuildTool
 			}
 
 			// Finalize and generate metadata for this target
-			if (!Rules.bDisableLinking)
+			if(!Rules.bDisableLinking)
 			{
 				// Also add any explicitly specified build products
-				if (Rules.AdditionalBuildProducts.Count > 0)
+				if(Rules.AdditionalBuildProducts.Count > 0)
 				{
 					Dictionary<string, string> Variables = GetTargetVariables(null);
-					foreach (string AdditionalBuildProduct in Rules.AdditionalBuildProducts)
+					foreach(string AdditionalBuildProduct in Rules.AdditionalBuildProducts)
 					{
 						FileReference BuildProductFile = new FileReference(Utils.ExpandVariables(AdditionalBuildProduct, Variables));
 						BuildProducts.Add(new KeyValuePair<FileReference, BuildProductType>(BuildProductFile, BuildProductType.RequiredResource));
@@ -1639,18 +1636,18 @@ namespace UnrealBuildTool
 
 				// Get the path to the version file unless this is a formal build (where it will be compiled in)
 				FileReference VersionFile = null;
-				if (Rules.LinkType != TargetLinkType.Monolithic && Binaries[0].Type == UEBuildBinaryType.Executable)
+				if(Rules.LinkType != TargetLinkType.Monolithic && Binaries[0].Type == UEBuildBinaryType.Executable)
 				{
 					UnrealTargetConfiguration VersionConfig = Configuration;
-					if (VersionConfig == UnrealTargetConfiguration.DebugGame && !bCompileMonolithic && TargetType != TargetType.Program && bUseSharedBuildEnvironment)
+					if(VersionConfig == UnrealTargetConfiguration.DebugGame && !bCompileMonolithic && TargetType != TargetType.Program && bUseSharedBuildEnvironment)
 					{
 						VersionConfig = UnrealTargetConfiguration.Development;
 					}
-					VersionFile = BuildVersion.GetFileNameForTarget(ExeDir, bCompileMonolithic ? TargetName : AppName, Platform, VersionConfig, Architecture);
+					VersionFile = BuildVersion.GetFileNameForTarget(ExeDir, bCompileMonolithic? TargetName : AppName, Platform, VersionConfig, Architecture);
 				}
 
 				// Also add the version file as a build product
-				if (VersionFile != null)
+				if(VersionFile != null)
 				{
 					BuildProducts.Add(new KeyValuePair<FileReference, BuildProductType>(VersionFile, BuildProductType.RequiredResource));
 				}
@@ -1670,7 +1667,7 @@ namespace UnrealBuildTool
 				StringBuilder WriteMetadataArguments = new StringBuilder();
 				WriteMetadataArguments.AppendFormat("-Input={0}", Utils.MakePathSafeToUseWithCommandLine(MetadataTargetFile));
 				WriteMetadataArguments.AppendFormat(" -Version={0}", WriteMetadataMode.CurrentVersionNumber);
-				if (Rules.bNoManifestChanges)
+				if(Rules.bNoManifestChanges)
 				{
 					WriteMetadataArguments.Append(" -NoManifestChanges");
 				}
@@ -1688,13 +1685,13 @@ namespace UnrealBuildTool
 
 				// Create actions to run the post build steps
 				FileReference[] PostBuildScripts = CreatePostBuildScripts();
-				foreach (FileReference PostBuildScript in PostBuildScripts)
+				foreach(FileReference PostBuildScript in PostBuildScripts)
 				{
 					FileReference OutputFile = new FileReference(PostBuildScript.FullName + ".ran");
 
 					Action PostBuildStepAction = new Action(ActionType.PostBuildStep);
 					PostBuildStepAction.CommandPath = BuildHostPlatform.Current.Shell;
-					if (BuildHostPlatform.Current.ShellType == ShellType.Cmd)
+					if(BuildHostPlatform.Current.ShellType == ShellType.Cmd)
 					{
 						PostBuildStepAction.CommandArguments = String.Format("/C \"call \"{0}\" && type NUL >\"{1}\"\"", PostBuildScript, OutputFile);
 					}
@@ -1714,22 +1711,22 @@ namespace UnrealBuildTool
 			}
 
 			// Build a list of all the files required to build
-			foreach (FileReference DependencyListFileName in Rules.DependencyListFileNames)
+			foreach(FileReference DependencyListFileName in Rules.DependencyListFileNames)
 			{
 				WriteDependencyList(DependencyListFileName, RuntimeDependencies, RuntimeDependencyTargetFileToSourceFile);
 			}
 
 			// If we're only generating the manifest, return now
-			foreach (FileReference ManifestFileName in Rules.ManifestFileNames)
+			foreach(FileReference ManifestFileName in Rules.ManifestFileNames)
 			{
 				GenerateManifest(ManifestFileName, BuildProducts);
 			}
 
 			// Check there are no EULA or restricted folder violations
-			if (!Rules.bDisableLinking)
+			if(!Rules.bDisableLinking)
 			{
 				// Check the distribution level of all binaries based on the dependencies they have
-				if (ProjectFile == null && !Rules.bOutputPubliclyDistributable)
+				if(ProjectFile == null && !Rules.bOutputPubliclyDistributable)
 				{
 					Dictionary<UEBuildModule, Dictionary<RestrictedFolder, DirectoryReference>> ModuleRestrictedFolderCache = new Dictionary<UEBuildModule, Dictionary<RestrictedFolder, DirectoryReference>>();
 
@@ -1738,12 +1735,12 @@ namespace UnrealBuildTool
 					{
 						bResult &= Binary.CheckRestrictedFolders(DirectoryReference.FromFile(ProjectFile), ModuleRestrictedFolderCache);
 					}
-					foreach (KeyValuePair<FileReference, FileReference> Pair in RuntimeDependencyTargetFileToSourceFile)
+					foreach(KeyValuePair<FileReference, FileReference> Pair in RuntimeDependencyTargetFileToSourceFile)
 					{
 						bResult &= CheckRestrictedFolders(Pair.Key, Pair.Value);
 					}
 
-					if (!bResult)
+					if(!bResult)
 					{
 						throw new BuildException("Unable to create binaries in less restricted locations than their input files.");
 					}
@@ -1754,7 +1751,7 @@ namespace UnrealBuildTool
 			}
 
 			// Add all the plugins to be tracked
-			foreach (FileReference PluginFile in global::UnrealBuildTool.Plugins.EnumeratePlugins(ProjectFile))
+			foreach(FileReference PluginFile in global::UnrealBuildTool.Plugins.EnumeratePlugins(ProjectFile))
 			{
 				FileItem PluginFileItem = FileItem.GetItemByFileReference(PluginFile);
 				Makefile.PluginFiles.Add(PluginFileItem);
@@ -1762,10 +1759,10 @@ namespace UnrealBuildTool
 
 			// Add all the input files to the predicate store
 			Makefile.AdditionalDependencies.Add(FileItem.GetItemByFileReference(TargetRulesFile));
-			foreach (UEBuildModule Module in Modules.Values)
+			foreach(UEBuildModule Module in Modules.Values)
 			{
 				Makefile.AdditionalDependencies.Add(FileItem.GetItemByFileReference(Module.RulesFile));
-				foreach (string ExternalDependency in Module.Rules.ExternalDependencies)
+				foreach(string ExternalDependency in Module.Rules.ExternalDependencies)
 				{
 					FileReference Location = FileReference.Combine(Module.RulesFile.Directory, ExternalDependency);
 					Makefile.AdditionalDependencies.Add(FileItem.GetItemByFileReference(Location));
@@ -1774,7 +1771,7 @@ namespace UnrealBuildTool
 			Makefile.AdditionalDependencies.UnionWith(Makefile.PluginFiles);
 
 			// Write a header containing public definitions for this target
-			if (Rules.ExportPublicHeader != null)
+			if(Rules.ExportPublicHeader != null)
 			{
 				UEBuildBinary Binary = Binaries[0];
 				FileReference Header = FileReference.Combine(Binary.OutputDir, Rules.ExportPublicHeader);
@@ -1787,21 +1784,6 @@ namespace UnrealBuildTool
 			return Makefile;
 		}
 
-		private void GenerateSNDBSIncludeRewriteRules()
-		{
-			FileReference IncludeRewriteRulesFile = FileReference.Combine(UnrealBuildTool.EngineDirectory, "Build", "SNDBS", "include-rewrite-rules.ini");
-			if (!Directory.Exists(IncludeRewriteRulesFile.Directory.FullName))
-			{
-				Directory.CreateDirectory(IncludeRewriteRulesFile.Directory.FullName);
-			}
-			List<string> IncludeRewriteRulesText = new List<string>();
-			IncludeRewriteRulesText.Add("[computed-include-rules]");
-			IncludeRewriteRulesText.Add(@"pattern1=^COMPILED_PLATFORM_HEADER\(\s*([^ ,]+)\)");
-			string PlatformString = Platform.ToString().StartsWith("Win") ? "Windows" : Platform.ToString();
-			IncludeRewriteRulesText.Add(String.Format("expansions1={0}/{0}$1", PlatformString));
-			File.WriteAllText(IncludeRewriteRulesFile.FullName, String.Join(Environment.NewLine, IncludeRewriteRulesText));
-		}
-
 		/// <summary>
 		/// Gets the output directory for the main executable
 		/// </summary>
diff --git a/Engine/Source/Programs/UnrealBuildTool/Executors/SNDBS.cs b/Engine/Source/Programs/UnrealBuildTool/Executors/SNDBS.cs
--- a/Engine/Source/Programs/UnrealBuildTool/Executors/SNDBS.cs
+++ b/Engine/Source/Programs/UnrealBuildTool/Executors/SNDBS.cs
@@ -14,6 +14,8 @@ namespace UnrealBuildTool
 {
 	class SNDBS : ActionExecutor
 	{
+		private FileReference IncludeRewriteRulesFile = FileReference.Combine(UnrealBuildTool.EngineDirectory, "Build", "SNDBS", "include-rewrite-rules.ini");
+
 		/// <summary>
 		/// Processor count multiplier for local execution. Can be below 1 to reserve CPU for other tasks.
 		/// </summary>
@@ -282,8 +284,7 @@ namespace UnrealBuildTool
                 string SCERoot = Environment.GetEnvironmentVariable("SCE_ROOT_DIR");
                 string SNDBSExecutable = Path.Combine(SCERoot, "Common/SN-DBS/bin/dbsbuild.exe");
 				DirectoryReference TemplatesDir = DirectoryReference.Combine(UnrealBuildTool.EngineDirectory, "Programs", "UnrealBuildTool", "SndbsTemplates");
-				FileReference IncludeRewriteRules = FileReference.Combine(UnrealBuildTool.EngineDirectory, "Build", "SNDBS", "include-rewrite-rules.ini");
-				string IncludeRewriteRulesArg = String.Format("--include-rewrite-rules \"{0}\"", IncludeRewriteRules.FullName);
+				string IncludeRewriteRulesArg = String.Format("--include-rewrite-rules \"{0}\"", IncludeRewriteRulesFile.FullName);
 				string VerbosityLevel = PrintDebugInfo ? "-v" : "-q";
                 ProcessStartInfo PSI = new ProcessStartInfo(SNDBSExecutable, String.Format("{0} -p UE4 -s \"{1}\" -templates \"{2}\" {3}", VerbosityLevel, FileReference.Combine(UnrealBuildTool.EngineDirectory, "Intermediate", "Build", "sndbs.bat").FullName, TemplatesDir.FullName, IncludeRewriteRulesArg));
 				PSI.RedirectStandardOutput = true;
@@ -365,6 +366,9 @@ namespace UnrealBuildTool
 					PrepareToolTemplate(CommandPath);
 				}
 
+				// Generate include-rewrite-rules.ini.
+				GenerateSNDBSIncludeRewriteRules();
+
 				// Use WMI to figure out physical cores, excluding hyper threading.
 				int NumCores = Utils.GetPhysicalProcessorCount();
 				
@@ -457,6 +461,13 @@ namespace UnrealBuildTool
 				Log.WriteLineIf(bLogDetailedActionStats,
 					LogEventType.Console, "Cumulative thread seconds ({0} processors): {1:0.00}", System.Environment.ProcessorCount, TotalThreadSeconds);
 			}
+
+			// Delete the include-rewrite-rules.ini if it was generated.
+			if (File.Exists(IncludeRewriteRulesFile.FullName))
+			{
+				File.Delete(IncludeRewriteRulesFile.FullName);
+			}
+
 			return SNDBSResult;
 		}
 
@@ -486,5 +497,34 @@ namespace UnrealBuildTool
 			}
 			File.WriteAllText(TemplateOutput.FullName, TemplateText);
 		}
+
+		private void GenerateSNDBSIncludeRewriteRules()
+		{
+			if (!Directory.Exists(IncludeRewriteRulesFile.Directory.FullName))
+			{
+				Directory.CreateDirectory(IncludeRewriteRulesFile.Directory.FullName);
+			}
+
+			List<string> IncludeRewriteRulesText = new List<string>();
+			IncludeRewriteRulesText.Add("[computed-include-rules]");
+			IncludeRewriteRulesText.Add(@"pattern1=^COMPILED_PLATFORM_HEADER\(\s*([^ ,]+)\)");
+
+			// Get all registered platforms. Most will just use the name as is, but some may have an override so
+			// add all distinct entries to the list.
+			IEnumerable<UnrealTargetPlatform> Platforms = UEBuildPlatform.GetRegisteredPlatforms();
+			List<string> PlatformNames = new List<string>();
+			foreach (UnrealTargetPlatform Platform in Platforms)
+			{
+				UEBuildPlatform PlatformData = UEBuildPlatform.GetBuildPlatform(Platform);
+				if (!PlatformNames.Contains(PlatformData.GetPlatformName()))
+				{
+					PlatformNames.Add(PlatformData.GetPlatformName());
+				}
+			}
+
+			IEnumerable<string> PlatformExpansions = PlatformNames.Select(p => String.Format("{0}/{0}$1", p));
+			IncludeRewriteRulesText.Add(String.Format("expansions1={0}", String.Join("|", PlatformExpansions)));
+			File.WriteAllText(IncludeRewriteRulesFile.FullName, String.Join(Environment.NewLine, IncludeRewriteRulesText));
+		}
 	}
 }
diff --git a/Engine/Source/Programs/UnrealBuildTool/Platform/Windows/UEBuildWindows.cs b/Engine/Source/Programs/UnrealBuildTool/Platform/Windows/UEBuildWindows.cs
--- a/Engine/Source/Programs/UnrealBuildTool/Platform/Windows/UEBuildWindows.cs
+++ b/Engine/Source/Programs/UnrealBuildTool/Platform/Windows/UEBuildWindows.cs
@@ -1478,6 +1478,14 @@ namespace UnrealBuildTool
 				   TryReadInstallDirRegistryKey32("Microsoft\\Microsoft SDKs\\NETFXSDK\\4.6.2", "KitsInstallationFolder", out OutInstallDir);
 		}
 
+		/// <summary>
+		/// Gets the platform name that should be used.
+		/// </summary>
+		public override string GetPlatformName()
+		{
+			return "Windows";
+		}
+
 		/// <summary>
 		/// If this platform can be compiled with SN-DBS
 		/// </summary>
@@ -1673,7 +1681,7 @@ namespace UnrealBuildTool
 			CompileEnvironment.Definitions.Add("PLATFORM_WINDOWS=1");
 			
 			// both Win32 and Win64 use Windows headers, so we enforce that here
-			CompileEnvironment.Definitions.Add("OVERRIDE_PLATFORM_HEADER_NAME=Windows");
+			CompileEnvironment.Definitions.Add(String.Format("OVERRIDE_PLATFORM_HEADER_NAME={0}", GetPlatformName()));
 
 			FileReference MorpheusShaderPath = FileReference.Combine(UnrealBuildTool.EngineDirectory, "Shaders", "Private", "Platform", "PS4", "PostProcessHMDMorpheus.usf");
 			if (FileReference.Exists(MorpheusShaderPath))
