Modifying IOS compilation process to be more flexible to alternate procedures.

diff --git a/Engine/Source/Programs/UnrealBuildTool/Platform/IOS/IOSToolChain.cs b/Engine/Source/Programs/UnrealBuildTool/Platform/IOS/IOSToolChain.cs
--- a/Engine/Source/Programs/UnrealBuildTool/Platform/IOS/IOSToolChain.cs
+++ b/Engine/Source/Programs/UnrealBuildTool/Platform/IOS/IOSToolChain.cs
@@ -681,8 +681,19 @@ namespace UnrealBuildTool
 					{
 						CompileAction.PrerequisiteItems.Add(CompileEnvironment.PrecompiledHeaderFile);
 					}
+
+					string ObjectFileExtension;
+					if (CompileEnvironment.AdditionalArguments != null && CompileEnvironment.AdditionalArguments.Contains("-emit-llvm"))
+					{
+						ObjectFileExtension = ".bc";
+					}
+					else
+					{
+						ObjectFileExtension = ".o";
+					}
+
 					// Add the object file to the produced item list.
-					FileItem ObjectFile = FileItem.GetItemByFileReference(FileReference.Combine(OutputDir, Path.GetFileName(SourceFile.AbsolutePath) + ".o"));
+					FileItem ObjectFile = FileItem.GetItemByFileReference(FileReference.Combine(OutputDir, Path.GetFileName(SourceFile.AbsolutePath) + ObjectFileExtension));
 
 					CompileAction.ProducedItems.Add(ObjectFile);
 					Result.ObjectFiles.Add(ObjectFile);
diff --git a/Engine/Source/Programs/UnrealBuildTool/Platform/IOS/UEBuildIOS.cs b/Engine/Source/Programs/UnrealBuildTool/Platform/IOS/UEBuildIOS.cs
--- a/Engine/Source/Programs/UnrealBuildTool/Platform/IOS/UEBuildIOS.cs
+++ b/Engine/Source/Programs/UnrealBuildTool/Platform/IOS/UEBuildIOS.cs
@@ -15,7 +15,7 @@ namespace UnrealBuildTool
 	/// <summary>
 	/// IOS-specific target settings
 	/// </summary>
-	public class IOSTargetRules
+	public partial class IOSTargetRules
 	{
 		/// <summary>
 		/// Whether to strip iOS symbols or not (implied by Shipping config)
@@ -75,7 +75,7 @@ namespace UnrealBuildTool
 	/// <summary>
 	/// Read-only wrapper for IOS-specific target settings
 	/// </summary>
-	public class ReadOnlyIOSTargetRules
+	public partial class ReadOnlyIOSTargetRules
 	{
 		/// <summary>
 		/// The private mutable settings object
