Removing CL#5826376 to fix iOS entitlements when building on Mac.
#tests iOS Client on Mac build

diff --git a/Engine/Source/Programs/UnrealBuildTool/Platform/IOS/IOSToolChain.cs b/Engine/Source/Programs/UnrealBuildTool/Platform/IOS/IOSToolChain.cs
--- a/Engine/Source/Programs/UnrealBuildTool/Platform/IOS/IOSToolChain.cs
+++ b/Engine/Source/Programs/UnrealBuildTool/Platform/IOS/IOSToolChain.cs
@@ -1870,6 +1870,9 @@ namespace UnrealBuildTool
 				UEDeployIOS DeployHandler = (Target.Platform == UnrealTargetPlatform.IOS ? new UEDeployIOS() : new UEDeployTVOS());
 				DeployHandler.ForDistribution = Target.bForDistribution;
 				DeployHandler.PrepTargetForDeployment(Target.ProjectFile, Target.TargetName, Target.Platform, Target.Configuration, Target.UPLScripts, Target.SdkVersion, Target.bCreateStubIPA, BundleID);
+
+				// write the entitlements file (building on Mac)
+				WriteEntitlements(Target);
 			}
 			
 
