Update timing information in VCToolChain.cs to have a secondary path controlled by bCompilerTrace in WindowsTargetRules and clean up the parsing actions status string so it isn't as noisy.

#rnx

diff --git a/Engine/Source/Programs/UnrealBuildTool/Platform/Windows/UEBuildWindows.cs b/Engine/Source/Programs/UnrealBuildTool/Platform/Windows/UEBuildWindows.cs
--- a/Engine/Source/Programs/UnrealBuildTool/Platform/Windows/UEBuildWindows.cs
+++ b/Engine/Source/Programs/UnrealBuildTool/Platform/Windows/UEBuildWindows.cs
@@ -210,6 +210,12 @@ namespace UnrealBuildTool
 		[XmlConfigFile(Category = "WindowsPlatform")]
 		public bool bClangTimeTrace = false;
 
+		/// <summary>
+		/// Outputs compile timing information so that it can be analyzed.
+		/// </summary>
+		[XmlConfigFile(Category = "WindowsPlatform")]
+		public bool bCompilerTrace = false;
+
 		/// <summary>
 		/// Bundle a working version of dbghelp.dll with the application, and use this to generate minidumps. This works around a bug with the Windows 10 Fall Creators Update (1709)
 		/// where rich PE headers larger than a certain size would result in corrupt minidumps.
@@ -401,6 +407,11 @@ namespace UnrealBuildTool
 			get { return Inner.bClangTimeTrace; }
 		}
 
+		public bool bCompilerTrace
+		{
+			get { return Inner.bCompilerTrace; }
+		}
+
 		public string GetVisualStudioCompilerVersionName()
 		{
 			return Inner.GetVisualStudioCompilerVersionName();
diff --git a/Engine/Source/Programs/UnrealBuildTool/Platform/Windows/VCToolChain.cs b/Engine/Source/Programs/UnrealBuildTool/Platform/Windows/VCToolChain.cs
--- a/Engine/Source/Programs/UnrealBuildTool/Platform/Windows/VCToolChain.cs
+++ b/Engine/Source/Programs/UnrealBuildTool/Platform/Windows/VCToolChain.cs
@@ -903,11 +903,15 @@ namespace UnrealBuildTool
 				AddSystemIncludePath(SharedArguments, IncludePath);
 			}
 
-			if (CompileEnvironment.bPrintTimingInfo)
+			if (CompileEnvironment.bPrintTimingInfo || Target.WindowsPlatform.bCompilerTrace)
 			{
 				if(Target.WindowsPlatform.Compiler == WindowsCompiler.VisualStudio2015_DEPRECATED || Target.WindowsPlatform.Compiler == WindowsCompiler.VisualStudio2017)
 				{
-					SharedArguments.Add("/Bt+ /d2cgsummary");
+					if (CompileEnvironment.bPrintTimingInfo)
+					{
+						SharedArguments.Add("/Bt+ /d2cgsummary");
+					}
+
 					if(EnvVars.ToolChainVersion >= VersionNumber.Parse("14.14.26316"))
 					{
 						SharedArguments.Add("/d1reportTime");
@@ -1171,7 +1175,7 @@ namespace UnrealBuildTool
 
 			Action ParseTimingInfoAction = Action.CreateRecursiveAction<ParseMsvcTimingInfoMode>(ActionType.ParseTimingInfo, ParseTimingArguments);
 			ParseTimingInfoAction.WorkingDirectory = UnrealBuildTool.EngineSourceDirectory;
-			ParseTimingInfoAction.StatusDescription = String.Format("Parsing Timing File '{0}'", CompileAction.TimingFile);
+			ParseTimingInfoAction.StatusDescription = Path.GetFileName(CompileAction.TimingFile.AbsolutePath);
 			ParseTimingInfoAction.bCanExecuteRemotely = true;
 			ParseTimingInfoAction.bCanExecuteRemotelyWithSNDBS = true;
 			ParseTimingInfoAction.PrerequisiteItems.Add(SourceFile);
@@ -1187,7 +1191,7 @@ namespace UnrealBuildTool
 			CompileAction.ProducedItems.Add(CompileAction.DependencyListFile);
 			CommandArguments.Add(String.Format("-dependencies={0}", Utils.MakePathSafeToUseWithCommandLine(CompileAction.DependencyListFile.Location)));
 
-			if (Target.bPrintToolChainTimingInfo)
+			if (Target.bPrintToolChainTimingInfo || Target.WindowsPlatform.bCompilerTrace)
 			{
 				CompileAction.TimingFile = FileItem.GetItemByFileReference(FileReference.Combine(OutputDir, String.Format("{0}.timing.txt", SourceFile.Location.GetFileName())));
 				CompileAction.ProducedItems.Add(CompileAction.TimingFile);
@@ -1207,7 +1211,7 @@ namespace UnrealBuildTool
 
 		public override void FinalizeOutput(ReadOnlyTargetRules Target, TargetMakefile Makefile)
 		{
-			if (Target.bPrintToolChainTimingInfo)
+			if (Target.bPrintToolChainTimingInfo || Target.WindowsPlatform.bCompilerTrace)
 			{
 				List<Action> ParseTimingActions = Makefile.Actions.Where(x => x.ActionType == ActionType.ParseTimingInfo).ToList();
 				List<FileItem> TimingJsonFiles = ParseTimingActions.SelectMany(a => a.ProducedItems.Where(i => i.HasExtension(".timing.bin"))).ToList();
