Fix invalid platform name when building for HTML5 on Linux.

diff --git a/Engine/Source/Programs/UnrealBuildTool/Platform/HTML5/HTML5SDKInfo.cs b/Engine/Source/Programs/UnrealBuildTool/Platform/HTML5/HTML5SDKInfo.cs
--- a/Engine/Source/Programs/UnrealBuildTool/Platform/HTML5/HTML5SDKInfo.cs
+++ b/Engine/Source/Programs/UnrealBuildTool/Platform/HTML5/HTML5SDKInfo.cs
@@ -62,8 +62,9 @@ namespace UnrealBuildTool
 		{
 			get
 			{
-				if (BuildHostPlatform.Current.Platform == UnrealTargetPlatform.Win64 || BuildHostPlatform.Current.Platform == UnrealTargetPlatform.Win32 || 
-					BuildHostPlatform.Current.Platform == UnrealTargetPlatform.Mac)
+				if (BuildHostPlatform.Current.Platform == UnrealTargetPlatform.Win64 
+					|| BuildHostPlatform.Current.Platform == UnrealTargetPlatform.Mac
+					|| BuildHostPlatform.Current.Platform == UnrealTargetPlatform.Linux)
 				{
 					return BuildHostPlatform.Current.Platform.ToString();
 				}
