Prevent recursive search through engine folder for config files.

diff --git a/Engine/Source/Programs/UnrealBuildTool/System/DataDrivenPlatformInfo.cs b/Engine/Source/Programs/UnrealBuildTool/System/DataDrivenPlatformInfo.cs
--- a/Engine/Source/Programs/UnrealBuildTool/System/DataDrivenPlatformInfo.cs
+++ b/Engine/Source/Programs/UnrealBuildTool/System/DataDrivenPlatformInfo.cs
@@ -45,37 +45,42 @@ namespace UnrealBuildTool
 				Dictionary<string, string> IniParents = new Dictionary<string, string>();
 
 				// look through all config dirs looking for the data driven ini file
-				foreach (string FilePath in Directory.EnumerateFiles(UnrealBuildTool.EngineDirectory.FullName, "DataDrivenPlatformInfo.ini", SearchOption.AllDirectories))
+				DirectoryReference EngineConfigDir = DirectoryReference.Combine(UnrealBuildTool.EngineDirectory, "Config");
+				foreach(DirectoryReference PlatformDir in DirectoryReference.EnumerateDirectories(EngineConfigDir))
 				{
-					// get the platform name from the path
-					string IniPlatformName = Path.GetFileName(Path.GetDirectoryName(FilePath));
-
-					// load the DataDrivenPlatformInfo from the path
-					ConfigFile Config = new ConfigFile(new FileReference(FilePath));
-					ConfigDataDrivenPlatformInfo NewInfo = new ConfigDataDrivenPlatformInfo();
-
-					// we must have the key section 
-					ConfigFileSection Section = null;
-					if (Config.TryGetSection("DataDrivenPlatformInfo", out Section))
+					FileReference FilePath = FileReference.Combine(PlatformDir, "DataDrivenPlatformInfo.ini");
+					if(FileReference.Exists(FilePath))
 					{
-						ConfigHierarchySection ParsedSection = new ConfigHierarchySection(new List<ConfigFileSection>() { Section });
+						// get the platform name from the path
+						string IniPlatformName = PlatformDir.GetDirectoryName();
 
-						// get string values
-						string IniParent;
-						if (ParsedSection.TryGetValue("IniParent", out IniParent))
-						{
-							IniParents[IniPlatformName] = IniParent;
-						}
+						// load the DataDrivenPlatformInfo from the path
+						ConfigFile Config = new ConfigFile(FilePath);
+						ConfigDataDrivenPlatformInfo NewInfo = new ConfigDataDrivenPlatformInfo();
 
-						// slightly nasty bool parsing for bool values
-						string Temp;
-						if (ParsedSection.TryGetValue("bIsConfidential", out Temp) == false || ConfigHierarchy.TryParse(Temp, out NewInfo.bIsConfidential) == false)
+						// we must have the key section 
+						ConfigFileSection Section = null;
+						if (Config.TryGetSection("DataDrivenPlatformInfo", out Section))
 						{
-							NewInfo.bIsConfidential = false;
-						}
+							ConfigHierarchySection ParsedSection = new ConfigHierarchySection(new List<ConfigFileSection>() { Section });
+
+							// get string values
+							string IniParent;
+							if (ParsedSection.TryGetValue("IniParent", out IniParent))
+							{
+								IniParents[IniPlatformName] = IniParent;
+							}
+
+							// slightly nasty bool parsing for bool values
+							string Temp;
+							if (ParsedSection.TryGetValue("bIsConfidential", out Temp) == false || ConfigHierarchy.TryParse(Temp, out NewInfo.bIsConfidential) == false)
+							{
+								NewInfo.bIsConfidential = false;
+							}
 
-						// create cache it
-						PlatformInfos[IniPlatformName] = NewInfo;
+							// create cache it
+							PlatformInfos[IniPlatformName] = NewInfo;
+						}
 					}
 				}
 
