UE-72100 - tvOS For Distribution packages fail with Code Signing Error if made after a Development Build

diff --git a/Engine/Source/Programs/UnrealBuildTool/ToolChain/RemoteMac.cs b/Engine/Source/Programs/UnrealBuildTool/ToolChain/RemoteMac.cs
--- a/Engine/Source/Programs/UnrealBuildTool/ToolChain/RemoteMac.cs
+++ b/Engine/Source/Programs/UnrealBuildTool/ToolChain/RemoteMac.cs
@@ -851,6 +851,9 @@ namespace UnrealBuildTool
 				TargetFiles.Add(TargetFile);
 			}
 
+			Execute("/", String.Format("rm -rf {0}/Intermediate/IOS/*.plist", GetRemotePath(UnrealBuildTool.EngineDirectory)));
+			Execute("/", String.Format("rm -rf {0}/Intermediate/TVOS/*.plist", GetRemotePath(UnrealBuildTool.EngineDirectory)));
+
 			// Write a file that protects all the scripts from being overridden by the standard engine filters
 			FileReference ScriptProtectList = FileReference.Combine(TempDir, "RsyncEngineScripts-Protect.txt");
 			using(StreamWriter Writer = new StreamWriter(ScriptProtectList.FullName))
