Run addtional Gradle targets as separate call to start with fresh heap

#android

diff --git a/Engine/Source/Programs/UnrealBuildTool/Platform/Android/UEDeployAndroid.cs b/Engine/Source/Programs/UnrealBuildTool/Platform/Android/UEDeployAndroid.cs
--- a/Engine/Source/Programs/UnrealBuildTool/Platform/Android/UEDeployAndroid.cs
+++ b/Engine/Source/Programs/UnrealBuildTool/Platform/Android/UEDeployAndroid.cs
@@ -4007,7 +4007,8 @@ namespace UnrealBuildTool
 						string GradleBuildType = bForDistribution ? ":app:assembleRelease" : ":app:assembleDebug";
 
 						// collect optional additional Gradle parameters from plugins
-						string GradleOptions = UPL.ProcessPluginNode(NDKArch, "gradleParameters", GradleBuildType); //  "--stacktrace --debug " + GradleBuildType);
+						string GradleOptions = GradleBuildType; // + "--stacktrace --debug"
+						string GradleAdditionalOptions = UPL.ProcessPluginNode(NDKArch, "gradleParameters", "");
 
 						// check for Android Studio project, call Gradle if doesn't exist (assume user will build with Android Studio)
 						string GradleAppImlFilename = Path.Combine(UE4BuildGradlePath, "app.iml");
@@ -4022,6 +4023,11 @@ namespace UnrealBuildTool
 							string ShellParametersEnd = Utils.IsRunningOnMono ? "'" : "";
 							RunCommandLineProgramWithExceptionAndFiltering(UE4BuildGradlePath, ShellExecutable, ShellParametersBegin + "\"" + GradleScriptPath + "\" " + GradleOptions + ShellParametersEnd, "Making .apk with Gradle...");
 
+							if (GradleAdditionalOptions != "")
+							{
+								RunCommandLineProgramWithExceptionAndFiltering(UE4BuildGradlePath, ShellExecutable, ShellParametersBegin + "\"" + GradleScriptPath + "\" " + GradleAdditionalOptions + ShellParametersEnd, "Additional Gradle steps...");
+							}
+							
 							// For build machine run a clean afterward to clean up intermediate files (does not remove final APK)
 							if (bIsBuildMachine)
 							{
