Add support for embedded builds to optionally append to DisplayName on Android
Provide optional second Gradle invocation after APK generated

#android

diff --git a/Engine/Source/Programs/AutomationTool/Android/AndroidPlatform.Automation.cs b/Engine/Source/Programs/AutomationTool/Android/AndroidPlatform.Automation.cs
--- a/Engine/Source/Programs/AutomationTool/Android/AndroidPlatform.Automation.cs
+++ b/Engine/Source/Programs/AutomationTool/Android/AndroidPlatform.Automation.cs
@@ -542,7 +542,7 @@ public class AndroidPlatform : Platform
 							Type = TargetType.Server;
 						}
 						LogInformation("SavePackageInfo");
-						Deploy.SavePackageInfo(Params.ShortProjectName, SC.ProjectRoot.FullName, Type);
+						Deploy.SavePackageInfo(Params.ShortProjectName, SC.ProjectRoot.FullName, Type, true);
 					}
 					Deploy.PrepForUATPackageOrDeploy(Params.RawProjectPath, Params.ShortProjectName, SC.ProjectRoot, SOName, SC.LocalRoot + "/Engine", Params.Distribution, CookFlavor, SC.StageTargets[0].Receipt.Configuration, false, bShouldCompileAsDll);
 				}
diff --git a/Engine/Source/Programs/UnrealBuildTool/Platform/Android/AndroidExports.cs b/Engine/Source/Programs/UnrealBuildTool/Platform/Android/AndroidExports.cs
--- a/Engine/Source/Programs/UnrealBuildTool/Platform/Android/AndroidExports.cs
+++ b/Engine/Source/Programs/UnrealBuildTool/Platform/Android/AndroidExports.cs
@@ -73,7 +73,8 @@ namespace UnrealBuildTool
 		/// <param name="ProjectName"></param>
 		/// <param name="ProjectDirectoryFullName"></param>
 		/// <param name="Type"></param>
-		bool SavePackageInfo(string ProjectName, string ProjectDirectoryFullName, TargetType Type);
+		/// <param name="bIsEmbedded"></param>
+		bool SavePackageInfo(string ProjectName, string ProjectDirectoryFullName, TargetType Type, bool bIsEmbedded);
 	}
 
 	/// <summary>
diff --git a/Engine/Source/Programs/UnrealBuildTool/Platform/Android/UEDeployAndroid.cs b/Engine/Source/Programs/UnrealBuildTool/Platform/Android/UEDeployAndroid.cs
--- a/Engine/Source/Programs/UnrealBuildTool/Platform/Android/UEDeployAndroid.cs
+++ b/Engine/Source/Programs/UnrealBuildTool/Platform/Android/UEDeployAndroid.cs
@@ -887,7 +887,7 @@ namespace UnrealBuildTool
 
 		private string CachedVersionDisplayName;
 
-		public string GetVersionDisplayName()
+		public string GetVersionDisplayName(bool bIsEmbedded)
 		{
 			if (string.IsNullOrEmpty(CachedVersionDisplayName))
 			{
@@ -909,12 +909,15 @@ namespace UnrealBuildTool
 					if (bAppendPlatformToVersionDisplayName)
 					{
 						VersionDisplayName = string.Format("{0}-Android", VersionDisplayName);
+					}
 
-						// add optional value from environment variable if defined
-						string OptionalAppendDisplayName = Environment.GetEnvironmentVariable("UEAndroidDisplayNameExtra");
-						if (OptionalAppendDisplayName != null)
+					// append optional text to version name if embedded build
+					if (bIsEmbedded)
+					{
+						string EmbeddedAppendDisplayName = "";
+						if (Ini.GetString("/Script/AndroidRuntimeSettings.AndroidRuntimeSettings", "EmbeddedAppendDisplayName", out EmbeddedAppendDisplayName))
 						{
-							VersionDisplayName = VersionDisplayName + OptionalAppendDisplayName;
+							VersionDisplayName = VersionDisplayName + EmbeddedAppendDisplayName;
 						}
 					}
 				}
@@ -1117,12 +1120,12 @@ namespace UnrealBuildTool
 			}
 		}
 
-		public void WriteCrashlyticsResources(string UEBuildPath, string PackageName, string ApplicationDisplayName)
+		public void WriteCrashlyticsResources(string UEBuildPath, string PackageName, string ApplicationDisplayName, bool bIsEmbedded)
 		{
 			System.DateTime CurrentDateTime = System.DateTime.Now;
 			string BuildID = Guid.NewGuid().ToString();
 
-			string VersionDisplayName = GetVersionDisplayName();
+			string VersionDisplayName = GetVersionDisplayName(bIsEmbedded);
 
 			StringBuilder CrashPropertiesContent = new StringBuilder("");
 			CrashPropertiesContent.Append("# This file is automatically generated by Crashlytics to uniquely\n");
@@ -2147,7 +2150,7 @@ namespace UnrealBuildTool
 		}
 
 
-		private string GenerateManifest(AndroidToolChain ToolChain, string ProjectName, TargetType InTargetType, string EngineDirectory, bool bIsForDistribution, bool bPackageDataInsideApk, string GameBuildFilesPath, bool bHasOBBFiles, bool bDisableVerifyOBBOnStartUp, string UE4Arch, string GPUArch, string CookFlavor, bool bUseExternalFilesDir, string Configuration, int SDKLevelInt)
+		private string GenerateManifest(AndroidToolChain ToolChain, string ProjectName, TargetType InTargetType, string EngineDirectory, bool bIsForDistribution, bool bPackageDataInsideApk, string GameBuildFilesPath, bool bHasOBBFiles, bool bDisableVerifyOBBOnStartUp, string UE4Arch, string GPUArch, string CookFlavor, bool bUseExternalFilesDir, string Configuration, int SDKLevelInt, bool bIsEmbedded)
 		{
 			// Read the engine version
 			string EngineVersion = ReadEngineVersion();
@@ -2168,7 +2171,7 @@ namespace UnrealBuildTool
 			// ini file to get settings from
 			ConfigHierarchy Ini = GetConfigCacheIni(ConfigHierarchyType.Engine);
 			string PackageName = GetPackageName(ProjectName);
-			string VersionDisplayName = GetVersionDisplayName();
+			string VersionDisplayName = GetVersionDisplayName(bIsEmbedded);
 			bool bEnableGooglePlaySupport;
 			Ini.GetBool("/Script/AndroidRuntimeSettings.AndroidRuntimeSettings", "bEnableGooglePlaySupport", out bEnableGooglePlaySupport);
 			bool bUseGetAccounts;
@@ -3081,13 +3084,13 @@ namespace UnrealBuildTool
 		}
 
 		private void CreateGradlePropertiesFiles(string Arch, int MinSDKVersion, int TargetSDKVersion, string CompileSDKVersion, string BuildToolsVersion, string PackageName,
-			string DestApkName, string NDKArch,	string UE4BuildFilesPath, string GameBuildFilesPath, string UE4BuildGradleAppPath, string UE4BuildPath, string UE4BuildGradlePath, bool bForDistribution)
+			string DestApkName, string NDKArch,	string UE4BuildFilesPath, string GameBuildFilesPath, string UE4BuildGradleAppPath, string UE4BuildPath, string UE4BuildGradlePath, bool bForDistribution, bool bIsEmbedded)
 		{
 			// Create gradle.properties
 			StringBuilder GradleProperties = new StringBuilder();
 
 			int StoreVersion = GetStoreVersion();
-			string VersionDisplayName = GetVersionDisplayName();
+			string VersionDisplayName = GetVersionDisplayName(bIsEmbedded);
 
 			ConfigHierarchy Ini = GetConfigCacheIni(ConfigHierarchyType.Engine);
 
@@ -3478,7 +3481,7 @@ namespace UnrealBuildTool
 			if (CrashlyticsPluginEnabled && !bGradleEnabled)
 			{
 				Trace.TraceInformation("Writing Crashlytics resources");
-				WriteCrashlyticsResources(Path.Combine(ProjectDirectory, "Build", "Android"), PackageName, ApplicationDisplayName);
+				WriteCrashlyticsResources(Path.Combine(ProjectDirectory, "Build", "Android"), PackageName, ApplicationDisplayName, bSkipGradleBuild);
 			}
 
 			// if the manifest matches, look at other settings stored in a file
@@ -3534,7 +3537,7 @@ namespace UnrealBuildTool
 			{
 				BuildList = from Arch in Arches
 							from GPUArch in GPUArchitectures
-							let manifest = GenerateManifest(ToolChain, ProjectName, InTargetType, EngineDirectory, bForDistribution, bPackageDataInsideApk, GameBuildFilesPath, RequiresOBB(bDisallowPackagingDataInApk, ObbFileLocation), bDisableVerifyOBBOnStartUp, Arch, GPUArch, CookFlavor, bUseExternalFilesDir, Configuration.ToString(), SDKLevelInt)
+							let manifest = GenerateManifest(ToolChain, ProjectName, InTargetType, EngineDirectory, bForDistribution, bPackageDataInsideApk, GameBuildFilesPath, RequiresOBB(bDisallowPackagingDataInApk, ObbFileLocation), bDisableVerifyOBBOnStartUp, Arch, GPUArch, CookFlavor, bUseExternalFilesDir, Configuration.ToString(), SDKLevelInt, bSkipGradleBuild)
 							select Tuple.Create(Arch, GPUArch, manifest);
 			}
 			else
@@ -3542,7 +3545,7 @@ namespace UnrealBuildTool
 				BuildList = from Arch in Arches
 							from GPUArch in GPUArchitectures
 							let manifestFile = Path.Combine(IntermediateAndroidPath, Arch + "_" + GPUArch + "_AndroidManifest.xml")
-							let manifest = GenerateManifest(ToolChain, ProjectName, InTargetType, EngineDirectory, bForDistribution, bPackageDataInsideApk, GameBuildFilesPath, RequiresOBB(bDisallowPackagingDataInApk, ObbFileLocation), bDisableVerifyOBBOnStartUp, Arch, GPUArch, CookFlavor, bUseExternalFilesDir, Configuration.ToString(), SDKLevelInt)
+							let manifest = GenerateManifest(ToolChain, ProjectName, InTargetType, EngineDirectory, bForDistribution, bPackageDataInsideApk, GameBuildFilesPath, RequiresOBB(bDisallowPackagingDataInApk, ObbFileLocation), bDisableVerifyOBBOnStartUp, Arch, GPUArch, CookFlavor, bUseExternalFilesDir, Configuration.ToString(), SDKLevelInt, bSkipGradleBuild)
 							let OldManifest = File.Exists(manifestFile) ? File.ReadAllText(manifestFile) : ""
 							where manifest != OldManifest
 							select Tuple.Create(Arch, GPUArch, manifest);
@@ -3982,7 +3985,7 @@ namespace UnrealBuildTool
 					File.WriteAllText(LocalPropertiesFilename, LocalProperties.ToString());
 
 					CreateGradlePropertiesFiles(Arch, MinSDKVersion, TargetSDKVersion, CompileSDKVersion, BuildToolsVersion, PackageName, DestApkName, NDKArch,
-						UE4BuildFilesPath, GameBuildFilesPath, UE4BuildGradleAppPath, UE4BuildPath, UE4BuildGradlePath, bForDistribution);
+						UE4BuildFilesPath, GameBuildFilesPath, UE4BuildGradleAppPath, UE4BuildPath, UE4BuildGradlePath, bForDistribution, bSkipGradleBuild);
 
 					if (!bSkipGradleBuild)
 					{
@@ -4007,8 +4010,8 @@ namespace UnrealBuildTool
 						string GradleBuildType = bForDistribution ? ":app:assembleRelease" : ":app:assembleDebug";
 
 						// collect optional additional Gradle parameters from plugins
-						string GradleOptions = GradleBuildType; // + "--stacktrace --debug"
-						string GradleAdditionalOptions = UPL.ProcessPluginNode(NDKArch, "gradleParameters", "");
+						string GradleOptions = UPL.ProcessPluginNode(NDKArch, "gradleParameters", GradleBuildType); //  "--stacktrace --debug " + GradleBuildType);
+						string GradleSecondCallOptions = UPL.ProcessPluginNode(NDKArch, "gradleSecondCallParameters", "");
 
 						// check for Android Studio project, call Gradle if doesn't exist (assume user will build with Android Studio)
 						string GradleAppImlFilename = Path.Combine(UE4BuildGradlePath, "app.iml");
@@ -4023,9 +4026,9 @@ namespace UnrealBuildTool
 							string ShellParametersEnd = Utils.IsRunningOnMono ? "'" : "";
 							RunCommandLineProgramWithExceptionAndFiltering(UE4BuildGradlePath, ShellExecutable, ShellParametersBegin + "\"" + GradleScriptPath + "\" " + GradleOptions + ShellParametersEnd, "Making .apk with Gradle...");
 
-							if (GradleAdditionalOptions != "")
+							if (GradleSecondCallOptions != "")
 							{
-								RunCommandLineProgramWithExceptionAndFiltering(UE4BuildGradlePath, ShellExecutable, ShellParametersBegin + "\"" + GradleScriptPath + "\" " + GradleAdditionalOptions + ShellParametersEnd, "Additional Gradle steps...");
+								RunCommandLineProgramWithExceptionAndFiltering(UE4BuildGradlePath, ShellExecutable, ShellParametersBegin + "\"" + GradleScriptPath + "\" " + GradleSecondCallOptions + ShellParametersEnd, "Additional Gradle steps...");
 							}
 							
 							// For build machine run a clean afterward to clean up intermediate files (does not remove final APK)
@@ -4121,8 +4124,6 @@ namespace UnrealBuildTool
 			DirectoryReference ProjectDirectory = DirectoryReference.FromFile(Receipt.ProjectFile) ?? UnrealBuildTool.EngineDirectory;
 			string TargetName = (Receipt.ProjectFile == null ? Receipt.TargetName : Receipt.ProjectFile.GetFileNameWithoutAnyExtensions());
 
-			SavePackageInfo(TargetName, ProjectDirectory.FullName, Receipt.TargetType);
-
 			AndroidToolChain ToolChain = ((AndroidPlatform)UEBuildPlatform.GetBuildPlatform(Receipt.Platform)).CreateTempToolChainForProject(Receipt.ProjectFile) as AndroidToolChain;
 
 			// get the receipt
@@ -4130,6 +4131,8 @@ namespace UnrealBuildTool
 
 			bool bShouldCompileAsDll = Receipt.HasValueForAdditionalProperty("CompileAsDll", "true");
 
+			SavePackageInfo(TargetName, ProjectDirectory.FullName, Receipt.TargetType, bShouldCompileAsDll);
+
 			// Get the output paths
 			BuildProductType ProductType = bShouldCompileAsDll ? BuildProductType.DynamicLibrary : BuildProductType.Executable;
 			List<FileReference> OutputPaths = Receipt.BuildProducts.Where(x => x.Type == ProductType).Select(x => x.Path).ToList();
@@ -4158,7 +4161,7 @@ namespace UnrealBuildTool
 		}
 
 		// Store generated package name in a text file for builds that do not generate an apk file 
-		public bool SavePackageInfo(string TargetName, string ProjectDirectory, TargetType InTargetType)
+		public bool SavePackageInfo(string TargetName, string ProjectDirectory, TargetType InTargetType, bool bIsEmbedded)
 		{
 			string PackageName = GetPackageName(TargetName);
 			string DestPackageNameFileName = Path.Combine(ProjectDirectory, "Binaries", "Android", "packageInfo.txt");
@@ -4166,7 +4169,7 @@ namespace UnrealBuildTool
 			string[] PackageInfoSource = new string[4];
 			PackageInfoSource[0] = PackageName;
 			PackageInfoSource[1] = GetStoreVersion().ToString();
-			PackageInfoSource[2] = GetVersionDisplayName();
+			PackageInfoSource[2] = GetVersionDisplayName(bIsEmbedded);
 			PackageInfoSource[3] = string.Format("name='com.epicgames.ue4.GameActivity.AppType' value='{0}'", InTargetType == TargetType.Game ? "" : InTargetType.ToString());
 
 			Log.TraceInformation("Writing packageInfo pkgName:{0} storeVersion:{1} versionDisplayName:{2} to {3}", PackageInfoSource[0], PackageInfoSource[1], PackageInfoSource[2], DestPackageNameFileName);
@@ -4216,7 +4219,7 @@ namespace UnrealBuildTool
 			// file and instead pushes files directly via deploy
 			AndroidToolChain ToolChain = new AndroidToolChain(ProjectFile, false, null, null);
 
-			SavePackageInfo(ProjectName, ProjectDirectory.FullName, Type);
+			SavePackageInfo(ProjectName, ProjectDirectory.FullName, Type, bSkipGradleBuild);
 
 			MakeApk(ToolChain, ProjectName, Type, ProjectDirectory.FullName, ExecutablePath, EngineDirectory, bForDistribution: bForDistribution, CookFlavor: CookFlavor, Configuration: Configuration,
 				bMakeSeparateApks: ShouldMakeSeparateApks(), bIncrementalPackage: false, bDisallowPackagingDataInApk: bIsDataDeploy, bDisallowExternalFilesDir: !bForDistribution || bIsDataDeploy, bSkipGradleBuild:bSkipGradleBuild);
