UE-72100 - tvOS For Distribution packages fail with Code Signing Error if made after a Development Build

diff --git a/Engine/Source/Programs/UnrealBuildTool/ToolChain/RemoteMac.cs b/Engine/Source/Programs/UnrealBuildTool/ToolChain/RemoteMac.cs
--- a/Engine/Source/Programs/UnrealBuildTool/ToolChain/RemoteMac.cs
+++ b/Engine/Source/Programs/UnrealBuildTool/ToolChain/RemoteMac.cs
@@ -851,9 +851,6 @@ namespace UnrealBuildTool
 				TargetFiles.Add(TargetFile);
 			}
 
-			Execute("/", String.Format("rm -rf {0}/Intermediate/IOS/*.plist", GetRemotePath(UnrealBuildTool.EngineDirectory)));
-			Execute("/", String.Format("rm -rf {0}/Intermediate/TVOS/*.plist", GetRemotePath(UnrealBuildTool.EngineDirectory)));
-
 			// Write a file that protects all the scripts from being overridden by the standard engine filters
 			FileReference ScriptProtectList = FileReference.Combine(TempDir, "RsyncEngineScripts-Protect.txt");
 			using(StreamWriter Writer = new StreamWriter(ScriptProtectList.FullName))
@@ -913,6 +910,11 @@ namespace UnrealBuildTool
 				UploadDirectory(ProjectDir, GetRemotePath(ProjectDir), ProjectFilters);
 			}
 
+			Execute("/", String.Format("rm -rf {0}/Intermediate/IOS/*.plist", GetRemotePath(UnrealBuildTool.EngineDirectory)));
+			Execute("/", String.Format("rm -rf {0}/Intermediate/IOS/*.plist", GetRemotePath(ProjectFile.Directory)));
+			Execute("/", String.Format("rm -rf {0}/Intermediate/TVOS/*.plist", GetRemotePath(UnrealBuildTool.EngineDirectory)));
+			Execute("/", String.Format("rm -rf {0}/Intermediate/TVOS/*.plist", GetRemotePath(ProjectFile.Directory)));
+
 			// Fixup permissions on any shell scripts
 			Execute(RemoteBaseDir, String.Format("chmod +x {0}/Build/BatchFiles/Mac/*.sh", EscapeShellArgument(GetRemotePath(UnrealBuildTool.EngineDirectory))));
 		}
