- Added UnrealBuildTool.GetAllEngineDirectories(string OptionaSubDirectory) that makes it easier to get, for instance, all Config directories under Engine and any PlatformExtensions
- Properly scan PlatformExtensions for DataDrivenPlatformInfo files

diff --git a/Engine/Source/Programs/UnrealBuildTool/System/DataDrivenPlatformInfo.cs b/Engine/Source/Programs/UnrealBuildTool/System/DataDrivenPlatformInfo.cs
--- a/Engine/Source/Programs/UnrealBuildTool/System/DataDrivenPlatformInfo.cs
+++ b/Engine/Source/Programs/UnrealBuildTool/System/DataDrivenPlatformInfo.cs
@@ -44,20 +44,29 @@ namespace UnrealBuildTool
 				PlatformInfos = new Dictionary<string, ConfigDataDrivenPlatformInfo>();
 				Dictionary<string, string> IniParents = new Dictionary<string, string>();
 
-				// look through all config dirs looking for the data driven ini file
-				DirectoryReference EngineConfigDir = DirectoryReference.Combine(UnrealBuildTool.EngineDirectory, "Config");
-				foreach(DirectoryReference PlatformDir in DirectoryReference.EnumerateDirectories(EngineConfigDir))
+				foreach (DirectoryReference EngineConfigDir in UnrealBuildTool.GetAllEngineDirectories("Config"))
 				{
-					FileReference FilePath = FileReference.Combine(PlatformDir, "DataDrivenPlatformInfo.ini");
-					if(FileReference.Exists(FilePath))
+					// look through all config dirs looking for the data driven ini file
+					foreach (string FilePath in Directory.EnumerateFiles(EngineConfigDir.FullName, "DataDrivenPlatformInfo.ini", SearchOption.AllDirectories))
 					{
 						// get the platform name from the path
-						string IniPlatformName = PlatformDir.GetDirectoryName();
+						string IniPlatformName;
+						// Foo/Engine/Config/<Platform>/DataDrivenPlatformInfo.ini
+						if (EngineConfigDir.IsUnderDirectory(UnrealBuildTool.EngineDirectory))
+						{
+							IniPlatformName = Path.GetFileName(Path.GetDirectoryName(FilePath));
+						}
+						// Foo/Platforms/<Platform>/Engine/Config/DataDrivenPlatformInfo.ini
+						else
+						{
+							IniPlatformName = Path.GetFileName(Path.GetDirectoryName(Path.GetDirectoryName(Path.GetDirectoryName(FilePath))));
+						}
 
 						// load the DataDrivenPlatformInfo from the path
-						ConfigFile Config = new ConfigFile(FilePath);
+						ConfigFile Config = new ConfigFile(new FileReference(FilePath));
 						ConfigDataDrivenPlatformInfo NewInfo = new ConfigDataDrivenPlatformInfo();
 
+
 						// we must have the key section 
 						ConfigFileSection Section = null;
 						if (Config.TryGetSection("DataDrivenPlatformInfo", out Section))
diff --git a/Engine/Source/Programs/UnrealBuildTool/System/RulesCompiler.cs b/Engine/Source/Programs/UnrealBuildTool/System/RulesCompiler.cs
--- a/Engine/Source/Programs/UnrealBuildTool/System/RulesCompiler.cs
+++ b/Engine/Source/Programs/UnrealBuildTool/System/RulesCompiler.cs
@@ -88,7 +88,7 @@ namespace UnrealBuildTool
 			List<DirectoryReference> RootFolders = new List<DirectoryReference>();
 			if (bIncludeEngine)
 			{
-				RootFolders.Add(UnrealBuildTool.EngineDirectory);
+				RootFolders.AddRange(UnrealBuildTool.GetAllEngineDirectories());
 			}
 			if(bIncludeEnterprise)
 			{
@@ -152,19 +152,6 @@ namespace UnrealBuildTool
 				}
 			}
 
-			// look for platform extension source code
-			if (DirectoryReference.Exists(UnrealBuildTool.PlatformExtensionsDirectory))
-			{
-				foreach (DirectoryReference PlatformFolder in DirectoryReference.EnumerateDirectories(UnrealBuildTool.PlatformExtensionsDirectory))
-				{
-					DirectoryReference PlatformSourceFolder = DirectoryReference.Combine(PlatformFolder, "Engine", "Source");
-					if (DirectoryReference.Exists(PlatformSourceFolder))
-					{
-						Folders.Add(PlatformSourceFolder);
-					}
-				}
-			}
-
 			// Iterate over all the folders to check
 			List<FileReference> SourceFiles = new List<FileReference>();
 			HashSet<FileReference> UniqueSourceFiles = new HashSet<FileReference>();
@@ -375,21 +362,16 @@ namespace UnrealBuildTool
 		{
 			if (EngineRulesAssembly == null)
 			{
-				List<PluginInfo> IncludedPlugins = Plugins.ReadEnginePlugins(UnrealBuildTool.EngineDirectory).ToList();
+				List<PluginInfo> IncludedPlugins = new List<PluginInfo>();
 
 				// look in the platform extensions engine directories for plugins
-				List<DirectoryReference> RootDirectories = new List<DirectoryReference>() { UnrealBuildTool.EngineDirectory };
-				if (DirectoryReference.Exists(UnrealBuildTool.PlatformExtensionsDirectory))
+				List<DirectoryReference> RootDirectories = new List<DirectoryReference>();
+				RootDirectories.AddRange(UnrealBuildTool.GetAllEngineDirectories());
+
+				// search for all engine plugins
+				foreach (DirectoryReference RootDir in RootDirectories) 
 				{
-					foreach (DirectoryReference PlatformFolder in DirectoryReference.EnumerateDirectories(UnrealBuildTool.PlatformExtensionsDirectory))
-					{
-						DirectoryReference PlatformEngineFolder = DirectoryReference.Combine(PlatformFolder, "Engine");
-						if (DirectoryReference.Exists(PlatformEngineFolder))
-						{
-							RootDirectories.Add(PlatformEngineFolder);
-							IncludedPlugins.AddRange(Plugins.ReadEnginePlugins(PlatformEngineFolder));
-						}
-					}
+					IncludedPlugins.AddRange(Plugins.ReadEnginePlugins(RootDir));
 				}
 
 				RulesScope EngineScope = new RulesScope("Engine", null);
@@ -532,42 +514,52 @@ namespace UnrealBuildTool
 					Parent = CreateEngineRulesAssembly(bUsePrecompiled, bSkipCompile);
 				}
 
+				DirectoryReference MainProjectDirectory = ProjectFileName.Directory;
+				DirectoryReference MainProjectSourceDirectory = DirectoryReference.Combine(MainProjectDirectory, "Source");
+
 				// Create a scope for things in this assembly
 				RulesScope Scope = new RulesScope("Project", Parent.Scope);
 
-				// Find all the rules under the project source directory
-				DirectoryReference ProjectDirectory = ProjectFileName.Directory;
-				DirectoryReference ProjectSourceDirectory = DirectoryReference.Combine(ProjectDirectory, "Source");
-
 				// Create a new context for modules created by this assembly
-				ModuleRulesContext DefaultModuleContext = new ModuleRulesContext(Scope, ProjectDirectory);
+				ModuleRulesContext DefaultModuleContext = new ModuleRulesContext(Scope, MainProjectDirectory);
 				DefaultModuleContext.bCanBuildDebugGame = true;
 				DefaultModuleContext.bCanHotReload = true;
 				DefaultModuleContext.bClassifyAsGameModuleForUHT = true;
 				DefaultModuleContext.bCanUseForSharedPCH = false;
-				
+
+				// gather modules from project and platforms
 				Dictionary<FileReference, ModuleRulesContext> ModuleFiles = new Dictionary<FileReference, ModuleRulesContext>();
-				AddModuleRulesWithContext(ProjectSourceDirectory, DefaultModuleContext, ModuleFiles);
-				
-				List<FileReference> TargetFiles = new List<FileReference>(FindAllRulesFiles(ProjectSourceDirectory, RulesFileType.Target));
+				List<FileReference> TargetFiles = new List<FileReference>();
 
 				// Find all the project plugins
-				List<PluginInfo> ProjectPlugins = new List<PluginInfo>(Plugins.ReadProjectPlugins(ProjectFileName.Directory));
+				List<PluginInfo> ProjectPlugins = new List<PluginInfo>();
 
-                // Add the project's additional plugin directories plugins too
-				if(Project.AdditionalPluginDirectories != null)
+				// Find all the rules/plugins under the project source directories
+				foreach (DirectoryReference ProjectDirectory in UnrealBuildTool.GetAllProjectDirectories(ProjectFileName))
 				{
-					foreach(string AdditionalPluginDirectory in Project.AdditionalPluginDirectories)
+					DirectoryReference ProjectSourceDirectory = DirectoryReference.Combine(ProjectDirectory, "Source");
+
+					AddModuleRulesWithContext(ProjectSourceDirectory, DefaultModuleContext, ModuleFiles);
+					TargetFiles.AddRange(FindAllRulesFiles(ProjectSourceDirectory, RulesFileType.Target));
+
+					ProjectPlugins.AddRange(Plugins.ReadProjectPlugins(ProjectDirectory));
+
+					// Add the project's additional plugin directories plugins too
+					if (Project.AdditionalPluginDirectories != null)
 					{
-						ProjectPlugins.AddRange(Plugins.ReadAdditionalPlugins(DirectoryReference.Combine(ProjectFileName.Directory, AdditionalPluginDirectory)));
+						foreach (string AdditionalPluginDirectory in Project.AdditionalPluginDirectories)
+						{
+							ProjectPlugins.AddRange(Plugins.ReadAdditionalPlugins(DirectoryReference.Combine(ProjectDirectory, AdditionalPluginDirectory)));
+						}
 					}
 				}
 
+
 				// Find all the plugin module rules
 				FindModuleRulesForPlugins(ProjectPlugins, DefaultModuleContext, ModuleFiles);
 
 				// Add the games project's intermediate source folder
-				DirectoryReference ProjectIntermediateSourceDirectory = DirectoryReference.Combine(ProjectDirectory, "Intermediate", "Source");
+				DirectoryReference ProjectIntermediateSourceDirectory = DirectoryReference.Combine(MainProjectDirectory, "Intermediate", "Source");
 				if (DirectoryReference.Exists(ProjectIntermediateSourceDirectory))
 				{
 					AddModuleRulesWithContext(ProjectIntermediateSourceDirectory, DefaultModuleContext, ModuleFiles);
@@ -575,14 +567,14 @@ namespace UnrealBuildTool
 				}
 
 				// Compile the assembly. If there are no module or target files, just use the parent assembly.
-				FileReference AssemblyFileName = FileReference.Combine(ProjectDirectory, "Intermediate", "Build", "BuildRules", ProjectFileName.GetFileNameWithoutExtension() + "ModuleRules" + FrameworkAssemblyExtension);
+				FileReference AssemblyFileName = FileReference.Combine(MainProjectDirectory, "Intermediate", "Build", "BuildRules", ProjectFileName.GetFileNameWithoutExtension() + "ModuleRules" + FrameworkAssemblyExtension);
 				if(ModuleFiles.Count == 0 && TargetFiles.Count == 0)
 				{
 					ProjectRulesAssembly = Parent;
 				}
 				else
 				{
-					ProjectRulesAssembly = new RulesAssembly(Scope, ProjectDirectory, ProjectPlugins, ModuleFiles, TargetFiles, AssemblyFileName, bContainsEngineModules: false, bUseBackwardsCompatibleDefaults: true, bReadOnly: UnrealBuildTool.IsProjectInstalled(), bSkipCompile: bSkipCompile, Parent: Parent);
+					ProjectRulesAssembly = new RulesAssembly(Scope, MainProjectDirectory, ProjectPlugins, ModuleFiles, TargetFiles, AssemblyFileName, bContainsEngineModules: false, bUseBackwardsCompatibleDefaults: true, bReadOnly: UnrealBuildTool.IsProjectInstalled(), bSkipCompile: bSkipCompile, Parent: Parent);
 				}
 				LoadedAssemblyMap.Add(ProjectFileName, ProjectRulesAssembly);
 			}
diff --git a/Engine/Source/Programs/UnrealBuildTool/UnrealBuildTool.cs b/Engine/Source/Programs/UnrealBuildTool/UnrealBuildTool.cs
--- a/Engine/Source/Programs/UnrealBuildTool/UnrealBuildTool.cs
+++ b/Engine/Source/Programs/UnrealBuildTool/UnrealBuildTool.cs
@@ -109,6 +109,43 @@ namespace UnrealBuildTool
 		/// </summary>
 		public static readonly DirectoryReference EnterpriseIntermediateDirectory = DirectoryReference.Combine(EnterpriseDirectory, "Intermediate");
 
+		/// <summary>
+		/// The main engine directory and all found platform extension engine directories
+		/// </summary>
+		public static DirectoryReference[] GetAllEngineDirectories(string Suffix="")
+		{
+			List<DirectoryReference> EngineDirectories = new List<DirectoryReference>() { DirectoryReference.Combine(EngineDirectory, Suffix) };
+			foreach (DirectoryReference PlatformDirectory in DirectoryReference.EnumerateDirectories(PlatformExtensionsDirectory))
+			{
+				DirectoryReference PlatformEngineDirectory = DirectoryReference.Combine(PlatformDirectory, "Engine", Suffix);
+				if (DirectoryReference.Exists(PlatformEngineDirectory))
+				{
+					EngineDirectories.Add(PlatformEngineDirectory);
+				}
+			}
+			return EngineDirectories.ToArray();
+		}
+
+		/// <summary>
+		/// The main project directory and all found platform extension project directories
+		/// </summary>
+		public static DirectoryReference[] GetAllProjectDirectories(FileReference ProjectFile, string Suffix = "")
+		{
+			// project name may not match the directory name, so use the ProjectFile's name as ProjectName
+			string ProjectName = ProjectFile.GetFileNameWithoutAnyExtensions();
+
+			List<DirectoryReference> ProjectDirectories = new List<DirectoryReference>() { DirectoryReference.Combine(ProjectFile.Directory, Suffix) };
+			foreach (DirectoryReference PlatformDirectory in DirectoryReference.EnumerateDirectories(PlatformExtensionsDirectory))
+			{
+				DirectoryReference PlatformEngineDirectory = DirectoryReference.Combine(PlatformDirectory, ProjectName, Suffix);
+				if (DirectoryReference.Exists(PlatformEngineDirectory))
+				{
+					ProjectDirectories.Add(PlatformEngineDirectory);
+				}
+			}
+			return ProjectDirectories.ToArray();
+		}
+
 		/// <summary>
 		/// The Remote Ini directory.  This should always be valid when compiling using a remote server.
 		/// </summary>
