- Fixed in UBT when /Platforms directory does not exist

diff --git a/Engine/Source/Programs/UnrealBuildTool/UnrealBuildTool.cs b/Engine/Source/Programs/UnrealBuildTool/UnrealBuildTool.cs
--- a/Engine/Source/Programs/UnrealBuildTool/UnrealBuildTool.cs
+++ b/Engine/Source/Programs/UnrealBuildTool/UnrealBuildTool.cs
@@ -115,12 +115,15 @@ namespace UnrealBuildTool
 		public static DirectoryReference[] GetAllEngineDirectories(string Suffix="")
 		{
 			List<DirectoryReference> EngineDirectories = new List<DirectoryReference>() { DirectoryReference.Combine(EngineDirectory, Suffix) };
-			foreach (DirectoryReference PlatformDirectory in DirectoryReference.EnumerateDirectories(PlatformExtensionsDirectory))
+			if (DirectoryReference.Exists(PlatformExtensionsDirectory))
 			{
-				DirectoryReference PlatformEngineDirectory = DirectoryReference.Combine(PlatformDirectory, "Engine", Suffix);
-				if (DirectoryReference.Exists(PlatformEngineDirectory))
+				foreach (DirectoryReference PlatformDirectory in DirectoryReference.EnumerateDirectories(PlatformExtensionsDirectory))
 				{
-					EngineDirectories.Add(PlatformEngineDirectory);
+					DirectoryReference PlatformEngineDirectory = DirectoryReference.Combine(PlatformDirectory, "Engine", Suffix);
+					if (DirectoryReference.Exists(PlatformEngineDirectory))
+					{
+						EngineDirectories.Add(PlatformEngineDirectory);
+					}
 				}
 			}
 			return EngineDirectories.ToArray();
@@ -135,12 +138,15 @@ namespace UnrealBuildTool
 			string ProjectName = ProjectFile.GetFileNameWithoutAnyExtensions();
 
 			List<DirectoryReference> ProjectDirectories = new List<DirectoryReference>() { DirectoryReference.Combine(ProjectFile.Directory, Suffix) };
-			foreach (DirectoryReference PlatformDirectory in DirectoryReference.EnumerateDirectories(PlatformExtensionsDirectory))
+			if (DirectoryReference.Exists(PlatformExtensionsDirectory))
 			{
-				DirectoryReference PlatformEngineDirectory = DirectoryReference.Combine(PlatformDirectory, ProjectName, Suffix);
-				if (DirectoryReference.Exists(PlatformEngineDirectory))
+				foreach (DirectoryReference PlatformDirectory in DirectoryReference.EnumerateDirectories(PlatformExtensionsDirectory))
 				{
-					ProjectDirectories.Add(PlatformEngineDirectory);
+					DirectoryReference PlatformEngineDirectory = DirectoryReference.Combine(PlatformDirectory, ProjectName, Suffix);
+					if (DirectoryReference.Exists(PlatformEngineDirectory))
+					{
+						ProjectDirectories.Add(PlatformEngineDirectory);
+					}
 				}
 			}
 			return ProjectDirectories.ToArray();
