Rename UniqueStringRegistry into UniqueStringRegistry.cs.

#rnx

diff --git a/Engine/Source/Programs/UnrealBuildTool/System/StringToInstanceRegistry.cs b/Engine/Source/Programs/UnrealBuildTool/System/StringToInstanceRegistry.cs
deleted file mode 100644
diff --git a/Engine/Source/Programs/UnrealBuildTool/System/UniqueStringRegistry.cs b/Engine/Source/Programs/UnrealBuildTool/System/UniqueStringRegistry.cs
new file mode 100644
--- /dev/null
+++ b/Engine/Source/Programs/UnrealBuildTool/System/UniqueStringRegistry.cs
@@ -0,0 +1,69 @@
+// Copyright 1998-2019 Epic Games, Inc. All Rights Reserved.
+
+using System;
+using System.Collections.Generic;
+using System.Linq;
+
+namespace UnrealBuildTool
+{
+	/// <summary>
+	/// Maps a unique string to an integer
+	/// </summary>
+	class UniqueStringRegistry
+	{
+		// protect the InstanceMap
+		private object LockObject = new object();
+
+		// holds a mapping of string name to single instance
+		private Dictionary<string, int> StringToInstanceMap = new Dictionary<string, int>(StringComparer.OrdinalIgnoreCase);
+
+		public UniqueStringRegistry()
+		{
+		}
+
+		public bool HasString(string Name)
+		{
+			return StringToInstanceMap.ContainsKey(Name);
+		}
+
+		public int FindOrAddByName(string Name)
+		{
+			// look for existing one
+			int Instance = -1;
+			if (!StringToInstanceMap.TryGetValue(Name, out Instance))
+			{
+				lock (LockObject)
+				{
+					if (!StringToInstanceMap.TryGetValue(Name, out Instance))
+					{
+						// copy over the dictionary and add, so that other threads can keep reading out of the old one until it's time
+						Dictionary<string, int> NewStringToInstanceMap = new Dictionary<string, int>(StringToInstanceMap, StringComparer.OrdinalIgnoreCase);
+
+						// make and add a new instance number
+						Instance = StringToInstanceMap.Count;
+						NewStringToInstanceMap[Name] = Instance;
+
+						// replace the class's map
+						StringToInstanceMap = NewStringToInstanceMap;
+					}
+				}
+			}
+			return Instance;
+		}
+
+		public string[] GetStringNames()
+		{
+			return StringToInstanceMap.Keys.ToArray();
+		}
+
+		public int[] GetStringIds()
+		{
+			return StringToInstanceMap.Values.ToArray();
+		}
+
+		public string GetStringForId(int Id)
+		{
+			return StringToInstanceMap.First(x => x.Value == Id).Key;
+		}
+	}
+}
\ No newline at end of file
diff --git a/Engine/Source/Programs/UnrealBuildTool/UnrealBuildTool.csproj b/Engine/Source/Programs/UnrealBuildTool/UnrealBuildTool.csproj
--- a/Engine/Source/Programs/UnrealBuildTool/UnrealBuildTool.csproj
+++ b/Engine/Source/Programs/UnrealBuildTool/UnrealBuildTool.csproj
@@ -229,7 +229,6 @@
     <Compile Include="System\SingleInstanceMutex.cs" />
     <Compile Include="System\SourceFileWorkingSet.cs" />
     <Compile Include="System\StartupTraceListener.cs" />
-    <Compile Include="System\StringToInstanceRegistry.cs" />
     <Compile Include="System\SupportedConfigurationsAttribute.cs" />
     <Compile Include="System\SupportedPlatformsAttribute.cs" />
     <Compile Include="System\TargetMakefile.cs" />
@@ -237,6 +236,7 @@
     <Compile Include="System\Timeline.cs" />
     <Compile Include="System\TimingData.cs" />
     <Compile Include="System\ToolMode.cs" />
+    <Compile Include="System\UniqueStringRegistry.cs" />
     <Compile Include="System\VersionNumber.cs" />
     <Compile Include="System\XmlConfig.cs" />
     <Compile Include="System\XmlConfigData.cs" />
