Implement the RestrictedPlatform type similarly to UnrealTargetPlatform. Platforms marked as confidential via data driven platform info are automatically registered with this list.

#rnx

diff --git a/Engine/Source/Programs/AutomationTool/AutomationUtils/DeploymentContext.cs b/Engine/Source/Programs/AutomationTool/AutomationUtils/DeploymentContext.cs
--- a/Engine/Source/Programs/AutomationTool/AutomationUtils/DeploymentContext.cs
+++ b/Engine/Source/Programs/AutomationTool/AutomationUtils/DeploymentContext.cs
@@ -420,7 +420,7 @@ public class DeploymentContext //: ProjectParams
 		{
 			RestrictedFolderNames.ExceptWith(PlatformExports.GetIncludedFolderNames(StagePlatform));
 		}
-		RestrictedFolderNames.UnionWith(RestrictedFolders.Names);
+		RestrictedFolderNames.UnionWith(RestrictedFolder.GetNames());
 		RestrictedFolderNames.Remove(StageTargetPlatform.IniPlatformType.ToString());
 
 		// Read the game config files
diff --git a/Engine/Source/Programs/AutomationTool/BuildGraph/BuildGraph.cs b/Engine/Source/Programs/AutomationTool/BuildGraph/BuildGraph.cs
--- a/Engine/Source/Programs/AutomationTool/BuildGraph/BuildGraph.cs
+++ b/Engine/Source/Programs/AutomationTool/BuildGraph/BuildGraph.cs
@@ -126,8 +126,8 @@ namespace AutomationTool
 			DefaultProperties["RootDir"] = CommandUtils.RootDirectory.FullName;
 			DefaultProperties["IsBuildMachine"] = IsBuildMachine ? "true" : "false";
 			DefaultProperties["HostPlatform"] = HostPlatform.Current.HostEditorPlatform.ToString();
-			DefaultProperties["RestrictedFolderNames"] = String.Join(";", RestrictedFolders.Names);
-			DefaultProperties["RestrictedFolderFilter"] = String.Join(";", RestrictedFolders.Names.Select(x => String.Format(".../{0}/...", x)));
+			DefaultProperties["RestrictedFolderNames"] = String.Join(";", RestrictedFolder.GetNames());
+			DefaultProperties["RestrictedFolderFilter"] = String.Join(";", RestrictedFolder.GetNames().Select(x => String.Format(".../{0}/...", x)));
 
 			// Attempt to read existing Build Version information
 			BuildVersion Version;
diff --git a/Engine/Source/Programs/AutomationTool/Scripts/BuildPluginCommand.Automation.cs b/Engine/Source/Programs/AutomationTool/Scripts/BuildPluginCommand.Automation.cs
--- a/Engine/Source/Programs/AutomationTool/Scripts/BuildPluginCommand.Automation.cs
+++ b/Engine/Source/Programs/AutomationTool/Scripts/BuildPluginCommand.Automation.cs
@@ -258,7 +258,10 @@ class BuildPlugin : BuildCommand
 		}
 
 		// Apply the standard exclusion rules
-		Filter.ExcludeRestrictedFolders();
+		foreach (string RestrictedFolderName in RestrictedFolder.GetNames())
+		{
+			Filter.AddRule(String.Format(".../{0}/...", RestrictedFolderName), FileFilterType.Exclude);
+		}
 
 		// Apply the filter to the plugin directory
 		return Filter.ApplyToDirectory(PluginFile.Directory, true);
diff --git a/Engine/Source/Programs/AutomationTool/Scripts/CheckRestrictedFolders.cs b/Engine/Source/Programs/AutomationTool/Scripts/CheckRestrictedFolders.cs
--- a/Engine/Source/Programs/AutomationTool/Scripts/CheckRestrictedFolders.cs
+++ b/Engine/Source/Programs/AutomationTool/Scripts/CheckRestrictedFolders.cs
@@ -26,7 +26,7 @@ namespace AutomationTool
 			}
 
 			// Find a list of restricted folders, and remove any names which are explicitly whitelisted
-			HashSet<string> RestrictedNames = new HashSet<string>(RestrictedFolders.Names, StringComparer.InvariantCultureIgnoreCase);
+			HashSet<string> RestrictedNames = new HashSet<string>(RestrictedFolder.GetNames(), StringComparer.InvariantCultureIgnoreCase);
 			foreach (string AllowParam in ParseParamValues("Allow"))
 			{
 				RestrictedNames.ExceptWith(AllowParam.Split('+'));
diff --git a/Engine/Source/Programs/DotNETCommon/DotNETUtilities/DotNETUtilities.csproj b/Engine/Source/Programs/DotNETCommon/DotNETUtilities/DotNETUtilities.csproj
--- a/Engine/Source/Programs/DotNETCommon/DotNETUtilities/DotNETUtilities.csproj
+++ b/Engine/Source/Programs/DotNETCommon/DotNETUtilities/DotNETUtilities.csproj
@@ -157,7 +157,6 @@
     <Compile Include="Perforce\Records\WhereRecord.cs" />
     <Compile Include="Properties\AssemblyInfo.cs" />
     <Compile Include="ReadOnlyHashSet.cs" />
-    <Compile Include="RestrictedFolders.cs" />
     <Compile Include="ResXResource\UEResXReader.cs" />
     <Compile Include="ResXResource\UEResXWriter.cs" />
     <Compile Include="StringUtils.cs" />
diff --git a/Engine/Source/Programs/DotNETCommon/DotNETUtilities/FileFilter.cs b/Engine/Source/Programs/DotNETCommon/DotNETUtilities/FileFilter.cs
--- a/Engine/Source/Programs/DotNETCommon/DotNETUtilities/FileFilter.cs
+++ b/Engine/Source/Programs/DotNETCommon/DotNETUtilities/FileFilter.cs
@@ -343,17 +343,6 @@ namespace Tools.DotNETCommon
 			}
 		}
 
-		/// <summary>
-		/// Excludes all confidential folders from the filter
-		/// </summary>
-		public void ExcludeRestrictedFolders()
-		{
-			foreach(string RestrictedFolderName in RestrictedFolders.Names)
-			{
-				AddRule(String.Format(".../{0}/...", RestrictedFolderName), FileFilterType.Exclude);
-			}
-		}
-
 		/// <summary>
 		/// Determines whether the given file matches the filter
 		/// </summary>
diff --git a/Engine/Source/Programs/DotNETCommon/DotNETUtilities/RestrictedFolders.cs b/Engine/Source/Programs/DotNETCommon/DotNETUtilities/RestrictedFolders.cs
deleted file mode 100644
diff --git a/Engine/Source/Programs/UnrealBuildTool/System/RestrictedFolders.cs b/Engine/Source/Programs/UnrealBuildTool/System/RestrictedFolders.cs
new file mode 100644
--- /dev/null
+++ b/Engine/Source/Programs/UnrealBuildTool/System/RestrictedFolders.cs
@@ -0,0 +1,224 @@
+// Copyright 1998-2019 Epic Games, Inc. All Rights Reserved.
+
+using System;
+using System.Collections.Generic;
+using System.IO;
+using System.Linq;
+using System.Text;
+using System.Threading.Tasks;
+using Tools.DotNETCommon;
+
+namespace UnrealBuildTool
+{
+	/// <summary>
+	/// Names of restricted folders. Note: The name of each entry is used to search for/create folders
+	/// </summary>
+	public partial struct RestrictedFolder
+	{
+		/// <summary>
+		/// Unique Id for this folder
+		/// </summary>
+		private int Id;
+
+		/// <summary>
+		/// Mapping for unique ids
+		/// </summary>
+		private static UniqueStringRegistry StringRegistry = new UniqueStringRegistry();
+
+		/// <summary>
+		/// Array of all restricted folder names
+		/// </summary>
+		private static string[] Names;
+
+		/// <summary>
+		/// Array of all restricted folders
+		/// </summary>
+		private static RestrictedFolder[] Values;
+
+		/// <summary>
+		/// Constructor
+		/// </summary>
+		/// <param name="Id">Id of the string</param>
+		private RestrictedFolder(int Id)
+		{
+			this.Id = Id;
+		}
+
+		/// <summary>
+		/// Creates a restricted folder instance from a string
+		/// </summary>
+		/// <param name="Name">Name of the folder</param>
+		/// <returns>New restricted folder instance</returns>
+		static private RestrictedFolder FindOrAddByName(string Name)
+		{
+			return new RestrictedFolder(StringRegistry.FindOrAddByName(Name));
+		}
+
+		/// <summary>
+		/// Tests for equality between two restricted folder instances
+		/// </summary>
+		/// <param name="A">First instance</param>
+		/// <param name="B">Second instance</param>
+		/// <returns>True if the two instances are equal</returns>
+		public static bool operator ==(RestrictedFolder A, RestrictedFolder B)
+		{
+			return A.Id == B.Id;
+		}
+
+		/// <summary>
+		/// Tests for inequality between two restricted folder instances
+		/// </summary>
+		/// <param name="A">First instance</param>
+		/// <param name="B">Second instance</param>
+		/// <returns>True if the two instances are not equal</returns>
+		public static bool operator !=(RestrictedFolder A, RestrictedFolder B)
+		{
+			return A.Id != B.Id;
+		}
+
+		/// <summary>
+		/// Tests whether two restricted folder instances are equal
+		/// </summary>
+		/// <param name="Other">The restricted folder to compare against</param>
+		/// <returns>True if the restricted folder is equal to the other instance</returns>
+		public override bool Equals(object Other)
+		{
+			return Other is RestrictedFolder && Id == ((RestrictedFolder)Other).Id;
+		}
+
+		/// <summary>
+		/// Gets a hash code for this object
+		/// </summary>
+		/// <returns>Hash code for the object</returns>
+		public override int GetHashCode()
+		{
+			return Id;
+		}
+
+		/// <summary>
+		/// Creates entries for all the confidential platforms. Should be called before returning any list of all folder values.
+		/// </summary>
+		private static void AddConfidentialPlatforms()
+		{
+			foreach (KeyValuePair<string, DataDrivenPlatformInfo.ConfigDataDrivenPlatformInfo> Pair in DataDrivenPlatformInfo.GetAllPlatformInfos())
+			{
+				if (Pair.Value.bIsConfidential)
+				{
+					FindOrAddByName(Pair.Key);
+				}
+			}
+		}
+
+		/// <summary>
+		/// Gets an array of all the restricted folder names
+		/// </summary>
+		/// <returns></returns>
+		public static string[] GetNames()
+		{
+			if(Names == null)
+			{
+				AddConfidentialPlatforms();
+				Names = StringRegistry.GetStringNames();
+			}
+			return Names;
+		}
+
+		/// <summary>
+		/// Ensures that we've added all the restricted folders, and return an array of them
+		/// </summary>
+		/// <returns>Array of restricted folder values</returns>
+		public static RestrictedFolder[] GetValues()
+		{
+			if(Values == null)
+			{
+				AddConfidentialPlatforms();
+				Values = Array.ConvertAll(StringRegistry.GetStringIds(), x => new RestrictedFolder(x));
+			}
+			return Values;
+		}
+
+		/// <summary>
+		/// Return the string representation
+		/// </summary>
+		/// <returns></returns>
+		public override string ToString()
+		{
+			return StringRegistry.GetStringForId(Id);
+		}
+	}
+
+	/// <summary>
+	/// Values for RestrictedFolder
+	/// </summary>
+	public partial struct RestrictedFolder
+	{
+		/// <summary>
+		/// Legacy. Should not be used any more.
+		/// </summary>
+		public static RestrictedFolder EpicInternal = FindOrAddByName("EpicInternal");
+
+		/// <summary>
+		/// Can be used by UE4 but not required
+		/// </summary>
+		public static RestrictedFolder CarefullyRedist = FindOrAddByName("CarefullyRedist");
+
+		/// <summary>
+		/// Epic Employees and Contractors
+		/// </summary>
+		public static RestrictedFolder NotForLicensees = FindOrAddByName("NotForLicensees");
+
+		/// <summary>
+		/// Epic Employees only
+		/// </summary>
+		public static RestrictedFolder NoRedist = FindOrAddByName("NoRedist");
+
+		/// <summary>
+		/// Playstation 4 source files
+		/// </summary>
+		public static RestrictedFolder PS4 = FindOrAddByName("PS4");
+
+		/// <summary>
+		/// XboxOne source files
+		/// </summary>
+		public static RestrictedFolder XboxOne = FindOrAddByName("XboxOne");
+
+		/// <summary>
+		/// Switch source files
+		/// </summary>
+		public static RestrictedFolder Switch = FindOrAddByName("Switch");
+
+		/// <summary>
+		/// Quail source files
+		/// </summary>
+		public static RestrictedFolder Quail = FindOrAddByName("Quail");
+	}
+
+	/// <summary>
+	/// Utility functions for getting restricted folder
+	/// </summary>
+	public static class RestrictedFolders
+	{
+		/// <summary>
+		/// Finds all the restricted folder names relative to a base directory
+		/// </summary>
+		/// <param name="BaseDir">The base directory to check against</param>
+		/// <param name="OtherDir">The file or directory to check</param>
+		/// <returns>Array of restricted folder names</returns>
+		public static List<RestrictedFolder> FindRestrictedFolders(DirectoryReference BaseDir, DirectoryReference OtherDir)
+		{
+			List<RestrictedFolder> Folders = new List<RestrictedFolder>();
+			if (OtherDir.IsUnderDirectory(BaseDir))
+			{
+				foreach(RestrictedFolder Value in RestrictedFolder.GetValues())
+				{
+					string Name = Value.ToString();
+					if (OtherDir.ContainsName(Name, BaseDir.FullName.Length))
+					{
+						Folders.Add(Value);
+					}
+				}
+			}
+			return Folders;
+		}
+	}
+}
diff --git a/Engine/Source/Programs/UnrealBuildTool/UnrealBuildTool.csproj b/Engine/Source/Programs/UnrealBuildTool/UnrealBuildTool.csproj
--- a/Engine/Source/Programs/UnrealBuildTool/UnrealBuildTool.csproj
+++ b/Engine/Source/Programs/UnrealBuildTool/UnrealBuildTool.csproj
@@ -195,6 +195,7 @@
     <Compile Include="System\CppIncludeParser.cs" />
     <Compile Include="System\FileMetadataPrefetch.cs" />
     <Compile Include="Modes\ParseMsvcTimingInfo.cs" />
+    <Compile Include="System\RestrictedFolders.cs" />
     <Compile Include="System\RulesScope.cs" />
     <Compile Include="System\SourceFileMetadataCache.cs" />
     <Compile Include="System\DataDrivenPlatformInfo.cs" />
