[This is actually Josh Adams]
- Fix setting up EMSDK on Linux

diff --git a/Engine/Source/Programs/UnrealBuildTool/Platform/HTML5/HTML5SDKInfo.cs b/Engine/Source/Programs/UnrealBuildTool/Platform/HTML5/HTML5SDKInfo.cs
--- a/Engine/Source/Programs/UnrealBuildTool/Platform/HTML5/HTML5SDKInfo.cs
+++ b/Engine/Source/Programs/UnrealBuildTool/Platform/HTML5/HTML5SDKInfo.cs
@@ -79,7 +79,7 @@ namespace UnrealBuildTool
 				{
 					return ".exe";
 				}
-				if (BuildHostPlatform.Current.Platform == UnrealTargetPlatform.Mac || BuildHostPlatform.Current.Platform == UnrealTargetPlatform.Mac)
+				if (BuildHostPlatform.Current.Platform == UnrealTargetPlatform.Mac || BuildHostPlatform.Current.Platform == UnrealTargetPlatform.Linux)
 				{
 					return "";
 				}
@@ -167,7 +167,7 @@ namespace UnrealBuildTool
 				{
 					return Path.Combine(SDKBase, "Win64", "python", PYTHON_VER, "python.exe");
 				}
-				if (BuildHostPlatform.Current.Platform == UnrealTargetPlatform.Mac || BuildHostPlatform.Current.Platform == UnrealTargetPlatform.Mac)
+				if (BuildHostPlatform.Current.Platform == UnrealTargetPlatform.Mac || BuildHostPlatform.Current.Platform == UnrealTargetPlatform.Linux)
 				{
 					return "/usr/bin/python";
 				}
