UBT: Add a <VCProjectFileGenerator><bBuildLiveCodingConsole> build setting that allows automatically building LiveCodingConsole automatically.

diff --git a/Engine/Source/Programs/UnrealBuildTool/ProjectFiles/VisualStudio/VCProject.cs b/Engine/Source/Programs/UnrealBuildTool/ProjectFiles/VisualStudio/VCProject.cs
--- a/Engine/Source/Programs/UnrealBuildTool/ProjectFiles/VisualStudio/VCProject.cs
+++ b/Engine/Source/Programs/UnrealBuildTool/ProjectFiles/VisualStudio/VCProject.cs
@@ -282,6 +282,7 @@ namespace UnrealBuildTool
 		bool bUsePerFileIntellisense;
 		bool bUsePrecompiled;
 		bool bEditorDependsOnShaderCompileWorker;
+		bool bBuildLiveCodingConsole;
 		string BuildToolOverride;
 
 		/// This is the platform name that Visual Studio is always guaranteed to support.  We'll use this as
@@ -305,8 +306,9 @@ namespace UnrealBuildTool
 		/// <param name="bUsePerFileIntellisense">If true, generates per-file intellisense data</param>
 		/// <param name="bUsePrecompiled">Whether to add the -UsePrecompiled argumemnt when building targets</param>
 		/// <param name="bEditorDependsOnShaderCompileWorker">Whether editor targets should also build ShaderCompileWorker</param>
+		/// <param name="bBuildLiveCodingConsole">Whether targets using live coding should also build LiveCodingConsole</param>
 		/// <param name="BuildToolOverride">Optional arguments to pass to UBT when building</param>
-		public VCProjectFile(FileReference InFilePath, FileReference InOnlyGameProject, VCProjectFileFormat InProjectFileFormat, bool bUseFastPDB, bool bUsePerFileIntellisense, bool bUsePrecompiled, bool bEditorDependsOnShaderCompileWorker, string BuildToolOverride)
+		public VCProjectFile(FileReference InFilePath, FileReference InOnlyGameProject, VCProjectFileFormat InProjectFileFormat, bool bUseFastPDB, bool bUsePerFileIntellisense, bool bUsePrecompiled, bool bEditorDependsOnShaderCompileWorker, bool bBuildLiveCodingConsole, string BuildToolOverride)
 			: base(InFilePath)
 		{
 			OnlyGameProject = InOnlyGameProject;
@@ -315,6 +317,7 @@ namespace UnrealBuildTool
 			this.bUsePerFileIntellisense = bUsePerFileIntellisense;
 			this.bUsePrecompiled = bUsePrecompiled;
 			this.bEditorDependsOnShaderCompileWorker = bEditorDependsOnShaderCompileWorker;
+			this.bBuildLiveCodingConsole = bBuildLiveCodingConsole;
 			this.BuildToolOverride = BuildToolOverride;
 		}
 
@@ -1394,15 +1397,33 @@ namespace UnrealBuildTool
 						BuildArguments.AppendFormat(" -Project={0}", UProjectPath);
 					}
 
-					if (bUsePrecompiled)
+					List<string> ExtraTargets = new List<string>();
+					if (!bUsePrecompiled)
 					{
-						BuildArguments.Append(" -UsePrecompiled");
+						if (TargetRulesObject.Type == TargetType.Editor && bEditorDependsOnShaderCompileWorker && !UnrealBuildTool.IsEngineInstalled())
+						{
+							ExtraTargets.Add("ShaderCompileWorker Win64 Development");
+						}
+						if (TargetRulesObject.bWithLiveCoding && bBuildLiveCodingConsole && !UnrealBuildTool.IsEngineInstalled())
+						{
+							ExtraTargets.Add("LiveCodingConsole Win64 Development");
+						}
 					}
-					else if(TargetRulesObject.Type == TargetType.Editor && bEditorDependsOnShaderCompileWorker && !UnrealBuildTool.IsEngineInstalled())
+
+					if(ExtraTargets.Count > 0)
 					{
 						BuildArguments.Replace("\"", "\\\"");
-						BuildArguments.Insert(0, "-Target=\"ShaderCompileWorker Win64 Development\" -Target=\"");
+						BuildArguments.Insert(0, "-Target=\"");
 						BuildArguments.Append("\"");
+						foreach(string ExtraTarget in ExtraTargets)
+						{
+							BuildArguments.AppendFormat(" -Target=\"{0}\"", ExtraTarget);
+						}
+					}
+
+					if (bUsePrecompiled)
+					{
+						BuildArguments.Append(" -UsePrecompiled");
 					}
 
 					// Always wait for the mutex between UBT invocations, so that building the whole solution doesn't fail.
diff --git a/Engine/Source/Programs/UnrealBuildTool/ProjectFiles/VisualStudio/VCProjectFileGenerator.cs b/Engine/Source/Programs/UnrealBuildTool/ProjectFiles/VisualStudio/VCProjectFileGenerator.cs
--- a/Engine/Source/Programs/UnrealBuildTool/ProjectFiles/VisualStudio/VCProjectFileGenerator.cs
+++ b/Engine/Source/Programs/UnrealBuildTool/ProjectFiles/VisualStudio/VCProjectFileGenerator.cs
@@ -81,6 +81,12 @@ namespace UnrealBuildTool
 		[XmlConfigFile(Category = "BuildConfiguration")]
 		bool bEditorDependsOnShaderCompileWorker = true;
 
+		/// <summary>
+		/// Whether to include a dependency on LiveCodingConsole when building targets that support live coding.
+		/// </summary>
+		[XmlConfigFile]
+		bool bBuildLiveCodingConsole = false;
+
 		/// <summary>
 		/// Override for the build tool to use in generated projects. If the compiler version is specified on the command line, we use the same argument on the 
 		/// command line for generated projects.
@@ -179,7 +185,7 @@ namespace UnrealBuildTool
 		/// <returns>The newly allocated project file object</returns>
 		protected override ProjectFile AllocateProjectFile(FileReference InitFilePath)
 		{
-			return new VCProjectFile(InitFilePath, OnlyGameProject, ProjectFileFormat, bAddFastPDBToProjects, bUsePerFileIntellisense, bUsePrecompiled, bEditorDependsOnShaderCompileWorker, BuildToolOverride);
+			return new VCProjectFile(InitFilePath, OnlyGameProject, ProjectFileFormat, bAddFastPDBToProjects, bUsePerFileIntellisense, bUsePrecompiled, bEditorDependsOnShaderCompileWorker, bBuildLiveCodingConsole, BuildToolOverride);
 		}
 
 
