LiveCoding: Add a target setting for using the debug version of the console application.

#rnx

diff --git a/Engine/Source/Developer/Windows/LiveCoding/LiveCoding.Build.cs b/Engine/Source/Developer/Windows/LiveCoding/LiveCoding.Build.cs
--- a/Engine/Source/Developer/Windows/LiveCoding/LiveCoding.Build.cs
+++ b/Engine/Source/Developer/Windows/LiveCoding/LiveCoding.Build.cs
@@ -11,6 +11,15 @@ public class LiveCoding : ModuleRules
 		PrivateDependencyModuleNames.Add("CoreUObject");
 		PrivateDependencyModuleNames.Add("Settings");
 
+		if(Target.bUseDebugLiveCodingConsole)
+        {
+            PrivateDefinitions.Add("USE_DEBUG_LIVE_CODING_CONSOLE=1");
+        }
+		else
+        {
+            PrivateDefinitions.Add("USE_DEBUG_LIVE_CODING_CONSOLE=0");
+        }
+
 		PublicAdditionalLibraries.Add(Path.Combine(ModuleDirectory, "Private", "External", "LC_JumpToSelf.lib"));
 	}
 }
diff --git a/Engine/Source/Developer/Windows/LiveCoding/Private/LiveCodingModule.cpp b/Engine/Source/Developer/Windows/LiveCoding/Private/LiveCodingModule.cpp
--- a/Engine/Source/Developer/Windows/LiveCoding/Private/LiveCodingModule.cpp
+++ b/Engine/Source/Developer/Windows/LiveCoding/Private/LiveCodingModule.cpp
@@ -46,9 +46,15 @@ void FLiveCodingModule::StartupModule()
 		ECVF_Cheat
 	);
 
+#if USE_DEBUG_LIVE_CODING_CONSOLE
+	static const TCHAR* DefaultConsolePath = TEXT("Binaries/Win64/LiveCodingConsole-Win64-Debug.exe");
+#else
+	static const TCHAR* DefaultConsolePath = TEXT("Binaries/Win64/LiveCodingConsole.exe");
+#endif 
+
 	ConsolePathVariable = ConsoleManager.RegisterConsoleVariable(
 		TEXT("LiveCoding.ConsolePath"),
-		FPaths::ConvertRelativePathToFull(FPaths::EngineDir() / TEXT("Binaries/Win64/LiveCodingConsole.exe")),
+		FPaths::ConvertRelativePathToFull(FPaths::EngineDir() / DefaultConsolePath),
 		TEXT("Path to the live coding console application"),
 		ECVF_Cheat
 	);
diff --git a/Engine/Source/Programs/UnrealBuildTool/Configuration/TargetRules.cs b/Engine/Source/Programs/UnrealBuildTool/Configuration/TargetRules.cs
--- a/Engine/Source/Programs/UnrealBuildTool/Configuration/TargetRules.cs
+++ b/Engine/Source/Programs/UnrealBuildTool/Configuration/TargetRules.cs
@@ -537,9 +537,16 @@ namespace UnrealBuildTool
 		[RequiresUniqueBuildEnvironment]
 		public bool bWithLiveCoding = false;
 
-        /// <summary>
-        /// Whether to turn on logging for test/shipping builds.
-        /// </summary>
+		/// <summary>
+		/// Whether to enable support for live coding
+		/// </summary>
+		[RequiresUniqueBuildEnvironment]
+		[XmlConfigFile(Category = "BuildConfiguration")]
+		public bool bUseDebugLiveCodingConsole = false;
+
+		/// <summary>
+		/// Whether to turn on logging for test/shipping builds.
+		/// </summary>
 		[RequiresUniqueBuildEnvironment]
         public bool bUseLoggingInShipping = false;
 
@@ -1870,7 +1877,12 @@ namespace UnrealBuildTool
 			get { return Inner.bWithLiveCoding; }
 		}
 
-        public bool bUseLoggingInShipping
+		public bool bUseDebugLiveCodingConsole
+		{
+			get { return Inner.bUseDebugLiveCodingConsole; }
+		}
+
+		public bool bUseLoggingInShipping
 		{
 			get { return Inner.bUseLoggingInShipping; }
 		}
diff --git a/Engine/Source/Programs/UnrealBuildTool/ProjectFiles/VisualStudio/VCProject.cs b/Engine/Source/Programs/UnrealBuildTool/ProjectFiles/VisualStudio/VCProject.cs
--- a/Engine/Source/Programs/UnrealBuildTool/ProjectFiles/VisualStudio/VCProject.cs
+++ b/Engine/Source/Programs/UnrealBuildTool/ProjectFiles/VisualStudio/VCProject.cs
@@ -1406,7 +1406,7 @@ namespace UnrealBuildTool
 						}
 						if (TargetRulesObject.bWithLiveCoding && bBuildLiveCodingConsole && !UnrealBuildTool.IsEngineInstalled())
 						{
-							ExtraTargets.Add("LiveCodingConsole Win64 Development");
+							ExtraTargets.Add(TargetRulesObject.bUseDebugLiveCodingConsole? "LiveCodingConsole Win64 Debug" : "LiveCodingConsole Win64 Development");
 						}
 					}
 
