PR #5825: Fix for regional separators in version parsing (Contributed by [REDACTED_PERSON])

#rnx

diff --git a/Engine/Source/Programs/UnrealBuildTool/System/VersionNumber.cs b/Engine/Source/Programs/UnrealBuildTool/System/VersionNumber.cs
--- a/Engine/Source/Programs/UnrealBuildTool/System/VersionNumber.cs
+++ b/Engine/Source/Programs/UnrealBuildTool/System/VersionNumber.cs
@@ -14,6 +14,11 @@ namespace UnrealBuildTool
 	[Serializable]
 	class VersionNumber : IComparable<VersionNumber>
 	{
+		/// <summary>
+		/// Set of delimiters for version numbers
+		/// </summary>
+		static readonly char[] Delimiters = { '.', ',' };
+
 		/// <summary>
 		/// The individual version components
 		/// </summary>
@@ -193,7 +198,7 @@ namespace UnrealBuildTool
 		public static VersionNumber Parse(string Text)
 		{
 			List<int> Components = new List<int>();
-			foreach(string TextElement in Text.Split('.'))
+			foreach(string TextElement in Text.Split(Delimiters))
 			{
 				Components.Add(int.Parse(TextElement));
 			}
@@ -209,7 +214,7 @@ namespace UnrealBuildTool
 		public static bool TryParse(string Text, out VersionNumber OutNumber)
 		{
 			List<int> Components = new List<int>();
-			foreach(string TextElement in Text.Split('.'))
+			foreach(string TextElement in Text.Split(Delimiters))
 			{
 				int Component;
 				if(!int.TryParse(TextElement, out Component))
