Live Coding: Disable for Win32 due to compile error.

#rnx

diff --git a/Engine/Source/Programs/UnrealBuildTool/Platform/Windows/UEBuildWindows.cs b/Engine/Source/Programs/UnrealBuildTool/Platform/Windows/UEBuildWindows.cs
--- a/Engine/Source/Programs/UnrealBuildTool/Platform/Windows/UEBuildWindows.cs
+++ b/Engine/Source/Programs/UnrealBuildTool/Platform/Windows/UEBuildWindows.cs
@@ -534,7 +534,7 @@ namespace UnrealBuildTool
 		{
 			base.ResetTarget(Target);
 
-			if(Target.Configuration != UnrealTargetConfiguration.Shipping)
+			if(Target.Platform == UnrealTargetPlatform.Win64 && Target.Configuration != UnrealTargetConfiguration.Shipping)
 			{
 				Target.bWithLiveCoding = true;
 				Target.WindowsPlatform.bCreateHotPatchableImage = true;
