UBT: Attempt to reduce the command line length to fix CIS errors on Linux.

diff --git a/Engine/Source/Programs/UnrealBuildTool/Platform/Linux/LinuxToolChain.cs b/Engine/Source/Programs/UnrealBuildTool/Platform/Linux/LinuxToolChain.cs
--- a/Engine/Source/Programs/UnrealBuildTool/Platform/Linux/LinuxToolChain.cs
+++ b/Engine/Source/Programs/UnrealBuildTool/Platform/Linux/LinuxToolChain.cs
@@ -1210,7 +1210,16 @@ namespace UnrealBuildTool
 			// Add include paths to the argument list.
 			foreach (DirectoryReference IncludePath in CompileEnvironment.UserIncludePaths)
 			{
-				Arguments += string.Format(" -I\"{0}\"", IncludePath.FullName.Replace('\\', '/'));
+				string IncludePathString;
+				if (IncludePath.IsUnderDirectory(UnrealBuildTool.RootDirectory))
+				{
+					IncludePathString = IncludePath.MakeRelativeTo(UnrealBuildTool.EngineSourceDirectory);
+				}
+				else
+				{
+					IncludePathString = IncludePath.FullName;
+				}
+				Arguments += string.Format(" -I\"{0}\"", IncludePathString.Replace('\\', '/'));
 			}
 			foreach (DirectoryReference IncludePath in CompileEnvironment.SystemIncludePaths)
 			{
