PlatformExtensions:
 - extensions are included in the solution
 - extension automation projects are included under Automation, and built as part of AutomationTool

diff --git a/Engine/Source/Programs/UnrealBuildTool/ProjectFiles/ProjectFileGenerator.cs b/Engine/Source/Programs/UnrealBuildTool/ProjectFiles/ProjectFileGenerator.cs
--- a/Engine/Source/Programs/UnrealBuildTool/ProjectFiles/ProjectFileGenerator.cs
+++ b/Engine/Source/Programs/UnrealBuildTool/ProjectFiles/ProjectFileGenerator.cs
@@ -1889,11 +1889,17 @@ namespace UnrealBuildTool
 				bool IsThirdPartyModule = CurModuleFile.ContainsName("ThirdParty", UnrealBuildTool.RootDirectory);
 
 				// check for engine, or platform extension engine folders
-				if( !bIncludeEngineSource && (CurModuleFile.IsUnderDirectory(UnrealBuildTool.EngineDirectory) ||
-					CurModuleFile.ContainsName("Engine", UnrealBuildTool.PlatformExtensionsDirectory)) )
+				if( !bIncludeEngineSource )
 				{
-					// We were asked to exclude engine modules from the generated projects
-					WantProjectFileForModule = false;
+					foreach (DirectoryReference EngineDirectory in UnrealBuildTool.GetAllEngineDirectories())
+					{
+						if (CurModuleFile.IsUnderDirectory(EngineDirectory))
+						{
+							// We were asked to exclude engine modules from the generated projects
+							WantProjectFileForModule = false;
+							break;
+						}
+					}
 				}
 
 				if( CurModuleFile.IsUnderDirectory(UnrealBuildTool.EnterpriseDirectory) && !bIncludeEnterpriseSource )
@@ -1996,12 +2002,16 @@ namespace UnrealBuildTool
 			string PossibleProgramTargetName = CurModuleFile.GetFileNameWithoutAnyExtensions();
 
 			// check for engine, or platform extension engine folders
-			if( CurModuleFile.IsUnderDirectory(UnrealBuildTool.EngineDirectory) ||
-				CurModuleFile.ContainsName("Engine", UnrealBuildTool.PlatformExtensionsDirectory) )
+			if( CurModuleFile.IsUnderDirectory(UnrealBuildTool.EngineDirectory))
 			{
 				ProjectFileNameBase = EngineProjectFileNameBase;
 				BaseFolder = UnrealBuildTool.EngineDirectory;
 			}
+			else if (CurModuleFile.IsUnderDirectory(UnrealBuildTool.PlatformExtensionsDirectory))
+			{
+				ProjectFileNameBase = EngineProjectFileNameBase;
+				BaseFolder = UnrealBuildTool.RootDirectory;
+			}
 			else if( CurModuleFile.IsUnderDirectory(UnrealBuildTool.EnterpriseSourceDirectory) ||
 				CurModuleFile.IsUnderDirectory(DirectoryReference.Combine(UnrealBuildTool.EnterpriseDirectory, "Plugins")) )
 			{
@@ -2073,16 +2083,16 @@ namespace UnrealBuildTool
 				bool IsEngineTarget = false;
 				bool IsEnterpriseTarget = false;
 				bool WantProjectFileForTarget = true;
-				if(TargetFilePath.IsUnderDirectory(UnrealBuildTool.EngineDirectory))
+				if(UnrealBuildTool.GetAllEngineDirectories().Any(x => TargetFilePath.IsUnderDirectory(x)))
 				{
 					// This is an engine target
 					IsEngineTarget = true;
 
-					if(TargetFilePath.IsUnderDirectory(EngineSourceProgramsDirectory))
+					if(UnrealBuildTool.GetAllEngineDirectories("Source/Programs").Any(x => TargetFilePath.IsUnderDirectory(x)))
 					{
 						WantProjectFileForTarget = IncludeEnginePrograms;
 					}
-					else if(TargetFilePath.IsUnderDirectory(UnrealBuildTool.EngineSourceDirectory))
+					else if(UnrealBuildTool.GetAllEngineDirectories("Source").Any(x => TargetFilePath.IsUnderDirectory(x)))
 					{
 						WantProjectFileForTarget = bIncludeEngineSource;
 					}
diff --git a/Engine/Source/Programs/UnrealBuildTool/System/RulesCompiler.cs b/Engine/Source/Programs/UnrealBuildTool/System/RulesCompiler.cs
--- a/Engine/Source/Programs/UnrealBuildTool/System/RulesCompiler.cs
+++ b/Engine/Source/Programs/UnrealBuildTool/System/RulesCompiler.cs
@@ -67,8 +67,9 @@ namespace UnrealBuildTool
 		/// <param name="bIncludeEngine"></param>
 		/// <param name="bIncludeEnterprise"></param>
 		/// <param name="bIncludeTempTargets">Whether to include targets generated by UAT to accomodate content-only projects that need to be compiled to include plugins</param>
+		/// <param name="bIncludePlatformExtensions"></param>
 		/// <returns></returns>
-		public static List<FileReference> FindAllRulesSourceFiles(RulesFileType RulesFileType, List<DirectoryReference> GameFolders, List<FileReference> ForeignPlugins, List<DirectoryReference> AdditionalSearchPaths, bool bIncludeEngine = true, bool bIncludeEnterprise = true, bool bIncludeTempTargets = true)
+		public static List<FileReference> FindAllRulesSourceFiles(RulesFileType RulesFileType, List<DirectoryReference> GameFolders, List<FileReference> ForeignPlugins, List<DirectoryReference> AdditionalSearchPaths, bool bIncludeEngine = true, bool bIncludeEnterprise = true, bool bIncludeTempTargets = true, bool bIncludePlatformExtensions = true)
 		{
 			List<DirectoryReference> Folders = new List<DirectoryReference>();
 
@@ -81,6 +82,10 @@ namespace UnrealBuildTool
 			{
 				Folders.Add(UnrealBuildTool.EnterpriseSourceDirectory);
 			}
+			if (bIncludePlatformExtensions)
+			{
+				Folders.Add(UnrealBuildTool.PlatformExtensionsDirectory);
+			}
 
 			// @todo plugin: Disallow modules from including plugin modules as dependency modules? (except when the module is part of that plugin)
 
