PlatformExtensions: adding a helper to UnrealTargetPlatform to check for group membership.

diff --git a/Engine/Source/Programs/UnrealBuildTool/Configuration/UEBuildTarget.cs b/Engine/Source/Programs/UnrealBuildTool/Configuration/UEBuildTarget.cs
--- a/Engine/Source/Programs/UnrealBuildTool/Configuration/UEBuildTarget.cs
+++ b/Engine/Source/Programs/UnrealBuildTool/Configuration/UEBuildTarget.cs
@@ -180,6 +180,15 @@ namespace UnrealBuildTool
 			return StringRegistry.HasString(Name);
 		}
 
+		/// <summary>
+		/// Helper that just calls UEBuildPlatform.IsPlatformInGroup
+		/// </summary>
+		/// <param name="Group"></param>
+		/// <returns></returns>
+		public bool IsInGroup(UnrealPlatformGroup Group)
+		{
+			return UEBuildPlatform.IsPlatformInGroup(this, Group);
+		}
 
 		/// <summary>
 		/// 32-bit Windows
