PR #5893: 4.22 sndbs fix (Contributed by [REDACTED_PERSON])

diff --git a/Engine/Source/Programs/UnrealBuildTool/Executors/SNDBS.cs b/Engine/Source/Programs/UnrealBuildTool/Executors/SNDBS.cs
--- a/Engine/Source/Programs/UnrealBuildTool/Executors/SNDBS.cs
+++ b/Engine/Source/Programs/UnrealBuildTool/Executors/SNDBS.cs
@@ -268,6 +268,12 @@ namespace UnrealBuildTool
 							JobNumber++;
 							NumScriptedActions++;
 							PrintDebugInfo |= Action.bPrintDebugInfo;
+
+							if (Action.DependencyListFile != null && File.Exists(Action.DependencyListFile.AbsolutePath))
+							{
+								Log.TraceVerbose("Deleting dependency list file {0}", Action.DependencyListFile.AbsolutePath);
+								File.Delete(Action.DependencyListFile.AbsolutePath);
+							}
 						}
 					}
 				}
