- Fixed LTO dSYM generation (both by saving the .lto.o file instead of using /tmp/lto.o, which fails), as well as use -j 1 to limit dsymutil threads that would crash internally in weird ways
- Generate Crashlytics info in dylib builds

diff --git a/Engine/Source/Programs/UnrealBuildTool/Platform/IOS/IOSToolChain.cs b/Engine/Source/Programs/UnrealBuildTool/Platform/IOS/IOSToolChain.cs
--- a/Engine/Source/Programs/UnrealBuildTool/Platform/IOS/IOSToolChain.cs
+++ b/Engine/Source/Programs/UnrealBuildTool/Platform/IOS/IOSToolChain.cs
@@ -181,7 +181,7 @@ namespace UnrealBuildTool
 					BuildProducts.Add(AssetFile, BuildProductType.RequiredResource);
 				}
 			}
-            if (Target.IOSPlatform.bGeneratedSYM && (ProjectSettings.bGenerateCrashReportSymbols || Target.bUseMallocProfiler) && Binary.Type == UEBuildBinaryType.Executable)
+            if (Target.IOSPlatform.bGeneratedSYM && (ProjectSettings.bGenerateCrashReportSymbols || Target.bUseMallocProfiler) && Binary.Type == UEBuildBinaryType.StaticLibrary)
             {
                 FileReference DebugFile = FileReference.Combine(Binary.OutputFilePath.Directory, Binary.OutputFilePath.GetFileNameWithoutExtension() + ".udebugsymbols");
                 BuildProducts.Add(DebugFile, BuildProductType.SymbolFile);
@@ -873,6 +873,22 @@ namespace UnrealBuildTool
 				LinkCommandArguments += string.Format(" @\"{0}\"", ResponsePath.FullName);
 			}
 
+			// if we are making an LTO build, write the lto file next to build so dsymutil can find it
+			if (LinkEnvironment.bAllowLTCG)
+			{
+				string LtoObjectFile;
+				if (Target.bShouldCompileAsDLL)
+				{
+					// go up a directory so we don't put this big file into the framework
+					LtoObjectFile = Path.Combine(Path.GetDirectoryName(Path.GetDirectoryName(OutputFile.AbsolutePath)), Path.GetFileName(OutputFile.AbsolutePath) + ".lto.o");
+				}
+				else
+				{
+					LtoObjectFile = OutputFile.AbsolutePath + ".lto.o";
+				}
+				LinkCommandArguments += string.Format(" -flto -Xlinker -object_path_lto -Xlinker \"{0}\"", LtoObjectFile);
+			}
+
 			// Add the output file to the command-line.
 			LinkCommandArguments += string.Format(" -o \"{0}\"", OutputFile.AbsolutePath);
 
@@ -1010,9 +1026,10 @@ namespace UnrealBuildTool
 			GenDebugAction.WorkingDirectory = GetMacDevSrcRoot();
 			GenDebugAction.CommandPath = BuildHostPlatform.Current.Shell;
 			string DsymutilPath = GetDsymutilPath();
+			// using -j 1 because there are some issues, at least with LTO builds, where dsymutil crashes, and -j 1 [1 thread] makes them stop. does not seem to be particularly slower. -j 2 works, but is the same speed
 			if(ProjectSettings.bGeneratedSYMBundle)
 			{
-				GenDebugAction.CommandArguments = string.Format("-c 'rm -rf \"{2}\"; \"{0}\" \"{1}\" -o \"{2}\"; cd \"{2}/..\"; zip -r -y -1 {3}.zip {3}'",
+				GenDebugAction.CommandArguments = string.Format("-c 'rm -rf \"{2}\"; \"{0}\" -j 1 \"{1}\" -o \"{2}\"; cd \"{2}/..\"; zip -r -y -1 {3}.zip {3}'",
                     DsymutilPath,
 					Executable.AbsolutePath,
                     OutputFile.AbsolutePath,
@@ -1022,7 +1039,7 @@ namespace UnrealBuildTool
             }
             else
 			{
-				GenDebugAction.CommandArguments = string.Format("-c 'rm -rf \"{2}\"; \"{0}\" \"{1}\" -f -o \"{2}\"'",
+				GenDebugAction.CommandArguments = string.Format("-c 'rm -rf \"{2}\"; \"{0}\" -j 1 \"{1}\" -f -o \"{2}\"'",
 						DsymutilPath,
 						Executable.AbsolutePath,
 						OutputFile.AbsolutePath);
@@ -1624,13 +1641,14 @@ namespace UnrealBuildTool
 
 	
         public static void PostBuildSync(IOSPostBuildSyncTarget Target)
-			{
+        {
 			ConfigHierarchy Ini = ConfigCache.ReadHierarchy(ConfigHierarchyType.Engine, Target.ProjectFile == null ? DirectoryReference.FromString(UnrealBuildTool.GetRemoteIniPath()): DirectoryReference.FromFile(Target.ProjectFile), UnrealTargetPlatform.IOS);
 			string BundleID;
 			Ini.GetString("/Script/IOSRuntimeSettings.IOSRuntimeSettings", "BundleIdentifier", out BundleID);
 
 			IOSProjectSettings ProjectSettings = ((IOSPlatform)UEBuildPlatform.GetBuildPlatform(Target.Platform)).ReadProjectSettings(Target.ProjectFile);
 
+			bool bPerformFullAppCreation = true;
 			if (BuildHostPlatform.Current.Platform == UnrealTargetPlatform.Mac)
 			{
 				if (IsCompiledAsFramework(Target.OutputPath.FullName))
@@ -1650,25 +1668,30 @@ namespace UnrealBuildTool
 
 					FileReference.Copy(PlistSrcLocation, PlistDstLocation, true);
 
-					// and do nothing else
-					return;
+					// do not perform any of the .app creation below
+					bPerformFullAppCreation = false;
 				}
 			}
 
 			string AppName = Target.TargetName;
-
 			string RemoteShadowDirectoryMac = Target.OutputPath.Directory.FullName;
-			FileReference StagedExecutablePath = GetStagedExecutablePath(Target.OutputPath, Target.TargetName);
-
-			// copy the executable
-			DirectoryReference.CreateDirectory(StagedExecutablePath.Directory);
-			FileReference.Copy(Target.OutputPath, StagedExecutablePath, true);
 
 			if (!Target.bSkipCrashlytics)
 			{
 				GenerateCrashlyticsData(RemoteShadowDirectoryMac, Path.GetFileName(Target.OutputPath.FullName), Target.ProjectDirectory.FullName, AppName);
 			}
 
+			// only make the app if needed
+			if (!bPerformFullAppCreation)
+			{
+				return;
+			}
+			
+			// copy the executable
+			FileReference StagedExecutablePath = GetStagedExecutablePath(Target.OutputPath, Target.TargetName);
+			DirectoryReference.CreateDirectory(StagedExecutablePath.Directory);
+			FileReference.Copy(Target.OutputPath, StagedExecutablePath, true);
+
 			if (Target.bCreateStubIPA)
 			{
 				// generate the dummy project so signing works
