- LTO dsyms now work, at least with dsymutil 10.0.1, and -j is no longer passed in to dsymutil 9

diff --git a/Engine/Source/Programs/UnrealBuildTool/Platform/IOS/IOSToolChain.cs b/Engine/Source/Programs/UnrealBuildTool/Platform/IOS/IOSToolChain.cs
--- a/Engine/Source/Programs/UnrealBuildTool/Platform/IOS/IOSToolChain.cs
+++ b/Engine/Source/Programs/UnrealBuildTool/Platform/IOS/IOSToolChain.cs
@@ -1007,12 +1007,13 @@ namespace UnrealBuildTool
 			return Path.Combine(Path.GetDirectoryName(ExecutablePath), Path.GetFileName(ExecutablePath) + ".dSYM");
 		}
 
-        /// <summary>
-        /// Generates debug info for a given executable
-        /// </summary>
-        /// <param name="Executable">FileItem describing the executable to generate debug info for</param>
+		/// <summary>
+		/// Generates debug info for a given executable
+		/// </summary>
+		/// <param name="Executable">FileItem describing the executable to generate debug info for</param>
 		/// <param name="Actions">List of actions to be executed. Additional actions will be added to this list.</param>
-        public List<FileItem> GenerateDebugInfo(FileItem Executable, List<Action> Actions)
+		/// <param name="bIsForLTOBuild">Was this build made with LTO enabled?</param>
+		private List<FileItem> GenerateDebugInfo(FileItem Executable, List<Action> Actions, bool bIsForLTOBuild)
 		{
             // Make a file item for the source and destination files
 			string FullDestPathRoot = GetdSYMPath(Executable);
@@ -1025,23 +1026,26 @@ namespace UnrealBuildTool
 
 			GenDebugAction.WorkingDirectory = GetMacDevSrcRoot();
 			GenDebugAction.CommandPath = BuildHostPlatform.Current.Shell;
-			string DsymutilPath = GetDsymutilPath();
+			string ExtraOptions;
+			string DsymutilPath = GetDsymutilPath(out ExtraOptions, bIsForLTOBuild:bIsForLTOBuild);
 			if(ProjectSettings.bGeneratedSYMBundle)
 			{
-				GenDebugAction.CommandArguments = string.Format("-c 'rm -rf \"{2}\"; \"{0}\" \"{1}\" -o \"{2}\"; cd \"{2}/..\"; zip -r -y -1 {3}.zip {3}'",
+				GenDebugAction.CommandArguments = string.Format("-c 'rm -rf \"{2}\"; \"{0}\" \"{1}\" {4} -o \"{2}\"; cd \"{2}/..\"; zip -r -y -1 {3}.zip {3}'",
                     DsymutilPath,
 					Executable.AbsolutePath,
                     OutputFile.AbsolutePath,
-					Path.GetFileName(FullDestPathRoot));
+					Path.GetFileName(FullDestPathRoot),
+					ExtraOptions);
                 GenDebugAction.ProducedItems.Add(ZipOutputFile);
                 Log.TraceInformation("Zip file: " + ZipOutputFile.AbsolutePath);
             }
             else
 			{
-				GenDebugAction.CommandArguments = string.Format("-c 'rm -rf \"{2}\"; \"{0}\" \"{1}\" -f -o \"{2}\"'",
+				GenDebugAction.CommandArguments = string.Format("-c 'rm -rf \"{2}\"; \"{0}\" \"{1}\" {3} -f -o \"{2}\"'",
 						DsymutilPath,
 						Executable.AbsolutePath,
-						OutputFile.AbsolutePath);
+						OutputFile.AbsolutePath,
+						ExtraOptions);
 			}
 			GenDebugAction.PrerequisiteItems.Add(Executable);
             GenDebugAction.ProducedItems.Add(OutputFile);
@@ -1396,7 +1400,7 @@ namespace UnrealBuildTool
             // For IOS/tvOS, generate the dSYM file if needed or requested
 			if (Target.IOSPlatform.bGeneratedSYM)
             {
- 				List<FileItem> Files = GenerateDebugInfo(Executable, Actions);
+ 				List<FileItem> Files = GenerateDebugInfo(Executable, Actions, BinaryLinkEnvironment.bAllowLTCG);
 				foreach (FileItem item in Files)
 				{
 					OutputFiles.Add(item);
@@ -1418,11 +1422,11 @@ namespace UnrealBuildTool
 				Action StripAction = new Action(ActionType.CreateAppBundle);
 				StripAction.WorkingDirectory = GetMacDevSrcRoot();
 				StripAction.CommandPath = BuildHostPlatform.Current.Shell;
-				StripAction.CommandArguments = String.Format("-c '\"{0}strip\" {1} \"{2}\" && touch \"{3}\"'", Settings.Value.ToolchainDir, StripArguments, Executable.Location, StripCompleteFile);
+				StripAction.CommandArguments = String.Format("-c '\"{0}strip\" {1} \"{2}\" -o \"{3}\" && touch \"{3}\"'", Settings.Value.ToolchainDir, StripArguments, Executable.Location, StripCompleteFile);
 				StripAction.PrerequisiteItems.Add(Executable);
 				StripAction.PrerequisiteItems.AddRange(OutputFiles);
 				StripAction.ProducedItems.Add(StripCompleteFile);
-				StripAction.StatusDescription = String.Format("Stripping symbols from {0}", Executable.AbsolutePath);
+				StripAction.StatusDescription = StripAction.CommandArguments; //String.Format("Stripping symbols from {0}", Executable.AbsolutePath);
 				StripAction.bCanExecuteRemotely = false;
 				Actions.Add(StripAction);
 
diff --git a/Engine/Source/Programs/UnrealBuildTool/Platform/Mac/MacToolChain.cs b/Engine/Source/Programs/UnrealBuildTool/Platform/Mac/MacToolChain.cs
--- a/Engine/Source/Programs/UnrealBuildTool/Platform/Mac/MacToolChain.cs
+++ b/Engine/Source/Programs/UnrealBuildTool/Platform/Mac/MacToolChain.cs
@@ -1111,10 +1111,13 @@ namespace UnrealBuildTool
 			// Deletes ay existing file on the building machine. Also, waits 30 seconds, if needed, for the input file to be created in an attempt to work around
 			// a problem where dsymutil would exit with an error saying the input file did not exist.
 			// Note that the source and dest are switched from a copy command
-			GenDebugAction.CommandArguments = string.Format("-c 'rm -rf \"{2}\"; for i in {{1..30}}; do if [ -f \"{1}\" ] ; then break; else echo\"Waiting for {1} before generating dSYM file.\"; sleep 1; fi; done; \"{0}\" -f \"{1}\" -o \"{2}\"'",
-                GetDsymutilPath(),
+			string ExtraOptions;
+			string DsymutilPath = GetDsymutilPath(out ExtraOptions, bIsForLTOBuild: false);
+			GenDebugAction.CommandArguments = string.Format("-c 'rm -rf \"{2}\"; for i in {{1..30}}; do if [ -f \"{1}\" ] ; then break; else echo\"Waiting for {1} before generating dSYM file.\"; sleep 1; fi; done; \"{0}\" {3} -f \"{1}\" -o \"{2}\"'",
+				DsymutilPath,
 				MachOBinary.AbsolutePath,
-				OutputFile.AbsolutePath);
+				OutputFile.AbsolutePath,
+				ExtraOptions);
 			if (LinkEnvironment.bIsCrossReferenced)
 			{
 				GenDebugAction.PrerequisiteItems.Add(FixDylibOutputFile);
diff --git a/Engine/Source/Programs/UnrealBuildTool/ToolChain/AppleToolChain.cs b/Engine/Source/Programs/UnrealBuildTool/ToolChain/AppleToolChain.cs
--- a/Engine/Source/Programs/UnrealBuildTool/ToolChain/AppleToolChain.cs
+++ b/Engine/Source/Programs/UnrealBuildTool/ToolChain/AppleToolChain.cs
@@ -160,7 +160,7 @@ namespace UnrealBuildTool
 			Utils.RunLocalProcessAndLogOutput(StartInfo);
 		}
 		
-		protected string GetDsymutilPath()
+		protected string GetDsymutilPath(out string ExtraOptions, bool bIsForLTOBuild=false)
 		{
 			FileReference DsymutilLocation = new FileReference("/usr/bin/dsymutil");
 
@@ -169,6 +169,7 @@ namespace UnrealBuildTool
 			string DsymutilVersionString = Utils.RunLocalProcessAndReturnStdOut(DsymutilLocation.FullName, "-version");
 
 			bool bUseInstalledDsymutil = true;
+			int Major = 0, Minor = 0, Patch = 0;
 
 			// tease out the version number
 			string[] Tokens = DsymutilVersionString.Split(" ".ToCharArray());
@@ -189,7 +190,6 @@ namespace UnrealBuildTool
 				}
 				else
 				{
-					int Major, Minor, Patch;
 					if (!int.TryParse(Versions[0], out Major) || !int.TryParse(Versions[1], out Minor) || !int.TryParse(Versions[2], out Patch))
 					{
 						Log.TraceInformationOnce("Unable to parse version tokens: {0}", Tokens[3]);
@@ -227,6 +227,8 @@ namespace UnrealBuildTool
 				}
 			}
 
+			// 10.0.1 has an issue with LTO builds where we need to limit the number of threads
+			ExtraOptions = (bIsForLTOBuild && Major == 10 && Minor == 0 && Patch == 1) ? "-j 1" : "";
 			return DsymutilLocation.FullName;
 		}
 	};
