- Fixed up change to stripping that did not do it in place

diff --git a/Engine/Source/Programs/UnrealBuildTool/Platform/IOS/IOSToolChain.cs b/Engine/Source/Programs/UnrealBuildTool/Platform/IOS/IOSToolChain.cs
--- a/Engine/Source/Programs/UnrealBuildTool/Platform/IOS/IOSToolChain.cs
+++ b/Engine/Source/Programs/UnrealBuildTool/Platform/IOS/IOSToolChain.cs
@@ -1422,11 +1422,11 @@ namespace UnrealBuildTool
 				Action StripAction = new Action(ActionType.CreateAppBundle);
 				StripAction.WorkingDirectory = GetMacDevSrcRoot();
 				StripAction.CommandPath = BuildHostPlatform.Current.Shell;
-				StripAction.CommandArguments = String.Format("-c '\"{0}strip\" {1} \"{2}\" -o \"{3}\" && touch \"{3}\"'", Settings.Value.ToolchainDir, StripArguments, Executable.Location, StripCompleteFile);
+				StripAction.CommandArguments = String.Format("-c '\"{0}strip\" {1} \"{2}\" && touch \"{3}\"'", Settings.Value.ToolchainDir, StripArguments, Executable.Location, StripCompleteFile);
 				StripAction.PrerequisiteItems.Add(Executable);
 				StripAction.PrerequisiteItems.AddRange(OutputFiles);
 				StripAction.ProducedItems.Add(StripCompleteFile);
-				StripAction.StatusDescription = StripAction.CommandArguments; //String.Format("Stripping symbols from {0}", Executable.AbsolutePath);
+				StripAction.StatusDescription = String.Format("Stripping symbols from {0}", Executable.AbsolutePath);
 				StripAction.bCanExecuteRemotely = false;
 				Actions.Add(StripAction);
 
