- Emitting dylib to the receipt as debug info

diff --git a/Engine/Source/Programs/UnrealBuildTool/Platform/IOS/IOSToolChain.cs b/Engine/Source/Programs/UnrealBuildTool/Platform/IOS/IOSToolChain.cs
--- a/Engine/Source/Programs/UnrealBuildTool/Platform/IOS/IOSToolChain.cs
+++ b/Engine/Source/Programs/UnrealBuildTool/Platform/IOS/IOSToolChain.cs
@@ -195,7 +195,7 @@ namespace UnrealBuildTool
 		/// <param name="OutputType">Type of build product</param>
 		public override bool ShouldAddDebugFileToReceipt(FileReference OutputFile, BuildProductType OutputType)
 		{
-			return OutputType == BuildProductType.Executable;
+			return OutputType == BuildProductType.Executable || OutputType == BuildProductType.DynamicLibrary;
 		}
 
 		string GetCompileArguments_Global(CppCompileEnvironment CompileEnvironment)
