- Fixed path for dSYM.zip for DLL compiles

diff --git a/Engine/Source/Programs/UnrealBuildTool/Platform/IOS/IOSToolChain.cs b/Engine/Source/Programs/UnrealBuildTool/Platform/IOS/IOSToolChain.cs
--- a/Engine/Source/Programs/UnrealBuildTool/Platform/IOS/IOSToolChain.cs
+++ b/Engine/Source/Programs/UnrealBuildTool/Platform/IOS/IOSToolChain.cs
@@ -1566,7 +1566,7 @@ namespace UnrealBuildTool
 			}
 		}
 
-		private static void GenerateCrashlyticsData(string ExecutableDirectory, string ExecutableName, string ProjectDir, string ProjectName)
+		private static void GenerateCrashlyticsData(string DsymZip, string ProjectDir, string ProjectName)
         {
 			Log.TraceInformation("Generating and uploading Crashlytics Data");
             string FabricPath = UnrealBuildTool.EngineDirectory + "/Intermediate/UnzippedFrameworks/Crashlytics/Fabric.embeddedframework";
@@ -1574,11 +1574,11 @@ namespace UnrealBuildTool
             {
 				string PlistFile = ProjectDir + "/Intermediate/IOS/" + ProjectName + "-Info.plist";
                 Process FabricProcess = new Process();
-                FabricProcess.StartInfo.WorkingDirectory = ExecutableDirectory;
+                FabricProcess.StartInfo.WorkingDirectory = Path.GetDirectoryName(DsymZip);
                 FabricProcess.StartInfo.FileName = "/bin/sh";
 				FabricProcess.StartInfo.Arguments = string.Format("-c 'chmod 777 \"{0}/Fabric.framework/upload-symbols\"; \"{0}/Fabric.framework/upload-symbols\" -a 7a4cebd0324af21696e5e321802c5e26ba541cad -p ios {1}'",
                     FabricPath,
-					ExecutableDirectory + "/" + ExecutableName + ".dSYM.zip");
+					DsymZip);
 
 				ProcessOutput Output = new ProcessOutput();
 
@@ -1652,6 +1652,7 @@ namespace UnrealBuildTool
 			IOSProjectSettings ProjectSettings = ((IOSPlatform)UEBuildPlatform.GetBuildPlatform(Target.Platform)).ReadProjectSettings(Target.ProjectFile);
 
 			bool bPerformFullAppCreation = true;
+			string PathToDsymZip = Target.OutputPath.FullName + ".dSYM.zip";
 			if (BuildHostPlatform.Current.Platform == UnrealTargetPlatform.Mac)
 			{
 				if (IsCompiledAsFramework(Target.OutputPath.FullName))
@@ -1673,15 +1674,15 @@ namespace UnrealBuildTool
 
 					// do not perform any of the .app creation below
 					bPerformFullAppCreation = false;
+					PathToDsymZip = Path.Combine(Target.OutputPath.Directory.ParentDirectory.FullName, Target.OutputPath.GetFileName() + ".dSYM.zip");
 				}
 			}
 
 			string AppName = Target.TargetName;
-			string RemoteShadowDirectoryMac = Target.OutputPath.Directory.FullName;
 
 			if (!Target.bSkipCrashlytics)
 			{
-				GenerateCrashlyticsData(RemoteShadowDirectoryMac, Path.GetFileName(Target.OutputPath.FullName), Target.ProjectDirectory.FullName, AppName);
+				GenerateCrashlyticsData(PathToDsymZip, Target.ProjectDirectory.FullName, AppName);
 			}
 
 			// only make the app if needed
@@ -1694,6 +1695,7 @@ namespace UnrealBuildTool
 			FileReference StagedExecutablePath = GetStagedExecutablePath(Target.OutputPath, Target.TargetName);
 			DirectoryReference.CreateDirectory(StagedExecutablePath.Directory);
 			FileReference.Copy(Target.OutputPath, StagedExecutablePath, true);
+			string RemoteShadowDirectoryMac = Target.OutputPath.Directory.FullName;
 
 			if (Target.bCreateStubIPA)
 			{
