Fix for a regression that caused nativized Blueprint struct types to fail UHT property validation at build time.

#rnx

diff --git a/Engine/Source/Programs/UnrealBuildTool/System/RulesCompiler.cs b/Engine/Source/Programs/UnrealBuildTool/System/RulesCompiler.cs
--- a/Engine/Source/Programs/UnrealBuildTool/System/RulesCompiler.cs
+++ b/Engine/Source/Programs/UnrealBuildTool/System/RulesCompiler.cs
@@ -616,6 +616,7 @@ namespace UnrealBuildTool
 
 				// Find all the modules
 				ModuleRulesContext PluginModuleContext = new ModuleRulesContext(Scope, PluginFileName.Directory);
+				PluginModuleContext.bClassifyAsGameModuleForUHT = !bContainsEngineModules;
 				FindModuleRulesForPlugins(ForeignPlugins, PluginModuleContext, ModuleFiles);
 
 				// Compile the assembly
