Proper fix for dsym issues in frameworks (integrated from 9.20)

diff --git a/Engine/Source/Programs/UnrealBuildTool/Platform/IOS/IOSToolChain.cs b/Engine/Source/Programs/UnrealBuildTool/Platform/IOS/IOSToolChain.cs
--- a/Engine/Source/Programs/UnrealBuildTool/Platform/IOS/IOSToolChain.cs
+++ b/Engine/Source/Programs/UnrealBuildTool/Platform/IOS/IOSToolChain.cs
@@ -1680,7 +1680,7 @@ namespace UnrealBuildTool
 
 			string AppName = Target.TargetName;
 
-			if (!Target.bSkipCrashlytics && !IsCompiledAsFramework(Target.OutputPath.FullName))
+			if (!Target.bSkipCrashlytics)
 			{
 				GenerateCrashlyticsData(PathToDsymZip, Target.ProjectDirectory.FullName, AppName);
 			}
