- Taking dylib exe data back out of receipt as it's putting the wrong debug info into the receipt

diff --git a/Engine/Source/Programs/UnrealBuildTool/Platform/IOS/IOSToolChain.cs b/Engine/Source/Programs/UnrealBuildTool/Platform/IOS/IOSToolChain.cs
--- a/Engine/Source/Programs/UnrealBuildTool/Platform/IOS/IOSToolChain.cs
+++ b/Engine/Source/Programs/UnrealBuildTool/Platform/IOS/IOSToolChain.cs
@@ -195,7 +195,7 @@ namespace UnrealBuildTool
 		/// <param name="OutputType">Type of build product</param>
 		public override bool ShouldAddDebugFileToReceipt(FileReference OutputFile, BuildProductType OutputType)
 		{
-			return OutputType == BuildProductType.Executable || OutputType == BuildProductType.DynamicLibrary;
+			return OutputType == BuildProductType.Executable;
 		}
 
 		string GetCompileArguments_Global(CppCompileEnvironment CompileEnvironment)
