- Fixed up the DebugInfo path for IOS framework builds, needed for Crash Reporting
- Reenabled crashlytics

diff --git a/Engine/Source/Programs/UnrealBuildTool/Configuration/UEBuildBinary.cs b/Engine/Source/Programs/UnrealBuildTool/Configuration/UEBuildBinary.cs
--- a/Engine/Source/Programs/UnrealBuildTool/Configuration/UEBuildBinary.cs
+++ b/Engine/Source/Programs/UnrealBuildTool/Configuration/UEBuildBinary.cs
@@ -527,7 +527,9 @@ namespace UnrealBuildTool
 			{
 				if (!String.IsNullOrEmpty(DebugExtension) && ToolChain.ShouldAddDebugFileToReceipt(OutputFile, OutputType) && bCreateDebugInfo)
 				{
-					BuildProducts.Add(OutputFile.ChangeExtension(DebugExtension), BuildProductType.SymbolFile);
+					// @todo this could be cleaned up if we replaced Platform.GetDebugExtensions() with ToolChain.GetDebugFiles(OutputFile)
+					// would need care in MacToolchain tho, so too risky for now
+					BuildProducts.Add(ToolChain.GetDebugFile(OutputFile, DebugExtension), BuildProductType.SymbolFile);
 				}
 			}
 		}
diff --git a/Engine/Source/Programs/UnrealBuildTool/Platform/IOS/IOSToolChain.cs b/Engine/Source/Programs/UnrealBuildTool/Platform/IOS/IOSToolChain.cs
--- a/Engine/Source/Programs/UnrealBuildTool/Platform/IOS/IOSToolChain.cs
+++ b/Engine/Source/Programs/UnrealBuildTool/Platform/IOS/IOSToolChain.cs
@@ -195,9 +195,21 @@ namespace UnrealBuildTool
 		/// <param name="OutputType">Type of build product</param>
 		public override bool ShouldAddDebugFileToReceipt(FileReference OutputFile, BuildProductType OutputType)
 		{
-			return OutputType == BuildProductType.Executable;
+			return OutputType == BuildProductType.Executable || OutputType == BuildProductType.DynamicLibrary;
 		}
 
+		public override FileReference GetDebugFile(FileReference OutputFile, string DebugExtension)
+		{
+			if (OutputFile.FullName.Contains(".framework"))
+			{
+				// need to put the debug info outside of the framework
+				return FileReference.Combine(OutputFile.Directory.ParentDirectory, OutputFile.ChangeExtension(DebugExtension).GetFileName());
+			}
+			//  by default, just change the extension to the debug extension
+			return OutputFile.ChangeExtension(DebugExtension);
+		}
+
+
 		string GetCompileArguments_Global(CppCompileEnvironment CompileEnvironment)
 		{
             string Result = "";
@@ -1680,7 +1692,7 @@ namespace UnrealBuildTool
 
 			string AppName = Target.TargetName;
 
-			if (!Target.bSkipCrashlytics && !IsCompiledAsFramework(Target.OutputPath.FullName))
+			if (!Target.bSkipCrashlytics)
 			{
 				GenerateCrashlyticsData(PathToDsymZip, Target.ProjectDirectory.FullName, AppName);
 			}
diff --git a/Engine/Source/Programs/UnrealBuildTool/ToolChain/UEToolChain.cs b/Engine/Source/Programs/UnrealBuildTool/ToolChain/UEToolChain.cs
--- a/Engine/Source/Programs/UnrealBuildTool/ToolChain/UEToolChain.cs
+++ b/Engine/Source/Programs/UnrealBuildTool/ToolChain/UEToolChain.cs
@@ -74,6 +74,13 @@ namespace UnrealBuildTool
 		{
 			return true;
 		}
+		
+		public virtual FileReference GetDebugFile(FileReference OutputFile, string DebugExtension)
+		{
+			//  by default, just change the extension to the debug extension
+			return OutputFile.ChangeExtension(DebugExtension);
+		}
+
 
 		public virtual void SetupBundleDependencies(List<UEBuildBinary> Binaries, string GameName)
 		{
