Fixing issue with generation of symbols for sequential builds on build machines when building engine as a framework (project-k).

diff --git a/Engine/Source/Programs/UnrealBuildTool/Platform/IOS/IOSToolChain.cs b/Engine/Source/Programs/UnrealBuildTool/Platform/IOS/IOSToolChain.cs
--- a/Engine/Source/Programs/UnrealBuildTool/Platform/IOS/IOSToolChain.cs
+++ b/Engine/Source/Programs/UnrealBuildTool/Platform/IOS/IOSToolChain.cs
@@ -1581,6 +1581,10 @@ namespace UnrealBuildTool
 		private static void GenerateCrashlyticsData(string DsymZip, string ProjectDir, string ProjectName)
         {
 			Log.TraceInformation("Generating and uploading Crashlytics Data");
+
+			// Clean this folder as it's used for extraction
+			CleanIntermediateDirectory(Path.Combine(UnrealBuildTool.EngineDirectory.FullName, "Intermediate", "Zipped"));
+
             string FabricPath = UnrealBuildTool.EngineDirectory + "/Intermediate/UnzippedFrameworks/Crashlytics/Fabric.embeddedframework";
             if (Directory.Exists(FabricPath) && Environment.GetEnvironmentVariable("IsBuildMachine") == "1")
             {
