Second attempt at fixing symbol failures for project-k builds

diff --git a/Engine/Source/Programs/UnrealBuildTool/Platform/IOS/IOSToolChain.cs b/Engine/Source/Programs/UnrealBuildTool/Platform/IOS/IOSToolChain.cs
--- a/Engine/Source/Programs/UnrealBuildTool/Platform/IOS/IOSToolChain.cs
+++ b/Engine/Source/Programs/UnrealBuildTool/Platform/IOS/IOSToolChain.cs
@@ -1583,9 +1583,9 @@ namespace UnrealBuildTool
 			Log.TraceInformation("Generating and uploading Crashlytics Data");
 
 			// Clean this folder as it's used for extraction
-			CleanIntermediateDirectory(Path.Combine(UnrealBuildTool.EngineDirectory.FullName, "Intermediate", "Zipped"));
+			Directory.Delete(Path.Combine(UnrealBuildTool.EngineDirectory.FullName, "Intermediate", "Unzipped"), true);
 
-            string FabricPath = UnrealBuildTool.EngineDirectory + "/Intermediate/UnzippedFrameworks/Crashlytics/Fabric.embeddedframework";
+			string FabricPath = UnrealBuildTool.EngineDirectory + "/Intermediate/UnzippedFrameworks/Crashlytics/Fabric.embeddedframework";
             if (Directory.Exists(FabricPath) && Environment.GetEnvironmentVariable("IsBuildMachine") == "1")
             {
 				string PlistFile = ProjectDir + "/Intermediate/IOS/" + ProjectName + "-Info.plist";
