Another attempt to fix a builder-only issue wth symbols for project K...

diff --git a/Engine/Source/Programs/UnrealBuildTool/Platform/IOS/IOSToolChain.cs b/Engine/Source/Programs/UnrealBuildTool/Platform/IOS/IOSToolChain.cs
--- a/Engine/Source/Programs/UnrealBuildTool/Platform/IOS/IOSToolChain.cs
+++ b/Engine/Source/Programs/UnrealBuildTool/Platform/IOS/IOSToolChain.cs
@@ -1583,7 +1583,13 @@ namespace UnrealBuildTool
 			Log.TraceInformation("Generating and uploading Crashlytics Data");
 
 			// Clean this folder as it's used for extraction
-			Directory.Delete(Path.Combine(UnrealBuildTool.EngineDirectory.FullName, "Intermediate", "Unzipped"), true);
+			string TempPath = Path.Combine(UnrealBuildTool.EngineDirectory.FullName, "Intermediate", "Unzipped");
+
+			if (Directory.Exists(TempPath))
+			{
+				Log.TraceInformation("Deleting temp path {0}", TempPath);
+				Directory.Delete(TempPath, true);
+			}
 
 			string FabricPath = UnrealBuildTool.EngineDirectory + "/Intermediate/UnzippedFrameworks/Crashlytics/Fabric.embeddedframework";
             if (Directory.Exists(FabricPath) && Environment.GetEnvironmentVariable("IsBuildMachine") == "1")
