Fix problems creating a new plugin in a project with a unique build environment.

diff --git a/Engine/Plugins/Editor/PluginBrowser/Source/PluginBrowser/Private/SNewPluginWizard.cpp b/Engine/Plugins/Editor/PluginBrowser/Source/PluginBrowser/Private/SNewPluginWizard.cpp
--- a/Engine/Plugins/Editor/PluginBrowser/Source/PluginBrowser/Private/SNewPluginWizard.cpp
+++ b/Engine/Plugins/Editor/PluginBrowser/Source/PluginBrowser/Private/SNewPluginWizard.cpp
@@ -750,6 +750,9 @@ FReply SNewPluginWizard::OnCreatePluginClicked()
 			bSucceeded = false;
 		}
 
+		// Reset the module paths cache. For unique build environments, the modules may be generated to the project binaries directory.
+		FModuleManager::Get().ResetModulePathsCache();
+
 		// Generate project files if we happen to be using a project file.
 		if (bSucceeded && !FDesktopPlatformModule::Get()->GenerateProjectFiles(FPaths::RootDir(), FPaths::GetProjectFilePath(), GWarn))
 		{
diff --git a/Engine/Source/Programs/UnrealBuildTool/Configuration/UEBuildTarget.cs b/Engine/Source/Programs/UnrealBuildTool/Configuration/UEBuildTarget.cs
--- a/Engine/Source/Programs/UnrealBuildTool/Configuration/UEBuildTarget.cs
+++ b/Engine/Source/Programs/UnrealBuildTool/Configuration/UEBuildTarget.cs
@@ -1594,10 +1594,25 @@ namespace UnrealBuildTool
 				}
 			}
 
-			// If we're just precompiling a plugin, only include output items which are under that directory
+			// If we're just precompiling a plugin, only include output items which are part of it
 			if(ForeignPlugin != null)
 			{
-				Makefile.OutputItems.RemoveAll(x => !x.Location.IsUnderDirectory(ForeignPlugin.Directory));
+				HashSet<FileItem> RetainOutputItems = new HashSet<FileItem>();
+				foreach(UEBuildPlugin Plugin in BuildPlugins)
+				{
+					if(Plugin.File == ForeignPlugin)
+					{
+						foreach (UEBuildModule Module in Plugin.Modules)
+						{
+							FileItem[] ModuleOutputItems;
+							if(Makefile.ModuleNameToOutputItems.TryGetValue(Module.Name, out ModuleOutputItems))
+							{
+								RetainOutputItems.UnionWith(ModuleOutputItems);
+							}
+						}
+					}
+				}
+				Makefile.OutputItems.RemoveAll(x => !RetainOutputItems.Contains(x));
 			}
 
 			// Allow the toolchain to modify the final output items
