Fix issue creating a new plugin after a hot reload has taken place.

diff --git a/Engine/Source/Programs/UnrealBuildTool/Modes/BuildMode.cs b/Engine/Source/Programs/UnrealBuildTool/Modes/BuildMode.cs
--- a/Engine/Source/Programs/UnrealBuildTool/Modes/BuildMode.cs
+++ b/Engine/Source/Programs/UnrealBuildTool/Modes/BuildMode.cs
@@ -539,8 +539,12 @@ namespace UnrealBuildTool
 			HotReloadState HotReloadState = null;
 			if(HotReloadMode == HotReloadMode.Disabled)
 			{
-				// Delete the previous state file
-				HotReload.DeleteTemporaryFiles(HotReloadStateFile);
+				// Make sure we're not doing a partial build from the editor (eg. compiling a new plugin)
+				if(TargetDescriptor.ForeignPlugin == null && TargetDescriptor.SingleFileToCompile == null)
+				{
+					// Delete the previous state file
+					HotReload.DeleteTemporaryFiles(HotReloadStateFile);
+				}
 			}
 			else
 			{
