UE-74772 - RemoteMac.cs fails to build if ProjectFile is null

diff --git a/Engine/Source/Programs/UnrealBuildTool/ToolChain/RemoteMac.cs b/Engine/Source/Programs/UnrealBuildTool/ToolChain/RemoteMac.cs
--- a/Engine/Source/Programs/UnrealBuildTool/ToolChain/RemoteMac.cs
+++ b/Engine/Source/Programs/UnrealBuildTool/ToolChain/RemoteMac.cs
@@ -911,9 +911,12 @@ namespace UnrealBuildTool
 			}
 
 			Execute("/", String.Format("rm -rf {0}/Intermediate/IOS/*.plist", GetRemotePath(UnrealBuildTool.EngineDirectory)));
-			Execute("/", String.Format("rm -rf {0}/Intermediate/IOS/*.plist", GetRemotePath(ProjectFile.Directory)));
 			Execute("/", String.Format("rm -rf {0}/Intermediate/TVOS/*.plist", GetRemotePath(UnrealBuildTool.EngineDirectory)));
-			Execute("/", String.Format("rm -rf {0}/Intermediate/TVOS/*.plist", GetRemotePath(ProjectFile.Directory)));
+			if (ProjectFile != null)
+			{
+				Execute("/", String.Format("rm -rf {0}/Intermediate/IOS/*.plist", GetRemotePath(ProjectFile.Directory)));
+				Execute("/", String.Format("rm -rf {0}/Intermediate/TVOS/*.plist", GetRemotePath(ProjectFile.Directory)));
+			}
 
 			// Fixup permissions on any shell scripts
 			Execute(RemoteBaseDir, String.Format("chmod +x {0}/Build/BatchFiles/Mac/*.sh", EscapeShellArgument(GetRemotePath(UnrealBuildTool.EngineDirectory))));
