When loading BuildConfiguration.xml we now strip whitespace characters around xml node values (except for string values). This allows a more convenient configuration. Also added a explicit format check for boolean values as incorrectly formatted values would simply be interpreted as false.

diff --git a/Engine/Source/Programs/UnrealBuildTool/System/XmlConfig.cs b/Engine/Source/Programs/UnrealBuildTool/System/XmlConfig.cs
--- a/Engine/Source/Programs/UnrealBuildTool/System/XmlConfig.cs
+++ b/Engine/Source/Programs/UnrealBuildTool/System/XmlConfig.cs
@@ -587,29 +587,42 @@ namespace UnrealBuildTool
 		/// <returns>The object that was parsed</returns>
 		static object ParseValue(Type FieldType, string Text)
 		{
+			// ignore whitespace in all fields except for Strings which we leave unprocessed
+			string TrimmedText = Text.Trim();
 			if(FieldType == typeof(string))
 			{
 				return Text;
 			}
 			else if(FieldType == typeof(bool) || FieldType == typeof(bool?))
 			{
-				return (Text == "1" || Text.Equals("true", StringComparison.InvariantCultureIgnoreCase));
+				if (TrimmedText == "1" || TrimmedText.Equals("true", StringComparison.InvariantCultureIgnoreCase))
+				{
+					return true;
+				}
+				else if (TrimmedText == "0" || TrimmedText.Equals("false", StringComparison.InvariantCultureIgnoreCase))
+				{
+					return false;
+				} 
+				else 
+				{
+					throw new Exception(String.Format("Unable to convert '{0}' to boolean. 'true/false/0/1' are the supported formats.", Text));
+				}
 			}
 			else if(FieldType == typeof(int))
 			{
-				return Int32.Parse(Text);
+				return Int32.Parse(TrimmedText);
 			}
 			else if(FieldType == typeof(float))
 			{
-				return Single.Parse(Text);
+				return Single.Parse(TrimmedText);
 			}
 			else if(FieldType == typeof(double))
 			{
-				return Double.Parse(Text);
+				return Double.Parse(TrimmedText);
 			}
 			else if(FieldType.IsEnum)
 			{
-				return Enum.Parse(FieldType, Text);
+				return Enum.Parse(FieldType, TrimmedText);
 			}
 			else
 			{
