UBT: Fix issues with UBT not detecting new files with reflection markup. Was previously not correctly sub-directories, nor checking for new UHT types being added to a module that was not previously being scanned.

diff --git a/Engine/Source/Programs/UnrealBuildTool/Configuration/UEBuildModuleCPP.cs b/Engine/Source/Programs/UnrealBuildTool/Configuration/UEBuildModuleCPP.cs
--- a/Engine/Source/Programs/UnrealBuildTool/Configuration/UEBuildModuleCPP.cs
+++ b/Engine/Source/Programs/UnrealBuildTool/Configuration/UEBuildModuleCPP.cs
@@ -1242,6 +1242,7 @@ namespace UnrealBuildTool
 			{
 				DirectoryItem ModuleDirectoryItem = DirectoryItem.GetItemByDirectoryReference(Dir);
 				FindInputFilesFromDirectoryRecursive(ModuleDirectoryItem, ExcludedNames, SourceDirectories, Makefile.DirectoryToSourceFiles, InputFiles);
+				Makefile.SourceDirectories.Add(ModuleDirectoryItem);
 			}
 
 			return InputFiles;
diff --git a/Engine/Source/Programs/UnrealBuildTool/System/TargetMakefile.cs b/Engine/Source/Programs/UnrealBuildTool/System/TargetMakefile.cs
--- a/Engine/Source/Programs/UnrealBuildTool/System/TargetMakefile.cs
+++ b/Engine/Source/Programs/UnrealBuildTool/System/TargetMakefile.cs
@@ -1,5 +1,6 @@
 // Copyright 1998-2019 Epic Games, Inc. All Rights Reserved.
 using System;
+using System.Collections.Concurrent;
 using System.Collections.Generic;
 using System.IO;
 using System.Linq;
@@ -20,7 +21,7 @@ namespace UnrealBuildTool
 		/// <summary>
 		/// The version number to write
 		/// </summary>
-		public const int CurrentVersion = 13;
+		public const int CurrentVersion = 14;
 
 		/// <summary>
 		/// The time at which the makefile was created
@@ -98,6 +99,11 @@ namespace UnrealBuildTool
 		/// </summary>
 		public HashSet<string> HotReloadModuleNames;
 
+		/// <summary>
+		/// List of all source directories
+		/// </summary>
+		public List<DirectoryItem> SourceDirectories;
+
 		/// <summary>
 		/// Set of all source directories. Any files being added or removed from these directories will invalidate the makefile.
 		/// </summary>
@@ -157,6 +163,7 @@ namespace UnrealBuildTool
 			this.OutputItems = new List<FileItem>();
 			this.ModuleNameToOutputItems = new Dictionary<string, FileItem[]>(StringComparer.OrdinalIgnoreCase);
 			this.HotReloadModuleNames = new HashSet<string>(StringComparer.OrdinalIgnoreCase);
+			this.SourceDirectories = new List<DirectoryItem>();
 			this.DirectoryToSourceFiles = new Dictionary<DirectoryItem, FileItem[]>();
 			this.WorkingSet = new HashSet<FileItem>();
 			this.CandidatesForWorkingSet = new HashSet<FileItem>();
@@ -187,6 +194,7 @@ namespace UnrealBuildTool
 			OutputItems = Reader.ReadList(() => Reader.ReadFileItem());
 			ModuleNameToOutputItems = Reader.ReadDictionary(() => Reader.ReadString(), () => Reader.ReadArray(() => Reader.ReadFileItem()), StringComparer.OrdinalIgnoreCase);
 			HotReloadModuleNames = Reader.ReadHashSet(() => Reader.ReadString(), StringComparer.OrdinalIgnoreCase);
+			SourceDirectories = Reader.ReadList(() => Reader.ReadDirectoryItem());
 			DirectoryToSourceFiles = Reader.ReadDictionary(() => Reader.ReadDirectoryItem(), () => Reader.ReadArray(() => Reader.ReadFileItem()));
 			WorkingSet = Reader.ReadHashSet(() => Reader.ReadFileItem());
 			CandidatesForWorkingSet = Reader.ReadHashSet(() => Reader.ReadFileItem());
@@ -217,6 +225,7 @@ namespace UnrealBuildTool
 			Writer.WriteList(OutputItems, Item => Writer.WriteFileItem(Item));
 			Writer.WriteDictionary(ModuleNameToOutputItems, k => Writer.WriteString(k), v => Writer.WriteArray(v, e => Writer.WriteFileItem(e)));
 			Writer.WriteHashSet(HotReloadModuleNames, x => Writer.WriteString(x));
+			Writer.WriteList(SourceDirectories, x => Writer.WriteDirectoryItem(x));
 			Writer.WriteDictionary(DirectoryToSourceFiles, k => Writer.WriteDirectoryItem(k), v => Writer.WriteArray(v, e => Writer.WriteFileItem(e)));
 			Writer.WriteHashSet(WorkingSet, x => Writer.WriteFileItem(x));
 			Writer.WriteHashSet(CandidatesForWorkingSet, x => Writer.WriteFileItem(x));
@@ -488,52 +497,26 @@ namespace UnrealBuildTool
 					}
 				}
 
-				// We do a check to see if any modules' headers have changed which have
-				// acquired or lost UHT types.  If so, which should be rare,
-				// we'll just invalidate the entire makefile and force it to be rebuilt.
+				// Load the metadata cache
+				SourceFileMetadataCache MetadataCache = SourceFileMetadataCache.CreateHierarchy(ProjectFile);
 
-				// Get all H files in processed modules newer than the makefile itself
-				HashSet<FileItem> HFilesNewerThanMakefile = new HashSet<FileItem>();
-				foreach(UHTModuleHeaderInfo ModuleHeaderInfo in Makefile.UObjectModuleHeaders)
+				// Find the set of files that contain reflection markup
+				ConcurrentBag<FileItem> NewFilesWithMarkupBag = new ConcurrentBag<FileItem>();
+				using (ThreadPoolWorkQueue Queue = new ThreadPoolWorkQueue())
 				{
-					foreach(FileItem HeaderFile in ModuleHeaderInfo.SourceFolder.EnumerateFiles())
+					foreach(DirectoryItem SourceDirectory in Makefile.SourceDirectories)
 					{
-						if(HeaderFile.HasExtension(".h") && HeaderFile.LastWriteTimeUtc > Makefile.CreateTimeUtc)
-						{
-							HFilesNewerThanMakefile.Add(HeaderFile);
-						}
+						Queue.Enqueue(() => FindFilesWithMarkup(SourceDirectory, MetadataCache, ExcludedFolderNames, NewFilesWithMarkupBag, Queue));
 					}
 				}
 
-				// Get all H files in all modules processed in the last makefile build
-				HashSet<FileItem> AllUHTHeaders = new HashSet<FileItem>(Makefile.UObjectModuleHeaders.SelectMany(x => x.HeaderFiles));
-
-				// Check whether any headers have been deleted. If they have, we need to regenerate the makefile since the module might now be empty. If we don't,
-				// and the file has been moved to a different module, we may include stale generated headers.
-				foreach (FileItem HeaderFile in AllUHTHeaders)
+				// Check whether the list has changed
+				List<FileItem> PrevFilesWithMarkup = Makefile.UObjectModuleHeaders.SelectMany(x => x.HeaderFiles).ToList();
+				List<FileItem> NextFilesWithMarkup = NewFilesWithMarkupBag.ToList();
+				if (NextFilesWithMarkup.Count != PrevFilesWithMarkup.Count || NextFilesWithMarkup.Intersect(PrevFilesWithMarkup).Count() != PrevFilesWithMarkup.Count)
 				{
-					if (!HeaderFile.Exists)
-					{
-						Log.TraceLog("File processed by UHT was deleted ({0}); invalidating makefile", HeaderFile);
-						ReasonNotLoaded = string.Format("UHT file was deleted");
-						return false;
-					}
-				}
-
-				// Makefile is invalid if:
-				// * There are any newer files which contain no UHT data, but were previously in the makefile
-				// * There are any newer files contain data which needs processing by UHT, but weren't not previously in the makefile
-				SourceFileMetadataCache MetadataCache = SourceFileMetadataCache.CreateHierarchy(ProjectFile);
-				foreach (FileItem HeaderFile in HFilesNewerThanMakefile)
-				{
-					bool bContainsUHTData = MetadataCache.ContainsReflectionMarkup(HeaderFile);
-					bool bWasProcessed = AllUHTHeaders.Contains(HeaderFile);
-					if (bContainsUHTData != bWasProcessed)
-					{
-						Log.TraceLog("{0} {1} contain UHT types and now {2} , ignoring it", HeaderFile, bWasProcessed ? "used to" : "didn't", bWasProcessed ? "doesn't" : "does");
-						ReasonNotLoaded = string.Format("new files with reflected types");
-						return false;
-					}
+					ReasonNotLoaded = "UHT files changed";
+					return false;
 				}
 
 				// If adaptive unity build is enabled, do a check to see if there are any source files that became part of the
@@ -599,6 +582,35 @@ namespace UnrealBuildTool
 			return false;
 		}
 
+		/// <summary>
+		/// Finds all the source files under a directory that contain reflection markup
+		/// </summary>
+		/// <param name="Directory">The directory to search</param>
+		/// <param name="MetadataCache">Cache of source file metadata</param>
+		/// <param name="ExcludedFolderNames">Set of folder names to ignore when recursing the directory tree</param>
+		/// <param name="FilesWithMarkup">Receives the set of files which contain reflection markup</param>
+		/// <param name="Queue">Queue to add sub-tasks to</param>
+		static void FindFilesWithMarkup(DirectoryItem Directory, SourceFileMetadataCache MetadataCache, ReadOnlyHashSet<string> ExcludedFolderNames, ConcurrentBag<FileItem> FilesWithMarkup, ThreadPoolWorkQueue Queue)
+		{
+			// Search through all the subfolders
+			foreach(DirectoryItem SubDirectory in Directory.EnumerateDirectories())
+			{
+				if(!ExcludedFolderNames.Contains(SubDirectory.Name))
+				{
+					Queue.Enqueue(() => FindFilesWithMarkup(SubDirectory, MetadataCache, ExcludedFolderNames, FilesWithMarkup, Queue));
+				}
+			}
+
+			// Check for all the headers in this folder
+			foreach(FileItem File in Directory.EnumerateFiles())
+			{
+				if(File.HasExtension(".h") && MetadataCache.ContainsReflectionMarkup(File))
+				{
+					FilesWithMarkup.Add(File);
+				}
+			}
+		}
+
 		/// <summary>
 		/// Gets the location of the makefile for particular target
 		/// </summary>
