Add MD5 hash check of Android build template files to invalidate intermedate source

#ue4
#android

diff --git a/Engine/Source/Programs/UnrealBuildTool/Platform/Android/UEDeployAndroid.cs b/Engine/Source/Programs/UnrealBuildTool/Platform/Android/UEDeployAndroid.cs
--- a/Engine/Source/Programs/UnrealBuildTool/Platform/Android/UEDeployAndroid.cs
+++ b/Engine/Source/Programs/UnrealBuildTool/Platform/Android/UEDeployAndroid.cs
@@ -10,6 +10,7 @@ using System.IO;
 using Microsoft.Win32;
 using System.Xml.Linq;
 using Tools.DotNETCommon;
+using System.Security.Cryptography;
 
 namespace UnrealBuildTool
 {
@@ -1630,7 +1631,7 @@ namespace UnrealBuildTool
 		}
 
 
-		private string GetAllBuildSettings(AndroidToolChain ToolChain, string BuildPath, bool bForDistribution, bool bMakeSeparateApks, bool bPackageDataInsideApk, bool bDisableVerifyOBBOnStartUp, bool bUseExternalFilesDir, bool bGradleEnabled)
+		private string GetAllBuildSettings(AndroidToolChain ToolChain, string BuildPath, bool bForDistribution, bool bMakeSeparateApks, bool bPackageDataInsideApk, bool bDisableVerifyOBBOnStartUp, bool bUseExternalFilesDir, bool bGradleEnabled, string TemplatesHashCode)
 		{
 			// make the settings string - this will be char by char compared against last time
 			StringBuilder CurrentSettings = new StringBuilder();
@@ -1646,6 +1647,7 @@ namespace UnrealBuildTool
 			CurrentSettings.AppendLine(string.Format("bDisableVerifyOBBOnStartUp={0}", bDisableVerifyOBBOnStartUp));
 			CurrentSettings.AppendLine(string.Format("bUseExternalFilesDir={0}", bUseExternalFilesDir));
 			CurrentSettings.AppendLine(string.Format("UPLHashCode={0}", UPLHashCode));
+			CurrentSettings.AppendLine(string.Format("TemplatesHashCode={0}", TemplatesHashCode));
 
 			// all AndroidRuntimeSettings ini settings in here
 			ConfigHierarchy Ini = GetConfigCacheIni(ConfigHierarchyType.Engine);
@@ -3463,7 +3465,8 @@ namespace UnrealBuildTool
 
 
 			// check to see if any "meta information" is newer than last time we build
-			string CurrentBuildSettings = GetAllBuildSettings(ToolChain, UE4BuildPath, bForDistribution, bMakeSeparateApks, bPackageDataInsideApk, bDisableVerifyOBBOnStartUp, bUseExternalFilesDir, bGradleEnabled);
+			string TemplatesHashCode = GenerateTemplatesHashCode(EngineDirectory);
+			string CurrentBuildSettings = GetAllBuildSettings(ToolChain, UE4BuildPath, bForDistribution, bMakeSeparateApks, bPackageDataInsideApk, bDisableVerifyOBBOnStartUp, bUseExternalFilesDir, bGradleEnabled, TemplatesHashCode);
 			string BuildSettingsCacheFile = Path.Combine(UE4BuildPath, "UEBuildSettings.txt");
 
 			// do we match previous build settings?
@@ -4274,6 +4277,53 @@ namespace UnrealBuildTool
 			}
 		}
 
+		private string GenerateTemplatesHashCode(string EngineDir)
+		{
+			string SourceDirectory = Path.Combine(EngineDir, "Build", "Android", "Java");
+
+			if (!Directory.Exists(SourceDirectory))
+			{
+				return "badpath";
+			}
+
+			MD5 md5 = MD5.Create();
+			byte[] TotalHashBytes = null;
+
+			string[] SourceFiles = Directory.GetFiles(SourceDirectory, "*.*", SearchOption.AllDirectories);
+			foreach (string Filename in SourceFiles)
+			{
+				using (FileStream stream = File.OpenRead(Filename))
+				{
+					byte[] FileHashBytes = md5.ComputeHash(stream);
+					if (TotalHashBytes != null)
+					{
+						int index = 0;
+						foreach (byte b in FileHashBytes)
+						{
+							TotalHashBytes[index] ^= b;
+							index++;
+						}
+					}
+					else
+					{
+						TotalHashBytes = FileHashBytes;
+					}
+				}
+			}
+
+			if (TotalHashBytes != null)
+			{
+				string HashCode = "";
+				foreach (byte b in TotalHashBytes)
+				{
+					HashCode += b.ToString("x2");
+				}
+				return HashCode;
+			}
+
+			return "empty";
+		}
+
 		private void UpdateGameActivity(string UE4Arch, string NDKArch, string EngineDir, string UE4BuildPath)
 		{
 			string SourceFilename = Path.Combine(EngineDir, "Build", "Android", "Java", "src", "com", "epicgames", "ue4", "GameActivity.java.template");
