UBT: Fix exception when querying older versions of Visual Studio that do not implement ISetupInstanceCatalog.

diff --git a/Engine/Source/Programs/UnrealBuildTool/Platform/Windows/UEBuildWindows.cs b/Engine/Source/Programs/UnrealBuildTool/Platform/Windows/UEBuildWindows.cs
--- a/Engine/Source/Programs/UnrealBuildTool/Platform/Windows/UEBuildWindows.cs
+++ b/Engine/Source/Programs/UnrealBuildTool/Platform/Windows/UEBuildWindows.cs
@@ -931,7 +931,7 @@ namespace UnrealBuildTool
 
 									int SortOrder = SortedInstallDirs.Count;
 
-									ISetupInstanceCatalog Catalog = (ISetupInstanceCatalog)Instance as ISetupInstanceCatalog;
+									ISetupInstanceCatalog Catalog = Instance as ISetupInstanceCatalog;
 									if (Catalog != null && Catalog.IsPrerelease())
 									{
 										SortOrder |= (1 << 16);
