Live Coding: Disable ICF and symbol stripping when live coding is enabled.

diff --git a/Engine/Source/Programs/UnrealBuildTool/Platform/Windows/UEBuildWindows.cs b/Engine/Source/Programs/UnrealBuildTool/Platform/Windows/UEBuildWindows.cs
--- a/Engine/Source/Programs/UnrealBuildTool/Platform/Windows/UEBuildWindows.cs
+++ b/Engine/Source/Programs/UnrealBuildTool/Platform/Windows/UEBuildWindows.cs
@@ -232,6 +232,16 @@ namespace UnrealBuildTool
 		/// </summary>
 		public bool bCreateHotPatchableImage = false;
 
+		/// <summary>
+		/// Strip unreferenced symbols (/OPT:REF)
+		/// </summary>
+		public bool bStripUnreferencedSymbols = false;
+
+		/// <summary>
+		/// Merge identical COMDAT sections together (/OPT:ICF)
+		/// </summary>
+		public bool bMergeIdenticalCOMDATs = false;
+
 		/// <summary>
 		/// Whether to put global symbols in their own sections (/Gw), allowing the linker to discard any that are unused.
 		/// </summary>
@@ -440,6 +450,16 @@ namespace UnrealBuildTool
 			get { return Inner.bCreateHotPatchableImage; }
 		}
 
+		public bool bStripUnreferencedSymbols
+		{
+			get { return Inner.bStripUnreferencedSymbols; }
+		}
+
+		public bool bMergeIdenticalCOMDATs
+		{
+			get { return Inner.bMergeIdenticalCOMDATs; }
+		}
+
 		public bool bOptimizeGlobalData
 		{
 			get { return Inner.bOptimizeGlobalData; }
@@ -592,6 +612,11 @@ namespace UnrealBuildTool
 				Target.bWithLiveCoding = true;
 				Target.WindowsPlatform.bCreateHotPatchableImage = true;
 			}
+			else if (Target.Configuration == UnrealTargetConfiguration.Test || Target.Configuration == UnrealTargetConfiguration.Shipping)
+			{
+				Target.WindowsPlatform.bStripUnreferencedSymbols = true;
+				Target.WindowsPlatform.bMergeIdenticalCOMDATs = true;
+			}
 		}
 
 		/// <summary>
diff --git a/Engine/Source/Programs/UnrealBuildTool/Platform/Windows/VCToolChain.cs b/Engine/Source/Programs/UnrealBuildTool/Platform/Windows/VCToolChain.cs
--- a/Engine/Source/Programs/UnrealBuildTool/Platform/Windows/VCToolChain.cs
+++ b/Engine/Source/Programs/UnrealBuildTool/Platform/Windows/VCToolChain.cs
@@ -786,22 +786,25 @@ namespace UnrealBuildTool
 			{
 				// Generate an EXE checksum.
 				Arguments.Add("/RELEASE");
+			}
 
-				// Eliminate unreferenced symbols.
+			// Eliminate unreferenced symbols.
+			if (Target.WindowsPlatform.bStripUnreferencedSymbols)
+			{
 				Arguments.Add("/OPT:REF");
-
-				// Remove redundant COMDATs.
-				Arguments.Add("/OPT:ICF");
 			}
-			//
-			//	Regular development binary. 
-			//
 			else
 			{
-				// Keep symbols that are unreferenced.
 				Arguments.Add("/OPT:NOREF");
+			}
 
-				// Disable identical COMDAT folding.
+			// Identical COMDAT folding
+			if (Target.WindowsPlatform.bMergeIdenticalCOMDATs)
+			{
+				Arguments.Add("/OPT:ICF");
+			}
+			else
+			{
 				Arguments.Add("/OPT:NOICF");
 			}
 
