Live Coding: Prevent compiling from the IDE while a Live Coding session is active. Engine runtime creates a global mutex with the executable path in the a name and leaks the handle, then UBT checks for it existing before allowing a build to start.

diff --git a/Engine/Source/Developer/Windows/LiveCoding/Private/LiveCodingModule.cpp b/Engine/Source/Developer/Windows/LiveCoding/Private/LiveCodingModule.cpp
--- a/Engine/Source/Developer/Windows/LiveCoding/Private/LiveCodingModule.cpp
+++ b/Engine/Source/Developer/Windows/LiveCoding/Private/LiveCodingModule.cpp
@@ -291,6 +291,25 @@ bool FLiveCodingModule::StartLiveCoding()
 		}
 		LppSetBuildArguments(*Arguments);
 
+		// Create a mutex that allows UBT to detect that we shouldn't hot-reload into this executable. The handle to it will be released automatically when the process exits.
+		FString ExecutablePath = FPaths::ConvertRelativePathToFull(FPlatformProcess::ExecutablePath());
+
+		FString MutexName = TEXT("Global\\LiveCoding_");
+		for (int Idx = 0; Idx < ExecutablePath.Len(); Idx++)
+		{
+			TCHAR Character = ExecutablePath[Idx];
+			if (Character == '/' || Character == '\\' || Character == ':')
+			{
+				MutexName += '+';
+			}
+			else
+			{
+				MutexName += Character;
+			}
+		}
+
+		ensure(CreateMutex(NULL, Windows::FALSE, *MutexName));
+
 		// Configure all the current modules. For non-commandlets, schedule it to be done in the first Tick() so we can batch everything together.
 		if (IsRunningCommandlet())
 		{
diff --git a/Engine/Source/Programs/UnrealBuildTool/Modes/BuildMode.cs b/Engine/Source/Programs/UnrealBuildTool/Modes/BuildMode.cs
--- a/Engine/Source/Programs/UnrealBuildTool/Modes/BuildMode.cs
+++ b/Engine/Source/Programs/UnrealBuildTool/Modes/BuildMode.cs
@@ -537,6 +537,12 @@ namespace UnrealBuildTool
 				}
 			}
 
+			// Guard against a live coding session for this target being active
+			if(HotReloadMode != HotReloadMode.LiveCoding && HotReload.IsLiveCodingSessionActive(Makefile))
+			{
+				throw new BuildException("Unable to start regular build while Live Coding is active. Press Ctrl+Alt+F11 to trigger a Live Coding compile.");
+			}
+
 			// Get the root prerequisite actions
 			List<Action> PrerequisiteActions = GatherPrerequisiteActions(TargetDescriptor, Makefile);
 
diff --git a/Engine/Source/Programs/UnrealBuildTool/System/HotReload.cs b/Engine/Source/Programs/UnrealBuildTool/System/HotReload.cs
--- a/Engine/Source/Programs/UnrealBuildTool/System/HotReload.cs
+++ b/Engine/Source/Programs/UnrealBuildTool/System/HotReload.cs
@@ -6,6 +6,7 @@ using System.Diagnostics;
 using System.IO;
 using System.Linq;
 using System.Text;
+using System.Threading;
 using System.Threading.Tasks;
 using Tools.DotNETCommon;
 
@@ -120,6 +121,42 @@ namespace UnrealBuildTool
 
 	static class HotReload
 	{
+		/// <summary>
+		/// Checks whether a live coding session is currently active for a target. If so, we don't want to allow modifying any object files before they're loaded.
+		/// </summary>
+		/// <param name="Makefile">Makefile for the target being built</param>
+		/// <returns>True if a live coding session is active, false otherwise</returns>
+		public static bool IsLiveCodingSessionActive(TargetMakefile Makefile)
+		{
+			// Find the first output executable
+			FileItem Executable = Makefile.OutputItems.FirstOrDefault(x => x.HasExtension(".exe"));
+
+			// Build the mutex name. This should match the name generated in LiveCodingModule.cpp.
+			StringBuilder MutexName = new StringBuilder("Global\\LiveCoding_");
+			for(int Idx = 0; Idx < Executable.AbsolutePath.Length; Idx++)
+			{
+				char Character = Executable.AbsolutePath[Idx];
+				if(Character == '/' || Character == '\\' || Character == ':')
+				{
+					MutexName.Append('+');
+				}
+				else
+				{
+					MutexName.Append(Character);
+				}
+			}
+			Log.TraceLog("Checking for live coding mutex: {0}", MutexName);
+
+			// Try to open the mutex
+			Mutex Mutex;
+			if(Mutex.TryOpenExisting(MutexName.ToString(), out Mutex))
+			{
+				Mutex.Dispose();
+				return true;
+			}
+			return false;
+		}
+
 		/// <summary>
 		/// Checks if the editor is currently running and this is a hot-reload
 		/// </summary>
