UBT: Fix exception when checking for live coding session for non-windows targets.

#rnx

diff --git a/Engine/Source/Programs/UnrealBuildTool/System/HotReload.cs b/Engine/Source/Programs/UnrealBuildTool/System/HotReload.cs
--- a/Engine/Source/Programs/UnrealBuildTool/System/HotReload.cs
+++ b/Engine/Source/Programs/UnrealBuildTool/System/HotReload.cs
@@ -130,30 +130,32 @@ namespace UnrealBuildTool
 		{
 			// Find the first output executable
 			FileItem Executable = Makefile.OutputItems.FirstOrDefault(x => x.HasExtension(".exe"));
-
-			// Build the mutex name. This should match the name generated in LiveCodingModule.cpp.
-			StringBuilder MutexName = new StringBuilder("Global\\LiveCoding_");
-			for(int Idx = 0; Idx < Executable.AbsolutePath.Length; Idx++)
+			if(Executable != null)
 			{
-				char Character = Executable.AbsolutePath[Idx];
-				if(Character == '/' || Character == '\\' || Character == ':')
+				// Build the mutex name. This should match the name generated in LiveCodingModule.cpp.
+				StringBuilder MutexName = new StringBuilder("Global\\LiveCoding_");
+				for(int Idx = 0; Idx < Executable.AbsolutePath.Length; Idx++)
 				{
-					MutexName.Append('+');
+					char Character = Executable.AbsolutePath[Idx];
+					if(Character == '/' || Character == '\\' || Character == ':')
+					{
+						MutexName.Append('+');
+					}
+					else
+					{
+						MutexName.Append(Character);
+					}
 				}
-				else
+				Log.TraceLog("Checking for live coding mutex: {0}", MutexName);
+
+				// Try to open the mutex
+				Mutex Mutex;
+				if(Mutex.TryOpenExisting(MutexName.ToString(), out Mutex))
 				{
-					MutexName.Append(Character);
+					Mutex.Dispose();
+					return true;
 				}
 			}
-			Log.TraceLog("Checking for live coding mutex: {0}", MutexName);
-
-			// Try to open the mutex
-			Mutex Mutex;
-			if(Mutex.TryOpenExisting(MutexName.ToString(), out Mutex))
-			{
-				Mutex.Dispose();
-				return true;
-			}
 			return false;
 		}
 
