Add support for tracking and validating compile times between builds so we can be alerted if a file's compile time suddenly bloats.

#rnx

diff --git a/Engine/Source/Programs/UnrealBuildTool/Modes/AggregateParsedTimingInfo.cs b/Engine/Source/Programs/UnrealBuildTool/Modes/AggregateParsedTimingInfo.cs
--- a/Engine/Source/Programs/UnrealBuildTool/Modes/AggregateParsedTimingInfo.cs
+++ b/Engine/Source/Programs/UnrealBuildTool/Modes/AggregateParsedTimingInfo.cs
@@ -72,6 +72,18 @@ namespace UnrealBuildTool
 				}
 			}
 
+			if (Arguments.HasValue("-CompileTimingFile="))
+			{
+				FileReference CompileTimingFile = Arguments.GetFileReference("-CompileTimingFile=");
+				Dictionary<string, double> CompileTimes = new Dictionary<string, double>();
+				foreach (KeyValuePair<string, TimingData> TimingData in FileTimingData.Children)
+				{
+					CompileTimes.Add(Json.EscapeString(TimingData.Key), TimingData.Value.InclusiveDuration);
+				}
+				string JsonCompileTimes = Json.Serialize(CompileTimes);
+				FileReference.WriteAllText(CompileTimingFile, JsonCompileTimes);
+			}
+
 			return 0;
 		}
 
diff --git a/Engine/Source/Programs/UnrealBuildTool/Platform/Windows/VCToolChain.cs b/Engine/Source/Programs/UnrealBuildTool/Platform/Windows/VCToolChain.cs
--- a/Engine/Source/Programs/UnrealBuildTool/Platform/Windows/VCToolChain.cs
+++ b/Engine/Source/Programs/UnrealBuildTool/Platform/Windows/VCToolChain.cs
@@ -1210,7 +1210,15 @@ namespace UnrealBuildTool
 					FileReference ManifestFile = FileReference.Combine(Makefile.ProjectIntermediateDirectory, $"{Target.Name}TimingManifest.txt");
 					File.WriteAllLines(ManifestFile.FullName, TimingJsonFiles.Select(f => f.AbsolutePath));
 
-					Action AggregateTimingInfoAction = Action.CreateRecursiveAction<AggregateParsedTimingInfo>(ActionType.ParseTimingInfo, $"-Name={Target.Name} -ManifestFile=\"{ManifestFile.FullName}\"");
+					FileReference ExpectedCompileTimeFile = FileReference.FromString(Path.Combine(Makefile.ProjectIntermediateDirectory.FullName, String.Format("{0}.json", Target.Name)));
+					List<string> ActionArgs = new List<string>()
+					{
+						String.Format("-Name={0}", Target.Name),
+						String.Format("-ManifestFile={0}", ManifestFile.FullName),
+						String.Format("-CompileTimingFile={0}", ExpectedCompileTimeFile),
+					};
+
+					Action AggregateTimingInfoAction = Action.CreateRecursiveAction<AggregateParsedTimingInfo>(ActionType.ParseTimingInfo, string.Join(" ", ActionArgs));
 					AggregateTimingInfoAction.WorkingDirectory = UnrealBuildTool.EngineSourceDirectory;
 					AggregateTimingInfoAction.StatusDescription = $"Aggregating {TimingJsonFiles.Count} Timing File(s)";
 					AggregateTimingInfoAction.bCanExecuteRemotely = false;
