Add -fdiagnostics-absolute-paths to clang builds on linux to get fully qualifed paths to all files involved when a error happens in a header

diff --git a/Engine/Source/Programs/UnrealBuildTool/Platform/Linux/LinuxToolChain.cs b/Engine/Source/Programs/UnrealBuildTool/Platform/Linux/LinuxToolChain.cs
--- a/Engine/Source/Programs/UnrealBuildTool/Platform/Linux/LinuxToolChain.cs
+++ b/Engine/Source/Programs/UnrealBuildTool/Platform/Linux/LinuxToolChain.cs
@@ -605,6 +605,13 @@ namespace UnrealBuildTool
 						Result += " -fcolor-diagnostics";
 					}
 				}
+
+				// output full paths to the files when the build fails, required 4.0+ of clang
+				if (CompilerVersionGreaterOrEqual(4, 0, 0))
+				{
+					Result += " -fdiagnostics-absolute-paths";
+				}
+
 				Result += " -Wno-unused-private-field";     // MultichannelTcpSocket.h triggers this, possibly more
 				// this hides the "warning : comparison of unsigned expression < 0 is always false" type warnings due to constant comparisons, which are possible with template arguments
 				Result += " -Wno-tautological-compare";
