Remove code that reads solution file when determining which compiler to use. This is very slow. We already have multiple ways of specifying which IDE to use.

#rnx

diff --git a/Engine/Source/Programs/UnrealBuildTool/Platform/Windows/UEBuildWindows.cs b/Engine/Source/Programs/UnrealBuildTool/Platform/Windows/UEBuildWindows.cs
--- a/Engine/Source/Programs/UnrealBuildTool/Platform/Windows/UEBuildWindows.cs
+++ b/Engine/Source/Programs/UnrealBuildTool/Platform/Windows/UEBuildWindows.cs
@@ -760,62 +760,6 @@ namespace UnrealBuildTool
 			    }
 			}
 
-			FileReference Solution = null;
-
-			do
-			{
-				if (ProjectFile != null)
-				{
-					Solution = FileReference.Combine(ProjectFile.Directory, ProjectFile.GetFileNameWithoutExtension() + ".sln");
-					if (FileReference.Exists(Solution))
-					{
-						break;
-					}
-				}
-
-				Solution = FileReference.Combine(UnrealBuildTool.RootDirectory, "UE4.sln");
-				if (FileReference.Exists(Solution))
-				{
-					break;
-				}
-			}
-			while (false);
-
-			if (FileReference.Exists(Solution))
-			{
-				foreach (var Line in FileReference.ReadAllLines(Solution))
-				{
-					if (!Line.StartsWith("VisualStudioVersion"))
-					{
-						continue;
-					}
-
-					string[] Pair = Line.Split('=');
-
-					if (Pair.Length != 2)
-					{
-						continue;
-					}
-
-					Version VersionValue;
-
-					if (!Version.TryParse(Pair[1].Trim(), out VersionValue))
-					{
-						break;
-					}
-
-					switch (VersionValue.Major)
-					{
-						case 15:
-							return WindowsCompiler.VisualStudio2017;
-						case 16:
-							return WindowsCompiler.VisualStudio2019;
-						default:
-							break;
-					}
-				}
-			}
-
 			// Second, default based on what's installed, test for 2015 first
 			if (HasCompiler(WindowsCompiler.VisualStudio2017))
 			{
