Fixing HoloLens installed build

#rnx

diff --git a/Engine/Plugins/Runtime/Database/SQLiteCore/SQLiteCore.uplugin b/Engine/Plugins/Runtime/Database/SQLiteCore/SQLiteCore.uplugin
--- a/Engine/Plugins/Runtime/Database/SQLiteCore/SQLiteCore.uplugin
+++ b/Engine/Plugins/Runtime/Database/SQLiteCore/SQLiteCore.uplugin
@@ -19,7 +19,10 @@
 		{
 			"Name": "SQLiteCore",
 			"Type": "Runtime",
-			"LoadingPhase": "PreDefault"
-		}
+			"LoadingPhase": "PreDefault",
+			"BlacklistPlatforms": [
+				"HoloLens"
+			]
+		},
 	]
 }
\ No newline at end of file
diff --git a/Engine/Plugins/Runtime/Database/SQLiteSupport/SQLiteSupport.uplugin b/Engine/Plugins/Runtime/Database/SQLiteSupport/SQLiteSupport.uplugin
--- a/Engine/Plugins/Runtime/Database/SQLiteSupport/SQLiteSupport.uplugin
+++ b/Engine/Plugins/Runtime/Database/SQLiteSupport/SQLiteSupport.uplugin
@@ -18,7 +18,10 @@
 		{
 			"Name": "SQLiteSupport",
 			"Type": "Runtime",
-			"LoadingPhase": "Default"
+			"LoadingPhase": "Default",
+			"BlacklistPlatforms": [
+				"HoloLens"
+			]
 		}
 	],
 	"Plugins": [
diff --git a/Engine/Source/Programs/UnrealBuildTool/Platform/HoloLens/HoloLensDeploy.cs b/Engine/Source/Programs/UnrealBuildTool/Platform/HoloLens/HoloLensDeploy.cs
--- a/Engine/Source/Programs/UnrealBuildTool/Platform/HoloLens/HoloLensDeploy.cs
+++ b/Engine/Source/Programs/UnrealBuildTool/Platform/HoloLens/HoloLensDeploy.cs
@@ -232,11 +232,11 @@ namespace UnrealBuildTool
 		private void MakePackage(TargetReceipt Receipt, TargetReceipt NewReceipt, WindowsArchitecture Architecture, List<string> UpdatedFiles)
 		{
 			string OutputName = String.Format("{0}_{1}_{2}_{3}", Receipt.TargetName, Receipt.Platform, Receipt.Configuration, WindowsExports.GetArchitectureSubpath(Architecture));
-			string IntermediateDirectory = Path.Combine(Receipt.ProjectDir.FullName, "Intermediate", "Deploy", WindowsExports.GetArchitectureSubpath(Architecture));
+			string IntermediateDirectory = Path.Combine(Receipt.ProjectDir != null ? Receipt.ProjectDir.FullName : UnrealBuildTool.EngineDirectory.FullName, "Intermediate", "Deploy", WindowsExports.GetArchitectureSubpath(Architecture));
 			string OutputDirectory = Receipt.Launch.Directory.FullName;
 			string OutputAppX = Path.Combine(OutputDirectory, OutputName + ".appx");
 			string SigningCertificate = @"Build\HoloLens\SigningCertificate.pfx";
-			string SigningCertificatePath = Path.Combine(Receipt.ProjectDir.FullName, SigningCertificate);
+			string SigningCertificatePath = Path.Combine(Receipt.ProjectDir != null ? Receipt.ProjectDir.FullName : UnrealBuildTool.EngineDirectory.FullName, SigningCertificate);
 
 			string MapFilename = Path.Combine(IntermediateDirectory, OutputName + ".pkgmap");
 			var LocalRoot = Receipt.ProjectDir;
@@ -264,7 +264,7 @@ namespace UnrealBuildTool
 							continue;
 						}
 
-						if (Product.Path.IsUnderDirectory(LocalRoot))
+						if (LocalRoot != null && Product.Path.IsUnderDirectory(LocalRoot))
 						{
 							Filename = Product.Path.MakeRelativeTo(LocalRoot.ParentDirectory);
 						}
@@ -291,7 +291,7 @@ namespace UnrealBuildTool
 						}
 
 						string Filename;
-						if (Dep.Path.IsUnderDirectory(LocalRoot))
+						if (LocalRoot != null && Dep.Path.IsUnderDirectory(LocalRoot))
 						{
 							Filename = Dep.Path.MakeRelativeTo(LocalRoot.ParentDirectory);
 						}
@@ -467,14 +467,14 @@ namespace UnrealBuildTool
 			System.DateTime PrepDeployStartTime = DateTime.UtcNow;
 
 			// Note: TargetReceipt.Read now expands path variables internally.
-			TargetReceipt NewReceipt;
-			var ReceiptFileName = TargetReceipt.GetDefaultPath(Receipt.ProjectDir, Receipt.TargetName, Receipt.Platform, Receipt.Configuration, "Multi");
-			if(!TargetReceipt.TryRead(ReceiptFileName, UnrealBuildTool.EngineDirectory, out NewReceipt))
+			TargetReceipt NewReceipt = null;
+			FileReference ReceiptFileName = TargetReceipt.GetDefaultPath(Receipt.ProjectDir != null ? Receipt.ProjectDir : UnrealBuildTool.EngineDirectory, Receipt.TargetName, Receipt.Platform, Receipt.Configuration, "Multi");
+			if (!TargetReceipt.TryRead(ReceiptFileName, UnrealBuildTool.EngineDirectory, out NewReceipt))
 			{
 				NewReceipt = new TargetReceipt(Receipt.ProjectFile, Receipt.TargetName, Receipt.TargetType, Receipt.Platform, Receipt.Configuration, Receipt.Version, "Multi");
 			}
 
-			AddWinMDReferencesFromReceipt(Receipt, Receipt.ProjectDir, UnrealBuildTool.EngineDirectory.ParentDirectory.FullName);
+			AddWinMDReferencesFromReceipt(Receipt, Receipt.ProjectDir != null ? Receipt.ProjectDir : UnrealBuildTool.EngineDirectory, UnrealBuildTool.EngineDirectory.ParentDirectory.FullName);
 
 			//PrepForUATPackageOrDeploy(InTarget.ProjectFile, InAppName, InTarget.ProjectDirectory.FullName, InTarget.OutputPath.FullName, TargetBuildEnvironment.RelativeEnginePath, false, "", false);
 			List<UnrealTargetConfiguration> TargetConfigs = new List<UnrealTargetConfiguration> { Receipt.Configuration };
@@ -497,10 +497,10 @@ namespace UnrealBuildTool
 
 			string AbsoluteExeDirectory = Path.GetDirectoryName(ExePaths[0]);
 			UnrealTargetPlatform Platform = UnrealTargetPlatform.HoloLens;
-			string IntermediateDirectory = Path.Combine(Receipt.ProjectDir.FullName, "Intermediate", "Deploy", WindowsExports.GetArchitectureSubpath(Arch));
-			List<string> UpdatedFiles = new HoloLensManifestGenerator().CreateManifest(Platform, Arch, AbsoluteExeDirectory, IntermediateDirectory, Receipt.ProjectFile, Receipt.ProjectDir.FullName, TargetConfigs, ExePaths, WinMDReferences);
+			string IntermediateDirectory = Path.Combine(Receipt.ProjectDir != null ? Receipt.ProjectDir.FullName : UnrealBuildTool.EngineDirectory.FullName, "Intermediate", "Deploy", WindowsExports.GetArchitectureSubpath(Arch));
+			List<string> UpdatedFiles = new HoloLensManifestGenerator().CreateManifest(Platform, Arch, AbsoluteExeDirectory, IntermediateDirectory, Receipt.ProjectFile, Receipt.ProjectDir != null ? Receipt.ProjectDir.FullName : UnrealBuildTool.EngineDirectory.FullName, TargetConfigs, ExePaths, WinMDReferences);
 
-			PrepForUATPackageOrDeploy(Receipt.ProjectFile, ProjectName, Receipt.ProjectDir.FullName, Arch, TargetConfigs, ExePaths, RelativeEnginePath, false, "", false);
+			PrepForUATPackageOrDeploy(Receipt.ProjectFile, ProjectName, Receipt.ProjectDir != null ? Receipt.ProjectDir.FullName : UnrealBuildTool.EngineDirectory.FullName, Arch, TargetConfigs, ExePaths, RelativeEnginePath, false, "", false);
 			MakePackage(Receipt, NewReceipt, Arch, UpdatedFiles);
 			CopyDataAndSymbolsBetweenReceipts(Receipt, NewReceipt, Arch);
 
diff --git a/Engine/Source/Programs/UnrealBuildTool/Platform/HoloLens/HoloLensProjectGenerator.cs b/Engine/Source/Programs/UnrealBuildTool/Platform/HoloLens/HoloLensProjectGenerator.cs
--- a/Engine/Source/Programs/UnrealBuildTool/Platform/HoloLens/HoloLensProjectGenerator.cs
+++ b/Engine/Source/Programs/UnrealBuildTool/Platform/HoloLens/HoloLensProjectGenerator.cs
@@ -156,7 +156,11 @@ namespace UnrealBuildTool
 		/// <returns>string    The name of the platform that VisualStudio recognizes</returns>
 		public override string GetVisualStudioPlatformName(UnrealTargetPlatform InPlatform, UnrealTargetConfiguration InConfiguration)
 		{
-			return "x64";
+			if (InPlatform == UnrealTargetPlatform.HoloLens)
+			{
+				return "arm64";
+			}
+			return InPlatform.ToString();
 		}
 
 		private bool IsValidHoloLensTarget(UnrealTargetPlatform InPlatform, TargetType InTargetType, FileReference InTargetFilePath)
diff --git a/Engine/Source/Runtime/ApplicationCore/Public/HoloLens/HoloLensPlatformSurvey.h b/Engine/Source/Runtime/ApplicationCore/Public/HoloLens/HoloLensPlatformSurvey.h
new file mode 100644
--- /dev/null
+++ b/Engine/Source/Runtime/ApplicationCore/Public/HoloLens/HoloLensPlatformSurvey.h
@@ -0,0 +1,8 @@
+// Copyright 1998-2019 Epic Games, Inc. All Rights Reserved.
+
+#pragma once
+
+#include "GenericPlatform/GenericPlatformSurvey.h"
+
+// generic implementation
+typedef FGenericPlatformSurvey FPlatformSurvey;
diff --git a/Engine/Source/Runtime/Online/BackgroundHTTP/Public/HoloLens/HoloLensPlatformBackgroundHttp.h b/Engine/Source/Runtime/Online/BackgroundHTTP/Public/HoloLens/HoloLensPlatformBackgroundHttp.h
new file mode 100644
--- /dev/null
+++ b/Engine/Source/Runtime/Online/BackgroundHTTP/Public/HoloLens/HoloLensPlatformBackgroundHttp.h
@@ -0,0 +1,8 @@
+// Copyright 1998-2019 Epic Games, Inc. All Rights Reserved.
+
+#pragma once
+
+#include "GenericPlatform/GenericPlatformBackgroundHttp.h"
+
+//Currently no implementation on PS4
+typedef FGenericPlatformBackgroundHttp FPlatformBackgroundHttp;
\ No newline at end of file
diff --git a/Engine/Source/Runtime/Windows/D3D11RHI/Private/Windows/WindowsD3D11Device.cpp b/Engine/Source/Runtime/Windows/D3D11RHI/Private/Windows/WindowsD3D11Device.cpp
--- a/Engine/Source/Runtime/Windows/D3D11RHI/Private/Windows/WindowsD3D11Device.cpp
+++ b/Engine/Source/Runtime/Windows/D3D11RHI/Private/Windows/WindowsD3D11Device.cpp
@@ -279,7 +279,11 @@ static void SafeTestForOptionalGraphicsTools(IDXGIAdapter* Adapter, D3D_FEATURE_
 	ID3D11Device* D3DDevice = NULL;
 	ID3D11DeviceContext* D3DDeviceContext = NULL;
 
+#if PLATFORM_HOLOLENS
+	static bool bIsWin10 = true;
+#else
 	static bool bIsWin10 = FWindowsPlatformMisc::VerifyWindowsVersion(10, 0);
+#endif
 	static bool bPrintErrorMessage = true;
 
 	// Test device creation.
