PR #5942: UBT: Fix duplicate platform check (Contributed by [REDACTED_PERSON])

#rnx

diff --git a/Engine/Source/Programs/UnrealBuildTool/ProjectFiles/ProjectFileGenerator.cs b/Engine/Source/Programs/UnrealBuildTool/ProjectFiles/ProjectFileGenerator.cs
--- a/Engine/Source/Programs/UnrealBuildTool/ProjectFiles/ProjectFileGenerator.cs
+++ b/Engine/Source/Programs/UnrealBuildTool/ProjectFiles/ProjectFileGenerator.cs
@@ -1839,7 +1839,7 @@ namespace UnrealBuildTool
 			}
 			if (BuildHostPlatform.Current.Platform == UnrealTargetPlatform.Win64)
 			{
-				return InPlatform == UnrealTargetPlatform.Win32 || InPlatform == UnrealTargetPlatform.Win32;
+				return InPlatform == UnrealTargetPlatform.Win32 || InPlatform == UnrealTargetPlatform.Win64;
 			}
 
 			throw new BuildException("Invalid RuntimePlatform:" + BuildHostPlatform.Current.Platform);
