Fix for broken IntelliSense data (PR# 5908 Thanks  zzz77!)

#rnx

diff --git a/Engine/Source/Programs/UnrealBuildTool/ProjectFiles/Project.cs b/Engine/Source/Programs/UnrealBuildTool/ProjectFiles/Project.cs
--- a/Engine/Source/Programs/UnrealBuildTool/ProjectFiles/Project.cs
+++ b/Engine/Source/Programs/UnrealBuildTool/ProjectFiles/Project.cs
@@ -307,7 +307,6 @@ namespace UnrealBuildTool
 					if(Def.EndsWith("_API", StringComparison.Ordinal))
 					{
                         CurDef = Def + "=";
-                        CurDef = Def + "_TYPE=";
 						Value = "";
 					}
 
