Suppress Hololens log spam when generating project files.

#rnx

diff --git a/Engine/Source/Programs/UnrealBuildTool/Platform/HoloLens/UEBuildHoloLens.cs b/Engine/Source/Programs/UnrealBuildTool/Platform/HoloLens/UEBuildHoloLens.cs
--- a/Engine/Source/Programs/UnrealBuildTool/Platform/HoloLens/UEBuildHoloLens.cs
+++ b/Engine/Source/Programs/UnrealBuildTool/Platform/HoloLens/UEBuildHoloLens.cs
@@ -171,12 +171,18 @@ namespace UnrealBuildTool
 			if (Target.Architecture.ToLower() == "arm64")
 			{
 				Target.HoloLensPlatform.Architecture = WindowsArchitecture.ARM64;
-				Log.TraceInformationOnce("Using ARM64 architecture for deploying to HoloLens device");
+				if(!Target.bGenerateProjectFiles)
+				{
+					Log.TraceInformationOnce("Using ARM64 architecture for deploying to HoloLens device");
+				}
 			}
 			else
 			{
 				Target.HoloLensPlatform.Architecture = WindowsArchitecture.x64;
-				Log.TraceInformationOnce("Using x64 architecture for deploying to HoloLens emulator");
+				if (!Target.bGenerateProjectFiles)
+				{
+					Log.TraceInformationOnce("Using x64 architecture for deploying to HoloLens emulator");
+				}
 			}
 
 			Target.WindowsPlatform.Compiler = Target.HoloLensPlatform.Compiler;
@@ -221,15 +227,18 @@ namespace UnrealBuildTool
 
 			Target.HoloLensPlatform.Win10SDKVersion = new Version(Environment.WindowsSdkVersion.ToString());
 
-			Log.TraceInformationOnce("Building using Windows SDK version {0} for HoloLens", Target.HoloLensPlatform.Win10SDKVersion);
-
-			if (Target.HoloLensPlatform.Win10SDKVersion < MinimumSDKVersionRecommended)
-			{
-				Log.TraceWarning("Your Windows SDK version {0} is older than the minimum recommended version ({1}) for HoloLens.  Consider upgrading.", Target.HoloLensPlatform.Win10SDKVersion, MinimumSDKVersionRecommended);
-			}
-			else if (Target.HoloLensPlatform.Win10SDKVersion > MaximumSDKVersionTested)
+			if(!Target.bGenerateProjectFiles)
 			{
-				Log.TraceInformationOnce("Your Windows SDK version ({0}) for HoloLens is newer than the highest tested with this version of UBT ({1}).  This is probably fine, but if you encounter issues consider using an earlier SDK.", Target.HoloLensPlatform.Win10SDKVersion, MaximumSDKVersionTested);
+				Log.TraceInformationOnce("Building using Windows SDK version {0} for HoloLens", Target.HoloLensPlatform.Win10SDKVersion);
+
+				if (Target.HoloLensPlatform.Win10SDKVersion < MinimumSDKVersionRecommended)
+				{
+					Log.TraceWarning("Your Windows SDK version {0} is older than the minimum recommended version ({1}) for HoloLens.  Consider upgrading.", Target.HoloLensPlatform.Win10SDKVersion, MinimumSDKVersionRecommended);
+				}
+				else if (Target.HoloLensPlatform.Win10SDKVersion > MaximumSDKVersionTested)
+				{
+					Log.TraceInformationOnce("Your Windows SDK version ({0}) for HoloLens is newer than the highest tested with this version of UBT ({1}).  This is probably fine, but if you encounter issues consider using an earlier SDK.", Target.HoloLensPlatform.Win10SDKVersion, MaximumSDKVersionTested);
+				}
 			}
 
 			HoloLensExports.InitWindowsSdkToolPath(Target.HoloLensPlatform.Win10SDKVersion.ToString());
