Fix incorrect executable path in VSCode projects that use unique build environments.

diff --git a/Engine/Source/Programs/UnrealBuildTool/ProjectFiles/VisualStudioCode/VSCodeProjectFileGenerator.cs b/Engine/Source/Programs/UnrealBuildTool/ProjectFiles/VisualStudioCode/VSCodeProjectFileGenerator.cs
--- a/Engine/Source/Programs/UnrealBuildTool/ProjectFiles/VisualStudioCode/VSCodeProjectFileGenerator.cs
+++ b/Engine/Source/Programs/UnrealBuildTool/ProjectFiles/VisualStudioCode/VSCodeProjectFileGenerator.cs
@@ -917,14 +917,9 @@ namespace UnrealBuildTool
 
 			// Get the executable name (minus any platform or config suffixes)
 			string BaseExeName = TargetName;
-			if (!bShouldCompileMonolithic && TargetRulesType != TargetType.Program)
+			if(Target.TargetRules != null && Target.TargetRules.BuildEnvironment == TargetBuildEnvironment.Shared && TargetRulesType != TargetType.Program)
 			{
-				// Figure out what the compiled binary will be called so that we can point the IDE to the correct file
-				string TargetConfigurationName = TargetRulesType.ToString();
-				if (TargetConfigurationName != TargetType.Game.ToString() && TargetConfigurationName != TargetType.Program.ToString())
-				{
-					BaseExeName = "UE4" + TargetConfigurationName;
-				}
+				BaseExeName = UEBuildTarget.GetAppNameForTargetType(TargetRulesType);
 			}
 
 			// Make the output file path
