Fix build error due to missing PYTHON_VER variable.

#rnx

diff --git a/Engine/Source/Programs/UnrealBuildTool/Platform/HTML5/HTML5SDKInfo.cs b/Engine/Source/Programs/UnrealBuildTool/Platform/HTML5/HTML5SDKInfo.cs
--- a/Engine/Source/Programs/UnrealBuildTool/Platform/HTML5/HTML5SDKInfo.cs
+++ b/Engine/Source/Programs/UnrealBuildTool/Platform/HTML5/HTML5SDKInfo.cs
@@ -16,7 +16,7 @@ namespace UnrealBuildTool
 	public class HTML5SDKInfo
 	{
 		static string NODE_VER = "8.9.1_64bit";
-//		static string PYTHON_VER = "2.7.13.1_64bit"; // Only used on Windows; other platforms use built-in Python.
+		static string PYTHON_VER = "2.7.13.1_64bit"; // Only used on Windows; other platforms use built-in Python.
 
 		static string LLVM_VER = "e1.38.31_64bit";
 		static string SDKVersion = "1.38.31";
