UBT: Force UHT to run when the path to the tool changes.

#rnx

diff --git a/Engine/Source/Programs/UnrealBuildTool/System/ExternalExecution.cs b/Engine/Source/Programs/UnrealBuildTool/System/ExternalExecution.cs
--- a/Engine/Source/Programs/UnrealBuildTool/System/ExternalExecution.cs
+++ b/Engine/Source/Programs/UnrealBuildTool/System/ExternalExecution.cs
@@ -1051,6 +1051,20 @@ namespace UnrealBuildTool
 					bUHTNeedsToRun = true;
 				}
 
+				// Check we're not using a different version of UHT
+				FileReference ToolInfoFile = ModuleInfoFileName.ChangeExtension(".uhtpath");
+				if(!bUHTNeedsToRun)
+				{
+					if(!FileReference.Exists(ToolInfoFile))
+					{
+						bUHTNeedsToRun = true;
+					}
+					else if(FileReference.ReadAllText(ToolInfoFile) != HeaderToolReceipt.FullName)
+					{
+						bUHTNeedsToRun = true;
+					}
+				}
+
 				// Get the file containing dependencies for the generated code
 				FileReference ExternalDependenciesFile = ModuleInfoFileName.ChangeExtension(".deps");
 				if (AreExternalDependenciesOutOfDate(ExternalDependenciesFile))
@@ -1172,6 +1186,10 @@ namespace UnrealBuildTool
 
 					Log.TraceInformation("Reflection code generated for {0} in {1} seconds", ActualTargetName, s.Elapsed.TotalSeconds);
 
+					// Update the tool info file
+					DirectoryReference.CreateDirectory(ToolInfoFile.Directory);
+					FileReference.WriteAllText(ToolInfoFile, HeaderToolReceipt.FullName);
+
 					// Now that UHT has successfully finished generating code, we need to update all cached FileItems in case their last write time has changed.
 					// Otherwise UBT might not detect changes UHT made.
 					using(Timeline.ScopeEvent("ExternalExecution.ResetCachedHeaderInfo()"))
