UBT: Fix race condition in spawning UHT and reading output.

#rnx

diff --git a/Engine/Source/Programs/UnrealBuildTool/System/ExternalExecution.cs b/Engine/Source/Programs/UnrealBuildTool/System/ExternalExecution.cs
--- a/Engine/Source/Programs/UnrealBuildTool/System/ExternalExecution.cs
+++ b/Engine/Source/Programs/UnrealBuildTool/System/ExternalExecution.cs
@@ -955,42 +955,21 @@ namespace UnrealBuildTool
 			}
 		}
 
-		/// <summary>
-		/// Run an external exe (and capture the output), given the exe path and the commandline.
-		/// </summary>
-		public static int RunExternalDotNETExecutable(string ExePath, string Commandline)
-		{
-#if NET_CORE
-			ProcessStartInfo ExeInfo = new ProcessStartInfo("dotnet", ExePath + " " + Commandline);
-#else
-			ProcessStartInfo ExeInfo = new ProcessStartInfo(ExePath, Commandline);
-#endif
-			Log.TraceVerbose("RunExternalExecutable {0} {1}", ExePath, Commandline);
-			ExeInfo.UseShellExecute = false;
-			ExeInfo.RedirectStandardOutput = true;
-			using (Process GameProcess = Process.Start(ExeInfo))
-			{
-				GameProcess.BeginOutputReadLine();
-				GameProcess.OutputDataReceived += PrintProcessOutputAsync;
-				GameProcess.WaitForExit();
-
-				return GameProcess.ExitCode;
-			}
-		}
-
 		/// <summary>
 		/// Run an external native executable (and capture the output), given the executable path and the commandline.
 		/// </summary>
 		public static int RunExternalNativeExecutable(FileReference ExePath, string Commandline)
 		{
-			ProcessStartInfo ExeInfo = new ProcessStartInfo(ExePath.FullName, Commandline);
 			Log.TraceVerbose("RunExternalExecutable {0} {1}", ExePath.FullName, Commandline);
-			ExeInfo.UseShellExecute = false;
-			ExeInfo.RedirectStandardOutput = true;
-			using (Process GameProcess = Process.Start(ExeInfo))
+			using (Process GameProcess = new Process())
 			{
-				GameProcess.BeginOutputReadLine();
+				GameProcess.StartInfo.FileName = ExePath.FullName;
+				GameProcess.StartInfo.Arguments = Commandline;
+				GameProcess.StartInfo.UseShellExecute = false;
+				GameProcess.StartInfo.RedirectStandardOutput = true;
 				GameProcess.OutputDataReceived += PrintProcessOutputAsync;
+				GameProcess.Start();
+				GameProcess.BeginOutputReadLine();
 				GameProcess.WaitForExit();
 
 				return GameProcess.ExitCode;
