Added ability to override Unity and Optimization per module in BuildConfiguration. This should eliminate most needs to disable unity builds in Build.cs (as disabling it for local iteration times was the primary argument for it).
Thus we renamed bFasterWithoutUnity to bUseUnity (note the inversion of the statement) which is what you should use if a module does not support building in Unity mode.

diff --git a/Engine/Source/Programs/UnrealBuildTool/Configuration/ModuleRules.cs b/Engine/Source/Programs/UnrealBuildTool/Configuration/ModuleRules.cs
--- a/Engine/Source/Programs/UnrealBuildTool/Configuration/ModuleRules.cs
+++ b/Engine/Source/Programs/UnrealBuildTool/Configuration/ModuleRules.cs
@@ -440,10 +440,31 @@ namespace UnrealBuildTool
 		/// </summary>
 		public string BinariesSubFolder = "";
 
+		private CodeOptimization? OptimizeCodeOverride;
+
 		/// <summary>
 		/// When this module's code should be optimized.
 		/// </summary>
-		public CodeOptimization OptimizeCode = CodeOptimization.Default;
+		public CodeOptimization OptimizeCode
+		{
+			get
+			{
+				if (OptimizeCodeOverride.HasValue)
+					return OptimizeCodeOverride.Value;
+
+				bool? ShouldOptimizeCode = null;
+				if (Target.EnableOptimizeCodeForModules?.Contains(Name) ?? false)
+					ShouldOptimizeCode = true;
+				if (Target.DisableOptimizeCodeForModules?.Contains(Name) ?? false)
+					ShouldOptimizeCode = false;
+
+				if (!ShouldOptimizeCode.HasValue)
+					return CodeOptimization.Default;
+
+				return ShouldOptimizeCode.Value ? CodeOptimization.Always : CodeOptimization.Never;
+			}
+			set { OptimizeCodeOverride = value; }
+		}
 
 		/// <summary>
 		/// Explicit private PCH for this module. Implies that this module will not use a shared PCH.
@@ -516,7 +537,29 @@ namespace UnrealBuildTool
 		/// <summary>
 		/// If true and unity builds are enabled, this module will build without unity.
 		/// </summary>
-		public bool bFasterWithoutUnity = false;
+		[Obsolete("bFasterWithoutUnity has been deprecated in favor of setting 'bUseUnity' on a per module basis in BuildConfiguration")]
+		public bool bFasterWithoutUnity
+		{
+			set { bUseUnity = value; }
+		}
+
+		private bool? bUseUnityOverride;
+		/// <summary>
+		/// If unity builds are enabled this can be used to override if this specific module will build using Unity.
+		/// This is set using the per module configurations in BuildConfiguration.
+		/// </summary>
+		public bool bUseUnity
+		{
+			set { bUseUnityOverride = value; }
+			get
+			{
+				bool UseUnity = true;
+				if (Target.DisableUnityBuildForModules?.Contains(Name) ?? false)
+					UseUnity = false;
+				return bUseUnityOverride ?? UseUnity;
+			}
+		}
+
 
 		/// <summary>
 		/// The number of source files in this module before unity build will be activated for that module.  If set to
@@ -629,8 +672,14 @@ namespace UnrealBuildTool
 		/// <summary>
 		/// List of system/library paths (directory of .lib files) - typically used for External (third party) modules
 		/// </summary>
+		[Obsolete("For external libraries use the full path in PublicAdditionalLibraries, if its a system library then use PublicSystemLibraries/PublicSystemLibraryPaths")]
 		public List<string> PublicLibraryPaths = new List<string>();
 
+		/// <summary>
+		/// List of system library paths (directory of .lib files) - for External (third party) modules please use the PublicAdditionalLibaries instead
+		/// </summary>
+		public List<string> PublicSystemLibraryPaths = new List<string>();
+
 		/// <summary>
 		/// List of search paths for libraries at runtime (eg. .so files)
 		/// </summary>
@@ -646,6 +695,11 @@ namespace UnrealBuildTool
 		/// </summary>
 		public List<string> PublicAdditionalLibraries = new List<string>();
 
+		/// <summary>
+		/// List of system libraries to use - these are typically referenced via name and then found via the system paths. If you need to reference a .lib file use the PublicAdditionalLibraries instead
+		/// </summary>
+		public List<string> PublicSystemLibraries = new List<string>();
+
 		/// <summary>
 		/// List of XCode frameworks (iOS and MacOS)
 		/// </summary>
diff --git a/Engine/Source/Programs/UnrealBuildTool/Configuration/TargetRules.cs b/Engine/Source/Programs/UnrealBuildTool/Configuration/TargetRules.cs
--- a/Engine/Source/Programs/UnrealBuildTool/Configuration/TargetRules.cs
+++ b/Engine/Source/Programs/UnrealBuildTool/Configuration/TargetRules.cs
@@ -1215,6 +1215,25 @@ namespace UnrealBuildTool
 		[CommandLine("-LinkerArguments=")]
 		public string AdditionalLinkerArguments;
 
+
+		/// <summary>
+		/// List of modules to disable unity builds for
+		/// </summary>
+		[XmlConfigFile(Category = "ModuleConfiguration", Name = "DisableUnityBuild")]
+		public string[] DisableUnityBuildForModules = null;
+
+		/// <summary>
+		///  List of modules to enable optimizations for
+		/// </summary>
+		[XmlConfigFile(Category = "ModuleConfiguration", Name = "EnableOptimizeCode")]
+		public string[] EnableOptimizeCodeForModules = null;
+
+		/// <summary>
+		/// List of modules to disable optimizations for
+		/// </summary>
+		[XmlConfigFile(Category = "ModuleConfiguration", Name = "DisableOptimizeCode")]
+		public string[] DisableOptimizeCodeForModules = null;
+
 		/// <summary>
 		/// When generating project files, specifies the name of the project file to use when there are multiple targets of the same type.
 		/// </summary>
@@ -2453,9 +2472,25 @@ namespace UnrealBuildTool
 			get { return Inner.bIsEngineInstalled; }
 		}
 
-		#if !__MonoCS__
-		#pragma warning restore C1591
-		#endif
+
+		public IReadOnlyList<string> DisableUnityBuildForModules
+		{
+			get { return Inner.DisableUnityBuildForModules; }
+		}
+
+		public IReadOnlyList<string> EnableOptimizeCodeForModules
+		{
+			get { return Inner.EnableOptimizeCodeForModules; }
+		}
+
+		public IReadOnlyList<string> DisableOptimizeCodeForModules
+		{
+			get { return Inner.DisableOptimizeCodeForModules; }
+		}
+
+#if !__MonoCS__
+#pragma warning restore C1591
+#endif
 		#endregion
 
 		/// <summary>
diff --git a/Engine/Source/Programs/UnrealBuildTool/Configuration/UEBuildModuleCPP.cs b/Engine/Source/Programs/UnrealBuildTool/Configuration/UEBuildModuleCPP.cs
--- a/Engine/Source/Programs/UnrealBuildTool/Configuration/UEBuildModuleCPP.cs
+++ b/Engine/Source/Programs/UnrealBuildTool/Configuration/UEBuildModuleCPP.cs
@@ -320,6 +320,7 @@ namespace UnrealBuildTool
 
 			// Should we use unity build mode for this module?
 			bool bModuleUsesUnityBuild = false;
+	
 			if (Target.bUseUnityBuild || Target.bForceUnityBuild)
 			{
 				if (Target.bForceUnityBuild)
@@ -327,9 +328,9 @@ namespace UnrealBuildTool
 					Log.TraceVerbose("Module '{0}' using unity build mode (bForceUnityBuild enabled for this module)", this.Name);
 					bModuleUsesUnityBuild = true;
 				}
-				else if (Rules.bFasterWithoutUnity)
+				else if (!Rules.bUseUnity)
 				{
-					Log.TraceVerbose("Module '{0}' not using unity build mode (bFasterWithoutUnity enabled for this module)", this.Name);
+					Log.TraceVerbose("Module '{0}' not using unity build mode (bUseUnity disabled for this module)", this.Name);
 					bModuleUsesUnityBuild = false;
 				}
 				else if (InputFiles.CPPFiles.Count < MinSourceFilesForUnityBuild)
@@ -1114,7 +1115,7 @@ namespace UnrealBuildTool
 			CppCompileEnvironment Result = new CppCompileEnvironment(BaseCompileEnvironment);
 
 			// Override compile environment
-			Result.bFasterWithoutUnity = Rules.bFasterWithoutUnity;
+			Result.bUseUnity = Rules.bUseUnity;
 			Result.bOptimizeCode = ShouldEnableOptimization(Rules.OptimizeCode, Target.Configuration, Rules.bTreatAsEngineModule);
 			Result.bUseRTTI |= Rules.bUseRTTI;
 			Result.bUseAVX = Rules.bUseAVX;
diff --git a/Engine/Source/Programs/UnrealBuildTool/System/CppCompileEnvironment.cs b/Engine/Source/Programs/UnrealBuildTool/System/CppCompileEnvironment.cs
--- a/Engine/Source/Programs/UnrealBuildTool/System/CppCompileEnvironment.cs
+++ b/Engine/Source/Programs/UnrealBuildTool/System/CppCompileEnvironment.cs
@@ -144,9 +144,10 @@ namespace UnrealBuildTool
 		public bool bEnableBufferSecurityChecks = true;
 
 		/// <summary>
-		/// If true and unity builds are enabled, this module will build without unity.
+		/// If unity builds are enabled this can be used to override if this specific module will build using Unity.
+		/// This is set using the per module configurations in BuildConfiguration.
 		/// </summary>
-		public bool bFasterWithoutUnity = false;
+		public bool bUseUnity = false;
 
 		/// <summary>
 		/// The number of source files in this module before unity build will be activated for that module.  If set to
@@ -386,7 +387,7 @@ namespace UnrealBuildTool
 			bUseRTTI = Other.bUseRTTI;
 			bUseInlining = Other.bUseInlining;
 			bUseAVX = Other.bUseAVX;
-			bFasterWithoutUnity = Other.bFasterWithoutUnity;
+			bUseUnity = Other.bUseUnity;
 			MinSourceFilesForUnityBuildOverride = Other.MinSourceFilesForUnityBuildOverride;
 			MinFilesUsingPrecompiledHeaderOverride = Other.MinFilesUsingPrecompiledHeaderOverride;
 			bBuildLocallyWithSNDBS = Other.bBuildLocallyWithSNDBS;
