Removed work-in-progress code that accidentially got submitted as part of per module configuration (change 7099374)

diff --git a/Engine/Source/Programs/UnrealBuildTool/Configuration/ModuleRules.cs b/Engine/Source/Programs/UnrealBuildTool/Configuration/ModuleRules.cs
--- a/Engine/Source/Programs/UnrealBuildTool/Configuration/ModuleRules.cs
+++ b/Engine/Source/Programs/UnrealBuildTool/Configuration/ModuleRules.cs
@@ -672,14 +672,8 @@ namespace UnrealBuildTool
 		/// <summary>
 		/// List of system/library paths (directory of .lib files) - typically used for External (third party) modules
 		/// </summary>
-		[Obsolete("For external libraries use the full path in PublicAdditionalLibraries, if its a system library then use PublicSystemLibraries/PublicSystemLibraryPaths")]
 		public List<string> PublicLibraryPaths = new List<string>();
 
-		/// <summary>
-		/// List of system library paths (directory of .lib files) - for External (third party) modules please use the PublicAdditionalLibaries instead
-		/// </summary>
-		public List<string> PublicSystemLibraryPaths = new List<string>();
-
 		/// <summary>
 		/// List of search paths for libraries at runtime (eg. .so files)
 		/// </summary>
@@ -695,11 +689,6 @@ namespace UnrealBuildTool
 		/// </summary>
 		public List<string> PublicAdditionalLibraries = new List<string>();
 
-		/// <summary>
-		/// List of system libraries to use - these are typically referenced via name and then found via the system paths. If you need to reference a .lib file use the PublicAdditionalLibraries instead
-		/// </summary>
-		public List<string> PublicSystemLibraries = new List<string>();
-
 		/// <summary>
 		/// List of XCode frameworks (iOS and MacOS)
 		/// </summary>
