Fix single-file compile when UHT needs to be executed. Was removing binaries from the target before determining the full set of modules for UHT to run.

#rnx

diff --git a/Engine/Source/Programs/UnrealBuildTool/Configuration/UEBuildTarget.cs b/Engine/Source/Programs/UnrealBuildTool/Configuration/UEBuildTarget.cs
--- a/Engine/Source/Programs/UnrealBuildTool/Configuration/UEBuildTarget.cs
+++ b/Engine/Source/Programs/UnrealBuildTool/Configuration/UEBuildTarget.cs
@@ -680,17 +680,6 @@ namespace UnrealBuildTool
 				Target.PreBuildSetup();
 			}
 
-			// If we're just compiling a single file, filter the list of binaries to only include the file we're interested in.
-			FileReference SingleFileToCompile = Descriptor.SingleFileToCompile;
-			if (SingleFileToCompile != null)
-			{
-				Target.Binaries.RemoveAll(x => !x.Modules.Any(y => SingleFileToCompile.IsUnderDirectory(y.ModuleDirectory)));
-				if(Target.Binaries.Count == 0)
-				{
-					throw new BuildException("Couldn't find any module containing {0} in {1}.", SingleFileToCompile, Target.TargetName);
-				}
-			}
-
 			return Target;
 		}
 
@@ -1426,7 +1415,7 @@ namespace UnrealBuildTool
 		/// <summary>
 		/// Builds the target, appending list of output files and returns building result.
 		/// </summary>
-		public TargetMakefile Build(BuildConfiguration BuildConfiguration, ISourceFileWorkingSet WorkingSet, bool bIsAssemblingBuild)
+		public TargetMakefile Build(BuildConfiguration BuildConfiguration, ISourceFileWorkingSet WorkingSet, bool bIsAssemblingBuild, FileReference SingleFileToCompile)
 		{
 			CppConfiguration CppConfiguration = GetCppConfiguration(Configuration);
 
@@ -1604,11 +1593,22 @@ namespace UnrealBuildTool
 				}
 			}
 
+			// Find the set of binaries to build. If we're just compiling a single file, filter the list of binaries to only include the file we're interested in.
+			List<UEBuildBinary> BuildBinaries = Binaries;
+			if (SingleFileToCompile != null)
+			{
+				BuildBinaries = Binaries.Where(x => x.Modules.Any(y => SingleFileToCompile.IsUnderDirectory(y.ModuleDirectory))).ToList();
+				if (BuildBinaries.Count == 0)
+				{
+					throw new BuildException("Couldn't find any module containing {0} in {1}.", SingleFileToCompile, TargetName);
+				}
+			}
+
 			// Build the target's binaries.
 			DirectoryReference ExeDir = GetExecutableDir();
 			using(Timeline.ScopeEvent("UEBuildBinary.Build()"))
 			{
-				foreach (UEBuildBinary Binary in Binaries)
+				foreach (UEBuildBinary Binary in BuildBinaries)
 				{
 					List<FileItem> BinaryOutputItems = Binary.Build(Rules, TargetToolChain, GlobalCompileEnvironment, GlobalLinkEnvironment, WorkingSet, ExeDir, Makefile);
 					Makefile.OutputItems.AddRange(BinaryOutputItems);
diff --git a/Engine/Source/Programs/UnrealBuildTool/Modes/BuildMode.cs b/Engine/Source/Programs/UnrealBuildTool/Modes/BuildMode.cs
--- a/Engine/Source/Programs/UnrealBuildTool/Modes/BuildMode.cs
+++ b/Engine/Source/Programs/UnrealBuildTool/Modes/BuildMode.cs
@@ -462,7 +462,7 @@ namespace UnrealBuildTool
 				using(Timeline.ScopeEvent("UEBuildTarget.Build()"))
 				{
 					const bool bIsAssemblingBuild = true;
-					Makefile = Target.Build(BuildConfiguration, WorkingSet, bIsAssemblingBuild);
+					Makefile = Target.Build(BuildConfiguration, WorkingSet, bIsAssemblingBuild, TargetDescriptor.SingleFileToCompile);
 				}
 
 				// Save the pre-build scripts onto the makefile
