Fix changes made in Hololens integration causing Visual Studio Express to crash when compiling.

#rnx

diff --git a/Engine/Source/Programs/UnrealBuildTool/Platform/Windows/VCEnvironment.cs b/Engine/Source/Programs/UnrealBuildTool/Platform/Windows/VCEnvironment.cs
--- a/Engine/Source/Programs/UnrealBuildTool/Platform/Windows/VCEnvironment.cs
+++ b/Engine/Source/Programs/UnrealBuildTool/Platform/Windows/VCEnvironment.cs
@@ -144,13 +144,15 @@ namespace UnrealBuildTool
 
 			// Add both toolchain paths to the PATH environment variable. There are some support DLLs which are only added to one of the paths, but which the toolchain in the other directory
 			// needs to run (eg. mspdbcore.dll).
-			if (Environment.Is64BitProcess && Architecture != WindowsArchitecture.x64)
+			if (Architecture == WindowsArchitecture.x64)
 			{
 				AddDirectoryToPath(GetVCToolPath(Compiler, ToolChainDir, WindowsArchitecture.x64));
+				AddDirectoryToPath(GetVCToolPath(Compiler, ToolChainDir, WindowsArchitecture.x86));
 			}
-			else if (!Environment.Is64BitProcess && Architecture != WindowsArchitecture.x86)
+			else if (Architecture == WindowsArchitecture.x86)
 			{
 				AddDirectoryToPath(GetVCToolPath(Compiler, ToolChainDir, WindowsArchitecture.x86));
+				AddDirectoryToPath(GetVCToolPath(Compiler, ToolChainDir, WindowsArchitecture.x64));
 			}
 
 
@@ -181,53 +183,64 @@ namespace UnrealBuildTool
 		/// <returns>Directory containing the 32-bit toolchain binaries</returns>
 		static DirectoryReference GetVCToolPath(WindowsCompiler Compiler, DirectoryReference VCToolChainDir, WindowsArchitecture Architecture)
 		{
-			FileReference CompilerPath = null;
 			if (Compiler >= WindowsCompiler.VisualStudio2017)
 			{
-				CompilerPath = FileReference.Combine(
-					VCToolChainDir,
-					"bin",
-					Environment.Is64BitProcess ? "Hostx64" : "Hostx86",
-					WindowsExports.GetArchitectureSubpath(Architecture),
-					"cl.exe");
+				FileReference NativeCompilerPath = FileReference.Combine(VCToolChainDir, "bin", "HostX64", WindowsExports.GetArchitectureSubpath(Architecture), "cl.exe");
+				if (FileReference.Exists(NativeCompilerPath))
+				{
+					return NativeCompilerPath.Directory;
+				}
+
+				FileReference CrossCompilerPath = FileReference.Combine(VCToolChainDir, "bin", "HostX86", WindowsExports.GetArchitectureSubpath(Architecture), "cl.exe");
+				if (FileReference.Exists(CrossCompilerPath))
+				{
+					return CrossCompilerPath.Directory;
+				}
 			}
 			else
 			{
 				if (Architecture == WindowsArchitecture.x86)
 				{
-					CompilerPath = FileReference.Combine(VCToolChainDir, "bin", "cl.exe");
+					FileReference CompilerPath = FileReference.Combine(VCToolChainDir, "bin", "cl.exe");
+					if(FileReference.Exists(CompilerPath))
+					{
+						return CompilerPath.Directory;
+					}
 				}
 				else if (Architecture == WindowsArchitecture.x64)
 				{
-					if (Environment.Is64BitProcess)
+					// Use the native 64-bit compiler if present
+					FileReference NativeCompilerPath = FileReference.Combine(VCToolChainDir, "bin", "amd64", "cl.exe");
+					if (FileReference.Exists(NativeCompilerPath))
 					{
-						CompilerPath = FileReference.Combine(VCToolChainDir, "bin", "amd64", "cl.exe");
+						return NativeCompilerPath.Directory;
 					}
 
-					if (!Environment.Is64BitProcess || !FileReference.Exists(CompilerPath))
+					// Otherwise use the amd64-on-x86 compiler. VS2012 Express only includes the latter.
+					FileReference CrossCompilerPath = FileReference.Combine(VCToolChainDir, "bin", "x86_amd64", "cl.exe");
+					if (FileReference.Exists(CrossCompilerPath))
 					{
-						CompilerPath = FileReference.Combine(VCToolChainDir, "bin", "x86_amd64", "cl.exe");
+						return CrossCompilerPath.Directory;
 					}
 				}
 				else if (Architecture == WindowsArchitecture.ARM32)
 				{
-					if (Environment.Is64BitProcess)
+					// Use the native 64-bit compiler if present
+					FileReference NativeCompilerPath = FileReference.Combine(VCToolChainDir, "bin", "amd64_arm", "cl.exe");
+					if (FileReference.Exists(NativeCompilerPath))
 					{
-						CompilerPath = FileReference.Combine(VCToolChainDir, "bin", "amd64_arm", "cl.exe");
+						return NativeCompilerPath.Directory;
 					}
 
-					if (!Environment.Is64BitProcess || !FileReference.Exists(CompilerPath))
+					// Otherwise use the amd64-on-x86 compiler. VS2012 Express only includes the latter.
+					FileReference CrossCompilerPath = FileReference.Combine(VCToolChainDir, "bin", "x86_arm", "cl.exe");
+					if (FileReference.Exists(CrossCompilerPath))
 					{
-						CompilerPath = FileReference.Combine(VCToolChainDir, "bin", "x86_arm", "cl.exe");
+						return CrossCompilerPath.Directory;
 					}
 				}
 			}
 
-			if (CompilerPath != null && FileReference.Exists(CompilerPath))
-			{
-				return CompilerPath.Directory;
-			}
-
 			throw new BuildException("No required compiler toolchain found in {0}", VCToolChainDir);
 		}
 
@@ -376,18 +389,18 @@ namespace UnrealBuildTool
 		{
 			// Add the standard Visual C++ include paths
 			IncludePaths.Add(DirectoryReference.Combine(ToolChainDir, "INCLUDE"));
-			string archPath = WindowsExports.GetArchitectureSubpath(Architecture);
+			string ArchFolder = WindowsExports.GetArchitectureSubpath(Architecture);
 
 			// Add the standard Visual C++ library paths
 			if (ToolChain >= WindowsCompiler.VisualStudio2017)
 			{
 				if (Platform == UnrealTargetPlatform.HoloLens)
 				{
-					LibraryPaths.Add(DirectoryReference.Combine(ToolChainDir, "lib", archPath, "store"));
+					LibraryPaths.Add(DirectoryReference.Combine(ToolChainDir, "lib", ArchFolder, "store"));
 				}
 				else
 				{
-					LibraryPaths.Add(DirectoryReference.Combine(ToolChainDir, "lib", archPath));
+					LibraryPaths.Add(DirectoryReference.Combine(ToolChainDir, "lib", ArchFolder));
 				}
 			}
 			else
@@ -423,7 +436,7 @@ namespace UnrealBuildTool
 				IncludePaths.Add(DirectoryReference.Combine(IncludeRootDir, "ucrt"));
 
 				DirectoryReference LibraryRootDir = DirectoryReference.Combine(Pair.Value, "lib", Pair.Key.ToString());
-				LibraryPaths.Add(DirectoryReference.Combine(LibraryRootDir, "ucrt", archPath));
+				LibraryPaths.Add(DirectoryReference.Combine(LibraryRootDir, "ucrt", ArchFolder));
 			}
 
 			// Add the NETFXSDK include path. We need this for SwarmInterface.
@@ -431,7 +444,7 @@ namespace UnrealBuildTool
 			if(WindowsPlatform.TryGetNetFxSdkInstallDir(out NetFxSdkDir))
 			{
 				IncludePaths.Add(DirectoryReference.Combine(NetFxSdkDir, "include", "um"));
-				LibraryPaths.Add(DirectoryReference.Combine(NetFxSdkDir, "lib", "um", archPath));
+				LibraryPaths.Add(DirectoryReference.Combine(NetFxSdkDir, "lib", "um", ArchFolder));
 			}
 
 			// Add the Windows SDK paths
@@ -444,8 +457,8 @@ namespace UnrealBuildTool
 				IncludePaths.Add(DirectoryReference.Combine(IncludeRootDir, "winrt"));
 
 				DirectoryReference LibraryRootDir = DirectoryReference.Combine(WindowsSdkDir, "lib", WindowsSdkVersion.ToString());
-				LibraryPaths.Add(DirectoryReference.Combine(LibraryRootDir, "ucrt", archPath));
-				LibraryPaths.Add(DirectoryReference.Combine(LibraryRootDir, "um", archPath));
+				LibraryPaths.Add(DirectoryReference.Combine(LibraryRootDir, "ucrt", ArchFolder));
+				LibraryPaths.Add(DirectoryReference.Combine(LibraryRootDir, "um", ArchFolder));
 			}
 			else
 			{
@@ -455,7 +468,7 @@ namespace UnrealBuildTool
 				IncludePaths.Add(DirectoryReference.Combine(IncludeRootDir, "winrt"));
 
 				DirectoryReference LibraryRootDir = DirectoryReference.Combine(WindowsSdkDir, "lib", "winv6.3");
-				LibraryPaths.Add(DirectoryReference.Combine(LibraryRootDir, "um", archPath));
+				LibraryPaths.Add(DirectoryReference.Combine(LibraryRootDir, "um", ArchFolder));
 			}
 		}
 
