Put dsym code that was lost in CL# 6109577 back where it should be
#rnx

diff --git a/Engine/Source/Programs/UnrealBuildTool/ToolChain/AppleToolChain.cs b/Engine/Source/Programs/UnrealBuildTool/ToolChain/AppleToolChain.cs
--- a/Engine/Source/Programs/UnrealBuildTool/ToolChain/AppleToolChain.cs
+++ b/Engine/Source/Programs/UnrealBuildTool/ToolChain/AppleToolChain.cs
@@ -164,11 +164,50 @@ namespace UnrealBuildTool
 		{
 			FileReference DsymutilLocation = new FileReference("/usr/bin/dsymutil");
 
+			// dsymutil before 10.0.1 has a bug that causes issues, it's fixed in autosdks but not everyone has those set up so for the timebeing we have
+			// a version in P4 - first determine if we need it
 			string DsymutilVersionString = Utils.RunLocalProcessAndReturnStdOut(DsymutilLocation.FullName, "-version");
 
-			// dsymutil 10.0.0 has a bug that causes issues, it's fixed in autosdks but not everyone has those set up so for the timebeing we have
-			// a version in P4...
-			if (DsymutilVersionString.StartsWith("Apple LLVM version 10.0.0", StringComparison.Ordinal))
+			bool bUseInstalledDsymutil = true;
+			int Major = 0, Minor = 0, Patch = 0;
+
+			// tease out the version number
+			string[] Tokens = DsymutilVersionString.Split(" ".ToCharArray());
+			
+			// sanity check
+			if (Tokens.Length < 4 || Tokens[3].Contains(".") == false)
+			{
+				Log.TraceInformationOnce("Unable to parse dsymutil version out of: {0}", DsymutilVersionString);
+
+				bUseInstalledDsymutil = false;
+			}
+			else
+			{
+				string[] Versions = Tokens[3].Split(".".ToCharArray());
+				if (Versions.Length < 3)
+				{
+					Log.TraceInformationOnce("Unable to parse version token: {0}", Tokens[3]);
+				}
+				else
+				{
+					if (!int.TryParse(Versions[0], out Major) || !int.TryParse(Versions[1], out Minor) || !int.TryParse(Versions[2], out Patch))
+					{
+						Log.TraceInformationOnce("Unable to parse version tokens: {0}", Tokens[3]);
+					}
+					else
+					{
+						Log.TraceInformationOnce("Parsed dsymutil version as {0}.{1}.{2}", Major, Minor, Patch);
+
+						if (Major < 10 || (Minor == 0 && Patch == 0))
+						{
+							bUseInstalledDsymutil = false;
+						}
+					}
+				}
+			}
+
+			// if the installed one is too old, use a fixed up one if it can
+			if (bUseInstalledDsymutil == false)
 			{
 				FileReference PatchedDsymutilLocation = FileReference.Combine(UnrealBuildTool.EngineDirectory, "Binaries/Mac/NotForLicensees/LLVM/bin/dsymutil");
 
@@ -189,7 +228,7 @@ namespace UnrealBuildTool
 			}
 
 			// 10.0.1 has an issue with LTO builds where we need to limit the number of threads
-			ExtraOptions = (DsymutilVersionString.StartsWith("Apple LLVM version 10.0.1", StringComparison.Ordinal)) ? "-j 1" : "";
+			ExtraOptions = (bIsForLTOBuild && Major == 10 && Minor == 0 && Patch == 1) ? "-j 1" : "";
 			return DsymutilLocation.FullName;
 		}
 	};
