Fix for 2019 project generation

diff --git a/Engine/Source/Programs/UnrealBuildTool/Platform/HoloLens/HoloLensProjectGenerator.cs b/Engine/Source/Programs/UnrealBuildTool/Platform/HoloLens/HoloLensProjectGenerator.cs
--- a/Engine/Source/Programs/UnrealBuildTool/Platform/HoloLens/HoloLensProjectGenerator.cs
+++ b/Engine/Source/Programs/UnrealBuildTool/Platform/HoloLens/HoloLensProjectGenerator.cs
@@ -77,12 +77,11 @@ namespace UnrealBuildTool
 				VCProjectFileContent.Append("		<WindowsTargetPlatformVersion>" + MaxTestedVersion + "</WindowsTargetPlatformVersion>" + ProjectFileGenerator.NewLine);
 			}
 
-			WindowsCompiler Compiler = WindowsCompiler.VisualStudio2017;
+			WindowsCompiler Compiler = WindowsPlatform.GetDefaultCompiler(TargetFilePath);
 			DirectoryReference PlatformWinMDLocation = HoloLens.GetCppCXMetadataLocation(Compiler, "Latest");
 			if (PlatformWinMDLocation == null || !FileReference.Exists(FileReference.Combine(PlatformWinMDLocation, "platform.winmd")))
 			{
-				Compiler = WindowsCompiler.VisualStudio2017;
-				PlatformWinMDLocation = HoloLens.GetCppCXMetadataLocation(Compiler, "Latest");
+				throw new BuildException("Unable to find platform.winmd for {0} toolchain; '{1}' is an invalid version", WindowsPlatform.GetCompilerName(Compiler), "Latest");
 			}
 			string FoundationWinMDPath = HoloLens.GetLatestMetadataPathForApiContract("Windows.Foundation.FoundationContract", Compiler);
 			string UniversalWinMDPath = HoloLens.GetLatestMetadataPathForApiContract("Windows.Foundation.UniversalApiContract", Compiler);
diff --git a/Engine/Source/Programs/UnrealBuildTool/Platform/HoloLens/UEBuildHoloLens.cs b/Engine/Source/Programs/UnrealBuildTool/Platform/HoloLens/UEBuildHoloLens.cs
--- a/Engine/Source/Programs/UnrealBuildTool/Platform/HoloLens/UEBuildHoloLens.cs
+++ b/Engine/Source/Programs/UnrealBuildTool/Platform/HoloLens/UEBuildHoloLens.cs
@@ -326,6 +326,10 @@ namespace UnrealBuildTool
 			{
 				return DirectoryReference.Combine(SelectedToolChainDir, "lib", "x86", "Store", "references");
 			}
+			else if (Compiler >= WindowsCompiler.VisualStudio2019)
+			{
+				return DirectoryReference.Combine(SelectedToolChainDir, "lib", "x86", "Store", "references");
+			}
 			else
 			{
 				return null;
