Fix errors building Clang targets.

#rnx

diff --git a/Engine/Source/Programs/UnrealBuildTool/Platform/Windows/VCEnvironment.cs b/Engine/Source/Programs/UnrealBuildTool/Platform/Windows/VCEnvironment.cs
--- a/Engine/Source/Programs/UnrealBuildTool/Platform/Windows/VCEnvironment.cs
+++ b/Engine/Source/Programs/UnrealBuildTool/Platform/Windows/VCEnvironment.cs
@@ -146,13 +146,13 @@ namespace UnrealBuildTool
 			// needs to run (eg. mspdbcore.dll).
 			if (Architecture == WindowsArchitecture.x64)
 			{
-				AddDirectoryToPath(GetVCToolPath(Compiler, ToolChainDir, WindowsArchitecture.x64));
-				AddDirectoryToPath(GetVCToolPath(Compiler, ToolChainDir, WindowsArchitecture.x86));
+				AddDirectoryToPath(GetVCToolPath(ToolChain, ToolChainDir, WindowsArchitecture.x64));
+				AddDirectoryToPath(GetVCToolPath(ToolChain, ToolChainDir, WindowsArchitecture.x86));
 			}
 			else if (Architecture == WindowsArchitecture.x86)
 			{
-				AddDirectoryToPath(GetVCToolPath(Compiler, ToolChainDir, WindowsArchitecture.x86));
-				AddDirectoryToPath(GetVCToolPath(Compiler, ToolChainDir, WindowsArchitecture.x64));
+				AddDirectoryToPath(GetVCToolPath(ToolChain, ToolChainDir, WindowsArchitecture.x86));
+				AddDirectoryToPath(GetVCToolPath(ToolChain, ToolChainDir, WindowsArchitecture.x64));
 			}
 
 
