Fix incorrect filename when launching DebugGame targets from VSCode (or any target that overrides the undecorated configuration)

diff --git a/Engine/Source/Programs/UnrealBuildTool/ProjectFiles/VisualStudioCode/VSCodeProjectFileGenerator.cs b/Engine/Source/Programs/UnrealBuildTool/ProjectFiles/VisualStudioCode/VSCodeProjectFileGenerator.cs
--- a/Engine/Source/Programs/UnrealBuildTool/ProjectFiles/VisualStudioCode/VSCodeProjectFileGenerator.cs
+++ b/Engine/Source/Programs/UnrealBuildTool/ProjectFiles/VisualStudioCode/VSCodeProjectFileGenerator.cs
@@ -916,21 +916,19 @@ namespace UnrealBuildTool
 			DirectoryReference OutputDirectory = DirectoryReference.Combine(RootDirectory, "Binaries", UBTPlatformName);
 
 			// Get the executable name (minus any platform or config suffixes)
-			string BaseExeName = TargetName;
-			if(Target.TargetRules != null && Target.TargetRules.BuildEnvironment == TargetBuildEnvironment.Shared && TargetRulesType != TargetType.Program)
+			string BinaryName;
+			if(Target.TargetRules.BuildEnvironment == TargetBuildEnvironment.Shared && TargetRulesType != TargetType.Program)
 			{
-				BaseExeName = UEBuildTarget.GetAppNameForTargetType(TargetRulesType);
+				BinaryName = UEBuildTarget.GetAppNameForTargetType(TargetRulesType);
 			}
-
-			// Make the output file path
-			string ExecutableFilename = FileReference.Combine(OutputDirectory, BaseExeName).ToString();
-
-			if ((Configuration != UnrealTargetConfiguration.Development) && ((Configuration !=  UnrealTargetConfiguration.DebugGame) || bShouldCompileMonolithic))
+			else
 			{
-				ExecutableFilename += "-" + UBTPlatformName + "-" + UBTConfigurationName;
+				BinaryName = TargetName;
 			}
-			ExecutableFilename += TargetRulesObject.Architecture;
-			ExecutableFilename += BuildPlatform.GetBinaryExtension(UEBuildBinaryType.Executable);
+
+			// Make the output file path
+			string BinaryFileName = UEBuildTarget.MakeBinaryFileName(BinaryName, Platform, Configuration, TargetRulesObject.Architecture, TargetRulesObject.UndecoratedConfiguration, UEBuildBinaryType.Executable);
+			string ExecutableFilename = FileReference.Combine(OutputDirectory, BinaryFileName).FullName;
 
 			// Include the path to the actual executable for a Mac app bundle
 			if (Platform == UnrealTargetPlatform.Mac && !Target.TargetRules.bIsBuildingConsoleApplication)
