Disable remote execution of ISPC. Suspect this is causing spurious mismatched case issues.

#rnx

diff --git a/Engine/Source/Programs/UnrealBuildTool/Platform/Windows/VCToolChain.cs b/Engine/Source/Programs/UnrealBuildTool/Platform/Windows/VCToolChain.cs
--- a/Engine/Source/Programs/UnrealBuildTool/Platform/Windows/VCToolChain.cs
+++ b/Engine/Source/Programs/UnrealBuildTool/Platform/Windows/VCToolChain.cs
@@ -1386,8 +1386,8 @@ namespace UnrealBuildTool
 				CompileAction.CommandPath = EnvVars.ISPCCompilerPath;
 				CompileAction.StatusDescription = Path.GetFileName(ISPCFile.AbsolutePath);
 
-				// Resource tool can run remotely if possible
-				CompileAction.bCanExecuteRemotely = true;
+				// Disable remote execution to workaround mismatched case on XGE
+				CompileAction.bCanExecuteRemotely = false;
 
 				List<string> Arguments = new List<string>();
 
@@ -1465,8 +1465,8 @@ namespace UnrealBuildTool
 				CompileAction.CommandPath = EnvVars.ISPCCompilerPath;
 				CompileAction.StatusDescription = Path.GetFileName(ISPCFile.AbsolutePath);
 
-				// Resource tool can run remotely if possible
-				CompileAction.bCanExecuteRemotely = true;
+				// Disable remote execution to workaround mismatched case on XGE
+				CompileAction.bCanExecuteRemotely = false;
 
 				List<string> Arguments = new List<string>();
 
