Speculative workaround for "Unexpected PDB error; OK(0)" - retry any failed links on another machine.

diff --git a/Engine/Source/Programs/UnrealBuildTool/Executors/XGE.cs b/Engine/Source/Programs/UnrealBuildTool/Executors/XGE.cs
--- a/Engine/Source/Programs/UnrealBuildTool/Executors/XGE.cs
+++ b/Engine/Source/Programs/UnrealBuildTool/Executors/XGE.cs
@@ -213,70 +213,8 @@ namespace UnrealBuildTool
 						ExportEnv.Add((string)(Pair.Key), (string)(Pair.Value));
 					}
 				}
-
-				int NumSortErrors = 0;
-				for (int ActionIndex = 0; ActionIndex < InActions.Count; ActionIndex++)
-				{
-					Action Action = InActions[ActionIndex];
-					foreach (Action PrerequisiteAction in Action.PrerequisiteActions)
-					{
-						if (InActions.Contains(PrerequisiteAction))
-						{
-							int DepIndex = InActions.IndexOf(PrerequisiteAction);
-							if (DepIndex > ActionIndex)
-							{
-								NumSortErrors++;
-							}
-						}
-					}
-				}
-				if (NumSortErrors > 0)
-				{
-					Actions = new List<Action>();
-					HashSet<int> UsedActions = new HashSet<int>();
-					for (int ActionIndex = 0; ActionIndex < InActions.Count; ActionIndex++)
-					{
-						if (UsedActions.Contains(ActionIndex))
-						{
-							continue;
-						}
-						Action Action = InActions[ActionIndex];
-						foreach(Action PrerequisiteAction in Action.PrerequisiteActions)
-						{
-							if (InActions.Contains(PrerequisiteAction))
-							{
-								int DepIndex = InActions.IndexOf(PrerequisiteAction);
-								if (UsedActions.Contains(DepIndex))
-								{
-									continue;
-								}
-								Actions.Add(PrerequisiteAction);
-								UsedActions.Add(DepIndex);
-							}
-						}
-						Actions.Add(Action);
-						UsedActions.Add(ActionIndex);
-					}
-					for (int ActionIndex = 0; ActionIndex < Actions.Count; ActionIndex++)
-					{
-						Action Action = Actions[ActionIndex];
-						foreach(Action PrerequisiteAction in Action.PrerequisiteActions)
-						{
-							if (Actions.Contains(PrerequisiteAction))
-							{
-								int DepIndex = Actions.IndexOf(PrerequisiteAction);
-								if (DepIndex > ActionIndex)
-								{
-									throw new BuildException("Action is not topologically sorted.\n  {0} {1}\nDependency\n  {2} {3}", Action.CommandPath, Action.CommandArguments, PrerequisiteAction.CommandPath, PrerequisiteAction.CommandArguments);
-								}
-							}
-						}
-					}
-				}
-
 			}
 
-
 			XmlDocument XGETaskDocument = new XmlDocument();
 
 			// <BuildSet FormatVersion="1">...</BuildSet>
@@ -364,6 +302,11 @@ namespace UnrealBuildTool
 							).ToArray()
 						)
 					);
+
+				if(Action.ActionType == ActionType.Link)
+				{
+					ToolElement.SetAttribute("AutoRecover", "Unexpected PDB error; OK(0)");
+				}
 			}
 
 			// <Project Name="Default" Env="Default">...</Project>
