PR #5971: Fix Stub Project Configuration Having Null Platform (Contributed by [REDACTED_PERSON])

#rnx

diff --git a/Engine/Source/Programs/UnrealBuildTool/ProjectFiles/VisualStudio/VCProject.cs b/Engine/Source/Programs/UnrealBuildTool/ProjectFiles/VisualStudio/VCProject.cs
--- a/Engine/Source/Programs/UnrealBuildTool/ProjectFiles/VisualStudio/VCProject.cs
+++ b/Engine/Source/Programs/UnrealBuildTool/ProjectFiles/VisualStudio/VCProject.cs
@@ -561,7 +561,7 @@ namespace UnrealBuildTool
 				ProjectConfigAndTargetCombinations = new List<ProjectConfigAndTargetCombination>();
 				if (IsStubProject)
 				{
-					ProjectConfigAndTargetCombination StubCombination = new ProjectConfigAndTargetCombination(null, UnrealTargetConfiguration.Unknown, StubProjectPlatformName, StubProjectConfigurationName, null);
+					ProjectConfigAndTargetCombination StubCombination = new ProjectConfigAndTargetCombination(UnrealTargetPlatform.Parse(StubProjectPlatformName), UnrealTargetConfiguration.Unknown, StubProjectPlatformName, StubProjectConfigurationName, null);
 					ProjectConfigAndTargetCombinations.Add(StubCombination);
 				}
 				else
