Add a proper error message when a valid mobile provision cannot be found for a target, rather than crashing.

#rnx

diff --git a/Engine/Source/Programs/UnrealBuildTool/Platform/IOS/IOSToolChain.cs b/Engine/Source/Programs/UnrealBuildTool/Platform/IOS/IOSToolChain.cs
--- a/Engine/Source/Programs/UnrealBuildTool/Platform/IOS/IOSToolChain.cs
+++ b/Engine/Source/Programs/UnrealBuildTool/Platform/IOS/IOSToolChain.cs
@@ -1812,6 +1812,11 @@ namespace UnrealBuildTool
 						MobileProvision.TryGetTeamUniqueId(out TeamUUID);
 					}
 
+					if(MobileProvisionFile == null)
+					{
+						throw new BuildException("Unable to find valid certificate/mobile provision pair.");
+					}
+
 					string ConfigName = Target.Configuration.ToString();
 					if (Target.TargetType != TargetType.Game && Target.TargetType != TargetType.Program)
 					{
