Fix UHT being run every time a target is compiled in an installed build. Precompiled modules were not being filtered out of the list of outdated headers for dependency checking.

diff --git a/Engine/Source/Programs/UnrealBuildTool/System/ExternalExecution.cs b/Engine/Source/Programs/UnrealBuildTool/System/ExternalExecution.cs
--- a/Engine/Source/Programs/UnrealBuildTool/System/ExternalExecution.cs
+++ b/Engine/Source/Programs/UnrealBuildTool/System/ExternalExecution.cs
@@ -350,23 +350,27 @@ namespace UnrealBuildTool
 	{
 		public DirectoryItem SourceFolder;
 		public List<FileItem> HeaderFiles;
+		public bool bUsePrecompiled;
 
-		public UHTModuleHeaderInfo(DirectoryItem SourceFolder, List<FileItem> HeaderFiles)
+		public UHTModuleHeaderInfo(DirectoryItem SourceFolder, List<FileItem> HeaderFiles, bool bUsePrecompiled)
 		{
 			this.SourceFolder = SourceFolder;
 			this.HeaderFiles = HeaderFiles;
+			this.bUsePrecompiled = bUsePrecompiled;
 		}
 
 		public UHTModuleHeaderInfo(BinaryArchiveReader Reader)
 		{
 			SourceFolder = Reader.ReadDirectoryItem();
 			HeaderFiles = Reader.ReadList(() => Reader.ReadFileItem());
+			bUsePrecompiled = Reader.ReadBool();
 		}
 
 		public void Write(BinaryArchiveWriter Writer)
 		{
 			Writer.WriteDirectoryItem(SourceFolder);
 			Writer.WriteList(HeaderFiles, Item => Writer.WriteFileItem(Item));
+			Writer.WriteBool(bUsePrecompiled);
 		}
 	}
 
@@ -589,7 +593,7 @@ namespace UnrealBuildTool
 					ReflectedHeaderFiles.AddRange(Info.PublicUObjectHeaders);
 					ReflectedHeaderFiles.AddRange(Info.PublicUObjectClassesHeaders);
 					ReflectedHeaderFiles.AddRange(Info.PrivateUObjectHeaders);
-					UObjectModuleHeaders.Add(new UHTModuleHeaderInfo(ModuleDirectoryItem, ReflectedHeaderFiles));
+					UObjectModuleHeaders.Add(new UHTModuleHeaderInfo(ModuleDirectoryItem, ReflectedHeaderFiles, Module.Rules.bUsePrecompiled));
 				}
 				else
 				{
diff --git a/Engine/Source/Programs/UnrealBuildTool/System/TargetMakefile.cs b/Engine/Source/Programs/UnrealBuildTool/System/TargetMakefile.cs
--- a/Engine/Source/Programs/UnrealBuildTool/System/TargetMakefile.cs
+++ b/Engine/Source/Programs/UnrealBuildTool/System/TargetMakefile.cs
@@ -21,7 +21,7 @@ namespace UnrealBuildTool
 		/// <summary>
 		/// The version number to write
 		/// </summary>
-		public const int CurrentVersion = 16;
+		public const int CurrentVersion = 17;
 
 		/// <summary>
 		/// The time at which the makefile was created
@@ -516,7 +516,7 @@ namespace UnrealBuildTool
 				}
 
 				// Check whether the list has changed
-				List<FileItem> PrevFilesWithMarkup = Makefile.UObjectModuleHeaders.SelectMany(x => x.HeaderFiles).ToList();
+				List<FileItem> PrevFilesWithMarkup = Makefile.UObjectModuleHeaders.Where(x => !x.bUsePrecompiled).SelectMany(x => x.HeaderFiles).ToList();
 				List<FileItem> NextFilesWithMarkup = NewFilesWithMarkupBag.ToList();
 				if (NextFilesWithMarkup.Count != PrevFilesWithMarkup.Count || NextFilesWithMarkup.Intersect(PrevFilesWithMarkup).Count() != PrevFilesWithMarkup.Count)
 				{
