LiveCoding: Allow patching executables after they have been staged to a different directory, so it runs correctly with packaged games.

#rnx

diff --git a/Engine/Source/Programs/UnrealBuildTool/Modes/BuildMode.cs b/Engine/Source/Programs/UnrealBuildTool/Modes/BuildMode.cs
--- a/Engine/Source/Programs/UnrealBuildTool/Modes/BuildMode.cs
+++ b/Engine/Source/Programs/UnrealBuildTool/Modes/BuildMode.cs
@@ -617,13 +617,13 @@ namespace UnrealBuildTool
 					string[] Lines = FileReference.ReadAllLines(TargetDescriptor.LiveCodingModules);
 
 					// Parse it out into a set of filenames
-					HashSet<FileReference> AllowedOutputFiles = new HashSet<FileReference>();
+					HashSet<string> AllowedOutputFileNames = new HashSet<string>(FileReference.Comparer);
 					foreach (string Line in Lines)
 					{
 						string TrimLine = Line.Trim();
 						if (TrimLine.Length > 0)
 						{
-							AllowedOutputFiles.Add(new FileReference(TrimLine));
+							AllowedOutputFileNames.Add(Path.GetFileName(TrimLine));
 						}
 					}
 
@@ -638,7 +638,7 @@ namespace UnrealBuildTool
 					}
 
 					// Find all the files that will be built that aren't allowed
-					List<FileReference> ProtectedOutputFiles = OutputFiles.Where(x => !AllowedOutputFiles.Contains(x)).ToList();
+					List<FileReference> ProtectedOutputFiles = OutputFiles.Where(x => !AllowedOutputFileNames.Contains(x.GetFileName())).ToList();
 					if (ProtectedOutputFiles.Count > 0)
 					{
 						FileReference.WriteAllLines(new FileReference(TargetDescriptor.LiveCodingModules.FullName + ".out"), ProtectedOutputFiles.Select(x => x.ToString()));
