UBT: Convert some windows target settings to properties which take on default values based on other properties. This fixes logic for determining whether live coding is enabled for a target, which needs to know the target type.

#rnx

diff --git a/Engine/Source/Programs/UnrealBuildTool/Configuration/TargetRules.cs b/Engine/Source/Programs/UnrealBuildTool/Configuration/TargetRules.cs
--- a/Engine/Source/Programs/UnrealBuildTool/Configuration/TargetRules.cs
+++ b/Engine/Source/Programs/UnrealBuildTool/Configuration/TargetRules.cs
@@ -541,7 +541,12 @@ namespace UnrealBuildTool
 		/// Whether to enable support for live coding
 		/// </summary>
 		[RequiresUniqueBuildEnvironment]
-		public bool bWithLiveCoding = false;
+		public bool bWithLiveCoding
+		{
+			get { return bWithLiveCodingPrivate ?? (Platform == UnrealTargetPlatform.Win64 && Configuration != UnrealTargetConfiguration.Shipping && Type != TargetType.Program); }
+			set { bWithLiveCodingPrivate = value; }
+		}
+		bool? bWithLiveCodingPrivate;
 
 		/// <summary>
 		/// Whether to enable support for live coding
@@ -1291,7 +1296,7 @@ namespace UnrealBuildTool
 		/// Windows-specific target settings.
 		/// </summary>
 		[ConfigSubObject]
-		public WindowsTargetRules WindowsPlatform = new WindowsTargetRules();
+		public WindowsTargetRules WindowsPlatform; // Requires 'this' parameter; initialized in constructor
 
 		/// <summary>
 		/// Xbox One-specific target settings.
@@ -1316,6 +1321,7 @@ namespace UnrealBuildTool
 			this.Architecture = Target.Architecture;
 			this.ProjectFile = Target.ProjectFile;
 			this.Version = Target.Version;
+			this.WindowsPlatform = new WindowsTargetRules(this);
 
 			// Read settings from config files
 			foreach(object ConfigurableObject in GetConfigurableObjects())
diff --git a/Engine/Source/Programs/UnrealBuildTool/Platform/Windows/UEBuildWindows.cs b/Engine/Source/Programs/UnrealBuildTool/Platform/Windows/UEBuildWindows.cs
--- a/Engine/Source/Programs/UnrealBuildTool/Platform/Windows/UEBuildWindows.cs
+++ b/Engine/Source/Programs/UnrealBuildTool/Platform/Windows/UEBuildWindows.cs
@@ -104,6 +104,11 @@ namespace UnrealBuildTool
 	/// </summary>
 	public class WindowsTargetRules
 	{
+		/// <summary>
+		/// The target rules which owns this object. Used to resolve some properties.
+		/// </summary>
+		TargetRules Target;
+
 		/// <summary>
 		/// Version of the compiler toolchain to use on Windows platform. A value of "default" will be changed to a specific version at UBT startup.
 		/// </summary>
@@ -230,17 +235,32 @@ namespace UnrealBuildTool
 		/// <summary>
 		/// Create an image that can be hot patched (/FUNCTIONPADMIN)
 		/// </summary>
-		public bool bCreateHotPatchableImage = false;
+		public bool bCreateHotPatchableImage
+		{
+			get { return bCreateHotPatchableImagePrivate ?? !Target.bWithLiveCoding; }
+			set { bCreateHotPatchableImagePrivate = true; }
+		}
+		private bool? bCreateHotPatchableImagePrivate;
 
 		/// <summary>
 		/// Strip unreferenced symbols (/OPT:REF)
 		/// </summary>
-		public bool bStripUnreferencedSymbols = false;
-
+		public bool bStripUnreferencedSymbols
+		{
+			get { return bStripUnreferencedSymbolsPrivate ?? ((Target.Configuration == UnrealTargetConfiguration.Test || Target.Configuration == UnrealTargetConfiguration.Shipping) && !Target.bWithLiveCoding); }
+			set { bStripUnreferencedSymbolsPrivate = value; }
+		}
+		private bool? bStripUnreferencedSymbolsPrivate;
+			
 		/// <summary>
 		/// Merge identical COMDAT sections together (/OPT:ICF)
 		/// </summary>
-		public bool bMergeIdenticalCOMDATs = false;
+		public bool bMergeIdenticalCOMDATs
+		{
+			get { return bMergeIdenticalCOMDATsPrivate ?? ((Target.Configuration == UnrealTargetConfiguration.Test || Target.Configuration == UnrealTargetConfiguration.Shipping) && !Target.bWithLiveCoding); }
+			set { bMergeIdenticalCOMDATsPrivate = value; }
+		}
+		private bool? bMergeIdenticalCOMDATsPrivate;
 
 		/// <summary>
 		/// Whether to put global symbols in their own sections (/Gw), allowing the linker to discard any that are unused.
@@ -322,6 +342,15 @@ namespace UnrealBuildTool
 					throw new BuildException("Unexpected WindowsCompiler version for GetVisualStudioCompilerVersionName().  Either not using a Visual Studio compiler or switch block needs to be updated");
 			}
 		}
+
+		/// <summary>
+		/// Constructor
+		/// </summary>
+		/// <param name="Target">The target rules which owns this object</param>
+		internal WindowsTargetRules(TargetRules Target)
+		{
+			this.Target = Target;
+		}
 	}
 
 	/// <summary>
@@ -606,17 +635,6 @@ namespace UnrealBuildTool
 		public override void ResetTarget(TargetRules Target)
 		{
 			base.ResetTarget(Target);
-
-			if(Target.Platform == UnrealTargetPlatform.Win64 && Target.Configuration != UnrealTargetConfiguration.Shipping && Target.Type != TargetType.Program)
-			{
-				Target.bWithLiveCoding = true;
-				Target.WindowsPlatform.bCreateHotPatchableImage = true;
-			}
-			else if (Target.Configuration == UnrealTargetConfiguration.Test || Target.Configuration == UnrealTargetConfiguration.Shipping)
-			{
-				Target.WindowsPlatform.bStripUnreferencedSymbols = true;
-				Target.WindowsPlatform.bMergeIdenticalCOMDATs = true;
-			}
 		}
 
 		/// <summary>
