Minor fallout from previous change.

#rnx

diff --git a/Engine/Source/Programs/LiveCodingConsole/LiveCodingConsole.Target.cs b/Engine/Source/Programs/LiveCodingConsole/LiveCodingConsole.Target.cs
--- a/Engine/Source/Programs/LiveCodingConsole/LiveCodingConsole.Target.cs
+++ b/Engine/Source/Programs/LiveCodingConsole/LiveCodingConsole.Target.cs
@@ -15,6 +15,7 @@ public class LiveCodingConsoleTarget : TargetRules
 		bUseMallocProfiler = false;
 		bCompileWithPluginSupport = true;
 		bIncludePluginsForTargetPlatforms = true;
+		bWithLiveCoding = true;
 
 		// Editor-only data, however, is needed
 		bBuildWithEditorOnlyData = true;
diff --git a/Engine/Source/Programs/UnrealBuildTool/Platform/Windows/UEBuildWindows.cs b/Engine/Source/Programs/UnrealBuildTool/Platform/Windows/UEBuildWindows.cs
--- a/Engine/Source/Programs/UnrealBuildTool/Platform/Windows/UEBuildWindows.cs
+++ b/Engine/Source/Programs/UnrealBuildTool/Platform/Windows/UEBuildWindows.cs
@@ -237,7 +237,7 @@ namespace UnrealBuildTool
 		/// </summary>
 		public bool bCreateHotPatchableImage
 		{
-			get { return bCreateHotPatchableImagePrivate ?? !Target.bWithLiveCoding; }
+			get { return bCreateHotPatchableImagePrivate ?? Target.bWithLiveCoding; }
 			set { bCreateHotPatchableImagePrivate = true; }
 		}
 		private bool? bCreateHotPatchableImagePrivate;
