- Fixed build after platform extension directory inversion

diff --git a/Engine/Source/Programs/UnrealBuildTool/System/DataDrivenPlatformInfo.cs b/Engine/Source/Programs/UnrealBuildTool/System/DataDrivenPlatformInfo.cs
--- a/Engine/Source/Programs/UnrealBuildTool/System/DataDrivenPlatformInfo.cs
+++ b/Engine/Source/Programs/UnrealBuildTool/System/DataDrivenPlatformInfo.cs
@@ -49,21 +49,23 @@ namespace UnrealBuildTool
 					// look through all config dirs looking for the data driven ini file
 					foreach (string FilePath in Directory.EnumerateFiles(EngineConfigDir.FullName, "DataDrivenPlatformInfo.ini", SearchOption.AllDirectories))
 					{
+						FileReference FileRef = new FileReference(FilePath);
+
 						// get the platform name from the path
 						string IniPlatformName;
-						// Foo/Engine/Config/<Platform>/DataDrivenPlatformInfo.ini
-						if (EngineConfigDir.IsUnderDirectory(UnrealBuildTool.EngineDirectory))
+						if (FileRef.IsUnderDirectory(DirectoryReference.Combine(UnrealBuildTool.EngineDirectory, "Config")))
 						{
+							// Foo/Engine/Config/<Platform>/DataDrivenPlatformInfo.ini
 							IniPlatformName = Path.GetFileName(Path.GetDirectoryName(FilePath));
 						}
-						// Foo/Platforms/<Platform>/Engine/Config/DataDrivenPlatformInfo.ini
 						else
 						{
-							IniPlatformName = Path.GetFileName(Path.GetDirectoryName(Path.GetDirectoryName(Path.GetDirectoryName(FilePath))));
+							// Foo/Engine/Platforms/<Platform>/Config/DataDrivenPlatformInfo.ini
+							IniPlatformName = Path.GetFileName(Path.GetDirectoryName(Path.GetDirectoryName(FilePath)));
 						}
 
 						// load the DataDrivenPlatformInfo from the path
-						ConfigFile Config = new ConfigFile(new FileReference(FilePath));
+						ConfigFile Config = new ConfigFile(FileRef);
 						ConfigDataDrivenPlatformInfo NewInfo = new ConfigDataDrivenPlatformInfo();
 
 
diff --git a/Engine/Source/Programs/UnrealBuildTool/System/RestrictedFolders.cs b/Engine/Source/Programs/UnrealBuildTool/System/RestrictedFolders.cs
--- a/Engine/Source/Programs/UnrealBuildTool/System/RestrictedFolders.cs
+++ b/Engine/Source/Programs/UnrealBuildTool/System/RestrictedFolders.cs
@@ -171,26 +171,6 @@ namespace UnrealBuildTool
 		/// Epic Employees only
 		/// </summary>
 		public static RestrictedFolder NoRedist = FindOrAddByName("NoRedist");
-
-		/// <summary>
-		/// Playstation 4 source files
-		/// </summary>
-		public static RestrictedFolder PS4 = FindOrAddByName("PS4");
-
-		/// <summary>
-		/// XboxOne source files
-		/// </summary>
-		public static RestrictedFolder XboxOne = FindOrAddByName("XboxOne");
-
-		/// <summary>
-		/// Switch source files
-		/// </summary>
-		public static RestrictedFolder Switch = FindOrAddByName("Switch");
-
-		/// <summary>
-		/// Quail source files
-		/// </summary>
-		public static RestrictedFolder Quail = FindOrAddByName("Quail");
 	}
 
 	/// <summary>
