Fix .sln generation for VSMac

diff --git a/Engine/Source/Programs/UnrealBuildTool/Platform/Android/AndroidProjectGenerator.cs b/Engine/Source/Programs/UnrealBuildTool/Platform/Android/AndroidProjectGenerator.cs
--- a/Engine/Source/Programs/UnrealBuildTool/Platform/Android/AndroidProjectGenerator.cs
+++ b/Engine/Source/Programs/UnrealBuildTool/Platform/Android/AndroidProjectGenerator.cs
@@ -38,6 +38,12 @@ namespace UnrealBuildTool
 			}
 		}
 
+		bool HostSupportsVSAndroid()
+		{
+			// Only supported on Windows (VS Mac is a thing..)
+			return BuildHostPlatform.Current.Platform.IsInGroup(UnrealPlatformGroup.Microsoft);
+		}
+
 		bool IsVSAndroidSupportInstalled()
 		{
 			if (VSSupportChecked)
@@ -46,6 +52,12 @@ namespace UnrealBuildTool
 			}
 			VSSupportChecked = true;
 
+
+			if (!HostSupportsVSAndroid())
+			{
+				return false;
+			}
+
 			//check to make sure Cross Platform Tools are installed for MS
 
 			// First check to see if the sandboxed SDK has been installed, it's always dropped to the same ProgramData file location
@@ -87,6 +99,12 @@ namespace UnrealBuildTool
 
 			CheckedForNsight = true;
 
+			if (!HostSupportsVSAndroid())
+			{
+				NsightInstalled = false;
+				return false;
+			}
+
 			// NOTE: there is now a registry key that we can use instead at:
 			//			HKEY_LOCAL_MACHINE\SOFTWARE\Wow6432Node\NVIDIA Corporation\Nsight Tegra\Version
 
