UE-69817 - iOS fails during build step with OSS Facebook plugin enabled due to spaces in filepath

diff --git a/Engine/Source/Programs/UnrealBuildTool/Platform/IOS/IOSToolChain.cs b/Engine/Source/Programs/UnrealBuildTool/Platform/IOS/IOSToolChain.cs
--- a/Engine/Source/Programs/UnrealBuildTool/Platform/IOS/IOSToolChain.cs
+++ b/Engine/Source/Programs/UnrealBuildTool/Platform/IOS/IOSToolChain.cs
@@ -1234,21 +1234,21 @@ namespace UnrealBuildTool
 				ExtractScript.AppendLine("#!/bin/sh");
 				ExtractScript.AppendLine("set -e");
 				// ExtractScript.AppendLine("set -x"); // For debugging
-				ExtractScript.AppendLine(String.Format("[ -d {0} ] && rm -rf {0}", Utils.EscapeShellArgument(Framework.OutputDirectory.FullName)));
-				ExtractScript.AppendLine(String.Format("unzip -q -o {0} -d {1}", Utils.EscapeShellArgument(Framework.ZipFile.FullName), Utils.EscapeShellArgument(Framework.OutputDirectory.ParentDirectory.FullName))); // Zip contains folder with the same name, hence ParentDirectory
-				ExtractScript.AppendLine(String.Format("touch {0}", Utils.EscapeShellArgument(Framework.ExtractedTokenFile.AbsolutePath)));
+				ExtractScript.AppendLine(String.Format("[ -d {0} ] && rm -rf {0}", Utils.MakePathSafeToUseWithCommandLine(Framework.OutputDirectory.FullName)));
+				ExtractScript.AppendLine(String.Format("unzip -q -o {0} -d {1}", Utils.MakePathSafeToUseWithCommandLine(Framework.ZipFile.FullName), Utils.MakePathSafeToUseWithCommandLine(Framework.OutputDirectory.ParentDirectory.FullName))); // Zip contains folder with the same name, hence ParentDirectory
+				ExtractScript.AppendLine(String.Format("touch {0}", Utils.MakePathSafeToUseWithCommandLine(Framework.ExtractedTokenFile.AbsolutePath)));
 
 				FileItem ExtractScriptFileItem = FileItem.CreateIntermediateTextFile(new FileReference(Framework.OutputDirectory.FullName + ".sh"), ExtractScript.ToString());
 
 				Action UnzipAction = new Action(ActionType.BuildProject);
 				UnzipAction.CommandPath = new FileReference("/bin/sh");
-				UnzipAction.CommandArguments = Utils.EscapeShellArgument(ExtractScriptFileItem.AbsolutePath);
+				UnzipAction.CommandArguments = Utils.MakePathSafeToUseWithCommandLine(ExtractScriptFileItem.AbsolutePath);
 				UnzipAction.WorkingDirectory = UnrealBuildTool.EngineDirectory;
 				UnzipAction.PrerequisiteItems.Add(InputFile);
 				UnzipAction.PrerequisiteItems.Add(ExtractScriptFileItem);
 				UnzipAction.ProducedItems.Add(Framework.ExtractedTokenFile);
 				UnzipAction.DeleteItems.Add(Framework.ExtractedTokenFile);
-				UnzipAction.StatusDescription = String.Format("Unzipping: {0} -> {1}", Framework.ZipFile, Framework.OutputDirectory);
+				UnzipAction.StatusDescription = String.Format("Unzipping : {0} -> {1}", Framework.ZipFile, Framework.OutputDirectory);
 				UnzipAction.bCanExecuteRemotely = false;
 				Actions.Add(UnzipAction);
 			}
