Fix exception when using custom build steps, related to UnrealTargetPlatform changes.

#rnx

diff --git a/Engine/Source/Programs/UnrealBuildTool/System/CustomBuildSteps.cs b/Engine/Source/Programs/UnrealBuildTool/System/CustomBuildSteps.cs
--- a/Engine/Source/Programs/UnrealBuildTool/System/CustomBuildSteps.cs
+++ b/Engine/Source/Programs/UnrealBuildTool/System/CustomBuildSteps.cs
@@ -14,7 +14,7 @@ namespace UnrealBuildTool
 	/// </summary>
 	public class CustomBuildSteps
 	{
-		SortedDictionary<UnrealTargetPlatform, string[]> HostPlatformToCommands = new SortedDictionary<UnrealTargetPlatform,string[]>();
+		Dictionary<UnrealTargetPlatform, string[]> HostPlatformToCommands = new Dictionary<UnrealTargetPlatform,string[]>();
 
 		/// <summary>
 		/// Construct a custom build steps object from a Json object.
@@ -62,7 +62,7 @@ namespace UnrealBuildTool
 		public void Write(JsonWriter Writer, string FieldName)
 		{
 			Writer.WriteObjectStart(FieldName);
-			foreach(KeyValuePair<UnrealTargetPlatform, string[]> Pair in HostPlatformToCommands)
+			foreach(KeyValuePair<UnrealTargetPlatform, string[]> Pair in HostPlatformToCommands.OrderBy(x => x.Key.ToString()))
 			{
 				Writer.WriteArrayStart(Pair.Key.ToString());
 				foreach(string Line in Pair.Value)
