Store the executable for each target in the makefile, so binary builds can construct the correct mutex name for detecting whether Live Coding is active.

#rnx

diff --git a/Engine/Source/Programs/UnrealBuildTool/Configuration/UEBuildTarget.cs b/Engine/Source/Programs/UnrealBuildTool/Configuration/UEBuildTarget.cs
--- a/Engine/Source/Programs/UnrealBuildTool/Configuration/UEBuildTarget.cs
+++ b/Engine/Source/Programs/UnrealBuildTool/Configuration/UEBuildTarget.cs
@@ -1499,7 +1499,7 @@ namespace UnrealBuildTool
 
 			// Create the makefile
 			string ExternalMetadata = UEBuildPlatform.GetBuildPlatform(Platform).GetExternalBuildMetadata(ProjectFile);
-			TargetMakefile Makefile = new TargetMakefile(TargetToolChain.GetVersionInfo(), ExternalMetadata, ReceiptFileName, ProjectIntermediateDirectory, TargetType, Rules.ConfigValueTracker, bDeployAfterCompile, bHasProjectScriptPlugin);
+			TargetMakefile Makefile = new TargetMakefile(TargetToolChain.GetVersionInfo(), ExternalMetadata, Binaries[0].OutputFilePaths[0], ReceiptFileName, ProjectIntermediateDirectory, TargetType, Rules.ConfigValueTracker, bDeployAfterCompile, bHasProjectScriptPlugin);
 
 			// Setup the hot reload module list
 			Makefile.HotReloadModuleNames = GetHotReloadModuleNames();
diff --git a/Engine/Source/Programs/UnrealBuildTool/System/HotReload.cs b/Engine/Source/Programs/UnrealBuildTool/System/HotReload.cs
--- a/Engine/Source/Programs/UnrealBuildTool/System/HotReload.cs
+++ b/Engine/Source/Programs/UnrealBuildTool/System/HotReload.cs
@@ -129,14 +129,14 @@ namespace UnrealBuildTool
 		public static bool IsLiveCodingSessionActive(TargetMakefile Makefile)
 		{
 			// Find the first output executable
-			FileItem Executable = Makefile.OutputItems.FirstOrDefault(x => x.HasExtension(".exe"));
+			FileReference Executable = Makefile.ExecutableFile;
 			if(Executable != null)
 			{
 				// Build the mutex name. This should match the name generated in LiveCodingModule.cpp.
 				StringBuilder MutexName = new StringBuilder("Global\\LiveCoding_");
-				for(int Idx = 0; Idx < Executable.AbsolutePath.Length; Idx++)
+				for(int Idx = 0; Idx < Executable.FullName.Length; Idx++)
 				{
-					char Character = Executable.AbsolutePath[Idx];
+					char Character = Executable.FullName[Idx];
 					if(Character == '/' || Character == '\\' || Character == ':')
 					{
 						MutexName.Append('+');
diff --git a/Engine/Source/Programs/UnrealBuildTool/System/TargetMakefile.cs b/Engine/Source/Programs/UnrealBuildTool/System/TargetMakefile.cs
--- a/Engine/Source/Programs/UnrealBuildTool/System/TargetMakefile.cs
+++ b/Engine/Source/Programs/UnrealBuildTool/System/TargetMakefile.cs
@@ -21,7 +21,7 @@ namespace UnrealBuildTool
 		/// <summary>
 		/// The version number to write
 		/// </summary>
-		public const int CurrentVersion = 17;
+		public const int CurrentVersion = 18;
 
 		/// <summary>
 		/// The time at which the makefile was created
@@ -38,6 +38,11 @@ namespace UnrealBuildTool
 		/// </summary>
 		public string ExternalMetadata;
 
+		/// <summary>
+		/// The main executable output by this target
+		/// </summary>
+		public FileReference ExecutableFile;
+
 		/// <summary>
 		/// Path to the receipt file for this target
 		/// </summary>
@@ -154,17 +159,19 @@ namespace UnrealBuildTool
 		/// </summary>
 		/// <param name="ToolchainInfo">String describing the toolchain used to build. This will be output before executing actions.</param>
 		/// <param name="ExternalMetadata">External build metadata from the platform</param>
+		/// <param name="ExecutableFile">Path to the executable or primary output binary for this target</param>
 		/// <param name="ReceiptFile">Path to the receipt file</param>
 		/// <param name="ProjectIntermediateDirectory">Path to the project intermediate directory</param>
 		/// <param name="TargetType">The type of target</param>
 		/// <param name="ConfigValueTracker">Set of dependencies on config files</param>
 		/// <param name="bDeployAfterCompile">Whether to deploy the target after compiling</param>
 		/// <param name="bHasProjectScriptPlugin">Whether the target has a project script plugin</param>
-		public TargetMakefile(string ToolchainInfo, string ExternalMetadata, FileReference ReceiptFile, DirectoryReference ProjectIntermediateDirectory, TargetType TargetType, ConfigValueTracker ConfigValueTracker, bool bDeployAfterCompile, bool bHasProjectScriptPlugin)
+		public TargetMakefile(string ToolchainInfo, string ExternalMetadata, FileReference ExecutableFile, FileReference ReceiptFile, DirectoryReference ProjectIntermediateDirectory, TargetType TargetType, ConfigValueTracker ConfigValueTracker, bool bDeployAfterCompile, bool bHasProjectScriptPlugin)
 		{
 			this.CreateTimeUtc = DateTime.UtcNow;
 			this.ToolchainInfo = ToolchainInfo;
 			this.ExternalMetadata = ExternalMetadata;
+			this.ExecutableFile = ExecutableFile;
 			this.ReceiptFile = ReceiptFile;
 			this.ProjectIntermediateDirectory = ProjectIntermediateDirectory;
 			this.TargetType = TargetType;
@@ -194,6 +201,7 @@ namespace UnrealBuildTool
 			CreateTimeUtc = new DateTime(Reader.ReadLong(), DateTimeKind.Utc);
 			ToolchainInfo = Reader.ReadString();
 			ExternalMetadata = Reader.ReadString();
+			ExecutableFile = Reader.ReadFileReference();
 			ReceiptFile = Reader.ReadFileReference();
 			ProjectIntermediateDirectory = Reader.ReadDirectoryReference();
 			TargetType = (TargetType)Reader.ReadInt();
@@ -226,6 +234,7 @@ namespace UnrealBuildTool
 			Writer.WriteLong(CreateTimeUtc.Ticks);
 			Writer.WriteString(ToolchainInfo);
 			Writer.WriteString(ExternalMetadata);
+			Writer.WriteFileReference(ExecutableFile);
 			Writer.WriteFileReference(ReceiptFile);
 			Writer.WriteDirectoryReference(ProjectIntermediateDirectory);
 			Writer.WriteInt((int)TargetType);
