Output each extracted framework to a different directory to fix contention when unzipping multiple frameworks in parallel.

diff --git a/Engine/Source/Programs/UnrealBuildTool/Configuration/UEBuildModule.cs b/Engine/Source/Programs/UnrealBuildTool/Configuration/UEBuildModule.cs
--- a/Engine/Source/Programs/UnrealBuildTool/Configuration/UEBuildModule.cs
+++ b/Engine/Source/Programs/UnrealBuildTool/Configuration/UEBuildModule.cs
@@ -192,7 +192,7 @@ namespace UnrealBuildTool
 					}
 					else
 					{
-						Framework = new UEBuildFramework(FrameworkRules.Name, FileReference.Combine(ModuleDirectory, FrameworkRules.ZipPath), DirectoryReference.Combine(UnrealBuildTool.EngineDirectory, "Intermediate", "UnzippedFrameworks", Name, Path.GetFileNameWithoutExtension(FrameworkRules.ZipPath)), FrameworkRules.CopyBundledAssets);
+						Framework = new UEBuildFramework(FrameworkRules.Name, FileReference.Combine(ModuleDirectory, FrameworkRules.ZipPath), DirectoryReference.Combine(UnrealBuildTool.EngineDirectory, "Intermediate", "UnzippedFrameworks", FrameworkRules.Name, Path.GetFileNameWithoutExtension(FrameworkRules.ZipPath)), FrameworkRules.CopyBundledAssets);
 					}
 					PublicAdditionalFrameworks.Add(Framework);
 				}
