CIS BUGHUNT - seeing if c# Process.Start() terminating faster than control returning to process.WaitForExit()

#rnx

diff --git a/Engine/Source/Programs/UnrealBuildTool/Platform/HTML5/HTML5SDKInfo.cs b/Engine/Source/Programs/UnrealBuildTool/Platform/HTML5/HTML5SDKInfo.cs
--- a/Engine/Source/Programs/UnrealBuildTool/Platform/HTML5/HTML5SDKInfo.cs
+++ b/Engine/Source/Programs/UnrealBuildTool/Platform/HTML5/HTML5SDKInfo.cs
@@ -393,14 +393,22 @@ namespace UnrealBuildTool
 // jic output dump is needed...
 //				processInfo.RedirectStandardError = true;
 //				processInfo.RedirectStandardOutput = true;
-				Process process = Process.Start(processInfo);
-//				process.OutputDataReceived += (object sender, DataReceivedEventArgs e) => Log.TraceInformation("output>>" + e.Data);
-//				process.BeginOutputReadLine();
-//				process.ErrorDataReceived += (object sender, DataReceivedEventArgs e) => Log.TraceInformation("error>>" + e.Data);
-//				process.BeginErrorReadLine();
-				process.WaitForExit();
-				Log.TraceInformation("emcc ExitCode: {0}", process.ExitCode);
-				process.Close();
+				try
+				{
+					Process process = Process.Start(processInfo);
+//					process.OutputDataReceived += (object sender, DataReceivedEventArgs e) => Log.TraceInformation("output>>" + e.Data);
+//					process.BeginOutputReadLine();
+//					process.ErrorDataReceived += (object sender, DataReceivedEventArgs e) => Log.TraceInformation("error>>" + e.Data);
+//					process.BeginErrorReadLine();
+					process.WaitForExit();
+					Log.TraceInformation("emcc ExitCode: {0}", process.ExitCode);
+					process.Close();
+				}
+				catch (Win32Exception ex)
+				{
+					// Process.Start() as terminated quick enough betore control has returned to process.WaitForExit()
+					Log.TraceInformation("Win32Exception ex.NativeErrorCode: {0}", ex.NativeErrorCode);
+				}
 				// uncomment OPTIMIZER (GUBP on build machines needs this)
 				// and PYTHON (reduce warnings on EMCC_DEBUG=1)
 				string pyth = Regex.Replace(PYTHON, @"\\", @"\\");
