CIS BUGHUNT - seeing if c# Process.Start() terminating faster than control returning to process.WaitForExit()

- fix build error...

#rnx

diff --git a/Engine/Source/Programs/UnrealBuildTool/Platform/HTML5/HTML5SDKInfo.cs b/Engine/Source/Programs/UnrealBuildTool/Platform/HTML5/HTML5SDKInfo.cs
--- a/Engine/Source/Programs/UnrealBuildTool/Platform/HTML5/HTML5SDKInfo.cs
+++ b/Engine/Source/Programs/UnrealBuildTool/Platform/HTML5/HTML5SDKInfo.cs
@@ -404,7 +404,7 @@ namespace UnrealBuildTool
 					Log.TraceInformation("emcc ExitCode: {0}", process.ExitCode);
 					process.Close();
 				}
-				catch (Win32Exception ex)
+				catch (System.ComponentModel.Win32Exception ex)
 				{
 					// Process.Start() as terminated quick enough betore control has returned to process.WaitForExit()
 					Log.TraceInformation("Win32Exception ex.NativeErrorCode: {0}", ex.NativeErrorCode);
