LiveCoding: Prevent object files being overwritten during a live coding compile. Object files are now built with a .lc suffix, leaving the original files intact. This ensures it can gather COFF data from the correct executables on relaunch, and means that any existing changes will be patched in as soon as the first Live Coding compile is initiated. Should fix several issues.

diff --git a/Engine/Source/Developer/Windows/LiveCodingServer/Private/External/LC_LiveModule.cpp b/Engine/Source/Developer/Windows/LiveCodingServer/Private/External/LC_LiveModule.cpp
--- a/Engine/Source/Developer/Windows/LiveCodingServer/Private/External/LC_LiveModule.cpp
+++ b/Engine/Source/Developer/Windows/LiveCodingServer/Private/External/LC_LiveModule.cpp
@@ -397,6 +397,16 @@ namespace
 	}
 
 	// BEGIN EPIC MOD - Allow mapping from object files to their unity object file
+	static void AddCompilandId(types::StringMap<uint32_t>& objFileToCompilandId, const std::wstring& objFile, uint32_t compilandId)
+	{
+		std::wstring WideObjectFile = file::NormalizePath(objFile.c_str());
+		ImmutableString ObjectFile = string::ToUtf8String(WideObjectFile);
+		if (objFileToCompilandId.find(ObjectFile) == objFileToCompilandId.end())
+		{
+			objFileToCompilandId.insert(std::make_pair(ObjectFile, compilandId));
+		}
+	}
+
 	static bool ReadLiveCodingInfo(const wchar_t* manifestFile, types::StringMap<uint32_t>& objFileToCompilandId)
 	{
 		// Read the file to a string
@@ -441,12 +451,11 @@ namespace
 					return false;
 				}
 
-				std::wstring WideObjectFile = file::NormalizePath((BaseDir + L"\\" + *SourceFileValue->AsString()).c_str());
-				ImmutableString ObjectFile = string::ToUtf8String(WideObjectFile);
-				if (objFileToCompilandId.find(ObjectFile) == objFileToCompilandId.end())
-				{
-					objFileToCompilandId.insert(std::make_pair(ObjectFile, UnityCompilandId));
-				}
+				std::wstring objFile = BaseDir + L"\\" + *SourceFileValue->AsString();
+				AddCompilandId(objFileToCompilandId, objFile, UnityCompilandId);
+
+				std::wstring patchedObjFile = file::RemoveExtension(objFile) + L".lc.obj";
+				AddCompilandId(objFileToCompilandId, patchedObjFile, UnityCompilandId);
 			}
 		}
 
diff --git a/Engine/Source/Programs/UnrealBuildTool/Configuration/TargetDescriptor.cs b/Engine/Source/Programs/UnrealBuildTool/Configuration/TargetDescriptor.cs
--- a/Engine/Source/Programs/UnrealBuildTool/Configuration/TargetDescriptor.cs
+++ b/Engine/Source/Programs/UnrealBuildTool/Configuration/TargetDescriptor.cs
@@ -64,6 +64,12 @@ namespace UnrealBuildTool
 		[CommandLine("-LiveCodingModules=")]
 		public FileReference LiveCodingModules = null;
 
+		/// <summary>
+		/// Path to the manifest for passing info about the output to live coding
+		/// </summary>
+		[CommandLine("-LiveCodingManifest=")]
+		public FileReference LiveCodingManifest = null;
+
 		/// <summary>
 		/// Constructor
 		/// </summary>
diff --git a/Engine/Source/Programs/UnrealBuildTool/Modes/BuildMode.cs b/Engine/Source/Programs/UnrealBuildTool/Modes/BuildMode.cs
--- a/Engine/Source/Programs/UnrealBuildTool/Modes/BuildMode.cs
+++ b/Engine/Source/Programs/UnrealBuildTool/Modes/BuildMode.cs
@@ -76,12 +76,6 @@ namespace UnrealBuildTool
 		[CommandLine("-WriteOutdatedActions=")]
 		public FileReference WriteOutdatedActionsFile = null;
 
-		/// <summary>
-		/// Path to the manifest for passing info about the output to live coding
-		/// </summary>
-		[CommandLine("-LiveCodingManifest=")]
-		public FileReference LiveCodingManifest = null;
-
 		/// <summary>
 		/// Main entry point
 		/// </summary>
@@ -200,7 +194,7 @@ namespace UnrealBuildTool
 					// Create the working set provider
 					using (ISourceFileWorkingSet WorkingSet = SourceFileWorkingSet.Create(UnrealBuildTool.RootDirectory, ProjectDirs))
 					{
-						Build(TargetDescriptors, BuildConfiguration, WorkingSet, Options, LiveCodingManifest, WriteOutdatedActionsFile);
+						Build(TargetDescriptors, BuildConfiguration, WorkingSet, Options, WriteOutdatedActionsFile);
 					}
 				}
 			}
@@ -220,10 +214,9 @@ namespace UnrealBuildTool
 		/// <param name="BuildConfiguration">Current build configuration</param>
 		/// <param name="WorkingSet">The source file working set</param>
 		/// <param name="Options">Additional options for the build</param>
-		/// <param name="LiveCodingManifest">Path to write the live coding manifest to</param>
 		/// <param name="WriteOutdatedActionsFile">Files to write the list of outdated actions to (rather than building them)</param>
 		/// <returns>Result from the compilation</returns>
-		public static void Build(List<TargetDescriptor> TargetDescriptors, BuildConfiguration BuildConfiguration, ISourceFileWorkingSet WorkingSet, BuildOptions Options, FileReference LiveCodingManifest, FileReference WriteOutdatedActionsFile)
+		public static void Build(List<TargetDescriptor> TargetDescriptors, BuildConfiguration BuildConfiguration, ISourceFileWorkingSet WorkingSet, BuildOptions Options, FileReference WriteOutdatedActionsFile)
 		{
 			// Create a makefile for each target
 			TargetMakefile[] Makefiles = new TargetMakefile[TargetDescriptors.Count];
@@ -232,13 +225,6 @@ namespace UnrealBuildTool
 				Makefiles[TargetIdx] = CreateMakefile(BuildConfiguration, TargetDescriptors[TargetIdx], WorkingSet);
 			}
 
-			// Output the Live Coding manifest
-			if(LiveCodingManifest != null)
-			{
-				List<Action> AllActions = Makefiles.SelectMany(x => x.Actions).ToList();
-				HotReload.WriteLiveCodingManifest(LiveCodingManifest, AllActions);
-			}
-
 			// Export the actions for each target
 			for (int TargetIdx = 0; TargetIdx < TargetDescriptors.Count; TargetIdx++)
 			{
@@ -646,6 +632,19 @@ namespace UnrealBuildTool
 				// Filter the prerequisite actions down to just the compile actions, then recompute all the actions to execute
 				PrerequisiteActions = new List<Action>(TargetActionsToExecute.Where(x => x.ActionType == ActionType.Compile));
 				TargetActionsToExecute = ActionGraph.GetActionsToExecute(Makefile.Actions, PrerequisiteActions, CppDependencies, History, BuildConfiguration.bIgnoreOutdatedImportLibraries);
+
+				// Update the action graph with these new paths
+				Dictionary<FileReference, FileReference> OriginalFileToPatchedFile = new Dictionary<FileReference, FileReference>();
+				HotReload.PatchActionGraphForLiveCoding(PrerequisiteActions, OriginalFileToPatchedFile);
+
+				// Get a new list of actions to execute now that the graph has been modified
+				TargetActionsToExecute = ActionGraph.GetActionsToExecute(Makefile.Actions, PrerequisiteActions, CppDependencies, History, BuildConfiguration.bIgnoreOutdatedImportLibraries);
+
+				// Output the Live Coding manifest
+				if(TargetDescriptor.LiveCodingManifest != null)
+				{
+					HotReload.WriteLiveCodingManifest(TargetDescriptor.LiveCodingManifest, Makefile.Actions, OriginalFileToPatchedFile);
+				}
 			}
 			else if (HotReloadMode == HotReloadMode.FromEditor || HotReloadMode == HotReloadMode.FromIDE)
 			{
diff --git a/Engine/Source/Programs/UnrealBuildTool/System/ExternalExecution.cs b/Engine/Source/Programs/UnrealBuildTool/System/ExternalExecution.cs
--- a/Engine/Source/Programs/UnrealBuildTool/System/ExternalExecution.cs
+++ b/Engine/Source/Programs/UnrealBuildTool/System/ExternalExecution.cs
@@ -1105,7 +1105,7 @@ namespace UnrealBuildTool
 
 						using(Timeline.ScopeEvent("Buildng UnrealHeaderTool"))
 						{
-							BuildMode.Build(TargetDescriptors, BuildConfiguration, WorkingSet, BuildOptions.Quiet, null, null);
+							BuildMode.Build(TargetDescriptors, BuildConfiguration, WorkingSet, BuildOptions.Quiet, null);
 						}
 					}
 
diff --git a/Engine/Source/Programs/UnrealBuildTool/System/HotReload.cs b/Engine/Source/Programs/UnrealBuildTool/System/HotReload.cs
--- a/Engine/Source/Programs/UnrealBuildTool/System/HotReload.cs
+++ b/Engine/Source/Programs/UnrealBuildTool/System/HotReload.cs
@@ -449,6 +449,106 @@ namespace UnrealBuildTool
 			return Char.IsLetterOrDigit(Character) || Character == '_';
 		}
 
+		/// <summary>
+		/// Patches a set of actions for use with live coding. The new action list will output object files to a different location.
+		/// </summary>
+		/// <param name="Actions">Set of actions</param>
+		/// <param name="OriginalFileToPatchedFile">Dictionary that receives a map of original object file to patched object file</param>
+		public static void PatchActionGraphForLiveCoding(IEnumerable<Action> Actions, Dictionary<FileReference, FileReference> OriginalFileToPatchedFile)
+		{
+			foreach (Action Action in Actions)
+			{
+				if(Action.ActionType == ActionType.Compile)
+				{
+					if(!Action.CommandPath.GetFileName().Equals("cl-filter.exe", StringComparison.OrdinalIgnoreCase))
+					{
+						throw new BuildException("Unable to patch action graph - unexpected executable in compile action ({0})", Action.CommandPath);
+					}
+
+					List<string> Arguments = Utils.ParseArgumentList(Action.CommandArguments);
+
+					// Find the index of the cl-filter argument delimiter
+					int DelimiterIdx = Arguments.IndexOf("--");
+					if(DelimiterIdx == -1)
+					{
+						throw new BuildException("Unable to patch action graph - missing '--' delimiter to cl-filter");
+					}
+
+					// Fix the dependencies path
+					const string DependenciesPrefix = "-dependencies=";
+
+					int DependenciesIdx = 0;
+					for(;;DependenciesIdx++)
+					{
+						if(DependenciesIdx == DelimiterIdx)
+						{
+							throw new BuildException("Unable to patch action graph - missing '{0}' argument to cl-filter", DependenciesPrefix);
+						}
+						else if(Arguments[DependenciesIdx].StartsWith(DependenciesPrefix, StringComparison.OrdinalIgnoreCase))
+						{
+							break;
+						}
+					}
+
+					FileReference OldDependenciesFile = new FileReference(Arguments[DependenciesIdx].Substring(DependenciesPrefix.Length));
+					FileItem OldDependenciesFileItem = Action.ProducedItems.First(x => x.Location == OldDependenciesFile);
+					Action.ProducedItems.Remove(OldDependenciesFileItem);
+
+					FileReference NewDependenciesFile = OldDependenciesFile.ChangeExtension(".lc.response");
+					FileItem NewDependenciesFileItem = FileItem.GetItemByFileReference(NewDependenciesFile);
+					Action.ProducedItems.Add(NewDependenciesFileItem);
+
+					Arguments[DependenciesIdx] = DependenciesPrefix + NewDependenciesFile.FullName;
+
+					// Fix the response file
+					int ResponseFileIdx = DelimiterIdx + 1;
+					for (; ; ResponseFileIdx++)
+					{
+						if (ResponseFileIdx == Arguments.Count)
+						{
+							throw new BuildException("Unable to patch action graph - missing response file argument to cl-filter");
+						}
+						else if (Arguments[ResponseFileIdx].StartsWith("@", StringComparison.Ordinal))
+						{
+							break;
+						}
+					}
+
+					FileReference OldResponseFile = new FileReference(Arguments[ResponseFileIdx].Substring(1));
+					FileReference NewResponseFile = new FileReference(OldResponseFile.FullName + ".lc");
+
+					const string OutputFilePrefix = "/Fo";
+
+					string[] ResponseLines = FileReference.ReadAllLines(OldResponseFile);
+					for(int Idx = 0; Idx < ResponseLines.Length; Idx++)
+					{
+						string ResponseLine = ResponseLines[Idx];
+						if(ResponseLine.StartsWith(OutputFilePrefix, StringComparison.Ordinal))
+						{
+							FileReference OldOutputFile = new FileReference(ResponseLine.Substring(3).Trim('\"'));
+							FileItem OldOutputFileItem = Action.ProducedItems.First(x => x.Location == OldOutputFile);
+							Action.ProducedItems.Remove(OldOutputFileItem);
+
+							FileReference NewOutputFile = OldOutputFile.ChangeExtension(".lc.obj");
+							FileItem NewOutputFileItem = FileItem.GetItemByFileReference(NewOutputFile);
+							Action.ProducedItems.Add(NewOutputFileItem);
+
+							OriginalFileToPatchedFile[OldOutputFile] = NewOutputFile;
+
+							ResponseLines[Idx] = OutputFilePrefix + "\"" + NewOutputFile.FullName + "\"";
+							break;
+						}
+					}
+					FileReference.WriteAllLines(NewResponseFile, ResponseLines);
+
+					Arguments[ResponseFileIdx] = "@" + NewResponseFile.FullName;
+
+					// Update the final arguments
+					Action.CommandArguments = Utils.FormatCommandLine(Arguments);
+				}
+			}
+		}
+
 		/// <summary>
 		/// Patch the action graph for hot reloading, mapping files according to the given dictionary.
 		/// </summary>
@@ -767,7 +867,8 @@ namespace UnrealBuildTool
 		/// </summary>
 		/// <param name="ManifestFile">File to write to</param>
 		/// <param name="Actions">List of actions that are part of the graph</param>
-		public static void WriteLiveCodingManifest(FileReference ManifestFile, List<Action> Actions)
+		/// <param name="OriginalFileToPatchedFile">Map of original object files to patched object files</param>
+		public static void WriteLiveCodingManifest(FileReference ManifestFile, List<Action> Actions, Dictionary<FileReference, FileReference> OriginalFileToPatchedFile)
 		{
 			// Find all the output object files
 			HashSet<FileItem> ObjectFiles = new HashSet<FileItem>();
@@ -803,7 +904,7 @@ namespace UnrealBuildTool
 					if(Action.ActionType == ActionType.Link)
 					{
 						FileItem OutputFile = Action.ProducedItems.FirstOrDefault(x => x.HasExtension(".exe") || x.HasExtension(".dll"));
-						if(OutputFile != null)
+						if(OutputFile != null && Action.PrerequisiteItems.Any(x => OriginalFileToPatchedFile.ContainsKey(x.Location)))
 						{
 							Writer.WriteObjectStart();
 							Writer.WriteValue("Output", OutputFile.Location.FullName);
@@ -811,9 +912,10 @@ namespace UnrealBuildTool
 							Writer.WriteArrayStart("Inputs");
 							foreach(FileItem InputFile in Action.PrerequisiteItems)
 							{
-								if(ObjectFiles.Contains(InputFile))
+								FileReference PatchedFile;
+								if(OriginalFileToPatchedFile.TryGetValue(InputFile.Location, out PatchedFile))
 								{
-									Writer.WriteValue(InputFile.AbsolutePath);
+									Writer.WriteValue(PatchedFile.FullName);
 								}
 							}
 							Writer.WriteArrayEnd();
diff --git a/Engine/Source/Programs/UnrealBuildTool/System/Utils.cs b/Engine/Source/Programs/UnrealBuildTool/System/Utils.cs
--- a/Engine/Source/Programs/UnrealBuildTool/System/Utils.cs
+++ b/Engine/Source/Programs/UnrealBuildTool/System/Utils.cs
@@ -1158,5 +1158,86 @@ namespace UnrealBuildTool
 				}
 			}
 		}
+
+		/// <summary>
+		/// Parses a command line into a list of arguments
+		/// </summary>
+		/// <param name="CommandLine">The command line to parse</param>
+		/// <returns>List of output arguments</returns>
+		public static List<string> ParseArgumentList(string CommandLine)
+		{
+			List<string> Arguments = new List<string>();
+
+			StringBuilder CurrentArgument = new StringBuilder();
+			for (int Idx = 0; Idx < CommandLine.Length; Idx++)
+			{
+				if (!Char.IsWhiteSpace(CommandLine[Idx]))
+				{
+					CurrentArgument.Clear();
+
+					bool bInQuotes = false;
+					for (; Idx < CommandLine.Length; Idx++)
+					{
+						if (CommandLine[Idx] == '\"')
+						{
+							bInQuotes ^= true;
+						}
+						else if (CommandLine[Idx] == ' ' && !bInQuotes)
+						{
+							break;
+						}
+						else
+						{
+							CurrentArgument.Append(CommandLine[Idx]);
+						}
+					}
+
+					Arguments.Add(CurrentArgument.ToString());
+				}
+			}
+
+			return Arguments;
+		}
+
+		/// <summary>
+		/// Formats a list of arguments as a command line, inserting quotes as necessary
+		/// </summary>
+		/// <param name="Arguments">List of arguments to format</param>
+		/// <returns>Command line string</returns>
+		public static string FormatCommandLine(List<string> Arguments)
+		{
+			StringBuilder CommandLine = new StringBuilder();
+			foreach (string Argument in Arguments)
+			{
+				if (CommandLine.Length > 0)
+				{
+					CommandLine.Append(' ');
+				}
+
+				int SpaceIdx = Argument.IndexOf(' ');
+				if (SpaceIdx == -1)
+				{
+					CommandLine.Append(Argument);
+				}
+				else
+				{
+					int EqualsIdx = Argument.IndexOf('=');
+					if (EqualsIdx != -1 && Argument[0] == '-')
+					{
+						CommandLine.Append(Argument, 0, EqualsIdx + 1);
+						CommandLine.Append('\"');
+						CommandLine.Append(Argument, EqualsIdx + 1, Argument.Length - (EqualsIdx + 1));
+						CommandLine.Append('\"');
+					}
+					else
+					{
+						CommandLine.Append('\"');
+						CommandLine.Append(Argument);
+						CommandLine.Append('\"');
+					}
+				}
+			}
+			return CommandLine.ToString();
+		}
 	}
 }
