Prevent LiveCodingConsole being built twice if bBuildLiveCodingConsole = true in BuildConfiguration.xml, and it's the target being built from Visual Studio.

#rnx

diff --git a/Engine/Source/Programs/UnrealBuildTool/ProjectFiles/VisualStudio/VCProject.cs b/Engine/Source/Programs/UnrealBuildTool/ProjectFiles/VisualStudio/VCProject.cs
--- a/Engine/Source/Programs/UnrealBuildTool/ProjectFiles/VisualStudio/VCProject.cs
+++ b/Engine/Source/Programs/UnrealBuildTool/ProjectFiles/VisualStudio/VCProject.cs
@@ -1407,7 +1407,7 @@ namespace UnrealBuildTool
 						{
 							ExtraTargets.Add("ShaderCompileWorker Win64 Development");
 						}
-						if (TargetRulesObject.bWithLiveCoding && bBuildLiveCodingConsole && !UnrealBuildTool.IsEngineInstalled())
+						if (TargetRulesObject.bWithLiveCoding && bBuildLiveCodingConsole && !UnrealBuildTool.IsEngineInstalled() && TargetRulesObject.Name != "LiveCodingConsole")
 						{
 							ExtraTargets.Add(TargetRulesObject.bUseDebugLiveCodingConsole? "LiveCodingConsole Win64 Debug" : "LiveCodingConsole Win64 Development");
 						}
