Merge Dev-VR->Release-4.21

diff --git a/Engine/Source/Programs/UnrealBuildTool/ProjectFiles/Xcode/XcodeProjectFileGenerator.cs b/Engine/Source/Programs/UnrealBuildTool/ProjectFiles/Xcode/XcodeProjectFileGenerator.cs
--- a/Engine/Source/Programs/UnrealBuildTool/ProjectFiles/Xcode/XcodeProjectFileGenerator.cs
+++ b/Engine/Source/Programs/UnrealBuildTool/ProjectFiles/Xcode/XcodeProjectFileGenerator.cs
@@ -106,7 +106,9 @@ namespace UnrealBuildTool
 			WorkspaceSettingsContent.Append("<!DOCTYPE plist PUBLIC \"-//Apple//DTD PLIST 1.0//EN\" \"http://www.apple.com/DTDs/PropertyList-1.0.dtd\">" + ProjectFileGenerator.NewLine);
 			WorkspaceSettingsContent.Append("<plist version=\"1.0\">" + ProjectFileGenerator.NewLine);
 			WorkspaceSettingsContent.Append("<dict>" + ProjectFileGenerator.NewLine);
-			WorkspaceSettingsContent.Append("\t<key>BuildLocationStyle</key>" + ProjectFileGenerator.NewLine);
+			WorkspaceSettingsContent.Append("\t<key>BuildSystemType</key>" + ProjectFileGenerator.NewLine);
+			WorkspaceSettingsContent.Append("\t<string>Original</string>" + ProjectFileGenerator.NewLine);
+            WorkspaceSettingsContent.Append("\t<key>BuildLocationStyle</key>" + ProjectFileGenerator.NewLine);
 			WorkspaceSettingsContent.Append("\t<string>UseTargetSettings</string>" + ProjectFileGenerator.NewLine);
 			WorkspaceSettingsContent.Append("\t<key>CustomBuildLocationType</key>" + ProjectFileGenerator.NewLine);
 			WorkspaceSettingsContent.Append("\t<string>RelativeToDerivedData</string>" + ProjectFileGenerator.NewLine);
