Fix temporary files containing dependency lists being created in Engine directory. Move them under Engine/Intermediate instead.

diff --git a/Engine/Source/Programs/UnrealBuildTool/Configuration/UEBuildTarget.cs b/Engine/Source/Programs/UnrealBuildTool/Configuration/UEBuildTarget.cs
--- a/Engine/Source/Programs/UnrealBuildTool/Configuration/UEBuildTarget.cs
+++ b/Engine/Source/Programs/UnrealBuildTool/Configuration/UEBuildTarget.cs
@@ -421,11 +421,11 @@ namespace UnrealBuildTool
 				DirectoryReference DependencyListDir;
 				if(RulesObject.ProjectFile == null)
 				{
-					DependencyListDir = DirectoryReference.Combine(UnrealBuildTool.EngineDirectory, RulesObject.Name, RulesObject.Configuration.ToString(), RulesObject.Platform.ToString());
+					DependencyListDir = DirectoryReference.Combine(UnrealBuildTool.EngineDirectory, "Intermediate", "DependencyLists", RulesObject.Name, RulesObject.Configuration.ToString(), RulesObject.Platform.ToString());
 				}
 				else
 				{
-					DependencyListDir = DirectoryReference.Combine(RulesObject.ProjectFile.Directory, RulesObject.Name, RulesObject.Configuration.ToString(), RulesObject.Platform.ToString());
+					DependencyListDir = DirectoryReference.Combine(RulesObject.ProjectFile.Directory, "Intermediate", "DependencyLists", RulesObject.Name, RulesObject.Configuration.ToString(), RulesObject.Platform.ToString());
 				}
 
 				FileReference DependencyListFile;
