Linux: Set FORCE_ANSI_ALLOCATOR when using -EnableTSan to avoid multiple definitions

diff --git a/Engine/Source/Programs/UnrealBuildTool/Platform/Linux/UEBuildLinux.cs b/Engine/Source/Programs/UnrealBuildTool/Platform/Linux/UEBuildLinux.cs
--- a/Engine/Source/Programs/UnrealBuildTool/Platform/Linux/UEBuildLinux.cs
+++ b/Engine/Source/Programs/UnrealBuildTool/Platform/Linux/UEBuildLinux.cs
@@ -248,8 +248,9 @@ namespace UnrealBuildTool
 				Target.bCompileCEF3 = false;
 			}
 
-			// Force using the ANSI allocator if ASan is enabled
-			if (Target.LinuxPlatform.bEnableAddressSanitizer)
+			// Force using the ANSI allocator if ASan or TSan is enabled
+			if (Target.LinuxPlatform.bEnableAddressSanitizer ||
+				Target.LinuxPlatform.bEnableThreadSanitizer)
 			{
 				Target.GlobalDefinitions.Add("FORCE_ANSI_ALLOCATOR=1");
 			}
