UBT: Avoid trying to figure out which mobile provision to use for a project when we already know the path to the provision we want. Also fix certificate not being updated if mobile provision changes.

diff --git a/Engine/Source/Programs/IOS/iPhonePackager/CompileTime.cs b/Engine/Source/Programs/IOS/iPhonePackager/CompileTime.cs
--- a/Engine/Source/Programs/IOS/iPhonePackager/CompileTime.cs
+++ b/Engine/Source/Programs/IOS/iPhonePackager/CompileTime.cs
@@ -213,10 +213,9 @@ namespace iPhonePackager
 		/// </summary>
 		static public void ExportCertificate()
 		{
-			string ProvisionWithPrefix;
-			if(!FindMobileProvision(Config.OverrideBundleName ?? "", out ProvisionWithPrefix, false))
+			if(Config.ProvisionFile == null)
 			{
-				Program.Error("Missing provision");
+				Program.Error("Missing -ProvisionFile=... argument");
 				return;
 			}
 
@@ -227,7 +226,7 @@ namespace iPhonePackager
 			}
 
 			// export the signing certificate to a file
-			MobileProvision Provision = MobileProvisionParser.ParseFile(ProvisionWithPrefix);
+			MobileProvision Provision = MobileProvisionParser.ParseFile(Config.ProvisionFile);
 			X509Certificate2 Certificate = CodeSignatureBuilder.FindCertificate(Provision);
             if (Certificate == null)
             {
diff --git a/Engine/Source/Programs/IOS/iPhonePackager/Config.cs b/Engine/Source/Programs/IOS/iPhonePackager/Config.cs
--- a/Engine/Source/Programs/IOS/iPhonePackager/Config.cs
+++ b/Engine/Source/Programs/IOS/iPhonePackager/Config.cs
@@ -350,6 +350,11 @@ namespace iPhonePackager
 		/// </summary>
 		public static string ProvisionUUID = "";
 
+        /// <summary>
+		/// provision to extract certificate for
+		/// </summary>
+		public static string ProvisionFile;
+
 		/// <summary>
 		/// IOS Team ID to be used for automatic signing
 		/// </summary>
diff --git a/Engine/Source/Programs/IOS/iPhonePackager/iPhonePackager.cs b/Engine/Source/Programs/IOS/iPhonePackager/iPhonePackager.cs
--- a/Engine/Source/Programs/IOS/iPhonePackager/iPhonePackager.cs
+++ b/Engine/Source/Programs/IOS/iPhonePackager/iPhonePackager.cs
@@ -283,6 +283,17 @@ namespace iPhonePackager
                                 return false;
                             }
                         }
+						else if (Arg == "-provisionfile")
+						{
+							if (Arguments.Length > ArgIndex + 1)
+							{
+								Config.ProvisionFile = Arguments [++ArgIndex];
+							}
+							else
+							{
+								return false;
+							}
+						}
 						else if (Arg == "-teamID")
 						{
 							// make sure there's at least one more arg
diff --git a/Engine/Source/Programs/UnrealBuildTool/ToolChain/RemoteMac.cs b/Engine/Source/Programs/UnrealBuildTool/ToolChain/RemoteMac.cs
--- a/Engine/Source/Programs/UnrealBuildTool/ToolChain/RemoteMac.cs
+++ b/Engine/Source/Programs/UnrealBuildTool/ToolChain/RemoteMac.cs
@@ -404,10 +404,16 @@ namespace UnrealBuildTool
 				Log.TraceInformation("[Remote] Uploading {0}", MobileProvisionFile);
 				UploadFile(MobileProvisionFile);
 
-				// Extract the certificate for the project
+				// Extract the certificate for the project. Try to avoid calling IPP if we already have it.
 				FileReference CertificateFile = FileReference.Combine(TempDir, "Certificate.p12");
-				if(!FileReference.Exists(CertificateFile))
+
+				FileReference CertificateInfoFile = FileReference.Combine(TempDir, "Certificate.txt");
+				string CertificateInfoContents = String.Format("{0}\n{1}", ProvisioningData.MobileProvisionFile, FileReference.GetLastWriteTimeUtc(ProvisioningData.MobileProvisionFile).Ticks);
+
+				if(!FileReference.Exists(CertificateFile) || !FileReference.Exists(CertificateInfoFile) || FileReference.ReadAllText(CertificateInfoFile) != CertificateInfoContents)
 				{
+					Log.TraceInformation("[Remote] Exporting certificate for {0}...", ProvisioningData.MobileProvisionFile);
+
 					StringBuilder Arguments = new StringBuilder("ExportCertificate");
 					if(TargetDesc.ProjectFile == null)
 					{
@@ -417,12 +423,12 @@ namespace UnrealBuildTool
 					{
 						Arguments.AppendFormat(" \"{0}\"", TargetDesc.ProjectFile.Directory);
 					}
+					Arguments.AppendFormat(" -provisionfile \"{0}\"", ProvisioningData.MobileProvisionFile);
+					Arguments.AppendFormat(" -outputcertificate \"{0}\"", CertificateFile);
 					if(TargetDesc.Platform == UnrealTargetPlatform.TVOS)
 					{
 						Arguments.Append(" -tvos");
 					}
-					Arguments.AppendFormat(" -outputcertificate \"{0}\"", CertificateFile);
-					Arguments.AppendFormat(" -provisioninguuid {0}", ProvisioningData.MobileProvisionUUID);
 
 					ProcessStartInfo StartInfo = new ProcessStartInfo();
 					StartInfo.FileName = FileReference.Combine(UnrealBuildTool.EngineDirectory, "Binaries", "DotNET", "IOS", "IPhonePackager.exe").FullName;
@@ -431,6 +437,8 @@ namespace UnrealBuildTool
 					{
 						throw new BuildException("IphonePackager failed.");
 					}
+
+					FileReference.WriteAllText(CertificateInfoFile, CertificateInfoContents);
 				}
 
 				// Upload the certificate to the remote
