PR #5138: Fix bForceBuildShaderFormats not being read correctly in TargetRules (Contributed by [REDACTED_PERSON])

diff --git a/Engine/Source/Programs/UnrealBuildTool/Configuration/TargetRules.cs b/Engine/Source/Programs/UnrealBuildTool/Configuration/TargetRules.cs
--- a/Engine/Source/Programs/UnrealBuildTool/Configuration/TargetRules.cs
+++ b/Engine/Source/Programs/UnrealBuildTool/Configuration/TargetRules.cs
@@ -1568,7 +1568,7 @@ namespace UnrealBuildTool
 
 		public bool bForceBuildShaderFormats
 		{
-			get { return Inner.bForceBuildTargetPlatforms; }
+			get { return Inner.bForceBuildShaderFormats; }
 		}
 
 		public bool bCompileSimplygon
