Merging CL 4347297 from Dev-VR - Swapping double to single quote to fix Mac Editor packaging for Lumin

diff --git a/Engine/Source/Programs/UnrealBuildTool/Platform/Android/AndroidToolChain.cs b/Engine/Source/Programs/UnrealBuildTool/Platform/Android/AndroidToolChain.cs
--- a/Engine/Source/Programs/UnrealBuildTool/Platform/Android/AndroidToolChain.cs
+++ b/Engine/Source/Programs/UnrealBuildTool/Platform/Android/AndroidToolChain.cs
@@ -1453,7 +1453,7 @@ namespace UnrealBuildTool
 							else
 							{
 								CompileAction.CommandPath = "/bin/sh";
-								CompileAction.CommandArguments = String.Format("-c \"{0} {1}\"", ClangPath, ResponseArgument);
+								CompileAction.CommandArguments = String.Format("-c \'{0} {1}\'", ClangPath, ResponseArgument);
 								CompileAction.CommandDescription = "Compile";
 							}
 						}
@@ -1684,7 +1684,7 @@ namespace UnrealBuildTool
 						}
 						else
 						{
-							LinkAction.CommandArguments = String.Format("-c \"{0} {1}\"", LinkAction.CommandPath, LinkAction.CommandArguments);
+							LinkAction.CommandArguments = String.Format("-c \'{0} {1}\'", LinkAction.CommandPath, LinkAction.CommandArguments);
 							LinkAction.CommandPath = "/bin/sh";
 							LinkAction.CommandDescription = "Link";
 						}
