Fix for x86 and x86_64 Android architectures with NDK14b

#ue4
#android

diff --git a/Engine/Source/Programs/UnrealBuildTool/Platform/Android/AndroidToolChain.cs b/Engine/Source/Programs/UnrealBuildTool/Platform/Android/AndroidToolChain.cs
--- a/Engine/Source/Programs/UnrealBuildTool/Platform/Android/AndroidToolChain.cs
+++ b/Engine/Source/Programs/UnrealBuildTool/Platform/Android/AndroidToolChain.cs
@@ -344,7 +344,7 @@ namespace UnrealBuildTool
 										" --sysroot=\"" + Path.Combine(NDKPath, "sysroot") + "\"" +
 										" -isystem " + Path.Combine(NDKPath, "sysroot/usr/include/i686-linux-android/") +
 										" -D__ANDROID_API__=" + NDKApiLevel32Int;
-				ToolchainParamsx86 = " -target x86_64-none-linux-android" +
+				ToolchainParamsx64 = " -target x86_64-none-linux-android" +
 										" --sysroot=\"" + Path.Combine(NDKPath, "sysroot") + "\"" +
 										" -isystem " + Path.Combine(NDKPath, "sysroot/usr/include/x86_64-linux-android/") +
 										" -D__ANDROID_API__=" + NDKApiLevel64Int;
