Allow plugins specified on the command line to override engine plugins.

diff --git a/Engine/Source/Programs/UnrealBuildTool/System/Plugins.cs b/Engine/Source/Programs/UnrealBuildTool/System/Plugins.cs
--- a/Engine/Source/Programs/UnrealBuildTool/System/Plugins.cs
+++ b/Engine/Source/Programs/UnrealBuildTool/System/Plugins.cs
@@ -26,7 +26,7 @@ namespace UnrealBuildTool
 	}
 
 	/// <summary>
-	/// Where a plugin was loaded from
+	/// Where a plugin was loaded from. The order of this enum is important; in the case of name collisions, larger-valued types will take precedence. Plugins of the same type may not be duplicated.
 	/// </summary>
 	public enum PluginType
 	{
@@ -169,11 +169,11 @@ namespace UnrealBuildTool
 				{
 					NameToPluginInfo.Add(Plugin.Name, Plugin);
 				}
-				else if(ExistingPluginInfo.Type == PluginType.Engine && Plugin.Type == PluginType.Project)
+				else if(Plugin.Type > ExistingPluginInfo.Type)
 				{
 					NameToPluginInfo[Plugin.Name] = Plugin;
 				}
-				else if(ExistingPluginInfo.Type != PluginType.Project || Plugin.Type != PluginType.Engine)
+				else if(Plugin.Type == ExistingPluginInfo.Type)
 				{
 					throw new BuildException(String.Format("Found '{0}' plugin in two locations ({1} and {2}). Plugin names must be unique.", Plugin.Name, ExistingPluginInfo.File, Plugin.File));
 				}
