UBT: Determine the correct user directory when remote compiling on Mac, rather than inferring it from the user name. The user name can be changed after an account is created.

diff --git a/Engine/Build/Rsync/RsyncEngine.txt b/Engine/Build/Rsync/RsyncEngine.txt
--- a/Engine/Build/Rsync/RsyncEngine.txt
+++ b/Engine/Build/Rsync/RsyncEngine.txt
@@ -217,10 +217,14 @@ protect /Binaries/ThirdParty/Mono/Mac/lib/mono/monodoc/monodoc.dll
 + /Binaries/ThirdParty/Mono/
 + /Binaries/ThirdParty/Mono/Mac/
 + /Binaries/ThirdParty/Mono/Mac/**
++ /Binaries/ThirdParty/**/
++ /Binaries/ThirdParty/**.dylib
 + /Build/
 + /Build/*
 + /Build/BatchFiles/
 + /Build/BatchFiles/**
++ /Build/Mac/
++ /Build/Mac/**
 + /Build/IOS/
 + /Build/IOS/**
 + /Config/
@@ -248,6 +252,7 @@ protect /Binaries/ThirdParty/Mono/Mac/lib/mono/monodoc/monodoc.dll
 + /Plugins/**/Source/**.m
 + /Plugins/**/Source/**.cs
 + /Plugins/**/Source/**.inl
++ /Plugins/**/Source/**.dylib
 + /Plugins/**.a
 + /Plugins/**.embeddedframework.zip
 + /Plugins/**.xml
@@ -291,5 +296,6 @@ protect /Binaries/ThirdParty/Mono/Mac/lib/mono/monodoc/monodoc.dll
 + /Source/ThirdParty/**.mm
 + /Source/ThirdParty/**.inl
 + /Source/ThirdParty/**.a
++ /Source/ThirdParty/**.dylib
 + /Source/ThirdParty/**.embeddedframework.zip
 + /Source/ThirdParty/libJPG/*.cpp
diff --git a/Engine/Source/Programs/UnrealBuildTool/ToolChain/RemoteMac.cs b/Engine/Source/Programs/UnrealBuildTool/ToolChain/RemoteMac.cs
--- a/Engine/Source/Programs/UnrealBuildTool/ToolChain/RemoteMac.cs
+++ b/Engine/Source/Programs/UnrealBuildTool/ToolChain/RemoteMac.cs
@@ -219,17 +219,6 @@ namespace UnrealBuildTool
 			RsyncAuthentication = ExpandVariables(RsyncAuthentication);
 			SshAuthentication = ExpandVariables(SshAuthentication);
 
-			// Get the remote base directory
-			RemoteBaseDir = String.Format("[REDACTED_USER_HOME]/UE4/Builds/{1}", UserName, Environment.MachineName);
-
-			// Build the list of directory mappings between the local and remote machines
-			Mappings = new List<RemoteMapping>();
-			Mappings.Add(new RemoteMapping(UnrealBuildTool.EngineDirectory, GetRemotePath(UnrealBuildTool.EngineDirectory)));
-			if(ProjectFile != null && !ProjectFile.IsUnderDirectory(UnrealBuildTool.EngineDirectory))
-			{
-				Mappings.Add(new RemoteMapping(ProjectFile.Directory, GetRemotePath(ProjectFile.Directory)));
-			}
-
 			// Build a list of arguments for SSH
 			CommonSshArguments = new List<string>();
 			CommonSshArguments.Add("-o BatchMode=yes");
@@ -248,6 +237,23 @@ namespace UnrealBuildTool
 			CommonRsyncArguments.Add("-m");
 			CommonRsyncArguments.Add("--chmod=ug=rwX,o=rxX");
 			CommonRsyncArguments.Add(String.Format("--rsh=\"{0} -p {1}\"", RsyncAuthentication, ServerPort));
+
+			// Get the remote base directory
+			StringBuilder Output;
+			if(ExecuteAndCaptureOutput("'echo ~'", out Output) != 0)
+			{
+				throw new BuildException("Unable to determine home directory for remote user");
+			}
+			RemoteBaseDir = String.Format("{0}/UE4/Builds/{1}", Output.ToString().Trim().TrimEnd('/'), Environment.MachineName);
+			Log.TraceInformation("[Remote] Using base directory '{0}'", RemoteBaseDir);
+
+			// Build the list of directory mappings between the local and remote machines
+			Mappings = new List<RemoteMapping>();
+			Mappings.Add(new RemoteMapping(UnrealBuildTool.EngineDirectory, GetRemotePath(UnrealBuildTool.EngineDirectory)));
+			if(ProjectFile != null && !ProjectFile.IsUnderDirectory(UnrealBuildTool.EngineDirectory))
+			{
+				Mappings.Add(new RemoteMapping(ProjectFile.Directory, GetRemotePath(ProjectFile.Directory)));
+			}
 		}
 
 		/// <summary>
@@ -907,6 +913,39 @@ namespace UnrealBuildTool
 			}
 		}
 
+		/// <summary>
+		/// Execute a remote command, capturing the output text
+		/// </summary>
+		/// <param name="Command">Command to be executed</param>
+		/// <param name="Output">Receives the output text</param>
+		/// <returns></returns>
+		protected int ExecuteAndCaptureOutput(string Command, out StringBuilder Output)
+		{
+			StringBuilder FullCommand = new StringBuilder();
+			foreach(string CommonSshArgument in CommonSshArguments)
+			{
+				FullCommand.AppendFormat("{0} ", CommonSshArgument);
+			}
+			FullCommand.Append(Command.Replace("\"", "\\\""));
+
+			using(Process SSHProcess = new Process())
+			{
+				Output = new StringBuilder();
+
+				StringBuilder OutputLocal = Output;
+				DataReceivedEventHandler OutputHandler = (E, Args) => { if(Args.Data != null){ OutputLocal.Append(Args.Data); } };
+
+				SSHProcess.StartInfo.FileName = SshExe.FullName;
+				SSHProcess.StartInfo.WorkingDirectory = SshExe.Directory.FullName;
+				SSHProcess.StartInfo.Arguments = FullCommand.ToString();
+				SSHProcess.OutputDataReceived += OutputHandler;
+				SSHProcess.ErrorDataReceived += OutputHandler;
+
+				Log.TraceLog("[SSH] {0} {1}", Utils.MakePathSafeToUseWithCommandLine(SSHProcess.StartInfo.FileName), SSHProcess.StartInfo.Arguments);
+				return Utils.RunLocalProcess(SSHProcess);
+			}
+		}
+
 		/// <summary>
 		/// Converts any remote paths within the given string to local format
 		/// </summary>
