UBT: Detect changes to source directories and invalidate the makefile if they are newer. Forces the makefile to be rebuilt when files are added or removed.

diff --git a/Engine/Source/Programs/UnrealBuildTool/Configuration/UEBuildTarget.cs b/Engine/Source/Programs/UnrealBuildTool/Configuration/UEBuildTarget.cs
--- a/Engine/Source/Programs/UnrealBuildTool/Configuration/UEBuildTarget.cs
+++ b/Engine/Source/Programs/UnrealBuildTool/Configuration/UEBuildTarget.cs
@@ -801,6 +801,11 @@ namespace UnrealBuildTool
 		/// </summary>
 		public FileReference DeployTargetFile;
 
+		/// <summary>
+		/// Directories which are scanned for source files. The UBT makefile uses this to check for files being added/removed.
+		/// </summary>
+		public HashSet<DirectoryReference> SourceDirectories = new HashSet<DirectoryReference>();
+
 		/// <summary>
 		/// A list of the module filenames which were used to build this target.
 		/// </summary>
@@ -856,6 +861,7 @@ namespace UnrealBuildTool
 			PostBuildStepScripts = (FileReference[])Info.GetValue("po", typeof(FileReference[]));
 			DeployTargetFile = (FileReference)Info.GetValue("dt", typeof(FileReference));
 			bHasProjectScriptPlugin = Info.GetBoolean("sp");
+			SourceDirectories = (HashSet<DirectoryReference>)Info.GetValue("sd", typeof(HashSet<DirectoryReference>));
 		}
 
 		public void GetObjectData(SerializationInfo Info, StreamingContext Context)
@@ -886,6 +892,7 @@ namespace UnrealBuildTool
 			Info.AddValue("po", PostBuildStepScripts);
 			Info.AddValue("dt", DeployTargetFile);
 			Info.AddValue("sp", bHasProjectScriptPlugin);
+			Info.AddValue("sd", SourceDirectories);
 		}
 
 		/// <summary>
@@ -4076,7 +4083,7 @@ namespace UnrealBuildTool
 					if (bDiscoverFiles && !RulesObject.bUsePrecompiled)
 					{
 						List<FileReference> SourceFilePaths = new List<FileReference>();
-						SourceFilePaths = SourceFileSearch.FindModuleSourceFiles(ModuleRulesFile: RulesObject.File);
+						SourceFilePaths = SourceFileSearch.FindModuleSourceFiles(ModuleRulesFile: RulesObject.File, SearchedDirectories: SourceDirectories);
 						FoundSourceFiles = GetCPlusPlusFilesToBuild(SourceFilePaths, ModuleDirectory, Platform);
 					}
 				}
diff --git a/Engine/Source/Programs/UnrealBuildTool/System/SourceFileSearch.cs b/Engine/Source/Programs/UnrealBuildTool/System/SourceFileSearch.cs
--- a/Engine/Source/Programs/UnrealBuildTool/System/SourceFileSearch.cs
+++ b/Engine/Source/Programs/UnrealBuildTool/System/SourceFileSearch.cs
@@ -38,11 +38,11 @@ namespace UnrealBuildTool
 		};
 
 		/// Finds mouse source files
-		public static List<FileReference> FindModuleSourceFiles(FileReference ModuleRulesFile, bool SearchSubdirectories = true)
+		public static List<FileReference> FindModuleSourceFiles(FileReference ModuleRulesFile, bool SearchSubdirectories = true, HashSet<DirectoryReference> SearchedDirectories = null)
 		{
 			// The module's "base directory" is simply the directory where its xxx.Build.cs file is stored.  We'll always
 			// harvest source files for this module in this base directory directory and all of its sub-directories.
-			return SourceFileSearch.FindFiles(ModuleRulesFile.Directory, SubdirectoryNamesToExclude: null, SearchSubdirectories: SearchSubdirectories);
+			return SourceFileSearch.FindFiles(ModuleRulesFile.Directory, SubdirectoryNamesToExclude: null, SearchSubdirectories: SearchSubdirectories, SearchedDirectories: SearchedDirectories);
 		}
 
 		/// <summary>
@@ -51,7 +51,8 @@ namespace UnrealBuildTool
 		/// <param name="DirectoryToSearch">Directory to search</param>
 		/// <param name="SubdirectoryNamesToExclude">Directory base names to ignore when searching subdirectories.  Can be null.</param>
 		/// <param name="SearchSubdirectories">True to include subdirectories, otherwise we only search the list of base directories</param>
-		public static List<FileReference> FindFiles(DirectoryReference DirectoryToSearch, List<string> SubdirectoryNamesToExclude = null, bool SearchSubdirectories = true)
+		/// <param name="SearchedDirectories">If non-null, is updated with a list of the directories searched</param>
+		public static List<FileReference> FindFiles(DirectoryReference DirectoryToSearch, List<string> SubdirectoryNamesToExclude = null, bool SearchSubdirectories = true, HashSet<DirectoryReference> SearchedDirectories = null)
 		{
 			// Build a list of directory names that we don't want to search under. We always ignore intermediate directories.
 			string[] ExcludedDirectorySuffixes;
@@ -68,17 +69,22 @@ namespace UnrealBuildTool
 			List<FileReference> FoundFiles = new List<FileReference>();
 			if (SearchSubdirectories)
 			{
-				FindFilesInternalRecursive(DirectoryToSearch, ExcludedDirectorySuffixes, FoundFiles);
+				FindFilesInternalRecursive(DirectoryToSearch, ExcludedDirectorySuffixes, FoundFiles, SearchedDirectories);
 			}
 			else
 			{
-				FindFilesInternal(DirectoryToSearch, ExcludedDirectorySuffixes, FoundFiles);
+				FindFilesInternal(DirectoryToSearch, ExcludedDirectorySuffixes, FoundFiles, SearchedDirectories);
 			}
 			return FoundFiles;
 		}
 
-		static void FindFilesInternal(DirectoryReference Directory, string[] ExcludedDirectorySuffixes, List<FileReference> FoundFiles)
+		static void FindFilesInternal(DirectoryReference Directory, string[] ExcludedDirectorySuffixes, List<FileReference> FoundFiles, HashSet<DirectoryReference> SearchedDirectories)
 		{
+			if(SearchedDirectories != null)
+			{
+				SearchedDirectories.Add(Directory);
+			}
+
 			foreach (FileReference File in DirectoryReference.EnumerateFiles(Directory))
 			{
 				if (ShouldInclude(File, DefaultExcludedFileSuffixes))
@@ -88,15 +94,15 @@ namespace UnrealBuildTool
 			}
 		}
 
-		static void FindFilesInternalRecursive(DirectoryReference Directory, string[] ExcludedDirectorySuffixes, List<FileReference> FoundFiles)
+		static void FindFilesInternalRecursive(DirectoryReference Directory, string[] ExcludedDirectorySuffixes, List<FileReference> FoundFiles, HashSet<DirectoryReference> SearchedDirectories)
 		{
-			FindFilesInternal(Directory, ExcludedDirectorySuffixes, FoundFiles);
+			FindFilesInternal(Directory, ExcludedDirectorySuffixes, FoundFiles, SearchedDirectories);
 
 			foreach (DirectoryReference SubDirectory in DirectoryReference.EnumerateDirectories(Directory))
 			{
 				if (ShouldInclude(SubDirectory, ExcludedDirectorySuffixes))
 				{
-					FindFilesInternalRecursive(SubDirectory, ExcludedDirectorySuffixes, FoundFiles);
+					FindFilesInternalRecursive(SubDirectory, ExcludedDirectorySuffixes, FoundFiles, SearchedDirectories);
 				}
 			}
 		}
diff --git a/Engine/Source/Programs/UnrealBuildTool/System/UBTMakefile.cs b/Engine/Source/Programs/UnrealBuildTool/System/UBTMakefile.cs
--- a/Engine/Source/Programs/UnrealBuildTool/System/UBTMakefile.cs
+++ b/Engine/Source/Programs/UnrealBuildTool/System/UBTMakefile.cs
@@ -18,8 +18,14 @@ namespace UnrealBuildTool
 	[Serializable]
 	class UBTMakefile : ISerializable
 	{
-		public const int CurrentVersion = 7;
+		/// <summary>
+		/// The version number to write
+		/// </summary>
+		public const int CurrentVersion = 8;
 
+		/// <summary>
+		/// The version number that was read
+		/// </summary>
 		public int Version;
 
 		/// <summary>
@@ -341,6 +347,17 @@ namespace UnrealBuildTool
 						}
 					}
 				}
+
+				foreach(DirectoryReference SourceDirectory in Target.SourceDirectories)
+				{
+					DirectoryInfo SourceDirectoryInfo = new DirectoryInfo(SourceDirectory.FullName);
+					if(!SourceDirectoryInfo.Exists || SourceDirectoryInfo.LastWriteTimeUtc > UBTMakefileInfo.LastWriteTimeUtc)
+					{
+						Log.TraceVerbose("Timestamp of {0} ({1}) is newer than makefile ({2})", SourceDirectory, SourceDirectoryInfo.LastWriteTimeUtc, UBTMakefileInfo.LastWriteTimeUtc);
+						ReasonNotLoaded = "source directory changed";
+						return null;
+					}
+				}
 			}
 
 			// We do a check to see if any modules' headers have changed which have
