Include the Visual Studio .natvis file as a solution item, to fix problems with VS2017 15.8.5 not finding it in a makefile project.

diff --git a/Engine/Source/Programs/UnrealBuildTool/ProjectFiles/VisualStudio/VCProjectFileGenerator.cs b/Engine/Source/Programs/UnrealBuildTool/ProjectFiles/VisualStudio/VCProjectFileGenerator.cs
--- a/Engine/Source/Programs/UnrealBuildTool/ProjectFiles/VisualStudio/VCProjectFileGenerator.cs
+++ b/Engine/Source/Programs/UnrealBuildTool/ProjectFiles/VisualStudio/VCProjectFileGenerator.cs
@@ -583,6 +583,13 @@ namespace UnrealBuildTool
 							"EndProject" + ProjectFileGenerator.NewLine
 						);
 				}
+
+				// Add the visualizers at the solution level. Doesn't seem to be picked up from a makefile project in VS2017 15.8.5.
+				VCSolutionFileContent.AppendLine(String.Format("Project(\"{{2150E333-8FDC-42A3-9474-1A3956D46DE8}}\") = \"Visualizers\", \"Visualizers\", \"{0}\"", Guid.NewGuid().ToString()));
+				VCSolutionFileContent.AppendLine("\tProjectSection(SolutionItems) = preProject");
+				VCSolutionFileContent.AppendLine("\t\tEngine\\Extras\\VisualStudioDebugging\\UE4.natvis = Engine\\Extras\\VisualStudioDebugging\\UE4.natvis");
+				VCSolutionFileContent.AppendLine("\tEndProjectSection");
+				VCSolutionFileContent.AppendLine("EndProject");
 			}
 
 			// Solution configuration platforms.  This is just a list of all of the platforms and configurations that
