Remove natvis file from UE4 project.

diff --git a/Engine/Source/Programs/UnrealBuildTool/ProjectFiles/VisualStudio/VCProjectFileGenerator.cs b/Engine/Source/Programs/UnrealBuildTool/ProjectFiles/VisualStudio/VCProjectFileGenerator.cs
--- a/Engine/Source/Programs/UnrealBuildTool/ProjectFiles/VisualStudio/VCProjectFileGenerator.cs
+++ b/Engine/Source/Programs/UnrealBuildTool/ProjectFiles/VisualStudio/VCProjectFileGenerator.cs
@@ -238,22 +238,6 @@ namespace UnrealBuildTool
 			base.ConfigureProjectFileGeneration(Arguments, ref IncludeAllPlatforms);
 		}
 
-		/// <summary>
-		/// Adds Extra files that are specific to Visual Studio projects
-		/// </summary>
-		/// <param name="EngineProject">Project to add files to</param>
-		protected override void AddEngineExtrasFiles(ProjectFile EngineProject)
-		{
-			base.AddEngineExtrasFiles(EngineProject);
-
-			// Add our UE4.natvis file
-			FileReference NatvisFilePath = FileReference.Combine(UnrealBuildTool.EngineDirectory, "Extras", "VisualStudioDebugging", "UE4.natvis");
-			if (FileReference.Exists(NatvisFilePath))
-			{
-				EngineProject.AddFileToProject(NatvisFilePath, UnrealBuildTool.EngineDirectory);
-			}
-		}
-
 		/// <summary>
 		/// Selects which platforms and build configurations we want in the project file
 		/// </summary>
