Fix warning about missing expat libs when generating projects from installed build on Mac.

diff --git a/Engine/Source/Programs/UnrealBuildTool/Configuration/UEBuildTarget.cs b/Engine/Source/Programs/UnrealBuildTool/Configuration/UEBuildTarget.cs
--- a/Engine/Source/Programs/UnrealBuildTool/Configuration/UEBuildTarget.cs
+++ b/Engine/Source/Programs/UnrealBuildTool/Configuration/UEBuildTarget.cs
@@ -1189,9 +1189,20 @@ namespace UnrealBuildTool
 		{
 			HashSet<FileReference> Files = new HashSet<FileReference>(RuntimeDependencySourceFiles);
 
+			// Figure out all the modules referenced by this target. This includes all the modules that are referenced, not just the ones compiled into binaries.
+			HashSet<UEBuildModule> Modules = new HashSet<UEBuildModule>();
+			foreach (UEBuildBinary Binary in Binaries)
+			{
+				foreach(UEBuildModule Module in Binary.Modules)
+				{
+					Modules.Add(Module);
+					Modules.UnionWith(Module.GetDependencies(true, true));
+				}
+			}
+
 			// Get the platform we're building for
 			UEBuildPlatform BuildPlatform = UEBuildPlatform.GetBuildPlatform(Platform);
-			foreach (UEBuildModule Module in Binaries.SelectMany(x => x.Modules))
+			foreach (UEBuildModule Module in Modules)
 			{
 				// Skip artificial modules
 				if(Module.RulesFile == null)
