UBT: Fix generating and checking for UHT manifest in the wrong location. If a target does not have a script plugin, the manifest should not be stored in the project binaries directory.

diff --git a/Engine/Source/Programs/UnrealBuildTool/Configuration/UEBuildTarget.cs b/Engine/Source/Programs/UnrealBuildTool/Configuration/UEBuildTarget.cs
--- a/Engine/Source/Programs/UnrealBuildTool/Configuration/UEBuildTarget.cs
+++ b/Engine/Source/Programs/UnrealBuildTool/Configuration/UEBuildTarget.cs
@@ -1172,7 +1172,7 @@ namespace UnrealBuildTool
 			}
 
 			// Finally clean UnrealHeaderTool if this target uses CoreUObject modules and we're not cleaning UHT already and we want UHT to be cleaned.
-			if (bIncludeUnrealHeaderTool && !RulesAssembly.IsReadOnly(ExternalExecution.GetHeaderToolReceiptFile(UnrealTargetConfiguration.Development)) && TargetName != "UnrealHeaderTool")
+			if (bIncludeUnrealHeaderTool && !RulesAssembly.IsReadOnly(ExternalExecution.GetHeaderToolReceiptFile(ProjectFile, UnrealTargetConfiguration.Development, bHasProjectScriptPlugin)) && TargetName != "UnrealHeaderTool")
 			{
 				ExternalExecution.RunExternalDotNETExecutable(UnrealBuildTool.GetUBTPath(), String.Format("UnrealHeaderTool {0} {1} -NoMutex -Clean -IgnoreJunk -NoLog", BuildHostPlatform.Current.Platform, UnrealTargetConfiguration.Development));
 			}
diff --git a/Engine/Source/Programs/UnrealBuildTool/System/ExternalExecution.cs b/Engine/Source/Programs/UnrealBuildTool/System/ExternalExecution.cs
--- a/Engine/Source/Programs/UnrealBuildTool/System/ExternalExecution.cs
+++ b/Engine/Source/Programs/UnrealBuildTool/System/ExternalExecution.cs
@@ -654,23 +654,24 @@ namespace UnrealBuildTool
 		/// Gets the path to the receipt for UHT
 		/// </summary>
 		/// <returns>Path to the UHT receipt</returns>
-		public static FileReference GetHeaderToolReceiptFile(UnrealTargetConfiguration Configuration)
+		public static FileReference GetHeaderToolReceiptFile(FileReference ProjectFile, UnrealTargetConfiguration Configuration, bool bHasProjectScriptPlugin)
 		{
-			return TargetReceipt.GetDefaultPath(UnrealBuildTool.EngineDirectory, "UnrealHeaderTool", BuildHostPlatform.Current.Platform, Configuration, "");
+			if(bHasProjectScriptPlugin && ProjectFile != null)
+			{
+				return TargetReceipt.GetDefaultPath(ProjectFile.Directory, "UnrealHeaderTool", BuildHostPlatform.Current.Platform, Configuration, "");
+			}
+			else
+			{
+				return TargetReceipt.GetDefaultPath(UnrealBuildTool.EngineDirectory, "UnrealHeaderTool", BuildHostPlatform.Current.Platform, Configuration, "");
+			}
 		}
 
 		/// <summary>
 		/// Gets UnrealHeaderTool.exe path. Does not care if UnrealheaderTool was build as a monolithic exe or not.
 		/// </summary>
-		static FileReference GetHeaderToolPath(UnrealTargetConfiguration Configuration)
+		static FileReference GetHeaderToolPath(FileReference ReceiptFile)
 		{
-			FileReference ReceiptFileName = GetHeaderToolReceiptFile(Configuration);
-			if(!FileReference.Exists(ReceiptFileName))
-			{
-				throw new BuildException("{0} is missing. Please rebuild UnrealHeaderTool manually.", ReceiptFileName);
-			}
-
-			TargetReceipt Receipt = TargetReceipt.Read(ReceiptFileName, UnrealBuildTool.EngineDirectory, null);
+			TargetReceipt Receipt = TargetReceipt.Read(ReceiptFile, UnrealBuildTool.EngineDirectory, null);
 			return Receipt.BuildProducts[0].Path;
 		}
 
@@ -680,12 +681,11 @@ namespace UnrealBuildTool
 		/// <returns>
 		/// Latest timestamp of UHT binaries or DateTime.MaxValue if UnrealHeaderTool is out of date and needs to be rebuilt.
 		/// </returns>
-		static bool GetHeaderToolTimestamp(UnrealTargetConfiguration Configuration, out DateTime Timestamp)
+		static bool GetHeaderToolTimestamp(FileReference ReceiptPath, out DateTime Timestamp)
 		{
 			using (ScopedTimer TimestampTimer = new ScopedTimer("GetHeaderToolTimestamp"))
 			{
 				// Try to read the receipt for UHT.
-				FileReference ReceiptPath = GetHeaderToolReceiptFile(Configuration);
 				if (!FileReference.Exists(ReceiptPath))
 				{
 					Timestamp = DateTime.MaxValue;
@@ -1104,9 +1104,12 @@ namespace UnrealBuildTool
 
 				UnrealTargetConfiguration UHTConfig = BuildConfiguration.bForceDebugUnrealHeaderTool ? UnrealTargetConfiguration.Debug : UnrealTargetConfiguration.Development;
 
+				// Figure out the receipt path
+				FileReference HeaderToolReceipt = GetHeaderToolReceiptFile(Target.ProjectFile, UHTConfig, Target.bHasProjectScriptPlugin);
+
 				// check if UHT is out of date
 				DateTime HeaderToolTimestamp = DateTime.MaxValue;
-				bool bHaveHeaderTool = !bIsBuildingUHT && GetHeaderToolTimestamp(UHTConfig, out HeaderToolTimestamp);
+				bool bHaveHeaderTool = !bIsBuildingUHT && GetHeaderToolTimestamp(HeaderToolReceipt, out HeaderToolTimestamp);
 
 				// ensure the headers are up to date
 				bool bUHTNeedsToRun = (BuildConfiguration.bForceHeaderGeneration || !bHaveHeaderTool || AreGeneratedCodeFilesOutOfDate(BuildConfiguration, UObjectModules, HeaderToolTimestamp, HotReload, bIsGatheringBuild, bIsAssemblingBuild));
@@ -1169,7 +1172,7 @@ namespace UnrealBuildTool
 						UBTArguments.Append(" -ignorejunk");
 
 						// Add UHT plugins to UBT command line as external plugins
-						if(Target.ProjectFile != null)
+						if(Target.bHasProjectScriptPlugin && Target.ProjectFile != null)
 						{
 							UBTArguments.AppendFormat(" -project=\"{0}\"", Target.ProjectFile);
 						}
@@ -1201,7 +1204,7 @@ namespace UnrealBuildTool
 					string ActualTargetName = String.IsNullOrEmpty(Target.GetTargetName()) ? "UE4" : Target.GetTargetName();
 					Log.TraceInformation("Parsing headers for {0}", ActualTargetName);
 
-					FileReference HeaderToolPath = GetHeaderToolPath(UHTConfig);
+					FileReference HeaderToolPath = GetHeaderToolPath(HeaderToolReceipt);
 					if (!FileReference.Exists(HeaderToolPath))
 					{
 						throw new BuildException("Unable to generate headers because UnrealHeaderTool binary was not found ({0}).", HeaderToolPath);
@@ -1232,7 +1235,7 @@ namespace UnrealBuildTool
 
 					Stopwatch s = new Stopwatch();
 					s.Start();
-					UHTResult = (ECompilationResult)RunExternalNativeExecutable(ExternalExecution.GetHeaderToolPath(UHTConfig), CmdLine);
+					UHTResult = (ECompilationResult)RunExternalNativeExecutable(ExternalExecution.GetHeaderToolPath(HeaderToolReceipt), CmdLine);
 					s.Stop();
 
 					if (UHTResult != ECompilationResult.Succeeded)
