Output Visual Studio project files in the same format as Visual Studio, and use StringBuilders everywhere to avoid large numbers of string copies.

diff --git a/Engine/Source/Programs/DotNETCommon/DotNETUtilities/StringUtils.cs b/Engine/Source/Programs/DotNETCommon/DotNETUtilities/StringUtils.cs
--- a/Engine/Source/Programs/DotNETCommon/DotNETUtilities/StringUtils.cs
+++ b/Engine/Source/Programs/DotNETCommon/DotNETUtilities/StringUtils.cs
@@ -9,6 +9,18 @@ namespace DotNETUtilities
 {
 	public static class StringUtils
 	{
+		/// <summary>
+		/// Extension method to allow formatting a string to a stringbuilder and appending a newline
+		/// </summary>
+		/// <param name="Builder">The string builder</param>
+		/// <param name="Format">Format string, as used for StringBuilder.AppendFormat</param>
+		/// <param name="Args">Arguments for the format string</param>
+		public static void AppendLine(this StringBuilder Builder, string Format, params object[] Args)
+		{
+			Builder.AppendFormat(Format, Args);
+			Builder.AppendLine();
+		}
+
 		/// <summary>
 		/// Case-sensitive replacement for String.EndsWith(), which seems to be pathalogically slow on Mono.
 		/// </summary>
diff --git a/Engine/Source/Programs/UnrealBuildTool/Platform/Android/AndroidProjectGenerator.cs b/Engine/Source/Programs/UnrealBuildTool/Platform/Android/AndroidProjectGenerator.cs
--- a/Engine/Source/Programs/UnrealBuildTool/Platform/Android/AndroidProjectGenerator.cs
+++ b/Engine/Source/Programs/UnrealBuildTool/Platform/Android/AndroidProjectGenerator.cs
@@ -193,17 +193,20 @@ namespace UnrealBuildTool
 		/// </summary>
 		/// <param name="InPlatform">  The UnrealTargetPlatform being built</param>
 		/// <param name="ProjectFileFormat"></param>
+		/// <param name="ProjectFileBuilder">String builder for the project file</param>
 		/// <returns>string    The custom property import lines for the project file; Empty string if it doesn't require one</returns>
-		public override string GetAdditionalVisualStudioPropertyGroups(UnrealTargetPlatform InPlatform, VCProjectFileFormat ProjectFileFormat)
+		public override void GetAdditionalVisualStudioPropertyGroups(UnrealTargetPlatform InPlatform, VCProjectFileFormat ProjectFileFormat, StringBuilder ProjectFileBuilder)
 		{
 			if(IsVSAndroidSupportInstalled() || !IsNsightInstalled(ProjectFileFormat))
 			{
-				return base.GetAdditionalVisualStudioPropertyGroups(InPlatform, ProjectFileFormat);
+				base.GetAdditionalVisualStudioPropertyGroups(InPlatform, ProjectFileFormat, ProjectFileBuilder);
+			}
+			else
+			{
+				ProjectFileBuilder.AppendLine("  <PropertyGroup Label=\"NsightTegraProject\">");
+				ProjectFileBuilder.AppendLine("    <NsightTegraProjectRevisionNumber>" + NsightVersionCode.ToString() + "</NsightTegraProjectRevisionNumber>");
+				ProjectFileBuilder.AppendLine("  </PropertyGroup>");
 			}
-
-			return "	<PropertyGroup Label=\"NsightTegraProject\">" + ProjectFileGenerator.NewLine +
-					"		<NsightTegraProjectRevisionNumber>" + NsightVersionCode.ToString() + "</NsightTegraProjectRevisionNumber>" + ProjectFileGenerator.NewLine +
-					"	</PropertyGroup>" + ProjectFileGenerator.NewLine;
 		}
 
 		/// <summary>
@@ -238,16 +241,19 @@ namespace UnrealBuildTool
 		/// <param name="InPlatform">  The UnrealTargetPlatform being built</param>
 		/// <param name="InConfiguration"> The UnrealTargetConfiguration being built</param>
 		/// <param name="InProjectFileFormat">The version of Visual Studio to target</param>
+		/// <param name="ProjectFileBuilder">String builder for the project file</param>
 		/// <returns>string    The custom configuration section for the project file; Empty string if it doesn't require one</returns>
-		public override string GetVisualStudioPlatformToolsetString(UnrealTargetPlatform InPlatform, UnrealTargetConfiguration InConfiguration, VCProjectFileFormat InProjectFileFormat)
+		public override void GetVisualStudioPlatformToolsetString(UnrealTargetPlatform InPlatform, UnrealTargetConfiguration InConfiguration, VCProjectFileFormat InProjectFileFormat, StringBuilder ProjectFileBuilder)
 		{
 			if (IsVSAndroidSupportInstalled() || !IsNsightInstalled(InProjectFileFormat))
 			{
-				return "\t\t<PlatformToolset>" + VCProjectFileGenerator.GetProjectFilePlatformToolsetVersionString(InProjectFileFormat) + "</PlatformToolset>" + ProjectFileGenerator.NewLine;
+				ProjectFileBuilder.AppendLine("    <PlatformToolset>" + VCProjectFileGenerator.GetProjectFilePlatformToolsetVersionString(InProjectFileFormat) + "</PlatformToolset>");
+			}
+			else
+			{
+				ProjectFileBuilder.AppendLine("    <PlatformToolset>" + VCProjectFileGenerator.GetProjectFilePlatformToolsetVersionString(InProjectFileFormat) + "</PlatformToolset>");
+				ProjectFileBuilder.AppendLine("    <AndroidNativeAPI>UseTarget</AndroidNativeAPI>");
 			}
-
-			return "\t\t<PlatformToolset>" + VCProjectFileGenerator.GetProjectFilePlatformToolsetVersionString(InProjectFileFormat) + "</PlatformToolset>" + ProjectFileGenerator.NewLine
-				+ "\t\t<AndroidNativeAPI>UseTarget</AndroidNativeAPI>" + ProjectFileGenerator.NewLine;
 		}
 
 		/// <summary>
@@ -261,11 +267,10 @@ namespace UnrealBuildTool
 		/// <param name="ProjectFilePath">Path to the project file</param>
 		/// <param name="NMakeOutputPath"></param>
 		/// <param name="InProjectFileFormat">Format for the generated project files</param>
+		/// <param name="ProjectFileBuilder">String builder for the project file</param>
 		/// <returns>The custom path lines for the project file; Empty string if it doesn't require one</returns>
-		public override string GetVisualStudioPathsEntries(UnrealTargetPlatform InPlatform, UnrealTargetConfiguration InConfiguration, TargetType TargetType, FileReference TargetRulesPath, FileReference ProjectFilePath, FileReference NMakeOutputPath, VCProjectFileFormat InProjectFileFormat)
+		public override void GetVisualStudioPathsEntries(UnrealTargetPlatform InPlatform, UnrealTargetConfiguration InConfiguration, TargetType TargetType, FileReference TargetRulesPath, FileReference ProjectFilePath, FileReference NMakeOutputPath, VCProjectFileFormat InProjectFileFormat, StringBuilder ProjectFileBuilder)
 		{
-			string PathsLines = "";
-
 			if (!IsVSAndroidSupportInstalled() && IsNsightInstalled(InProjectFileFormat))
 			{
 
@@ -299,55 +304,59 @@ namespace UnrealBuildTool
 				AdditionalLibDirs += ";" + IntermediateDirectoryPath + @"\obj\local\x86";
 				AdditionalLibDirs += @";$(AdditionalLibraryDirectories)";
 
-				PathsLines =
-					"		<IncludePath />" + ProjectFileGenerator.NewLine +
-					"		<ReferencePath />" + ProjectFileGenerator.NewLine +
-					"		<LibraryPath />" + ProjectFileGenerator.NewLine +
-					"		<LibraryWPath />" + ProjectFileGenerator.NewLine +
-					"		<SourcePath />" + ProjectFileGenerator.NewLine +
-					"		<ExcludePath />" + ProjectFileGenerator.NewLine +
-					"		<AndroidAttach>False</AndroidAttach>" + ProjectFileGenerator.NewLine +
-					"		<DebuggerFlavor>AndroidDebugger</DebuggerFlavor>" + ProjectFileGenerator.NewLine +
-					"		<OverrideAPKPath>" + APKPath + "</OverrideAPKPath>" + ProjectFileGenerator.NewLine +
-					"		<AdditionalLibraryDirectories>" + AdditionalLibDirs + "</AdditionalLibraryDirectories>" + ProjectFileGenerator.NewLine +
-					"		<BuildXmlPath>" + BuildXmlPath + "</BuildXmlPath>" + ProjectFileGenerator.NewLine +
-					"		<AndroidManifestPath>" + AndroidManifestPath + "</AndroidManifestPath>" + ProjectFileGenerator.NewLine;
+				ProjectFileBuilder.AppendLine("    <IncludePath />");
+				ProjectFileBuilder.AppendLine("    <ReferencePath />");
+				ProjectFileBuilder.AppendLine("    <LibraryPath />");
+				ProjectFileBuilder.AppendLine("    <LibraryWPath />");
+				ProjectFileBuilder.AppendLine("    <SourcePath />");
+				ProjectFileBuilder.AppendLine("    <ExcludePath />");
+				ProjectFileBuilder.AppendLine("    <AndroidAttach>False</AndroidAttach>");
+				ProjectFileBuilder.AppendLine("    <DebuggerFlavor>AndroidDebugger</DebuggerFlavor>");
+				ProjectFileBuilder.AppendLine("    <OverrideAPKPath>" + APKPath + "</OverrideAPKPath>");
+				ProjectFileBuilder.AppendLine("    <AdditionalLibraryDirectories>" + AdditionalLibDirs + "</AdditionalLibraryDirectories>");
+				ProjectFileBuilder.AppendLine("    <BuildXmlPath>" + BuildXmlPath + "</BuildXmlPath>");
+				ProjectFileBuilder.AppendLine("    <AndroidManifestPath>" + AndroidManifestPath + "</AndroidManifestPath>");
 			}
 			else 
 			{
-				PathsLines = base.GetVisualStudioPathsEntries(InPlatform, InConfiguration, TargetType, TargetRulesPath, ProjectFilePath, NMakeOutputPath, InProjectFileFormat);
+				base.GetVisualStudioPathsEntries(InPlatform, InConfiguration, TargetType, TargetRulesPath, ProjectFilePath, NMakeOutputPath, InProjectFileFormat, ProjectFileBuilder);
 			}
-
-			return PathsLines;
-
 		}
 		
 		/// <summary>
 		/// Return any custom property settings. These will be included right after Global properties to make values available to all other imports.
 		/// </summary>
 		/// <param name="InPlatform">  The UnrealTargetPlatform being built</param>
+		/// <param name="ProjectFileBuilder">String builder for the project file</param>
 		/// <returns>string    The custom property import lines for the project file; Empty string if it doesn't require one</returns>
-		public override string GetVisualStudioGlobalProperties(UnrealTargetPlatform InPlatform)
+		public override void GetVisualStudioGlobalProperties(UnrealTargetPlatform InPlatform, StringBuilder ProjectFileBuilder)
 		{
 			if (IsVSAndroidSupportInstalled())
 			{
-				return "	<Import Project=\"$(AndroidTargetsPath)\\Android.Tools.props\"/>" + ProjectFileGenerator.NewLine;
+				ProjectFileBuilder.AppendLine("  <Import Project=\"$(AndroidTargetsPath)\\Android.Tools.props\"/>");
+			}
+			else
+			{
+				base.GetVisualStudioGlobalProperties(InPlatform, ProjectFileBuilder);
 			}
-			return base.GetVisualStudioGlobalProperties(InPlatform);
 		}
 
 		/// <summary>
 		/// Return any custom property settings. These will be included in the ImportGroup section
 		/// </summary>
 		/// <param name="InPlatform">  The UnrealTargetPlatform being built</param>
+		/// <param name="ProjectFileBuilder">String builder for the project file</param>
 		/// <returns>string    The custom property import lines for the project file; Empty string if it doesn't require one</returns>
-		public override string GetVisualStudioImportGroupProperties(UnrealTargetPlatform InPlatform)
+		public override void GetVisualStudioImportGroupProperties(UnrealTargetPlatform InPlatform, StringBuilder ProjectFileBuilder)
 		{
 			if (IsVSAndroidSupportInstalled())
 			{
-				return "		<Import Project=\"VSAndroidUnreal.props\" /> " + ProjectFileGenerator.NewLine;
+				ProjectFileBuilder.AppendLine("    <Import Project=\"VSAndroidUnreal.props\" />");
+			}
+			else
+			{
+				base.GetVisualStudioImportGroupProperties(InPlatform, ProjectFileBuilder);
 			}
-			return base.GetVisualStudioImportGroupProperties(InPlatform);
 		}
 
 		/// <summary>
@@ -712,7 +721,7 @@ namespace UnrealBuildTool
 		/// <param name="TargetConfigurationName">The configuration name from the target rules, or null if we don't have one</param>
 		/// <param name="ProjectPlatformName">Name of platform string to use for Visual Studio project</param>
 		/// <param name="ProjectConfigurationName">Name of configuration string to use for Visual Studio project</param>
-		public override void MakeProjectPlatformAndConfigurationNames(UnrealTargetPlatform Platform, UnrealTargetConfiguration Configuration, string TargetConfigurationName, out string ProjectPlatformName, out string ProjectConfigurationName)
+		public override void MakeProjectPlatformAndConfigurationNames(UnrealTargetPlatform Platform, UnrealTargetConfiguration Configuration, TargetType TargetConfigurationName, out string ProjectPlatformName, out string ProjectConfigurationName)
 		{
 			ProjectConfigurationName = UnrealTargetConfiguration.Debug.ToString();
 			ProjectPlatformName = "ARM";
diff --git a/Engine/Source/Programs/UnrealBuildTool/Platform/Lumin/LuminProjectGenerator.cs b/Engine/Source/Programs/UnrealBuildTool/Platform/Lumin/LuminProjectGenerator.cs
--- a/Engine/Source/Programs/UnrealBuildTool/Platform/Lumin/LuminProjectGenerator.cs
+++ b/Engine/Source/Programs/UnrealBuildTool/Platform/Lumin/LuminProjectGenerator.cs
@@ -56,10 +56,11 @@ namespace UnrealBuildTool
 		/// <param name="InPlatform">  The UnrealTargetPlatform being built</param>
 		/// <param name="InConfiguration"> The UnrealTargetConfiguration being built</param>
 		/// <param name="InProjectFileFormat"> The verison of VS</param> 
+		/// <param name="ProjectFileBuilder">String builder for the project file</param>
 		/// <returns>string    The custom configuration section for the project file; Empty string if it doesn't require one</returns>
-		public override string GetVisualStudioPlatformToolsetString(UnrealTargetPlatform InPlatform, UnrealTargetConfiguration InConfiguration, VCProjectFileFormat InProjectFileFormat)
+		public override void GetVisualStudioPlatformToolsetString(UnrealTargetPlatform InPlatform, UnrealTargetConfiguration InConfiguration, VCProjectFileFormat InProjectFileFormat, StringBuilder ProjectFileBuilder)
 		{
-			return "\t\t<PlatformToolset>" + VCProjectFileGenerator.GetProjectFilePlatformToolsetVersionString(InProjectFileFormat) + "</PlatformToolset>" + ProjectFileGenerator.NewLine;
+			ProjectFileBuilder.AppendLine("    <PlatformToolset>" + VCProjectFileGenerator.GetProjectFilePlatformToolsetVersionString(InProjectFileFormat) + "</PlatformToolset>");
 		}
 
 		/// <summary>
@@ -73,11 +74,12 @@ namespace UnrealBuildTool
 		/// <param name="ProjectFilePath">Path to the project file</param>
 		/// <param name="NMakeOutputPath"></param>
 		/// <param name="InProjectFileFormat">Format for the generated project files</param>
+		/// <param name="ProjectFileBuilder">The project file content</param>
 		/// <returns>The custom path lines for the project file; Empty string if it doesn't require one</returns>
-		public override string GetVisualStudioPathsEntries(UnrealTargetPlatform InPlatform, UnrealTargetConfiguration InConfiguration, TargetType TargetType, FileReference TargetRulesPath, FileReference ProjectFilePath, FileReference NMakeOutputPath, VCProjectFileFormat InProjectFileFormat)
+		public override void GetVisualStudioPathsEntries(UnrealTargetPlatform InPlatform, UnrealTargetConfiguration InConfiguration, TargetType TargetType, FileReference TargetRulesPath, FileReference ProjectFilePath, FileReference NMakeOutputPath, VCProjectFileFormat InProjectFileFormat, StringBuilder ProjectFileBuilder)
 		{
-			string PathsLines = string.Empty;
-			PathsLines = base.GetVisualStudioPathsEntries(InPlatform, InConfiguration, TargetType, TargetRulesPath, ProjectFilePath, NMakeOutputPath, InProjectFileFormat);
+			base.GetVisualStudioPathsEntries(InPlatform, InConfiguration, TargetType, TargetRulesPath, ProjectFilePath, NMakeOutputPath, InProjectFileFormat, ProjectFileBuilder);
+
 			if (IsVSLuminSupportInstalled(InProjectFileFormat) && TargetType == TargetType.Game && InPlatform == UnrealTargetPlatform.Lumin)
 			{
 				string MLSDK = Environment.GetEnvironmentVariable("MLSDK");
@@ -115,9 +117,8 @@ namespace UnrealBuildTool
 													"<AutoStopAtFunction>{6}</AutoStopAtFunction>" + ProjectFileGenerator.NewLine +
 													"<EnablePrettyPrinting>{7}</EnablePrettyPrinting>" + ProjectFileGenerator.NewLine +
 													"<MLDownloadOnStart>{8}</MLDownloadOnStart>" + ProjectFileGenerator.NewLine;
-				PathsLines = PathsLines + ProjectFileGenerator.NewLine + string.Format(CustomPathEntriesTemplate, MLSDK, PackageFile, ELFFile, DebuggerFlavor, Attach, EnableAutoStop, AutoStopAtFunction, EnablePrettyPrinting, MLDownloadOnStart);
+				ProjectFileBuilder.Append(ProjectFileGenerator.NewLine + string.Format(CustomPathEntriesTemplate, MLSDK, PackageFile, ELFFile, DebuggerFlavor, Attach, EnableAutoStop, AutoStopAtFunction, EnablePrettyPrinting, MLDownloadOnStart));
 			}
-			return PathsLines;
 		}
 	}
 }
diff --git a/Engine/Source/Programs/UnrealBuildTool/ProjectFiles/VisualStudio/VCProject.cs b/Engine/Source/Programs/UnrealBuildTool/ProjectFiles/VisualStudio/VCProject.cs
--- a/Engine/Source/Programs/UnrealBuildTool/ProjectFiles/VisualStudio/VCProject.cs
+++ b/Engine/Source/Programs/UnrealBuildTool/ProjectFiles/VisualStudio/VCProject.cs
@@ -11,6 +11,7 @@ using System.Text;
 using System.Diagnostics;
 using System.Security;
 using Tools.DotNETCommon;
+using DotNETUtilities;
 
 namespace UnrealBuildTool
 {
@@ -109,7 +110,7 @@ namespace UnrealBuildTool
 		/// <param name="TargetConfigurationName">The configuration name from the target rules, or null if we don't have one</param>
 		/// <param name="ProjectPlatformName">Name of platform string to use for Visual Studio project</param>
 		/// <param name="ProjectConfigurationName">Name of configuration string to use for Visual Studio project</param>
-		public abstract void MakeProjectPlatformAndConfigurationNames(UnrealTargetPlatform Platform, UnrealTargetConfiguration Configuration, string TargetConfigurationName, out string ProjectPlatformName, out string ProjectConfigurationName);
+		public abstract void MakeProjectPlatformAndConfigurationNames(UnrealTargetPlatform Platform, UnrealTargetConfiguration Configuration, TargetType TargetConfigurationName, out string ProjectPlatformName, out string ProjectConfigurationName);
 
 		static UnrealTargetConfiguration[] GetSupportedConfigurations(TargetRules Rules)
 		{
@@ -288,7 +289,7 @@ namespace UnrealBuildTool
 		/// <param name="TargetConfigurationName">The configuration name from the target rules, or null if we don't have one</param>
 		/// <param name="ProjectPlatformName">Name of platform string to use for Visual Studio project</param>
 		/// <param name="ProjectConfigurationName">Name of configuration string to use for Visual Studio project</param>
-		public override void MakeProjectPlatformAndConfigurationNames(UnrealTargetPlatform Platform, UnrealTargetConfiguration Configuration, string TargetConfigurationName, out string ProjectPlatformName, out string ProjectConfigurationName)
+		public override void MakeProjectPlatformAndConfigurationNames(UnrealTargetPlatform Platform, UnrealTargetConfiguration Configuration, TargetType TargetConfigurationName, out string ProjectPlatformName, out string ProjectConfigurationName)
 		{
 			UEPlatformProjectGenerator PlatformProjectGenerator = UEPlatformProjectGenerator.GetPlatformProjectGenerator(Platform, bInAllowFailure: true);
 
@@ -315,10 +316,7 @@ namespace UnrealBuildTool
 				ProjectPlatformName = ProjectPlatformNameOverride == "" ? VCProjectFileGenerator.DefaultPlatformName : ProjectPlatformNameOverride;
 			}
 
-			if (!String.IsNullOrEmpty(TargetConfigurationName))
-			{
-				ProjectConfigurationName += "_" + TargetConfigurationName;
-			}
+			ProjectConfigurationName += "_" + TargetConfigurationName.ToString();
 		}
 
 		class ProjectConfigAndTargetCombination
@@ -409,7 +407,7 @@ namespace UnrealBuildTool
 								if (IsValidProjectPlatformAndConfiguration(ProjectTarget, Platform, Configuration))
 								{
 									string ProjectPlatformName, ProjectConfigurationName;
-									MakeProjectPlatformAndConfigurationNames(Platform, Configuration, ProjectTarget.TargetRules.Type.ToString(), out ProjectPlatformName, out ProjectConfigurationName);
+									MakeProjectPlatformAndConfigurationNames(Platform, Configuration, ProjectTarget.TargetRules.Type, out ProjectPlatformName, out ProjectConfigurationName);
 
 									ProjectConfigAndTargetCombination Combination = new ProjectConfigAndTargetCombination(Platform, Configuration, ProjectPlatformName, ProjectConfigurationName, ProjectTarget);
 									ProjectConfigAndTargetCombinations.Add(Combination);
@@ -519,25 +517,19 @@ namespace UnrealBuildTool
 			bool FiltersFileIsNeeded = false;
 
 			// Project file header
-			VCProjectFileContent.Append(
-				"<?xml version=\"1.0\" encoding=\"utf-8\"?>" + ProjectFileGenerator.NewLine +
-				ProjectFileGenerator.NewLine +
-				"<Project DefaultTargets=\"Build\" ToolsVersion=\"" + VCProjectFileGenerator.GetProjectFileToolVersionString(ProjectFileFormat) + "\" xmlns=\"http://schemas.microsoft.com/developer/msbuild/2003\">" + ProjectFileGenerator.NewLine);
+			VCProjectFileContent.AppendLine("<?xml version=\"1.0\" encoding=\"utf-8\"?>");
+			VCProjectFileContent.AppendLine("<Project DefaultTargets=\"Build\" ToolsVersion=\"{0}\" xmlns=\"http://schemas.microsoft.com/developer/msbuild/2003\">", VCProjectFileGenerator.GetProjectFileToolVersionString(ProjectFileFormat));
 
 			bool bGenerateUserFileContent = UEPlatformProjectGenerator.PlatformRequiresVSUserFileGeneration(InPlatforms, InConfigurations);
 			if (bGenerateUserFileContent)
 			{
-				VCUserFileContent.Append(
-					"<?xml version=\"1.0\" encoding=\"utf-8\"?>" + ProjectFileGenerator.NewLine +
-					ProjectFileGenerator.NewLine +
-					"<Project ToolsVersion=\"" + VCProjectFileGenerator.GetProjectFileToolVersionString(ProjectFileFormat) + "\" xmlns=\"http://schemas.microsoft.com/developer/msbuild/2003\">" + ProjectFileGenerator.NewLine
-					);
+				VCUserFileContent.AppendLine("<?xml version=\"1.0\" encoding=\"utf-8\"?>");
+				VCUserFileContent.AppendLine("<Project ToolsVersion=\"{0}\" xmlns=\"http://schemas.microsoft.com/developer/msbuild/2003\">", VCProjectFileGenerator.GetProjectFileToolVersionString(ProjectFileFormat));
 			}
 
 			BuildProjectConfigAndTargetCombinations(InPlatforms, InConfigurations);
 
-			VCProjectFileContent.Append(
-				"	<ItemGroup Label=\"ProjectConfigurations\">" + ProjectFileGenerator.NewLine);
+			VCProjectFileContent.AppendLine("  <ItemGroup Label=\"ProjectConfigurations\">");
 
 			// Make a list of the platforms and configs as project-format names
 			List<UnrealTargetPlatform> ProjectPlatforms = new List<UnrealTargetPlatform>();
@@ -566,24 +558,19 @@ namespace UnrealBuildTool
 				foreach (Tuple<string, UnrealTargetPlatform> PlatformTuple in ProjectPlatformNameAndPlatforms)
 				{
 					string ProjectPlatformName = PlatformTuple.Item1;
-					VCProjectFileContent.Append(
-							"		<ProjectConfiguration Include=\"" + ProjectConfigurationName + "|" + ProjectPlatformName + "\">" + ProjectFileGenerator.NewLine +
-							"			<Configuration>" + ProjectConfigurationName + "</Configuration>" + ProjectFileGenerator.NewLine +
-							"			<Platform>" + ProjectPlatformName + "</Platform>" + ProjectFileGenerator.NewLine +
-							"		</ProjectConfiguration>" + ProjectFileGenerator.NewLine);
+					VCProjectFileContent.AppendLine("    <ProjectConfiguration Include=\"{0}|{1}\">", ProjectConfigurationName, ProjectPlatformName);
+					VCProjectFileContent.AppendLine("      <Configuration>{0}</Configuration>", ProjectConfigurationName);
+					VCProjectFileContent.AppendLine("      <Platform>{0}</Platform>", ProjectPlatformName);
+					VCProjectFileContent.AppendLine("    </ProjectConfiguration>");
 				}
 			}
 
-			VCProjectFileContent.Append(
-				"	</ItemGroup>" + ProjectFileGenerator.NewLine);
+			VCProjectFileContent.AppendLine("  </ItemGroup>");
 
-			VCFiltersFileContent.Append(
-				"<?xml version=\"1.0\" encoding=\"utf-8\"?>" + ProjectFileGenerator.NewLine +
-				ProjectFileGenerator.NewLine +
-				"<Project ToolsVersion=\"" + VCProjectFileGenerator.GetProjectFileToolVersionString(ProjectFileFormat) + "\" xmlns=\"http://schemas.microsoft.com/developer/msbuild/2003\">" + ProjectFileGenerator.NewLine);
+			VCFiltersFileContent.AppendLine("<?xml version=\"1.0\" encoding=\"utf-8\"?>");
+			VCFiltersFileContent.AppendLine("<Project ToolsVersion=\"{0}\" xmlns=\"http://schemas.microsoft.com/developer/msbuild/2003\">", VCProjectFileGenerator.GetProjectFileToolVersionString(ProjectFileFormat));
 
 			// Platform specific PropertyGroups, etc.
-			StringBuilder AdditionalPropertyGroups = new StringBuilder();
 			if (!IsStubProject)
 			{
 				foreach (UnrealTargetPlatform Platform in ProjectPlatforms)
@@ -591,24 +578,21 @@ namespace UnrealBuildTool
 					UEPlatformProjectGenerator ProjGenerator = UEPlatformProjectGenerator.GetPlatformProjectGenerator(Platform, true);
 					if (ProjGenerator != null && ProjGenerator.HasVisualStudioSupport(Platform, UnrealTargetConfiguration.Development, ProjectFileFormat))
 					{
-						AdditionalPropertyGroups.Append(ProjGenerator.GetAdditionalVisualStudioPropertyGroups(Platform, ProjectFileFormat));
+						ProjGenerator.GetAdditionalVisualStudioPropertyGroups(Platform, ProjectFileFormat, VCProjectFileContent);
 					}
 				}
-
-				VCProjectFileContent.Append(AdditionalPropertyGroups);
 			}
 
 			// Project globals (project GUID, project type, SCC bindings, etc)
 			{
-				VCProjectFileContent.Append(
-					"	<PropertyGroup Label=\"Globals\">" + ProjectFileGenerator.NewLine +
-					"		<ProjectGuid>" + ProjectGUID.ToString("B").ToUpperInvariant() + "</ProjectGuid>" + ProjectFileGenerator.NewLine +
-					"		<Keyword>MakeFileProj</Keyword>" + ProjectFileGenerator.NewLine +
-					"		<RootNamespace>" + ProjectName + "</RootNamespace>" + ProjectFileGenerator.NewLine +
-					"		<PlatformToolset>" + VCProjectFileGenerator.GetProjectFilePlatformToolsetVersionString(ProjectFileFormat) + "</PlatformToolset>" + ProjectFileGenerator.NewLine +
-					"		<MinimumVisualStudioVersion>" + VCProjectFileGenerator.GetProjectFileToolVersionString(ProjectFileFormat) + "</MinimumVisualStudioVersion>" + ProjectFileGenerator.NewLine +
-					"		<TargetRuntime>Native</TargetRuntime>" + ProjectFileGenerator.NewLine +
-					"	</PropertyGroup>" + ProjectFileGenerator.NewLine);
+				VCProjectFileContent.AppendLine("  <PropertyGroup Label=\"Globals\">");
+				VCProjectFileContent.AppendLine("    <ProjectGuid>{0}</ProjectGuid>", ProjectGUID.ToString("B").ToUpperInvariant());
+				VCProjectFileContent.AppendLine("    <Keyword>MakeFileProj</Keyword>");
+				VCProjectFileContent.AppendLine("    <RootNamespace>{0}</RootNamespace>", ProjectName);
+				VCProjectFileContent.AppendLine("    <PlatformToolset>{0}</PlatformToolset>", VCProjectFileGenerator.GetProjectFilePlatformToolsetVersionString(ProjectFileFormat));
+				VCProjectFileContent.AppendLine("    <MinimumVisualStudioVersion>{0}</MinimumVisualStudioVersion>", VCProjectFileGenerator.GetProjectFileToolVersionString(ProjectFileFormat));
+				VCProjectFileContent.AppendLine("    <TargetRuntime>Native</TargetRuntime>");
+				VCProjectFileContent.AppendLine("  </PropertyGroup>");
 			}
 
 			// look for additional import lines for all platforms for non stub projects
@@ -619,7 +603,7 @@ namespace UnrealBuildTool
 					UEPlatformProjectGenerator ProjGenerator = UEPlatformProjectGenerator.GetPlatformProjectGenerator(Platform, true);
 					if (ProjGenerator != null && ProjGenerator.HasVisualStudioSupport(Platform, UnrealTargetConfiguration.Development, ProjectFileFormat))
 					{
-						VCProjectFileContent.Append(ProjGenerator.GetVisualStudioGlobalProperties(Platform));
+						ProjGenerator.GetVisualStudioGlobalProperties(Platform, VCProjectFileContent);
 					}
 				}
 			}
@@ -628,12 +612,10 @@ namespace UnrealBuildTool
 			{
 				// TODO: Restrict this to only the Lumin platform targets, routing via GetVisualStudioGlobalProperties().
 				// Currently hacking here because returning true from HasVisualStudioSupport() for lumin causes bunch of faiures in VS.
-				string LuminDebuggerItemGroup =
-						"	<ItemGroup>" + ProjectFileGenerator.NewLine +
-						"		<ProjectCapability Include=\"MLProject\" />" + ProjectFileGenerator.NewLine +
-						"		<PropertyPageSchema Include=\"$(LOCALAPPDATA)\\Microsoft\\VisualStudio\\MagicLeap\\debugger.xaml\" />" + ProjectFileGenerator.NewLine +
-						"	</ItemGroup>" + ProjectFileGenerator.NewLine;
-				VCProjectFileContent.Append(LuminDebuggerItemGroup);
+				VCProjectFileContent.AppendLine("  <ItemGroup>");
+				VCProjectFileContent.AppendLine("    <ProjectCapability Include=\"MLProject\" />");
+				VCProjectFileContent.AppendLine("    <PropertyPageSchema Include=\"$(LOCALAPPDATA)\\Microsoft\\VisualStudio\\MagicLeap\\debugger.xaml\" />");
+				VCProjectFileContent.AppendLine("  </ItemGroup>");
 			}
 
 			// Write each project configuration PreDefaultProps section
@@ -649,8 +631,7 @@ namespace UnrealBuildTool
 				}
 			}
 
-			VCProjectFileContent.Append(
-				"	<Import Project=\"$(VCTargetsPath)\\Microsoft.Cpp.Default.props\" />" + ProjectFileGenerator.NewLine);
+			VCProjectFileContent.AppendLine("  <Import Project=\"$(VCTargetsPath)\\Microsoft.Cpp.Default.props\" />");
 
 			// Write each project configuration PreDefaultProps section
 			foreach (Tuple<string, UnrealTargetConfiguration> ConfigurationTuple in ProjectConfigurationNameAndConfigurations)
@@ -665,11 +646,9 @@ namespace UnrealBuildTool
 				}
 			}
 
-			VCProjectFileContent.Append(
-				"	<Import Project=\"$(VCTargetsPath)\\Microsoft.Cpp.props\" />" + ProjectFileGenerator.NewLine +
-				"	<ImportGroup Label=\"ExtensionSettings\" />" + ProjectFileGenerator.NewLine +
-				"	<PropertyGroup Label=\"UserMacros\" />" + ProjectFileGenerator.NewLine
-				);
+			VCProjectFileContent.AppendLine("  <Import Project=\"$(VCTargetsPath)\\Microsoft.Cpp.props\" />");
+			VCProjectFileContent.AppendLine("  <ImportGroup Label=\"ExtensionSettings\" />");
+			VCProjectFileContent.AppendLine("  <PropertyGroup Label=\"UserMacros\" />");
 
 			// Write each project configuration
 			foreach (ProjectConfigAndTargetCombination Combination in ProjectConfigAndTargetCombinations)
@@ -714,11 +693,9 @@ namespace UnrealBuildTool
 					LocalAliasedFiles.Add(new AliasedFile(ProjectRelativeSourceFile, FilterRelativeSourceDirectory));
 				}
 
-				VCFiltersFileContent.Append(
-					"	<ItemGroup>" + ProjectFileGenerator.NewLine);
+				VCFiltersFileContent.AppendLine("  <ItemGroup>");
 
-				VCProjectFileContent.Append(
-					"	<ItemGroup>" + ProjectFileGenerator.NewLine);
+				VCProjectFileContent.AppendLine("  <ItemGroup>");
 
 				// Add all file directories to the filters file as solution filters
 				HashSet<string> FilterDirectories = new HashSet<string>();
@@ -761,46 +738,40 @@ namespace UnrealBuildTool
 
 					if (!string.IsNullOrEmpty(PCHFileName))
 					{
-						VCProjectFileContent.Append(
-							"		<" + VCFileType + " Include=\"" + EscapeFileName(AliasedFile.FileSystemPath) + "\">" + ProjectFileGenerator.NewLine +
-							"			<AdditionalOptions>$(AdditionalOptions) /Yu" + PCHFileName + "</AdditionalOptions>" + ProjectFileGenerator.NewLine +
-							"		</" + VCFileType + " >" + ProjectFileGenerator.NewLine);
+						VCProjectFileContent.AppendLine("    <{0} Include=\"{1}\">" + VCFileType, EscapeFileName(AliasedFile.FileSystemPath));
+						VCProjectFileContent.AppendLine("      <AdditionalOptions>$(AdditionalOptions) /Yu{0}</AdditionalOptions>", PCHFileName);
+						VCProjectFileContent.AppendLine("    </{0}>", VCFileType);
 					}
 					else
 					{
-						VCProjectFileContent.Append(
-							"		<" + VCFileType + " Include=\"" + EscapeFileName(AliasedFile.FileSystemPath) + "\" />" + ProjectFileGenerator.NewLine);
+						VCProjectFileContent.AppendLine("    <{0} Include=\"{1}\"/>", VCFileType, EscapeFileName(AliasedFile.FileSystemPath));
 					}
 
 					if (!String.IsNullOrWhiteSpace(AliasedFile.ProjectPath))
 					{
-						VCFiltersFileContent.Append(
-							"		<" + VCFileType + " Include=\"" + EscapeFileName(AliasedFile.FileSystemPath) + "\">" + ProjectFileGenerator.NewLine +
-							"			<Filter>" + Utils.CleanDirectorySeparators(AliasedFile.ProjectPath) + "</Filter>" + ProjectFileGenerator.NewLine +
-							"		</" + VCFileType + " >" + ProjectFileGenerator.NewLine);
+						VCFiltersFileContent.AppendLine("    <{0} Include=\"{1}\">", VCFileType, EscapeFileName(AliasedFile.FileSystemPath));
+						VCFiltersFileContent.AppendLine("      <Filter>{0}</Filter>", Utils.CleanDirectorySeparators(AliasedFile.ProjectPath));
+						VCFiltersFileContent.AppendLine("    </{0}>", VCFileType);
 
 						FiltersFileIsNeeded = true;
 					}
 					else
 					{
 						// No need to specify the root directory relative to the project (it would just be an empty string!)
-						VCFiltersFileContent.Append(
-							"		<" + VCFileType + " Include=\"" + EscapeFileName(AliasedFile.FileSystemPath) + "\" />" + ProjectFileGenerator.NewLine);
+						VCFiltersFileContent.AppendLine("    <{0} Include=\"{1}\" />", VCFileType, EscapeFileName(AliasedFile.FileSystemPath));
 					}
 				}
 
-				VCProjectFileContent.Append(
-					"	</ItemGroup>" + ProjectFileGenerator.NewLine);
+				VCProjectFileContent.AppendLine("  </ItemGroup>");
 
-				VCFiltersFileContent.Append(
-					"	</ItemGroup>" + ProjectFileGenerator.NewLine);
+				VCFiltersFileContent.AppendLine("  </ItemGroup>");
 			}
 
 			// For Installed engine builds, include engine source in the source search paths if it exists. We never build it locally, so the debugger can't find it.
 			if (UnrealBuildTool.IsEngineInstalled() && !IsStubProject)
 			{
-				VCProjectFileContent.Append("	<PropertyGroup>" + ProjectFileGenerator.NewLine);
-				VCProjectFileContent.Append("		<SourcePath>");
+				VCProjectFileContent.AppendLine("  <PropertyGroup>");
+				VCProjectFileContent.Append("    <SourcePath>");
 				foreach (string DirectoryName in Directory.EnumerateDirectories(UnrealBuildTool.EngineSourceDirectory.FullName, "*", SearchOption.AllDirectories))
 				{
 					if (Directory.EnumerateFiles(DirectoryName, "*.cpp").Any())
@@ -809,8 +780,8 @@ namespace UnrealBuildTool
 						VCProjectFileContent.Append(";");
 					}
 				}
-				VCProjectFileContent.Append("</SourcePath>" + ProjectFileGenerator.NewLine);
-				VCProjectFileContent.Append("	</PropertyGroup>" + ProjectFileGenerator.NewLine);
+				VCProjectFileContent.AppendLine("</SourcePath>");
+				VCProjectFileContent.AppendLine("  </PropertyGroup>");
 			}
 
 			// Write IntelliSense info
@@ -818,28 +789,13 @@ namespace UnrealBuildTool
 				// @todo projectfiles: Currently we are storing defines/include paths for ALL configurations rather than using ConditionString and storing
 				//      this data uniquely for each target configuration.  IntelliSense may behave better if we did that, but it will result in a LOT more
 				//      data being stored into the project file, and might make the IDE perform worse when switching configurations!
-				VCProjectFileContent.Append(
-					"	<PropertyGroup>" + ProjectFileGenerator.NewLine +
-					"		<NMakePreprocessorDefinitions>$(NMakePreprocessorDefinitions)" + (VCPreprocessorDefinitions.Length > 0 ? (";" + VCPreprocessorDefinitions) : "") + "</NMakePreprocessorDefinitions>" + ProjectFileGenerator.NewLine +
-					"		<NMakeIncludeSearchPath>$(NMakeIncludeSearchPath)" + (VCIncludeSearchPaths.Length > 0 ? (";" + VCIncludeSearchPaths) : "") + "</NMakeIncludeSearchPath>" + ProjectFileGenerator.NewLine +
-					"		<NMakeForcedIncludes>$(NMakeForcedIncludes)</NMakeForcedIncludes>" + ProjectFileGenerator.NewLine +
-					"		<NMakeAssemblySearchPath>$(NMakeAssemblySearchPath)</NMakeAssemblySearchPath>" + ProjectFileGenerator.NewLine +
-					"		<NMakeForcedUsingAssemblies>$(NMakeForcedUsingAssemblies)</NMakeForcedUsingAssemblies>" + ProjectFileGenerator.NewLine +
-					"	</PropertyGroup>" + ProjectFileGenerator.NewLine);
-			}
-
-			// look for additional import lines for all platforms for non stub projects
-			StringBuilder AdditionalTargetSettings = new StringBuilder();
-			if (!IsStubProject)
-			{
-				foreach (UnrealTargetPlatform Platform in ProjectPlatforms)
-				{
-					UEPlatformProjectGenerator ProjGenerator = UEPlatformProjectGenerator.GetPlatformProjectGenerator(Platform, true);
-					if (ProjGenerator != null && ProjGenerator.HasVisualStudioSupport(Platform, UnrealTargetConfiguration.Development, ProjectFileFormat))
-					{
-						AdditionalTargetSettings.Append(ProjGenerator.GetVisualStudioTargetOverrides(Platform, ProjectFileFormat));
-					}
-				}
+				VCProjectFileContent.AppendLine("  <PropertyGroup>");
+				VCProjectFileContent.AppendLine("    <NMakePreprocessorDefinitions>$(NMakePreprocessorDefinitions){0}</NMakePreprocessorDefinitions>", (VCPreprocessorDefinitions.Length > 0 ? (";" + VCPreprocessorDefinitions) : ""));
+				VCProjectFileContent.AppendLine("    <NMakeIncludeSearchPath>$(NMakeIncludeSearchPath){0}</NMakeIncludeSearchPath>", (VCIncludeSearchPaths.Length > 0 ? (";" + VCIncludeSearchPaths) : ""));
+				VCProjectFileContent.AppendLine("    <NMakeForcedIncludes>$(NMakeForcedIncludes)</NMakeForcedIncludes>");
+				VCProjectFileContent.AppendLine("    <NMakeAssemblySearchPath>$(NMakeAssemblySearchPath)</NMakeAssemblySearchPath>");
+				VCProjectFileContent.AppendLine("    <NMakeForcedUsingAssemblies>$(NMakeForcedUsingAssemblies)</NMakeForcedUsingAssemblies>");
+				VCProjectFileContent.AppendLine("  </PropertyGroup>");
 			}
 
 			string OutputManifestString = "";
@@ -874,24 +830,30 @@ namespace UnrealBuildTool
 				}
 			}
 
-			VCProjectFileContent.Append(
-					OutputManifestString +	// output manifest must come before the Cpp.targets file.
-					"	<ItemDefinitionGroup>" + ProjectFileGenerator.NewLine +
-					"	</ItemDefinitionGroup>" + ProjectFileGenerator.NewLine +
-					"	<Import Project=\"$(VCTargetsPath)\\Microsoft.Cpp.targets\" />" + ProjectFileGenerator.NewLine +
-					AdditionalTargetSettings.ToString() +
-					"	<ImportGroup Label=\"ExtensionTargets\">" + ProjectFileGenerator.NewLine +
-					"	</ImportGroup>" + ProjectFileGenerator.NewLine +
-					"</Project>" + ProjectFileGenerator.NewLine);
+			VCProjectFileContent.Append(OutputManifestString); // output manifest must come before the Cpp.targets file.
+			VCProjectFileContent.AppendLine("  <ItemDefinitionGroup>");
+			VCProjectFileContent.AppendLine("  </ItemDefinitionGroup>");
+			VCProjectFileContent.AppendLine("  <Import Project=\"$(VCTargetsPath)\\Microsoft.Cpp.targets\" />");
+			if (!IsStubProject)
+			{
+				foreach (UnrealTargetPlatform Platform in ProjectPlatforms)
+				{
+					UEPlatformProjectGenerator ProjGenerator = UEPlatformProjectGenerator.GetPlatformProjectGenerator(Platform, true);
+					if (ProjGenerator != null && ProjGenerator.HasVisualStudioSupport(Platform, UnrealTargetConfiguration.Development, ProjectFileFormat))
+					{
+						ProjGenerator.GetVisualStudioTargetOverrides(Platform, ProjectFileFormat, VCProjectFileContent);
+					}
+				}
+			}
+			VCProjectFileContent.AppendLine("  <ImportGroup Label=\"ExtensionTargets\">");
+			VCProjectFileContent.AppendLine("  </ImportGroup>");
+			VCProjectFileContent.AppendLine("</Project>");
 
-			VCFiltersFileContent.Append(
-				"</Project>" + ProjectFileGenerator.NewLine);
+			VCFiltersFileContent.AppendLine("</Project>");
 
 			if (bGenerateUserFileContent)
 			{
-				VCUserFileContent.Append(
-					"</Project>" + ProjectFileGenerator.NewLine
-					);
+				VCUserFileContent.AppendLine("</Project>");
 			}
 
 			// Save the project file
@@ -992,10 +954,9 @@ namespace UnrealBuildTool
 						//       matches identically with the pre-existing file
 						string FilterGUID = Guid.NewGuid().ToString("B").ToUpperInvariant();
 
-						VCFiltersFileContent.Append(
-							"		<Filter Include=\"" + LeadingDirectory + "\">" + ProjectFileGenerator.NewLine +
-							"			<UniqueIdentifier>" + FilterGUID + "</UniqueIdentifier>" + ProjectFileGenerator.NewLine +
-							"		</Filter>" + ProjectFileGenerator.NewLine);
+						VCFiltersFileContent.AppendLine("    <Filter Include=\"{0}\">", LeadingDirectory);
+						VCFiltersFileContent.AppendLine("      <UniqueIdentifier>{0}</UniqueIdentifier>", FilterGUID);
+						VCFiltersFileContent.AppendLine("    </Filter>");
 
 						FiltersFileIsNeeded = true;
 					}
@@ -1048,15 +1009,17 @@ namespace UnrealBuildTool
 			string ProjectConfigurationAndPlatformName = ProjectConfigurationName + "|" + ProjectPlatformName;
 			string ConditionString = "Condition=\"'$(Configuration)|$(Platform)'=='" + ProjectConfigurationAndPlatformName + "'\"";
 
-			string PlatformToolsetString = (ProjGenerator != null) ? ProjGenerator.GetVisualStudioPreDefaultString(TargetPlatform, TargetConfiguration) : "";
-
-			if (!String.IsNullOrEmpty(PlatformToolsetString))
+			if(ProjGenerator != null)
 			{
-				VCProjectFileContent.Append(
-					"	<PropertyGroup " + ConditionString + " Label=\"Configuration\">" + ProjectFileGenerator.NewLine +
-							PlatformToolsetString +
-					"	</PropertyGroup>" + ProjectFileGenerator.NewLine
-				);
+				StringBuilder PlatformToolsetString = new StringBuilder();
+				ProjGenerator.GetVisualStudioPreDefaultString(TargetPlatform, TargetConfiguration, PlatformToolsetString);
+
+				if (PlatformToolsetString.Length > 0)
+				{
+					VCProjectFileContent.AppendLine("  <PropertyGroup " + ConditionString + " Label=\"Configuration\">", ConditionString);
+					VCProjectFileContent.Append(PlatformToolsetString);
+					VCProjectFileContent.AppendLine("  </PropertyGroup>");
+				}
 			}
 		}
 
@@ -1072,19 +1035,21 @@ namespace UnrealBuildTool
 			string ProjectConfigurationAndPlatformName = ProjectConfigurationName + "|" + ProjectPlatformName;
 			string ConditionString = "Condition=\"'$(Configuration)|$(Platform)'=='" + ProjectConfigurationAndPlatformName + "'\"";
 
-			string PlatformToolsetString = (ProjGenerator != null) ? ProjGenerator.GetVisualStudioPlatformToolsetString(TargetPlatform, TargetConfiguration, ProjectFileFormat) : "";
-			if (String.IsNullOrEmpty(PlatformToolsetString))
+			StringBuilder PlatformToolsetString = new StringBuilder();
+			if(ProjGenerator != null)
+			{
+				ProjGenerator.GetVisualStudioPlatformToolsetString(TargetPlatform, TargetConfiguration, ProjectFileFormat, PlatformToolsetString);
+			}
+			if (PlatformToolsetString.Length == 0)
 			{
-				PlatformToolsetString = "		<PlatformToolset>" + VCProjectFileGenerator.GetProjectFilePlatformToolsetVersionString(ProjectFileFormat) + "</PlatformToolset>" + ProjectFileGenerator.NewLine;
+				PlatformToolsetString.AppendLine("    <PlatformToolset>" + VCProjectFileGenerator.GetProjectFilePlatformToolsetVersionString(ProjectFileFormat) + "</PlatformToolset>");
 			}
 
 			string PlatformConfigurationType = (ProjGenerator == null) ? "Makefile" : ProjGenerator.GetVisualStudioPlatformConfigurationType(TargetPlatform, ProjectFileFormat);
-			VCProjectFileContent.Append(
-				"	<PropertyGroup " + ConditionString + " Label=\"Configuration\">" + ProjectFileGenerator.NewLine +
-				"		<ConfigurationType>" + PlatformConfigurationType + "</ConfigurationType>" + ProjectFileGenerator.NewLine +
-						PlatformToolsetString +
-				"	</PropertyGroup>" + ProjectFileGenerator.NewLine
-				);
+			VCProjectFileContent.AppendLine("  <PropertyGroup {0} Label=\"Configuration\">", ConditionString);
+			VCProjectFileContent.AppendLine("    <ConfigurationType>{0}</ConfigurationType>", PlatformConfigurationType);
+			VCProjectFileContent.Append(PlatformToolsetString);
+			VCProjectFileContent.AppendLine("  </PropertyGroup>");
 		}
 
 		// Anonymous function that writes project configuration data
@@ -1108,12 +1073,13 @@ namespace UnrealBuildTool
 			string ConditionString = "Condition=\"'$(Configuration)|$(Platform)'=='" + Combination.ProjectConfigurationAndPlatformName + "'\"";
 
 			{
-				string ImportGroupProperties = (ProjGenerator != null) ? ProjGenerator.GetVisualStudioImportGroupProperties(Platform) : "";
-				VCProjectFileContent.Append(
-					"	<ImportGroup " + ConditionString + " Label=\"PropertySheets\">" + ProjectFileGenerator.NewLine +
-					"		<Import Project=\"$(UserRootDir)\\Microsoft.Cpp.$(Platform).user.props\" Condition=\"exists('$(UserRootDir)\\Microsoft.Cpp.$(Platform).user.props')\" Label=\"LocalAppDataPlatform\" />" + ProjectFileGenerator.NewLine +
-							ImportGroupProperties +
-					"	</ImportGroup>" + ProjectFileGenerator.NewLine);
+				VCProjectFileContent.AppendLine("  <ImportGroup {0} Label=\"PropertySheets\">", ConditionString);
+				VCProjectFileContent.AppendLine("    <Import Project=\"$(UserRootDir)\\Microsoft.Cpp.$(Platform).user.props\" Condition=\"exists('$(UserRootDir)\\Microsoft.Cpp.$(Platform).user.props')\" Label=\"LocalAppDataPlatform\" />");
+				if(ProjGenerator != null)
+				{
+					ProjGenerator.GetVisualStudioImportGroupProperties(Platform, VCProjectFileContent);
+				}
+				VCProjectFileContent.AppendLine("  </ImportGroup>");
 
 				DirectoryReference ProjectDirectory = ProjectFilePath.Directory;
 
@@ -1121,29 +1087,30 @@ namespace UnrealBuildTool
 				{
 					string ProjectRelativeUnusedDirectory = NormalizeProjectPath(DirectoryReference.Combine(UnrealBuildTool.EngineDirectory, "Intermediate", "Build", "Unused"));
 
-					VCProjectFileContent.Append(
-						"	<PropertyGroup " + ConditionString + ">" + ProjectFileGenerator.NewLine +
-						"		<OutDir>" + ProjectRelativeUnusedDirectory + Path.DirectorySeparatorChar + "</OutDir>" + ProjectFileGenerator.NewLine +
-						"		<IntDir>" + ProjectRelativeUnusedDirectory + Path.DirectorySeparatorChar + "</IntDir>" + ProjectFileGenerator.NewLine +
-						"		<NMakeBuildCommandLine>@rem Nothing to do.</NMakeBuildCommandLine>" + ProjectFileGenerator.NewLine +
-						"		<NMakeReBuildCommandLine>@rem Nothing to do.</NMakeReBuildCommandLine>" + ProjectFileGenerator.NewLine +
-						"		<NMakeCleanCommandLine>@rem Nothing to do.</NMakeCleanCommandLine>" + ProjectFileGenerator.NewLine +
-						"		<NMakeOutput/>" + ProjectFileGenerator.NewLine +
-						"	</PropertyGroup>" + ProjectFileGenerator.NewLine);
+					VCProjectFileContent.AppendLine("  <PropertyGroup {0}>", ConditionString);
+					VCProjectFileContent.AppendLine("    <OutDir>{0}{1}</OutDir>", ProjectRelativeUnusedDirectory, Path.DirectorySeparatorChar);
+					VCProjectFileContent.AppendLine("    <IntDir>{0}{1}</IntDir>", ProjectRelativeUnusedDirectory, Path.DirectorySeparatorChar);
+					VCProjectFileContent.AppendLine("    <NMakeBuildCommandLine>@rem Nothing to do.</NMakeBuildCommandLine>");
+					VCProjectFileContent.AppendLine("    <NMakeReBuildCommandLine>@rem Nothing to do.</NMakeReBuildCommandLine>");
+					VCProjectFileContent.AppendLine("    <NMakeCleanCommandLine>@rem Nothing to do.</NMakeCleanCommandLine>");
+					VCProjectFileContent.AppendLine("    <NMakeOutput/>");
+					VCProjectFileContent.AppendLine("  </PropertyGroup>");
 				}
 				else if (UnrealBuildTool.IsEngineInstalled() && Combination.ProjectTarget != null && Combination.ProjectTarget.TargetRules != null && !Combination.ProjectTarget.SupportedPlatforms.Contains(Combination.Platform))
 				{
 					string ProjectRelativeUnusedDirectory = NormalizeProjectPath(DirectoryReference.Combine(UnrealBuildTool.EngineDirectory, "Intermediate", "Build", "Unused"));
 
-					VCProjectFileContent.AppendFormat(
-						"	<PropertyGroup " + ConditionString + ">" + ProjectFileGenerator.NewLine +
-						"		<OutDir>" + ProjectRelativeUnusedDirectory + Path.DirectorySeparatorChar + "</OutDir>" + ProjectFileGenerator.NewLine +
-						"		<IntDir>" + ProjectRelativeUnusedDirectory + Path.DirectorySeparatorChar + "</IntDir>" + ProjectFileGenerator.NewLine +
-						"		<NMakeBuildCommandLine>@echo {0} is not a supported platform for {1}. Valid platforms are {2}.</NMakeBuildCommandLine>" + ProjectFileGenerator.NewLine +
-						"		<NMakeReBuildCommandLine>@echo {0} is not a supported platform for {1}. Valid platforms are {2}.</NMakeReBuildCommandLine>" + ProjectFileGenerator.NewLine +
-						"		<NMakeCleanCommandLine>@echo {0} is not a supported platform for {1}. Valid platforms are {2}.</NMakeCleanCommandLine>" + ProjectFileGenerator.NewLine +
-						"		<NMakeOutput/>" + ProjectFileGenerator.NewLine +
-						"	</PropertyGroup>" + ProjectFileGenerator.NewLine, Combination.Platform, Combination.ProjectTarget.TargetFilePath.GetFileNameWithoutAnyExtensions(), String.Join(", ", Combination.ProjectTarget.SupportedPlatforms.Select(x => x.ToString())));
+					string TargetName = Combination.ProjectTarget.TargetFilePath.GetFileNameWithoutAnyExtensions();
+					string ValidPlatforms = String.Join(", ", Combination.ProjectTarget.SupportedPlatforms.Select(x => x.ToString()));
+
+					VCProjectFileContent.AppendLine("  <PropertyGroup {0}>", ConditionString);
+					VCProjectFileContent.AppendLine("    <OutDir>{0}{1}</OutDir>", ProjectRelativeUnusedDirectory, Path.DirectorySeparatorChar);
+					VCProjectFileContent.AppendLine("    <IntDir>{0}{1}</IntDir>", ProjectRelativeUnusedDirectory, Path.DirectorySeparatorChar);
+					VCProjectFileContent.AppendLine("    <NMakeBuildCommandLine>@echo {0} is not a supported platform for {1}. Valid platforms are {2}.</NMakeBuildCommandLine>", Combination.Platform, TargetName, ValidPlatforms);
+					VCProjectFileContent.AppendLine("    <NMakeReBuildCommandLine>@echo {0} is not a supported platform for {1}. Valid platforms are {2}.</NMakeReBuildCommandLine>", Combination.Platform, TargetName, ValidPlatforms);
+					VCProjectFileContent.AppendLine("    <NMakeCleanCommandLine>@echo {0} is not a supported platform for {1}. Valid platforms are {2}.</NMakeCleanCommandLine>", Combination.Platform, TargetName, ValidPlatforms);
+					VCProjectFileContent.AppendLine("    <NMakeOutput/>");
+					VCProjectFileContent.AppendLine("  </PropertyGroup>");
 				}
 				else
 				{
@@ -1202,17 +1169,22 @@ namespace UnrealBuildTool
 					NMakePath += TargetRulesObject.Architecture;
 					NMakePath += BuildPlatform.GetBinaryExtension(UEBuildBinaryType.Executable);
 
-					VCProjectFileContent.Append(
-						"	<PropertyGroup " + ConditionString + ">" + ProjectFileGenerator.NewLine);
+					VCProjectFileContent.AppendLine("  <PropertyGroup {0}>", ConditionString);
+
+					StringBuilder PathsStringBuilder = new StringBuilder();
+					if(ProjGenerator != null)
+					{
+						ProjGenerator.GetVisualStudioPathsEntries(Platform, Configuration, TargetRulesObject.Type, TargetFilePath, ProjectFilePath, NMakePath, ProjectFileFormat, PathsStringBuilder);
+					}
 
-					string PathStrings = (ProjGenerator != null) ? ProjGenerator.GetVisualStudioPathsEntries(Platform, Configuration, TargetRulesObject.Type, TargetFilePath, ProjectFilePath, NMakePath, ProjectFileFormat) : "";
+					string PathStrings = PathsStringBuilder.ToString();
 					if (string.IsNullOrEmpty(PathStrings) || (PathStrings.Contains("<IntDir>") == false))
 					{
 						string ProjectRelativeUnusedDirectory = "$(ProjectDir)..\\Build\\Unused";
-						VCProjectFileContent.Append(
-							PathStrings +
-							"		<OutDir>" + ProjectRelativeUnusedDirectory + Path.DirectorySeparatorChar + "</OutDir>" + ProjectFileGenerator.NewLine +
-							"		<IntDir>" + ProjectRelativeUnusedDirectory + Path.DirectorySeparatorChar + "</IntDir>" + ProjectFileGenerator.NewLine);
+						VCProjectFileContent.Append(PathStrings);
+
+						VCProjectFileContent.AppendLine("    <OutDir>{0}{1}</OutDir>", ProjectRelativeUnusedDirectory, Path.DirectorySeparatorChar);
+						VCProjectFileContent.AppendLine("    <IntDir>{0}{1}</IntDir>", ProjectRelativeUnusedDirectory, Path.DirectorySeparatorChar);
 					}
 					else
 					{
@@ -1253,23 +1225,10 @@ namespace UnrealBuildTool
 					}
 
 					// NMake Build command line
-					VCProjectFileContent.Append("		<NMakeBuildCommandLine>");
-					VCProjectFileContent.Append(EscapePath(NormalizeProjectPath(FileReference.Combine(BatchFilesDirectory, "Build.bat"))) + BuildArguments.ToString());
-					VCProjectFileContent.Append("</NMakeBuildCommandLine>" + ProjectFileGenerator.NewLine);
-
-					// NMake ReBuild command line
-					VCProjectFileContent.Append("		<NMakeReBuildCommandLine>");
-					VCProjectFileContent.Append(EscapePath(NormalizeProjectPath(FileReference.Combine(BatchFilesDirectory, "Rebuild.bat"))) + BuildArguments.ToString());
-					VCProjectFileContent.Append("</NMakeReBuildCommandLine>" + ProjectFileGenerator.NewLine);
-
-					// NMake Clean command line
-					VCProjectFileContent.Append("		<NMakeCleanCommandLine>");
-					VCProjectFileContent.Append(EscapePath(NormalizeProjectPath(FileReference.Combine(BatchFilesDirectory, "Clean.bat"))) + BuildArguments.ToString());
-					VCProjectFileContent.Append("</NMakeCleanCommandLine>" + ProjectFileGenerator.NewLine);
-
-					VCProjectFileContent.Append("		<NMakeOutput>");
-					VCProjectFileContent.Append(NormalizeProjectPath(NMakePath.FullName));
-					VCProjectFileContent.Append("</NMakeOutput>" + ProjectFileGenerator.NewLine);
+					VCProjectFileContent.AppendLine("    <NMakeBuildCommandLine>{0}{1}</NMakeBuildCommandLine>", EscapePath(NormalizeProjectPath(FileReference.Combine(BatchFilesDirectory, "Build.bat"))), BuildArguments.ToString());
+					VCProjectFileContent.AppendLine("    <NMakeReBuildCommandLine>{0}{1}</NMakeReBuildCommandLine>", EscapePath(NormalizeProjectPath(FileReference.Combine(BatchFilesDirectory, "Rebuild.bat"))), BuildArguments.ToString());
+					VCProjectFileContent.AppendLine("    <NMakeCleanCommandLine>{0}{1}</NMakeCleanCommandLine>", EscapePath(NormalizeProjectPath(FileReference.Combine(BatchFilesDirectory, "Clean.bat"))), BuildArguments.ToString());
+					VCProjectFileContent.AppendLine("    <NMakeOutput>{0}</NMakeOutput>", NormalizeProjectPath(NMakePath.FullName));
 
 					if (TargetRulesObject.Type == TargetType.Game || TargetRulesObject.Type == TargetType.Client || TargetRulesObject.Type == TargetType.Server)
 					{
@@ -1277,7 +1236,7 @@ namespace UnrealBuildTool
 						UEPlatformProjectGenerator.GenerateGamePlatformSpecificProperties(Platform, Configuration, TargetRulesObject.Type, VCProjectFileContent, RootDirectory, TargetFilePath);
 					}
 
-					VCProjectFileContent.Append("	</PropertyGroup>" + ProjectFileGenerator.NewLine);
+					VCProjectFileContent.AppendLine("  </PropertyGroup>");
 
 					string LayoutDirString = (ProjGenerator != null) ? ProjGenerator.GetVisualStudioLayoutDirSection(Platform, Configuration, ConditionString, Combination.ProjectTarget.TargetRules.Type, Combination.ProjectTarget.TargetFilePath, ProjectFilePath, NMakePath, ProjectFileFormat) : "";
 					VCProjectFileContent.Append(LayoutDirString);
@@ -1289,8 +1248,7 @@ namespace UnrealBuildTool
 
 					if ((Platform == UnrealTargetPlatform.Win32) || (Platform == UnrealTargetPlatform.Win64))
 					{
-						VCUserFileContent.Append(
-							"	<PropertyGroup " + ConditionString + ">" + ProjectFileGenerator.NewLine);
+						VCUserFileContent.AppendLine("  <PropertyGroup {0}>", ConditionString);
 						if (TargetRulesObject.Type != TargetType.Game)
 						{
 							string DebugOptions = "";
@@ -1305,16 +1263,10 @@ namespace UnrealBuildTool
 								DebugOptions += ProjectName;
 							}
 
-							VCUserFileContent.Append(
-								"		<LocalDebuggerCommandArguments>" + DebugOptions + "</LocalDebuggerCommandArguments>" + ProjectFileGenerator.NewLine
-								);
+							VCUserFileContent.AppendLine("    <LocalDebuggerCommandArguments>{0}</LocalDebuggerCommandArguments>", DebugOptions);
 						}
-						VCUserFileContent.Append(
-							"		<DebuggerFlavor>WindowsLocalDebugger</DebuggerFlavor>" + ProjectFileGenerator.NewLine
-							);
-						VCUserFileContent.Append(
-							"	</PropertyGroup>" + ProjectFileGenerator.NewLine
-							);
+						VCUserFileContent.AppendLine("    <DebuggerFlavor>WindowsLocalDebugger</DebuggerFlavor>");
+						VCUserFileContent.AppendLine("  </PropertyGroup>");
 					}
 
 					string PlatformUserFileStrings = (ProjGenerator != null) ? ProjGenerator.GetVisualStudioUserFileStrings(Platform, Configuration, ConditionString, TargetRulesObject, Combination.ProjectTarget.TargetFilePath, ProjectFilePath) : "";
@@ -1432,7 +1384,7 @@ namespace UnrealBuildTool
 		/// <param name="TargetConfigurationName">The configuration name from the target rules, or null if we don't have one</param>
 		/// <param name="ProjectPlatformName">Name of platform string to use for Visual Studio project</param>
 		/// <param name="ProjectConfigurationName">Name of configuration string to use for Visual Studio project</param>
-		public override void MakeProjectPlatformAndConfigurationNames(UnrealTargetPlatform Platform, UnrealTargetConfiguration Configuration, string TargetConfigurationName, out string ProjectPlatformName, out string ProjectConfigurationName)
+		public override void MakeProjectPlatformAndConfigurationNames(UnrealTargetPlatform Platform, UnrealTargetConfiguration Configuration, TargetType TargetConfigurationName, out string ProjectPlatformName, out string ProjectConfigurationName)
 		{
 			ProjectConfigurationName = Configuration.ToString();
 			ProjectPlatformName = VCProjectFileGenerator.DotNetPlatformName;
diff --git a/Engine/Source/Programs/UnrealBuildTool/ProjectFiles/VisualStudio/VCProjectFileGenerator.cs b/Engine/Source/Programs/UnrealBuildTool/ProjectFiles/VisualStudio/VCProjectFileGenerator.cs
--- a/Engine/Source/Programs/UnrealBuildTool/ProjectFiles/VisualStudio/VCProjectFileGenerator.cs
+++ b/Engine/Source/Programs/UnrealBuildTool/ProjectFiles/VisualStudio/VCProjectFileGenerator.cs
@@ -312,17 +312,25 @@ namespace UnrealBuildTool
 		/// </summary>
 		class VCSolutionConfigCombination
 		{
+			/// <summary>
 			/// Visual Studio solution configuration name for this config+platform
+			/// </summary>
 			public string VCSolutionConfigAndPlatformName;
 
+			/// <summary>
 			/// Configuration name
+			/// </summary>
 			public UnrealTargetConfiguration Configuration;
 
+			/// <summary>
 			/// Platform name
+			/// </summary>
 			public UnrealTargetPlatform Platform;
 
-			/// The target configuration name
-			public string TargetConfigurationName;
+			/// <summary>
+			/// The target type
+			/// </summary>
+			public TargetType TargetConfigurationName;
 
 			public override string ToString()
 			{
@@ -335,18 +343,17 @@ namespace UnrealBuildTool
 		/// Composes a string to use for the Visual Studio solution configuration, given a build configuration and target rules configuration name
 		/// </summary>
 		/// <param name="Configuration">The build configuration</param>
-		/// <param name="TargetConfigurationName">The target rules configuration name</param>
+		/// <param name="TargetType">The type of target being built</param>
 		/// <returns>The generated solution configuration name</returns>
-		string MakeSolutionConfigurationName(UnrealTargetConfiguration Configuration, string TargetConfigurationName)
+		string MakeSolutionConfigurationName(UnrealTargetConfiguration Configuration, TargetType TargetType)
 		{
 			string SolutionConfigName = Configuration.ToString();
 
 			// Don't bother postfixing "Game" or "Program" -- that will be the default when using "Debug", "Development", etc.
 			// Also don't postfix "RocketGame" when we're building Rocket game projects.  That's the only type of game there is in that case!
-			if (!TargetConfigurationName.Equals(TargetType.Game.ToString(), StringComparison.InvariantCultureIgnoreCase) &&
-				!TargetConfigurationName.Equals(TargetType.Program.ToString(), StringComparison.InvariantCultureIgnoreCase))
+			if (TargetType != TargetType.Game && TargetType != TargetType.Program)
 			{
-				SolutionConfigName += " " + TargetConfigurationName;
+				SolutionConfigName += " " + TargetType.ToString();
 			}
 
 			return SolutionConfigName;
@@ -590,7 +597,7 @@ namespace UnrealBuildTool
 						VCSolutionFileContent.Append(
 							"	GlobalSection(SolutionConfigurationPlatforms) = preSolution" + ProjectFileGenerator.NewLine);
 
-						Dictionary<string, Tuple<UnrealTargetConfiguration, string>> SolutionConfigurationsValidForProjects = new Dictionary<string, Tuple<UnrealTargetConfiguration, string>>();
+						Dictionary<string, Tuple<UnrealTargetConfiguration, TargetType>> SolutionConfigurationsValidForProjects = new Dictionary<string, Tuple<UnrealTargetConfiguration, TargetType>>();
 						HashSet<UnrealTargetPlatform> PlatformsValidForProjects = new HashSet<UnrealTargetPlatform>();
 
 						foreach (UnrealTargetConfiguration CurConfiguration in SupportedConfigurations)
@@ -618,14 +625,14 @@ namespace UnrealBuildTool
 														PlatformsValidForProjects.Add(CurPlatform);
 
 														// Default to a target configuration name of "Game", since that will collapse down to an empty string
-														string TargetConfigurationName = TargetType.Game.ToString();
+														TargetType TargetType = TargetType.Game;
 														if (ProjectTarget.TargetRules != null)
 														{
-															TargetConfigurationName = ProjectTarget.TargetRules.Type.ToString();
+															TargetType = ProjectTarget.TargetRules.Type;
 														}
 
-														string SolutionConfigName = MakeSolutionConfigurationName(CurConfiguration, TargetConfigurationName);
-														SolutionConfigurationsValidForProjects[SolutionConfigName] = new Tuple<UnrealTargetConfiguration, string>(CurConfiguration, TargetConfigurationName);
+														string SolutionConfigName = MakeSolutionConfigurationName(CurConfiguration, TargetType);
+														SolutionConfigurationsValidForProjects[SolutionConfigName] = new Tuple<UnrealTargetConfiguration, TargetType>(CurConfiguration, TargetType);
 													}
 												}
 											}
@@ -637,12 +644,12 @@ namespace UnrealBuildTool
 
 						foreach (UnrealTargetPlatform CurPlatform in PlatformsValidForProjects)
 						{
-							foreach (KeyValuePair<string, Tuple<UnrealTargetConfiguration, string>> SolutionConfigKeyValue in SolutionConfigurationsValidForProjects)
+							foreach (KeyValuePair<string, Tuple<UnrealTargetConfiguration, TargetType>> SolutionConfigKeyValue in SolutionConfigurationsValidForProjects)
 							{
 								// e.g.  "Development|Win64 = Development|Win64"
 								string SolutionConfigName = SolutionConfigKeyValue.Key;
 								UnrealTargetConfiguration Configuration = SolutionConfigKeyValue.Value.Item1;
-								string TargetConfigurationName = SolutionConfigKeyValue.Value.Item2;
+								TargetType TargetType = SolutionConfigKeyValue.Value.Item2;
 
 								string SolutionPlatformName = CurPlatform.ToString();
 
@@ -653,7 +660,7 @@ namespace UnrealBuildTool
 											VCSolutionConfigAndPlatformName = SolutionConfigAndPlatformPair,
 											Configuration = Configuration,
 											Platform = CurPlatform,
-											TargetConfigurationName = TargetConfigurationName
+											TargetConfigurationName = TargetType
 										}
 									);
 							}
@@ -712,13 +719,13 @@ namespace UnrealBuildTool
 								{
 									// Handle aliasing of Program and Game target configuration names
 									if ((IsProgramProject && GameOrProgramConfigsAlreadyMapped.Add(SolutionConfigCombination.VCSolutionConfigAndPlatformName)) ||
-										IsProgramProject && SolutionConfigCombination.TargetConfigurationName != TargetType.Game.ToString() ||
-										!IsProgramProject && SolutionConfigCombination.TargetConfigurationName != TargetType.Program.ToString())
+										IsProgramProject && SolutionConfigCombination.TargetConfigurationName != TargetType.Game ||
+										!IsProgramProject && SolutionConfigCombination.TargetConfigurationName != TargetType.Program)
 									{
-										string TargetConfigurationName = SolutionConfigCombination.TargetConfigurationName;
+										TargetType TargetConfigurationName = SolutionConfigCombination.TargetConfigurationName;
 										if (IsProgramProject)
 										{
-											TargetConfigurationName = TargetType.Program.ToString();
+											TargetConfigurationName = TargetType.Program;
 										}
 
 										// Now, we want to find a target in this project that maps to the current solution config combination.  Only up to one target should
@@ -729,7 +736,7 @@ namespace UnrealBuildTool
 											bool IsMatchingCombination = VCProjectFile.IsValidProjectPlatformAndConfiguration(ProjectTarget, SolutionConfigCombination.Platform, SolutionConfigCombination.Configuration);
 											if (ProjectTarget.TargetRules != null)
 											{
-												if (TargetConfigurationName != ProjectTarget.TargetRules.Type.ToString())
+												if (TargetConfigurationName != ProjectTarget.TargetRules.Type)
 												{
 													// Solution configuration name for this combination doesn't match this target's configuration name.  It's not buildable.
 													IsMatchingCombination = false;
@@ -740,7 +747,7 @@ namespace UnrealBuildTool
 												// UBT gets a pass because it is a dependency of every single configuration combination
 												if (CurProject != UBTProject &&
 													!CurProject.ShouldBuildForAllSolutionTargets &&
-													TargetConfigurationName != TargetType.Game.ToString())
+													TargetConfigurationName != TargetType.Game)
 												{
 													// Can't build non-generated project in configurations except for the default (Game)
 													IsMatchingCombination = false;
@@ -786,11 +793,11 @@ namespace UnrealBuildTool
 
 											if (IsProgramProject)
 											{
-												TargetConfigurationName = TargetType.Program.ToString();
+												TargetConfigurationName = TargetType.Program;
 											}
 											else
 											{
-												TargetConfigurationName = TargetType.Game.ToString();
+												TargetConfigurationName = TargetType.Game;
 											}
 										}
 
@@ -807,7 +814,7 @@ namespace UnrealBuildTool
 										}
 
 										// Always allow SCW and UnrealLighmass to build in editor configurations
-										if (MatchingProjectTarget == null && SolutionConfigCombination.TargetConfigurationName == TargetType.Editor.ToString() && SolutionConfigCombination.Platform == UnrealTargetPlatform.Win64)
+										if (MatchingProjectTarget == null && SolutionConfigCombination.TargetConfigurationName == TargetType.Editor && SolutionConfigCombination.Platform == UnrealTargetPlatform.Win64)
 										{
 											if (CurProject == ShaderCompileWorkerProject)
 											{
@@ -988,7 +995,7 @@ namespace UnrealBuildTool
 					VCSolutionOptions Options = new VCSolutionOptions(ProjectFileFormat);
 
 					// Set the default configuration and startup project
-					VCSolutionConfigCombination DefaultConfig = SolutionConfigCombinations.Find(x => x.Configuration == UnrealTargetConfiguration.Development && x.Platform == UnrealTargetPlatform.Win64 && x.TargetConfigurationName == "Editor");
+					VCSolutionConfigCombination DefaultConfig = SolutionConfigCombinations.Find(x => x.Configuration == UnrealTargetConfiguration.Development && x.Platform == UnrealTargetPlatform.Win64 && x.TargetConfigurationName == TargetType.Editor);
 					if (DefaultConfig != null)
 					{
 						List<VCBinarySetting> Settings = new List<VCBinarySetting>();
diff --git a/Engine/Source/Programs/UnrealBuildTool/System/UEPlatformProjectGenerator.cs b/Engine/Source/Programs/UnrealBuildTool/System/UEPlatformProjectGenerator.cs
--- a/Engine/Source/Programs/UnrealBuildTool/System/UEPlatformProjectGenerator.cs
+++ b/Engine/Source/Programs/UnrealBuildTool/System/UEPlatformProjectGenerator.cs
@@ -170,12 +170,12 @@ namespace UnrealBuildTool
 		/// <summary>
 		/// Return project configuration settings that must be included before the default props file
 		/// </summary>
-		/// <param name="InPlatform">  The UnrealTargetPlatform being built</param>
-		/// <param name="InConfiguration"> The UnrealTargetConfiguration being built</param>
+		/// <param name="Platform">The UnrealTargetPlatform being built</param>
+		/// <param name="Configuration">The UnrealTargetConfiguration being built</param>
+		/// <param name="ProjectFileBuilder">String builder for the project file</param>
 		/// <returns>string    The custom configuration section for the project file; Empty string if it doesn't require one</returns>
-		public virtual string GetVisualStudioPreDefaultString(UnrealTargetPlatform InPlatform, UnrealTargetConfiguration InConfiguration)
+		public virtual void GetVisualStudioPreDefaultString(UnrealTargetPlatform Platform, UnrealTargetConfiguration Configuration, StringBuilder ProjectFileBuilder)
 		{
-			return "";
 		}
 
 		/// <summary>
@@ -184,10 +184,10 @@ namespace UnrealBuildTool
 		/// <param name="InPlatform">  The UnrealTargetPlatform being built</param>
 		/// <param name="InConfiguration"> The UnrealTargetConfiguration being built</param>
 		/// <param name="InProjectFileFormat">The visual studio project file format being generated</param>
+		/// <param name="ProjectFileBuilder">String builder for the project file</param>
 		/// <returns>string    The custom configuration section for the project file; Empty string if it doesn't require one</returns>
-		public virtual string GetVisualStudioPlatformToolsetString(UnrealTargetPlatform InPlatform, UnrealTargetConfiguration InConfiguration, VCProjectFileFormat InProjectFileFormat)
+		public virtual void GetVisualStudioPlatformToolsetString(UnrealTargetPlatform InPlatform, UnrealTargetConfiguration InConfiguration, VCProjectFileFormat InProjectFileFormat, StringBuilder ProjectFileBuilder)
 		{
-			return "";
 		}
 
 		/// <summary>
@@ -195,10 +195,10 @@ namespace UnrealBuildTool
 		/// </summary>
 		/// <param name="InPlatform">  The UnrealTargetPlatform being built</param>
 		/// <param name="InProjectFileFormat">The visual studio project file format being generated</param>
+		/// <param name="ProjectFileBuilder">String builder for the project file</param>
 		/// <returns>string    The custom property import lines for the project file; Empty string if it doesn't require one</returns>
-		public virtual string GetAdditionalVisualStudioPropertyGroups(UnrealTargetPlatform InPlatform, VCProjectFileFormat InProjectFileFormat)
+		public virtual void GetAdditionalVisualStudioPropertyGroups(UnrealTargetPlatform InPlatform, VCProjectFileFormat InProjectFileFormat, StringBuilder ProjectFileBuilder)
 		{
-			return "";
 		}
 
 		/// <summary>
@@ -223,43 +223,41 @@ namespace UnrealBuildTool
 		/// <param name="ProjectFilePath"></param>
 		/// <param name="NMakeOutputPath"></param>
 		/// <param name="InProjectFileFormat">The visual studio project file format being generated</param>
+		/// <param name="ProjectFileBuilder">String builder for the project file</param>
 		/// <returns>The custom path lines for the project file; Empty string if it doesn't require one</returns>
-		public virtual string GetVisualStudioPathsEntries(UnrealTargetPlatform InPlatform, UnrealTargetConfiguration InConfiguration, TargetType TargetType, FileReference TargetRulesPath, FileReference ProjectFilePath, FileReference NMakeOutputPath, VCProjectFileFormat InProjectFileFormat)
+		public virtual void GetVisualStudioPathsEntries(UnrealTargetPlatform InPlatform, UnrealTargetConfiguration InConfiguration, TargetType TargetType, FileReference TargetRulesPath, FileReference ProjectFilePath, FileReference NMakeOutputPath, VCProjectFileFormat InProjectFileFormat, StringBuilder ProjectFileBuilder)
 		{
 			// NOTE: We are intentionally overriding defaults for these paths with empty strings.  We never want Visual Studio's
 			//       defaults for these fields to be propagated, since they are version-sensitive paths that may not reflect
 			//       the environment that UBT is building in.  We'll set these environment variables ourselves!
 			// NOTE: We don't touch 'ExecutablePath' because that would result in Visual Studio clobbering the system "Path"
 			//       environment variable
-			string PathsLines =
-				"		<IncludePath />\n" +
-				"		<ReferencePath />\n" +
-				"		<LibraryPath />\n" +
-				"		<LibraryWPath />\n" +
-				"		<SourcePath />\n" +
-				"		<ExcludePath />\n";
-
-			return PathsLines;
+			ProjectFileBuilder.AppendLine("    <IncludePath />");
+			ProjectFileBuilder.AppendLine("    <ReferencePath />");
+			ProjectFileBuilder.AppendLine("    <LibraryPath />");
+			ProjectFileBuilder.AppendLine("    <LibraryWPath />");
+			ProjectFileBuilder.AppendLine("    <SourcePath />");
+			ProjectFileBuilder.AppendLine("    <ExcludePath />");
 		}
 
 		/// <summary>
 		/// Return any custom property settings. These will be included in the ImportGroup section
 		/// </summary>
 		/// <param name="InPlatform">  The UnrealTargetPlatform being built</param>
+		/// <param name="ProjectFileBuilder">String builder for the project file</param>
 		/// <returns>string    The custom property import lines for the project file; Empty string if it doesn't require one</returns>
-		public virtual string GetVisualStudioImportGroupProperties(UnrealTargetPlatform InPlatform)
+		public virtual void GetVisualStudioImportGroupProperties(UnrealTargetPlatform InPlatform, StringBuilder ProjectFileBuilder)
 		{
-			return "";
 		}
 
 		/// <summary>
 		/// Return any custom property settings. These will be included right after Global properties to make values available to all other imports.
 		/// </summary>
 		/// <param name="InPlatform">  The UnrealTargetPlatform being built</param>
+		/// <param name="ProjectFileBuilder">String builder for the project file</param>
 		/// <returns>string    The custom property import lines for the project file; Empty string if it doesn't require one</returns>
-		public virtual string GetVisualStudioGlobalProperties(UnrealTargetPlatform InPlatform)
+		public virtual void GetVisualStudioGlobalProperties(UnrealTargetPlatform InPlatform, StringBuilder ProjectFileBuilder)
 		{
-			return "";
 		}
 
 		/// <summary>
@@ -267,10 +265,10 @@ namespace UnrealBuildTool
 		/// </summary>
 		/// <param name="InPlatform">  The UnrealTargetPlatform being built</param>
 		/// <param name="InProjectFileFormat">The visual studio project file format being generated</param>
+		/// <param name="ProjectFileBuilder">String builder for the project file</param>
 		/// <returns>string    The custom property import lines for the project file; Empty string if it doesn't require one</returns>
-		public virtual string GetVisualStudioTargetOverrides(UnrealTargetPlatform InPlatform, VCProjectFileFormat InProjectFileFormat)
+		public virtual void GetVisualStudioTargetOverrides(UnrealTargetPlatform InPlatform, VCProjectFileFormat InProjectFileFormat, StringBuilder ProjectFileBuilder)
 		{
-			return "";
 		}
 
 		/// <summary>
