Fix issues with mismatched solution/project configurations.

* Every C++ project now has an "invalid" configuration. Trying to build any solution configuration which does not apply to the project will now build that, which prints an error rather than building a completely different configuration.
* Available configurations are now parsed from C# files, so the "Release" configuration will now be used if "Development" is not available, and platforms are supported correctly.
* Programs are now built in game and editor configurations.

diff --git a/Engine/Source/Programs/DotNETCommon/DotNETUtilities/StringUtils.cs b/Engine/Source/Programs/DotNETCommon/DotNETUtilities/StringUtils.cs
--- a/Engine/Source/Programs/DotNETCommon/DotNETUtilities/StringUtils.cs
+++ b/Engine/Source/Programs/DotNETCommon/DotNETUtilities/StringUtils.cs
@@ -5,7 +5,7 @@ using System.Linq;
 using System.Text;
 using System.Threading.Tasks;
 
-namespace DotNETUtilities
+namespace Tools.DotNETCommon
 {
 	public static class StringUtils
 	{
diff --git a/Engine/Source/Programs/UnrealBuildTool/ProjectFiles/ProjectFileGenerator.cs b/Engine/Source/Programs/UnrealBuildTool/ProjectFiles/ProjectFileGenerator.cs
--- a/Engine/Source/Programs/UnrealBuildTool/ProjectFiles/ProjectFileGenerator.cs
+++ b/Engine/Source/Programs/UnrealBuildTool/ProjectFiles/ProjectFileGenerator.cs
@@ -199,17 +199,6 @@ namespace UnrealBuildTool
 		/// </summary>
 		protected string GameProjectName = null;
 
-		/// <summary>
-		/// Global static that only adds platforms that are supported when generating a given target.
-		/// This was the old behavior, and it resulted in scenarios where having an unsupported platform selected
-		/// in the platform drop-down would silently 'switch' to building Win32.
-		/// The new behavior is to add all platforms when generating a target, and then check if it is supported
-		/// at build time. If it is not, then a BuildException is thrown informing the user of an unsupported platform.
-		/// NOTE: This only matters when using "-AllProjects".  It can increase the project file load times though, because of all
-		///       of the extra project configuration combinations we need to store
-		/// </summary>
-		public static bool bCreateDummyConfigsForUnsupportedPlatforms = true;
-
 		/// <summary>
 		/// Whether we should include configurations for "Test" and "Shipping" in generated projects (pass "-NoShippingConfigs" to disable this)
 		/// </summary>
@@ -1094,14 +1083,6 @@ namespace UnrealBuildTool
 							bIncludeTestAndShippingConfigs = false;
 							break;
 
-						case "-DUMMYCONFIGS":
-							bCreateDummyConfigsForUnsupportedPlatforms = true;
-							break;
-
-						case "-NODUMMYCONFIGS":
-							bCreateDummyConfigsForUnsupportedPlatforms = false;
-							break;
-
 						case "-ALLLANGUAGES":
 							bAllDocumentationLanguages = true;
 							break;
diff --git a/Engine/Source/Programs/UnrealBuildTool/ProjectFiles/VisualStudio/VCProject.cs b/Engine/Source/Programs/UnrealBuildTool/ProjectFiles/VisualStudio/VCProject.cs
--- a/Engine/Source/Programs/UnrealBuildTool/ProjectFiles/VisualStudio/VCProject.cs
+++ b/Engine/Source/Programs/UnrealBuildTool/ProjectFiles/VisualStudio/VCProject.cs
@@ -8,10 +8,10 @@ using System.Xml.XPath;
 using System.Xml.Linq;
 using System.Linq;
 using System.Text;
+using System.Text.RegularExpressions;
 using System.Diagnostics;
 using System.Security;
 using Tools.DotNETCommon;
-using DotNETUtilities;
 
 namespace UnrealBuildTool
 {
@@ -112,6 +112,14 @@ namespace UnrealBuildTool
 		/// <param name="ProjectConfigurationName">Name of configuration string to use for Visual Studio project</param>
 		public abstract void MakeProjectPlatformAndConfigurationNames(UnrealTargetPlatform Platform, UnrealTargetConfiguration Configuration, TargetType TargetConfigurationName, out string ProjectPlatformName, out string ProjectConfigurationName);
 
+		public string MakeConfigurationAndPlatformPair(UnrealTargetPlatform Platform, UnrealTargetConfiguration Configuration, TargetType TargetConfigurationName)
+		{
+			string ProjectPlatformName;
+			string ProjectConfigName;
+			MakeProjectPlatformAndConfigurationNames(Platform, Configuration, TargetConfigurationName, out ProjectPlatformName, out ProjectConfigName);
+			return String.Format("{0}|{1}", ProjectConfigName, ProjectPlatformName);
+		}
+
 		static UnrealTargetConfiguration[] GetSupportedConfigurations(TargetRules Rules)
 		{
 			// Otherwise take the SupportedConfigurationsAttribute from the first type in the inheritance chain that supports it
@@ -125,11 +133,7 @@ namespace UnrealBuildTool
 			}
 
 			// Otherwise, get the default for the target type
-			if (Rules.Type == TargetType.Program)
-			{
-				return new[] { UnrealTargetConfiguration.Debug, UnrealTargetConfiguration.Development };
-			}
-			else if(Rules.Type == TargetType.Editor)
+			if(Rules.Type == TargetType.Editor)
 			{
 				return new[] { UnrealTargetConfiguration.Debug, UnrealTargetConfiguration.DebugGame, UnrealTargetConfiguration.Development };
 			}
@@ -176,20 +180,9 @@ namespace UnrealBuildTool
 
 			List<UnrealTargetConfiguration> SupportedConfigurations = new List<UnrealTargetConfiguration>();
 			List<UnrealTargetPlatform> SupportedPlatforms = new List<UnrealTargetPlatform>();
-			if (!ProjectFileGenerator.bCreateDummyConfigsForUnsupportedPlatforms)
-			{
-				if (ProjectTarget.TargetRules != null)
-				{
-					SupportedPlatforms.AddRange(ProjectTarget.SupportedPlatforms);
-				}
-			}
-			else
-			{
-				SupportedPlatforms.AddRange(Utils.GetPlatformsInClass(UnrealPlatformClass.All));
-			}
-
 			if (ProjectTarget.TargetRules != null)
 			{
+				SupportedPlatforms.AddRange(ProjectTarget.SupportedPlatforms);
 				SupportedConfigurations.AddRange(GetSupportedConfigurations(ProjectTarget.TargetRules));
 			}
 
@@ -551,6 +544,12 @@ namespace UnrealBuildTool
 				}
 			}
 
+			// Add the "invalid" configuration for each platform. We use this when the solution configuration does not match any project configuration.
+			VCProjectFileContent.AppendLine("    <ProjectConfiguration Include=\"Invalid|{0}\">", VCProjectFileGenerator.DefaultPlatformName);
+			VCProjectFileContent.AppendLine("      <Configuration>Invalid</Configuration>");
+			VCProjectFileContent.AppendLine("      <Platform>{0}</Platform>", VCProjectFileGenerator.DefaultPlatformName);
+			VCProjectFileContent.AppendLine("    </ProjectConfiguration>");
+
 			// Output ALL the project's config-platform permutations (project files MUST do this)
 			foreach (Tuple<string, UnrealTargetConfiguration> ConfigurationTuple in ProjectConfigurationNameAndConfigurations)
 			{
@@ -633,6 +632,11 @@ namespace UnrealBuildTool
 
 			VCProjectFileContent.AppendLine("  <Import Project=\"$(VCTargetsPath)\\Microsoft.Cpp.Default.props\" />");
 
+			// Write the invalid configuration data
+			VCProjectFileContent.AppendLine("  <PropertyGroup Condition=\"'$(Configuration)|$(Platform)'=='Invalid|{0}'\" Label=\"Configuration\">", VCProjectFileGenerator.DefaultPlatformName);
+			VCProjectFileContent.AppendLine("    <ConfigurationType>Makefile</ConfigurationType>");
+			VCProjectFileContent.AppendLine("  </PropertyGroup>");
+
 			// Write each project configuration PreDefaultProps section
 			foreach (Tuple<string, UnrealTargetConfiguration> ConfigurationTuple in ProjectConfigurationNameAndConfigurations)
 			{
@@ -650,6 +654,14 @@ namespace UnrealBuildTool
 			VCProjectFileContent.AppendLine("  <ImportGroup Label=\"ExtensionSettings\" />");
 			VCProjectFileContent.AppendLine("  <PropertyGroup Label=\"UserMacros\" />");
 
+			// Write the invalid configuration
+			const string InvalidMessage = "echo The selected platform/configuration is not valid for this target.";
+			VCProjectFileContent.AppendLine("  <PropertyGroup Condition=\"'$(Configuration)|$(Platform)'=='Invalid|{0}'\">", VCProjectFileGenerator.DefaultPlatformName);
+			VCProjectFileContent.AppendLine("    <NMakeBuildCommandLine>{0}</NMakeBuildCommandLine>", InvalidMessage);
+			VCProjectFileContent.AppendLine("    <NMakeReBuildCommandLine>{0}</NMakeReBuildCommandLine>", InvalidMessage);
+			VCProjectFileContent.AppendLine("    <NMakeCleanCommandLine>{0}</NMakeCleanCommandLine>", InvalidMessage);
+			VCProjectFileContent.AppendLine("  </PropertyGroup>");
+
 			// Write each project configuration
 			foreach (ProjectConfigAndTargetCombination Combination in ProjectConfigAndTargetCombinations)
 			{
@@ -1290,6 +1302,16 @@ namespace UnrealBuildTool
 			get { return "{FAE04EC0-301F-11D3-BF4B-00C04F79EFBC}"; }
 		}
 
+		/// <summary>
+		/// Platforms that this project supports
+		/// </summary>
+		public HashSet<string> Platforms = new HashSet<string>();
+
+		/// <summary>
+		/// Configurations that this project supports
+		/// </summary>
+		public HashSet<string> Configurations = new HashSet<string>();
+
 		/// <summary>
 		/// Constructs a new project file object
 		/// </summary>
@@ -1297,6 +1319,44 @@ namespace UnrealBuildTool
 		public VCSharpProjectFile(FileReference InitFilePath)
 			: base(InitFilePath)
 		{
+			try
+			{
+				XmlDocument Document = new XmlDocument();
+				Document.Load(InitFilePath.FullName);
+
+				// Check the root element is the right type
+				if (Document.DocumentElement.Name != "Project")
+				{
+					throw new BuildException("Unexpected root element '{0}' in project file", Document.DocumentElement.Name);
+				}
+
+				// Parse all the configurations and platforms 
+				// Parse the basic structure of the document, updating properties and recursing into other referenced projects as we go
+				foreach (XmlElement Element in Document.DocumentElement.ChildNodes.OfType<XmlElement>())
+				{
+					if(Element.Name == "PropertyGroup")
+					{
+						string Condition = Element.GetAttribute("Condition");
+						if(!String.IsNullOrEmpty(Condition))
+						{
+							Match Match = Regex.Match(Condition, "^\\s*'\\$\\(Configuration\\)\\|\\$\\(Platform\\)'\\s*==\\s*'(.+)\\|(.+)'\\s*$");
+							if(Match.Success && Match.Groups.Count == 3)
+							{
+								Configurations.Add(Match.Groups[1].Value);
+								Platforms.Add(Match.Groups[2].Value);
+							}
+							else
+							{
+								Log.TraceWarning("Unable to parse configuration/platform from condition '{0}': {1}", InitFilePath, Condition);
+							}
+						}
+					}
+				}
+			}
+			catch(Exception Ex)
+			{
+				Log.TraceWarning("Unable to parse {0}: {1}", InitFilePath, Ex.ToString());
+			}
 		}
 
 		/// <summary>
@@ -1331,49 +1391,6 @@ namespace UnrealBuildTool
 			return Info.IsDotNETCoreProject();
 		}
 
-		/// <summary>
-		/// Reads the list of dependencies from the specified project file.
-		/// </summary>
-		public List<string> GetCSharpDependencies()
-		{
-			List<string> RelativeFilePaths = new List<string>();
-			XmlDocument Doc = new XmlDocument();
-			Doc.Load(ProjectFilePath.FullName);
-
-			string[] Tags = new string[] { "Compile", "Page", "Resource" };
-			foreach (string Tag in Tags)
-			{
-				XmlNodeList Elements = Doc.GetElementsByTagName(Tag);
-				foreach (XmlElement Element in Elements)
-				{
-					RelativeFilePaths.Add(Element.GetAttribute("Include"));
-				}
-			}
-
-			return RelativeFilePaths;
-		}
-
-		/// <summary>
-		/// Adds a C# dot net (system) assembly reference to this project
-		/// </summary>
-		/// <param name="AssemblyReference">The full path to the assembly file on disk</param>
-		public void AddDotNetAssemblyReference(string AssemblyReference)
-		{
-			if (!DotNetAssemblyReferences.Contains(AssemblyReference))
-			{
-				DotNetAssemblyReferences.Add(AssemblyReference);
-			}
-		}
-
-		/// <summary>
-		/// Adds a C# assembly reference to this project, such as a third party assembly needed for this project to compile
-		/// </summary>
-		/// <param name="AssemblyReference">The full path to the assembly file on disk</param>
-		public void AddAssemblyReference(FileReference AssemblyReference)
-		{
-			AssemblyReferences.Add(AssemblyReference);
-		}
-
 		/// <summary>
 		/// Given a target platform and configuration, generates a platform and configuration name string to use in Visual Studio projects.
 		/// Unlike with solution configurations, Visual Studio project configurations only support certain types of platforms, so we'll
@@ -1386,8 +1403,34 @@ namespace UnrealBuildTool
 		/// <param name="ProjectConfigurationName">Name of configuration string to use for Visual Studio project</param>
 		public override void MakeProjectPlatformAndConfigurationNames(UnrealTargetPlatform Platform, UnrealTargetConfiguration Configuration, TargetType TargetConfigurationName, out string ProjectPlatformName, out string ProjectConfigurationName)
 		{
-			ProjectConfigurationName = Configuration.ToString();
-			ProjectPlatformName = VCProjectFileGenerator.DotNetPlatformName;
+			// Find the matching platform name
+			if(Platform == UnrealTargetPlatform.Win32 && Platforms.Contains("x86"))
+			{
+				ProjectPlatformName = "x86";
+			}
+			else if(Platforms.Contains("x64"))
+			{
+				ProjectPlatformName = "x64";
+			}
+			else
+			{
+				ProjectPlatformName = "Any CPU";
+			}
+
+			// Find the matching configuration
+			string ConfigurationName = Configuration.ToString();
+			if(Configurations.Contains(ConfigurationName))
+			{
+				ProjectConfigurationName = ConfigurationName;
+			}
+			else if(Configurations.Contains("Development"))
+			{
+				ProjectConfigurationName = "Development";
+			}
+			else
+			{
+				ProjectConfigurationName = "Release";
+			}
 		}
 
 		/// <summary>
@@ -1401,10 +1444,6 @@ namespace UnrealBuildTool
 			throw new BuildException("Support for writing C# projects from UnrealBuildTool has been removed.");
 		}
 
-		/// Assemblies this project is dependent on
-		protected readonly List<FileReference> AssemblyReferences = new List<FileReference>();
-		/// System assemblies this project is dependent on
-		protected readonly List<string> DotNetAssemblyReferences = new List<string>() { "System", "System.Core", "System.Data", "System.Xml" };
 		/// Cache of parsed info about this project
 		protected readonly Dictionary<UnrealTargetConfiguration, CsProjectInfo> CachedProjectInfo = new Dictionary<UnrealTargetConfiguration, CsProjectInfo>();
 	}
diff --git a/Engine/Source/Programs/UnrealBuildTool/ProjectFiles/VisualStudio/VCProjectFileGenerator.cs b/Engine/Source/Programs/UnrealBuildTool/ProjectFiles/VisualStudio/VCProjectFileGenerator.cs
--- a/Engine/Source/Programs/UnrealBuildTool/ProjectFiles/VisualStudio/VCProjectFileGenerator.cs
+++ b/Engine/Source/Programs/UnrealBuildTool/ProjectFiles/VisualStudio/VCProjectFileGenerator.cs
@@ -7,6 +7,7 @@ using System.IO;
 using System.Linq;
 using System.Xml.Linq;
 using Tools.DotNETCommon;
+using DotNETUtilities;
 
 namespace UnrealBuildTool
 {
@@ -223,10 +224,6 @@ namespace UnrealBuildTool
 		/// installed version of Visual Studio (e.g, "iOS")
 		public const string DefaultPlatformName = "Win32";
 
-		/// The platform name that must be used for .NET projects
-		public const string DotNetPlatformName = "Any CPU";
-
-
 		/// <summary>
 		/// Configures project generator based on command-line options
 		/// </summary>
@@ -406,42 +403,30 @@ namespace UnrealBuildTool
 			// Setup solution file content
 			StringBuilder VCSolutionFileContent = new StringBuilder();
 
-			const string VersionTag = "# UnrealEngineGeneratedSolutionVersion=1.0";
-
 			// Solution file header
 			if (ProjectFileFormat == VCProjectFileFormat.VisualStudio2017)
 			{
-				VCSolutionFileContent.Append(
-					ProjectFileGenerator.NewLine +
-					"Microsoft Visual Studio Solution File, Format Version 12.00" + ProjectFileGenerator.NewLine +
-					"# Visual Studio 15" + ProjectFileGenerator.NewLine +
-					"VisualStudioVersion = 15.0.25807.0" + ProjectFileGenerator.NewLine +
-					"MinimumVisualStudioVersion = 10.0.40219.1" + ProjectFileGenerator.NewLine);
+				VCSolutionFileContent.AppendLine("Microsoft Visual Studio Solution File, Format Version 12.00");
+				VCSolutionFileContent.AppendLine("# Visual Studio 15");
+				VCSolutionFileContent.AppendLine("VisualStudioVersion = 15.0.25807.0");
+				VCSolutionFileContent.AppendLine("MinimumVisualStudioVersion = 10.0.40219.1");
 			}
 			else if (ProjectFileFormat == VCProjectFileFormat.VisualStudio2015)
 			{
-				VCSolutionFileContent.Append(
-					ProjectFileGenerator.NewLine +
-					"Microsoft Visual Studio Solution File, Format Version 12.00" + ProjectFileGenerator.NewLine +
-					"# Visual Studio 14" + ProjectFileGenerator.NewLine +
-					"VisualStudioVersion = 14.0.22310.1" + ProjectFileGenerator.NewLine +
-					"MinimumVisualStudioVersion = 10.0.40219.1" + ProjectFileGenerator.NewLine);
+				VCSolutionFileContent.AppendLine("Microsoft Visual Studio Solution File, Format Version 12.00");
+				VCSolutionFileContent.AppendLine("# Visual Studio 14");
+				VCSolutionFileContent.AppendLine("VisualStudioVersion = 14.0.22310.1");
+				VCSolutionFileContent.AppendLine("MinimumVisualStudioVersion = 10.0.40219.1");
 			}
 			else if (ProjectFileFormat == VCProjectFileFormat.VisualStudio2013)
 			{
-				VCSolutionFileContent.Append(
-					ProjectFileGenerator.NewLine +
-					"Microsoft Visual Studio Solution File, Format Version 12.00" + ProjectFileGenerator.NewLine +
-					"# Visual Studio 2013" + ProjectFileGenerator.NewLine +
-					VersionTag + ProjectFileGenerator.NewLine);
+				VCSolutionFileContent.AppendLine("Microsoft Visual Studio Solution File, Format Version 12.00");
+				VCSolutionFileContent.AppendLine("# Visual Studio 2013");
             }
             else if (ProjectFileFormat == VCProjectFileFormat.VisualStudio2012)
             {
-                VCSolutionFileContent.Append(
-                    ProjectFileGenerator.NewLine +
-                    "Microsoft Visual Studio Solution File, Format Version 12.00" + ProjectFileGenerator.NewLine +
-                    "# Visual Studio 2012" + ProjectFileGenerator.NewLine +
-                    VersionTag + ProjectFileGenerator.NewLine);
+                VCSolutionFileContent.AppendLine("Microsoft Visual Studio Solution File, Format Version 12.00");
+				VCSolutionFileContent.AppendLine("# Visual Studio 2012");
             }
 			else
 			{
@@ -498,27 +483,21 @@ namespace UnrealBuildTool
 					foreach (VisualStudioSolutionFolder CurFolder in AllSolutionFolders)
 					{
 						string FolderGUIDString = CurFolder.FolderGUID.ToString("B").ToUpperInvariant();
-						VCSolutionFileContent.Append(
-								"Project(\"" + SolutionFolderEntryGUID + "\") = \"" + CurFolder.FolderName + "\", \"" + CurFolder.FolderName + "\", \"" + FolderGUIDString + "\"" + ProjectFileGenerator.NewLine);
+						VCSolutionFileContent.AppendLine("Project(\"" + SolutionFolderEntryGUID + "\") = \"" + CurFolder.FolderName + "\", \"" + CurFolder.FolderName + "\", \"" + FolderGUIDString + "\"");
 
 						// Add any files that are inlined right inside the solution folder
 						if (CurFolder.Files.Count > 0)
 						{
-							VCSolutionFileContent.Append(
-									"	ProjectSection(SolutionItems) = preProject" + ProjectFileGenerator.NewLine);
+							VCSolutionFileContent.AppendLine("	ProjectSection(SolutionItems) = preProject");
 							foreach (string CurFile in CurFolder.Files)
 							{
 								// Syntax is:  <relative file path> = <relative file path>
-								VCSolutionFileContent.Append(
-									"		" + CurFile + " = " + CurFile + ProjectFileGenerator.NewLine);
+								VCSolutionFileContent.AppendLine("		" + CurFile + " = " + CurFile);
 							}
-							VCSolutionFileContent.Append(
-									"	EndProjectSection" + ProjectFileGenerator.NewLine);
+							VCSolutionFileContent.AppendLine("	EndProjectSection");
 						}
 
-						VCSolutionFileContent.Append(
-								"EndProject" + ProjectFileGenerator.NewLine
-							);
+						VCSolutionFileContent.AppendLine("EndProject");
 					}
 				}
 
@@ -536,8 +515,7 @@ namespace UnrealBuildTool
 					// Use the existing project's GUID that's already known to us
 					string ProjectGUID = CurProject.ProjectGUID.ToString("B").ToUpperInvariant();
 
-					VCSolutionFileContent.Append(
-							"Project(\"" + ProjectTypeGUID + "\") = \"" + ProjectNameInSolution + "\", \"" + CurProject.ProjectFilePath.MakeRelativeTo(ProjectFileGenerator.MasterProjectPath) + "\", \"" + ProjectGUID + "\"" + ProjectFileGenerator.NewLine);
+					VCSolutionFileContent.AppendLine("Project(\"" + ProjectTypeGUID + "\") = \"" + ProjectNameInSolution + "\", \"" + CurProject.ProjectFilePath.MakeRelativeTo(ProjectFileGenerator.MasterProjectPath) + "\", \"" + ProjectGUID + "\"");
 
 					// Setup dependency on UnrealBuildTool, if we need that.  This makes sure that UnrealBuildTool is
 					// freshly compiled before kicking off any build operations on this target project
@@ -557,22 +535,20 @@ namespace UnrealBuildTool
 
 						if (Dependencies.Count > 0)
 						{
-							VCSolutionFileContent.Append("\tProjectSection(ProjectDependencies) = postProject" + ProjectFileGenerator.NewLine);
+							VCSolutionFileContent.AppendLine("\tProjectSection(ProjectDependencies) = postProject");
 
 							// Setup any addition dependencies this project has...
 							foreach (ProjectFile DependsOnProject in Dependencies)
 							{
 								string DependsOnProjectGUID = ((MSBuildProjectFile)DependsOnProject).ProjectGUID.ToString("B").ToUpperInvariant();
-								VCSolutionFileContent.Append("\t\t" + DependsOnProjectGUID + " = " + DependsOnProjectGUID + ProjectFileGenerator.NewLine);
+								VCSolutionFileContent.AppendLine("\t\t" + DependsOnProjectGUID + " = " + DependsOnProjectGUID);
 							}
 
-							VCSolutionFileContent.Append("\tEndProjectSection" + ProjectFileGenerator.NewLine);
+							VCSolutionFileContent.AppendLine("\tEndProjectSection");
 						}
 					}
 
-					VCSolutionFileContent.Append(
-							"EndProject" + ProjectFileGenerator.NewLine
-						);
+					VCSolutionFileContent.AppendLine("EndProject");
 				}
 
 				// Add the visualizers at the solution level. Doesn't seem to be picked up from a makefile project in VS2017 15.8.5.
@@ -590,12 +566,10 @@ namespace UnrealBuildTool
 			// The "Global" section has source control, solution configurations, project configurations,
 			// preferences, and project hierarchy data
 			{
-				VCSolutionFileContent.Append(
-					"Global" + ProjectFileGenerator.NewLine);
+				VCSolutionFileContent.AppendLine("Global");
 				{
 					{
-						VCSolutionFileContent.Append(
-							"	GlobalSection(SolutionConfigurationPlatforms) = preSolution" + ProjectFileGenerator.NewLine);
+						VCSolutionFileContent.AppendLine("	GlobalSection(SolutionConfigurationPlatforms) = preSolution");
 
 						Dictionary<string, Tuple<UnrealTargetConfiguration, TargetType>> SolutionConfigurationsValidForProjects = new Dictionary<string, Tuple<UnrealTargetConfiguration, TargetType>>();
 						HashSet<UnrealTargetPlatform> PlatformsValidForProjects = new HashSet<UnrealTargetPlatform>();
@@ -679,24 +653,19 @@ namespace UnrealBuildTool
 							// We alias "Game" and "Program" to both have the same solution configuration, so we're careful not to add the same combination twice.
 							if (!AppendedSolutionConfigAndPlatformNames.Contains(SolutionConfigCombination.VCSolutionConfigAndPlatformName))
 							{
-								VCSolutionFileContent.Append(
-									"		" + SolutionConfigCombination.VCSolutionConfigAndPlatformName + " = " + SolutionConfigCombination.VCSolutionConfigAndPlatformName + ProjectFileGenerator.NewLine);
+								VCSolutionFileContent.AppendLine("		" + SolutionConfigCombination.VCSolutionConfigAndPlatformName + " = " + SolutionConfigCombination.VCSolutionConfigAndPlatformName);
 								AppendedSolutionConfigAndPlatformNames.Add(SolutionConfigCombination.VCSolutionConfigAndPlatformName);
 							}
 						}
 
-						VCSolutionFileContent.Append(
-							"	EndGlobalSection" + ProjectFileGenerator.NewLine);
+						VCSolutionFileContent.AppendLine("	EndGlobalSection");
 					}
 
 
 					// Assign each project's "project configuration" to our "solution platform + configuration" pairs.  This
 					// also sets up which projects are actually built when building the solution.
 					{
-						VCSolutionFileContent.Append(
-							"	GlobalSection(ProjectConfigurationPlatforms) = postSolution" + ProjectFileGenerator.NewLine);
-
-						List<VCSolutionConfigCombination> CombinationsThatWereMatchedToProjects = new List<VCSolutionConfigCombination>();
+						VCSolutionFileContent.AppendLine("	GlobalSection(ProjectConfigurationPlatforms) = postSolution");
 
 						foreach (MSBuildProjectFile CurProject in AllProjectFiles)
 						{
@@ -712,16 +681,31 @@ namespace UnrealBuildTool
 								{
 									throw new BuildException("Expecting project '" + CurProject.ProjectFilePath + "' to have at least one ProjectTarget associated with it!");
 								}
-								bool IsProgramProject = CurProject.ProjectTargets[0].TargetRules != null && CurProject.ProjectTargets[0].TargetRules.Type == TargetType.Program;
 
-								HashSet<string> GameOrProgramConfigsAlreadyMapped = new HashSet<string>();
 								foreach (VCSolutionConfigCombination SolutionConfigCombination in SolutionConfigCombinations)
 								{
-									// Handle aliasing of Program and Game target configuration names
-									if ((IsProgramProject && GameOrProgramConfigsAlreadyMapped.Add(SolutionConfigCombination.VCSolutionConfigAndPlatformName)) ||
-										IsProgramProject && SolutionConfigCombination.TargetConfigurationName != TargetType.Game ||
-										!IsProgramProject && SolutionConfigCombination.TargetConfigurationName != TargetType.Program)
+									// Have to match every solution configuration combination to a project configuration (or use the invalid one) 
+									string ProjectConfigAndPlatformPair;
+									if(CurProject is VCSharpProjectFile)
+									{
+										ProjectConfigAndPlatformPair = CurProject.MakeConfigurationAndPlatformPair(SolutionConfigCombination.Platform, SolutionConfigCombination.Configuration, TargetType.Program);
+									}
+									else
+									{
+										ProjectConfigAndPlatformPair = String.Format("Invalid|{0}", DefaultPlatformName);
+									}
+
+									// Whether the configuration should be built automatically as part of the solution
+									bool bBuildByDefault = false;
+
+									// Whether this configuration should deploy by default (requires bBuildByDefault)
+									bool bDeployByDefault = false;
+
+									// Programs are built in editor configurations (since the editor is like a desktop program too) and game configurations (since we omit the "game" qualification in the configuration name).
+									bool IsProgramProject = CurProject.ProjectTargets[0].TargetRules != null && CurProject.ProjectTargets[0].TargetRules.Type == TargetType.Program;
+									if(!IsProgramProject || SolutionConfigCombination.TargetConfigurationName == TargetType.Game || SolutionConfigCombination.TargetConfigurationName == TargetType.Editor)
 									{
+										// Get the target type we expect to find for this project
 										TargetType TargetConfigurationName = SolutionConfigCombination.TargetConfigurationName;
 										if (IsProgramProject)
 										{
@@ -730,190 +714,134 @@ namespace UnrealBuildTool
 
 										// Now, we want to find a target in this project that maps to the current solution config combination.  Only up to one target should
 										// and every solution config combination should map to at least one target in one project (otherwise we shouldn't have added it!).
-										ProjectTarget MatchingProjectTarget = null;
+										List<ProjectTarget> MatchingProjectTargets = new List<ProjectTarget>();
 										foreach (ProjectTarget ProjectTarget in CurProject.ProjectTargets)
 										{
-											bool IsMatchingCombination = VCProjectFile.IsValidProjectPlatformAndConfiguration(ProjectTarget, SolutionConfigCombination.Platform, SolutionConfigCombination.Configuration);
-											if (ProjectTarget.TargetRules != null)
-											{
-												if (TargetConfigurationName != ProjectTarget.TargetRules.Type)
-												{
-													// Solution configuration name for this combination doesn't match this target's configuration name.  It's not buildable.
-													IsMatchingCombination = false;
-												}
-											}
-											else
+											if(VCProjectFile.IsValidProjectPlatformAndConfiguration(ProjectTarget, SolutionConfigCombination.Platform, SolutionConfigCombination.Configuration))
 											{
-												// UBT gets a pass because it is a dependency of every single configuration combination
-												if (CurProject != UBTProject &&
-													!CurProject.ShouldBuildForAllSolutionTargets &&
-													TargetConfigurationName != TargetType.Game)
+												if (ProjectTarget.TargetRules != null)
 												{
-													// Can't build non-generated project in configurations except for the default (Game)
-													IsMatchingCombination = false;
-												}
-											}
-
-											if (IsMatchingCombination)
-											{
-												if (MatchingProjectTarget != null)
-												{
-													// Not expecting more than one target to match a single solution configuration per project!
-													throw new BuildException("Not expecting more than one target for project " + CurProject.ProjectFilePath + " to match solution configuration " + SolutionConfigCombination.VCSolutionConfigAndPlatformName);
+													if (TargetConfigurationName == ProjectTarget.TargetRules.Type)
+													{
+														MatchingProjectTargets.Add(ProjectTarget);
+													}
 												}
-
-												MatchingProjectTarget = ProjectTarget;
-
-												// NOTE: For faster perf, we could "break" here and bail out early, but that would circumvent the error checking
-												//		 for multiple targets within a project that may map to a single solution configuration.
-											}
-										}
-
-										UnrealTargetConfiguration SolutionConfiguration = SolutionConfigCombination.Configuration;
-										UnrealTargetPlatform SolutionPlatform = SolutionConfigCombination.Platform;
-
-
-										if (MatchingProjectTarget == null)
-										{
-											// The current configuration/platform and target configuration name doesn't map to anything our project actually supports.
-											// We'll map it to a default config.
-											SolutionConfiguration = UnrealTargetConfiguration.Development;
-
-											// Prefer using Win64 as the default, but fall back to a platform the project file actually supports if needed.  This is for
-											// projects that can never be compiled in Windows, such as UnrealLaunchDaemon which is an iOS-only program
-											SolutionPlatform = UnrealTargetPlatform.Win64;
-											if (CurProject.ProjectTargets[0].TargetRules != null)
-											{
-												if (!CurProject.ProjectTargets[0].SupportedPlatforms.Contains(SolutionPlatform))
+												else
 												{
-													SolutionPlatform = CurProject.ProjectTargets[0].SupportedPlatforms[0];
+													// UBT gets a pass because it is a dependency of every single configuration combination
+													if (CurProject == UBTProject || CurProject.ShouldBuildForAllSolutionTargets || TargetConfigurationName == TargetType.Game)
+													{
+														MatchingProjectTargets.Add(ProjectTarget);
+													}
 												}
 											}
-
-
-											if (IsProgramProject)
-											{
-												TargetConfigurationName = TargetType.Program;
-											}
-											else
-											{
-												TargetConfigurationName = TargetType.Game;
-											}
-										}
-
-
-										// If the project wants to always build in "Development", regardless of what the solution
-										// configuration is set to, then we'll do that here.  This is used for projects like
-										// UnrealBuildTool and ShaderCompileWorker
-										if (MatchingProjectTarget != null)
-										{
-											if (MatchingProjectTarget.ForceDevelopmentConfiguration)
-											{
-												SolutionConfiguration = UnrealTargetConfiguration.Development;
-											}
 										}
 
 										// Always allow SCW and UnrealLighmass to build in editor configurations
-										if (MatchingProjectTarget == null && SolutionConfigCombination.TargetConfigurationName == TargetType.Editor && SolutionConfigCombination.Platform == UnrealTargetPlatform.Win64)
+										if (MatchingProjectTargets.Count == 0 && SolutionConfigCombination.TargetConfigurationName == TargetType.Editor && SolutionConfigCombination.Platform == UnrealTargetPlatform.Win64)
 										{
 											if (CurProject == ShaderCompileWorkerProject)
 											{
-												MatchingProjectTarget = ShaderCompileWorkerProject.ProjectTargets[0];
+												MatchingProjectTargets.Add(ShaderCompileWorkerProject.ProjectTargets[0]);
 											}
 											else if (CurProject == UnrealLightmassProject)
 											{
-												MatchingProjectTarget = UnrealLightmassProject.ProjectTargets[0];
+												MatchingProjectTargets.Add(UnrealLightmassProject.ProjectTargets[0]);
 											}
 										}
 
-										string ProjectConfigName;
-										string ProjectPlatformName;
-										if (CurProject.IsStubProject)
-										{
-											if (SolutionPlatform != UnrealTargetPlatform.Unknown || SolutionConfiguration != UnrealTargetConfiguration.Unknown)
-											{
-												throw new BuildException("Stub project was expecting platform and configuration type to be set to Unknown");
-											}
-											ProjectPlatformName = MSBuildProjectFile.StubProjectPlatformName;
-											ProjectConfigName = MSBuildProjectFile.StubProjectConfigurationName;
-										}
-										else
+										// Make sure there's only one matching project target
+										if(MatchingProjectTargets.Count > 1)
 										{
-											CurProject.MakeProjectPlatformAndConfigurationNames(SolutionPlatform, SolutionConfiguration, TargetConfigurationName, out ProjectPlatformName, out ProjectConfigName);
+											throw new BuildException("Not expecting more than one target for project " + CurProject.ProjectFilePath + " to match solution configuration " + SolutionConfigCombination.VCSolutionConfigAndPlatformName);
 										}
 
-										string ProjectConfigAndPlatformPair = ProjectConfigName.ToString() + "|" + ProjectPlatformName.ToString();
-
-										// e.g.  "{4232C52C-680F-4850-8855-DC39419B5E9B}.Debug|iOS.ActiveCfg = iOS_Debug|Win32"
-										string CurProjectGUID = CurProject.ProjectGUID.ToString("B").ToUpperInvariant();
-										VCSolutionFileContent.Append(
-											"		" + CurProjectGUID + "." + SolutionConfigCombination.VCSolutionConfigAndPlatformName + ".ActiveCfg = " + ProjectConfigAndPlatformPair + ProjectFileGenerator.NewLine);
+										// If we found a matching project, get matching configuration
+										if(MatchingProjectTargets.Count == 1)
+										{
+											// Get the matching target
+											ProjectTarget MatchingProjectTarget = MatchingProjectTargets[0];
 
+											// If the project wants to always build in "Development", regardless of what the solution configuration is set to, then we'll do that here.
+											UnrealTargetConfiguration ProjectConfiguration = SolutionConfigCombination.Configuration;
+											if (MatchingProjectTarget.ForceDevelopmentConfiguration && SolutionConfigCombination.TargetConfigurationName != TargetType.Game)
+											{
+												ProjectConfiguration = UnrealTargetConfiguration.Development;
+											}
 
-										// Set whether this project configuration should be built when the user initiates "build solution"
-										if (MatchingProjectTarget != null && CurProject.ShouldBuildByDefaultForSolutionTargets)
-										{
-											// Some targets are "dummy targets"; they only exist to show user friendly errors in VS. Weed them out here, and don't set them to build by default.
-											List<UnrealTargetPlatform> SupportedPlatforms = null;
-											if (MatchingProjectTarget.TargetRules != null)
+											// Get the matching project configuration
+											UnrealTargetPlatform ProjectPlatform = SolutionConfigCombination.Platform;
+											if (CurProject.IsStubProject)
 											{
-												SupportedPlatforms = new List<UnrealTargetPlatform>();
-												SupportedPlatforms.AddRange(MatchingProjectTarget.SupportedPlatforms);
+												if (ProjectPlatform != UnrealTargetPlatform.Unknown || ProjectConfiguration != UnrealTargetConfiguration.Unknown)
+												{
+													throw new BuildException("Stub project was expecting platform and configuration type to be set to Unknown");
+												}
+												ProjectConfigAndPlatformPair = String.Format("{0}|{1}", MSBuildProjectFile.StubProjectConfigurationName, MSBuildProjectFile.StubProjectPlatformName);
 											}
-											if (SupportedPlatforms == null || SupportedPlatforms.Contains(SolutionConfigCombination.Platform))
+											else
 											{
-												VCSolutionFileContent.Append(
-														"		" + CurProjectGUID + "." + SolutionConfigCombination.VCSolutionConfigAndPlatformName + ".Build.0 = " + ProjectConfigAndPlatformPair + ProjectFileGenerator.NewLine);
+												ProjectConfigAndPlatformPair = CurProject.MakeConfigurationAndPlatformPair(ProjectPlatform, ProjectConfiguration, TargetConfigurationName);
+											}
 
-												UEPlatformProjectGenerator ProjGen = UEPlatformProjectGenerator.GetPlatformProjectGenerator(SolutionConfigCombination.Platform, true);
-												if (MatchingProjectTarget.ProjectDeploys ||
-													((ProjGen != null) && (ProjGen.GetVisualStudioDeploymentEnabled(SolutionPlatform, SolutionConfiguration) == true)))
+											// Set whether this project configuration should be built when the user initiates "build solution"
+											if (CurProject.ShouldBuildByDefaultForSolutionTargets)
+											{
+												// Some targets are "dummy targets"; they only exist to show user friendly errors in VS. Weed them out here, and don't set them to build by default.
+												List<UnrealTargetPlatform> SupportedPlatforms = null;
+												if (MatchingProjectTarget.TargetRules != null)
+												{
+													SupportedPlatforms = new List<UnrealTargetPlatform>();
+													SupportedPlatforms.AddRange(MatchingProjectTarget.SupportedPlatforms);
+												}
+												if (SupportedPlatforms == null || SupportedPlatforms.Contains(SolutionConfigCombination.Platform))
 												{
-													VCSolutionFileContent.Append(
-															"		" + CurProjectGUID + "." + SolutionConfigCombination.VCSolutionConfigAndPlatformName + ".Deploy.0 = " + ProjectConfigAndPlatformPair + ProjectFileGenerator.NewLine);
+													bBuildByDefault = true;
+
+													UEPlatformProjectGenerator ProjGen = UEPlatformProjectGenerator.GetPlatformProjectGenerator(SolutionConfigCombination.Platform, true);
+													if (MatchingProjectTarget.ProjectDeploys ||
+														((ProjGen != null) && (ProjGen.GetVisualStudioDeploymentEnabled(ProjectPlatform, ProjectConfiguration) == true)))
+													{
+														bDeployByDefault = true;
+													}
 												}
 											}
 										}
+									}
 
-										CombinationsThatWereMatchedToProjects.Add(SolutionConfigCombination);
+									// e.g.  "{4232C52C-680F-4850-8855-DC39419B5E9B}.Debug|iOS.ActiveCfg = iOS_Debug|Win32"
+									string CurProjectGUID = CurProject.ProjectGUID.ToString("B").ToUpperInvariant();
+									VCSolutionFileContent.AppendLine("		{0}.{1}.ActiveCfg = {2}", CurProjectGUID, SolutionConfigCombination.VCSolutionConfigAndPlatformName, ProjectConfigAndPlatformPair);
+									if (bBuildByDefault)
+									{
+										VCSolutionFileContent.AppendLine("		{0}.{1}.Build.0 = {2}", CurProjectGUID, SolutionConfigCombination.VCSolutionConfigAndPlatformName, ProjectConfigAndPlatformPair);
+										if(bDeployByDefault)
+										{
+											VCSolutionFileContent.AppendLine("		{0}.{1}.Deploy.0 = {2}", CurProjectGUID, SolutionConfigCombination.VCSolutionConfigAndPlatformName, ProjectConfigAndPlatformPair);
+										}
 									}
 								}
 							}
 						}
 
-						// Check for problems
-						foreach (VCSolutionConfigCombination SolutionConfigCombination in SolutionConfigCombinations)
-						{
-							if (!CombinationsThatWereMatchedToProjects.Contains(SolutionConfigCombination))
-							{
-								throw new BuildException("Unable to find a ProjectTarget that matches the solution configuration/platform mapping: " + SolutionConfigCombination.Configuration.ToString() + ", " + SolutionConfigCombination.Platform.ToString() + ", " + SolutionConfigCombination.TargetConfigurationName);
-							}
-						}
-						VCSolutionFileContent.Append(
-							"	EndGlobalSection" + ProjectFileGenerator.NewLine);
+						VCSolutionFileContent.AppendLine("	EndGlobalSection");
 					}
 
 
 					// Setup other solution properties
 					{
-						VCSolutionFileContent.Append(
-							"	GlobalSection(SolutionProperties) = preSolution" + ProjectFileGenerator.NewLine);
-
 						// HideSolutionNode sets whether or not the top-level solution entry is completely hidden in the UI.
 						// We don't want that, as we need users to be able to right click on the solution tree item.
-						VCSolutionFileContent.Append(
-							"		HideSolutionNode = FALSE" + ProjectFileGenerator.NewLine);
-
-						VCSolutionFileContent.Append(
-							"	EndGlobalSection" + ProjectFileGenerator.NewLine);
+						VCSolutionFileContent.AppendLine("	GlobalSection(SolutionProperties) = preSolution");
+						VCSolutionFileContent.AppendLine("		HideSolutionNode = FALSE");
+						VCSolutionFileContent.AppendLine("	EndGlobalSection");
 					}
 
 
 
 					// Solution directory hierarchy
 					{
-						VCSolutionFileContent.Append(
-							"	GlobalSection(NestedProjects) = preSolution" + ProjectFileGenerator.NewLine);
+						VCSolutionFileContent.AppendLine("	GlobalSection(NestedProjects) = preSolution");
 
 						// Every entry in this section is in the format "Guid1 = Guid2".  Guid1 is the child project (or solution
 						// filter)'s GUID, and Guid2 is the solution filter directory to parent the child project (or solution
@@ -929,15 +857,13 @@ namespace UnrealBuildTool
 									foreach (MSBuildProjectFile ChildProject in CurFolder.ChildProjects)
 									{
 										//	e.g. "{BF6FB09F-A2A6-468F-BE6F-DEBE07EAD3EA} = {C43B6BB5-3EF0-4784-B896-4099753BCDA9}"
-										LocalVCSolutionFileContent.Append(
-											"		" + ChildProject.ProjectGUID.ToString("B").ToUpperInvariant() + " = " + CurFolderGUIDString + ProjectFileGenerator.NewLine);
+										LocalVCSolutionFileContent.AppendLine("		" + ChildProject.ProjectGUID.ToString("B").ToUpperInvariant() + " = " + CurFolderGUIDString);
 									}
 
 									foreach (VisualStudioSolutionFolder SubFolder in CurFolder.SubFolders)
 									{
 										//	e.g. "{BF6FB09F-A2A6-468F-BE6F-DEBE07EAD3EA} = {C43B6BB5-3EF0-4784-B896-4099753BCDA9}"
-										LocalVCSolutionFileContent.Append(
-											"		" + SubFolder.FolderGUID.ToString("B").ToUpperInvariant() + " = " + CurFolderGUIDString + ProjectFileGenerator.NewLine);
+										LocalVCSolutionFileContent.AppendLine("		" + SubFolder.FolderGUID.ToString("B").ToUpperInvariant() + " = " + CurFolderGUIDString);
 									}
 
 									// Recurse into subfolders
@@ -946,13 +872,11 @@ namespace UnrealBuildTool
 							};
 						FolderProcessorFunction(VCSolutionFileContent, RootFolder.SubFolders);
 
-						VCSolutionFileContent.Append(
-							"	EndGlobalSection" + ProjectFileGenerator.NewLine);
+						VCSolutionFileContent.AppendLine("	EndGlobalSection");
 					}
 				}
 
-				VCSolutionFileContent.Append(
-					"EndGlobal" + ProjectFileGenerator.NewLine);
+				VCSolutionFileContent.AppendLine("EndGlobal");
 			}
 
 
