UBT: Set the permissions on directories containing scripts to be executed as part of remote builds.

diff --git a/Engine/Source/Programs/UnrealBuildTool/ToolChain/RemoteMac.cs b/Engine/Source/Programs/UnrealBuildTool/ToolChain/RemoteMac.cs
--- a/Engine/Source/Programs/UnrealBuildTool/ToolChain/RemoteMac.cs
+++ b/Engine/Source/Programs/UnrealBuildTool/ToolChain/RemoteMac.cs
@@ -236,6 +236,7 @@ namespace UnrealBuildTool
 			BasicRsyncArguments.Add("--compress");
 			BasicRsyncArguments.Add("--verbose");
 			BasicRsyncArguments.Add(String.Format("--rsh=\"{0} -p {1}\"", RsyncAuthentication, ServerPort));
+			BasicRsyncArguments.Add("--chmod=ug=rwX,o=rxX");
 
 			// Build a list of arguments for Rsync filters
 			CommonRsyncArguments = new List<string>(BasicRsyncArguments);
@@ -244,7 +245,6 @@ namespace UnrealBuildTool
 			CommonRsyncArguments.Add("--delete-excluded"); // Delete anything not in the source directory
 			CommonRsyncArguments.Add("--times"); // Preserve modification times
 			CommonRsyncArguments.Add("--prune-empty-dirs"); // Remove empty directories from the file list
-			CommonRsyncArguments.Add("--chmod=ug=rwX,o=rxX");
 
 			// Get the remote base directory
 			StringBuilder Output;
