Disabled the clion ignore file generation on Windows since it isn't supported correctly there

diff --git a/Engine/Source/Programs/UnrealBuildTool/ProjectFiles/CMake/CMakefileGenerator.cs b/Engine/Source/Programs/UnrealBuildTool/ProjectFiles/CMake/CMakefileGenerator.cs
--- a/Engine/Source/Programs/UnrealBuildTool/ProjectFiles/CMake/CMakefileGenerator.cs
+++ b/Engine/Source/Programs/UnrealBuildTool/ProjectFiles/CMake/CMakefileGenerator.cs
@@ -782,7 +782,10 @@ namespace UnrealBuildTool
 
 		protected override bool WriteProjectFiles()
 		{
-			WriteCLionIgnoreDirs();
+			if (BuildHostPlatform.Current.Platform != UnrealTargetPlatform.Win64)
+			{
+				WriteCLionIgnoreDirs();
+			}
 			return WriteCMakeLists();
 		}
 
