Fix runtime dependencies not being enumerated for modules that aren't linked into a binary.

diff --git a/Engine/Source/Programs/UnrealBuildTool/Configuration/UEBuildTarget.cs b/Engine/Source/Programs/UnrealBuildTool/Configuration/UEBuildTarget.cs
--- a/Engine/Source/Programs/UnrealBuildTool/Configuration/UEBuildTarget.cs
+++ b/Engine/Source/Programs/UnrealBuildTool/Configuration/UEBuildTarget.cs
@@ -1184,10 +1184,25 @@ namespace UnrealBuildTool
 		/// Writes a list of all the externally referenced files required to use the precompiled data for this target
 		/// </summary>
 		/// <param name="Location">Path to the dependency list</param>
-		/// <param name="RuntimeDependencySourceFiles">All the source files for runtime dependencies. This is only evaluated at build time.</param>
-		void WriteDependencyList(FileReference Location, IEnumerable<FileReference> RuntimeDependencySourceFiles)
+		/// <param name="RuntimeDependencies">List of all the runtime dependencies for this target</param>
+		/// <param name="RuntimeDependencyTargetFileToSourceFile">Map of runtime dependencies to their location in the source tree, before they are staged</param>
+		void WriteDependencyList(FileReference Location, List<RuntimeDependency> RuntimeDependencies, Dictionary<FileReference, FileReference> RuntimeDependencyTargetFileToSourceFile)
 		{
-			HashSet<FileReference> Files = new HashSet<FileReference>(RuntimeDependencySourceFiles);
+			HashSet<FileReference> Files = new HashSet<FileReference>();
+
+			// Find all the runtime dependency files in their original location
+			foreach(RuntimeDependency RuntimeDependency in RuntimeDependencies)
+			{
+				FileReference SourceFile;
+				if(RuntimeDependencyTargetFileToSourceFile.TryGetValue(RuntimeDependency.Path, out SourceFile))
+				{
+					Files.Add(SourceFile);
+				}
+				else
+				{
+					Files.Add(RuntimeDependency.Path);
+				}
+			}
 
 			// Figure out all the modules referenced by this target. This includes all the modules that are referenced, not just the ones compiled into binaries.
 			HashSet<UEBuildModule> Modules = new HashSet<UEBuildModule>();
@@ -2233,10 +2248,10 @@ namespace UnrealBuildTool
 				return ECompilationResult.Canceled;
 			}
 
-			// Build a list of all the externally files
+			// Build a list of all the files required to build
 			foreach(FileReference DependencyListFileName in Rules.DependencyListFileNames)
 			{
-				WriteDependencyList(DependencyListFileName, RuntimeDependencyTargetFileToSourceFile.Values);
+				WriteDependencyList(DependencyListFileName, RuntimeDependencies, RuntimeDependencyTargetFileToSourceFile);
 			}
 
 			// If we're only generating the manifest, return now
