UBT: Fix remote compilation of foreign projects, where parent directory for engine files may not exist when scripts are uploaded.

diff --git a/Engine/Source/Programs/UnrealBuildTool/ToolChain/RemoteMac.cs b/Engine/Source/Programs/UnrealBuildTool/ToolChain/RemoteMac.cs
--- a/Engine/Source/Programs/UnrealBuildTool/ToolChain/RemoteMac.cs
+++ b/Engine/Source/Programs/UnrealBuildTool/ToolChain/RemoteMac.cs
@@ -645,6 +645,7 @@ namespace UnrealBuildTool
 		void UploadFiles(DirectoryReference LocalDirectory, string RemoteDirectory, FileReference LocalFileList)
 		{
 			List<string> Arguments = new List<string>(BasicRsyncArguments);
+			Arguments.Add(String.Format("--rsync-path=\"mkdir -p {0} && rsync\"", RemoteDirectory));
 			Arguments.Add(String.Format("--files-from=\"{0}\"", GetLocalCygwinPath(LocalFileList)));
 			Arguments.Add(String.Format("\"{0}/\"", GetLocalCygwinPath(LocalDirectory)));
 			Arguments.Add(String.Format("\"{0}@{1}\":'{2}/'", UserName, ServerName, RemoteDirectory));
