Add a dummy output to the invalid configuration, to force it to always run.

diff --git a/Engine/Source/Programs/UnrealBuildTool/ProjectFiles/VisualStudio/VCProject.cs b/Engine/Source/Programs/UnrealBuildTool/ProjectFiles/VisualStudio/VCProject.cs
--- a/Engine/Source/Programs/UnrealBuildTool/ProjectFiles/VisualStudio/VCProject.cs
+++ b/Engine/Source/Programs/UnrealBuildTool/ProjectFiles/VisualStudio/VCProject.cs
@@ -660,6 +660,7 @@ namespace UnrealBuildTool
 			VCProjectFileContent.AppendLine("    <NMakeBuildCommandLine>{0}</NMakeBuildCommandLine>", InvalidMessage);
 			VCProjectFileContent.AppendLine("    <NMakeReBuildCommandLine>{0}</NMakeReBuildCommandLine>", InvalidMessage);
 			VCProjectFileContent.AppendLine("    <NMakeCleanCommandLine>{0}</NMakeCleanCommandLine>", InvalidMessage);
+			VCProjectFileContent.AppendLine("    <NMakeOutput>Invalid Output</NMakeOutput>", InvalidMessage);
 			VCProjectFileContent.AppendLine("  </PropertyGroup>");
 
 			// Write each project configuration
