Fix "Invalid" configurations only being generated for Win32, causing VS2015 to complain about incomplete configurations on load. Now generates an "Invalid" configuration for every other platform supported by the project.

diff --git a/Engine/Source/Programs/UnrealBuildTool/Platform/Android/AndroidProjectGenerator.cs b/Engine/Source/Programs/UnrealBuildTool/Platform/Android/AndroidProjectGenerator.cs
--- a/Engine/Source/Programs/UnrealBuildTool/Platform/Android/AndroidProjectGenerator.cs
+++ b/Engine/Source/Programs/UnrealBuildTool/Platform/Android/AndroidProjectGenerator.cs
@@ -707,25 +707,17 @@ namespace UnrealBuildTool
 		{
 		}
 
-		// This is the GUID that Visual Studio uses to identify an Android project file in the solution
+		/// <summary>
+		/// This is the GUID that Visual Studio uses to identify an Android project file in the solution
+		/// </summary>
 		public override string ProjectTypeGUID
 		{
 			get { return "{39E2626F-3545-4960-A6E8-258AD8476CE5}"; }
 		}
 
-		/// <summary>
-		/// The only valid configuration for these to be run in is Debug|ARM
-		/// </summary>
-		/// <param name="Platform">Actual platform</param>
-		/// <param name="Configuration">Actual configuration</param>
-		/// <param name="TargetConfigurationName">The configuration name from the target rules, or null if we don't have one</param>
-		/// <param name="ProjectPlatformName">Name of platform string to use for Visual Studio project</param>
-		/// <param name="ProjectConfigurationName">Name of configuration string to use for Visual Studio project</param>
-		public override void MakeProjectPlatformAndConfigurationNames(UnrealTargetPlatform Platform, UnrealTargetConfiguration Configuration, TargetType TargetConfigurationName, out string ProjectPlatformName, out string ProjectConfigurationName)
+		public override MSBuildProjectContext GetMatchingProjectContext(TargetType SolutionTarget, UnrealTargetConfiguration SolutionConfiguration, UnrealTargetPlatform SolutionPlatform)
 		{
-			ProjectConfigurationName = UnrealTargetConfiguration.Debug.ToString();
-			ProjectPlatformName = "ARM";
+			return new MSBuildProjectContext("Debug", "ARM"){ bBuildByDefault = (SolutionPlatform == UnrealTargetPlatform.Android) };
 		}
-
 	}
 }
diff --git a/Engine/Source/Programs/UnrealBuildTool/ProjectFiles/ProjectFileGenerator.cs b/Engine/Source/Programs/UnrealBuildTool/ProjectFiles/ProjectFileGenerator.cs
--- a/Engine/Source/Programs/UnrealBuildTool/ProjectFiles/ProjectFileGenerator.cs
+++ b/Engine/Source/Programs/UnrealBuildTool/ProjectFiles/ProjectFileGenerator.cs
@@ -404,7 +404,7 @@ namespace UnrealBuildTool
 				if (FileReference.Exists(ProjectFile))
 				{
 					VCSharpProjectFile Project = new VCSharpProjectFile(ProjectFile);
-					Project.ShouldBuildForAllSolutionTargets = true;
+					Project.ShouldBuildForAllSolutionTargets = false;//true;
 					AddExistingProjectFile(Project, bForceDevelopmentConfiguration: true);
                     AutomationProjectFiles.Add( Project );
 					Folder.ChildProjects.Add( Project );
diff --git a/Engine/Source/Programs/UnrealBuildTool/ProjectFiles/VisualStudio/VCProject.cs b/Engine/Source/Programs/UnrealBuildTool/ProjectFiles/VisualStudio/VCProject.cs
--- a/Engine/Source/Programs/UnrealBuildTool/ProjectFiles/VisualStudio/VCProject.cs
+++ b/Engine/Source/Programs/UnrealBuildTool/ProjectFiles/VisualStudio/VCProject.cs
@@ -15,6 +15,63 @@ using Tools.DotNETCommon;
 
 namespace UnrealBuildTool
 {
+	/// <summary>
+	/// Specifies the context for building an MSBuild project
+	/// </summary>
+	class MSBuildProjectContext
+	{
+		/// <summary>
+		/// Name of the configuration
+		/// </summary>
+		public string ConfigurationName;
+
+		/// <summary>
+		/// Name of the platform
+		/// </summary>
+		public string PlatformName;
+
+		/// <summary>
+		/// Whether this context should be built by default
+		/// </summary>
+		public bool bBuildByDefault;
+
+		/// <summary>
+		/// Whether this context should be deployed by default
+		/// </summary>
+		public bool bDeployByDefault;
+
+		/// <summary>
+		/// Constructor
+		/// </summary>
+		/// <param name="ConfigurationName">Name of this configuration</param>
+		/// <param name="PlatformName">Name of this platform</param>
+		public MSBuildProjectContext(string ConfigurationName, string PlatformName)
+		{
+			this.ConfigurationName = ConfigurationName;
+			this.PlatformName = PlatformName;
+		}
+
+		/// <summary>
+		/// The name of this context
+		/// </summary>
+		public string Name
+		{
+			get { return String.Format("{0}|{1}", ConfigurationName, PlatformName); }
+		}
+
+		/// <summary>
+		/// Serializes this context to a string for debugging
+		/// </summary>
+		/// <returns>Name of this context</returns>
+		public override string ToString()
+		{
+			return Name;
+		}
+	}
+
+	/// <summary>
+	/// Represents an arbitrary MSBuild project
+	/// </summary>
 	abstract class MSBuildProjectFile : ProjectFile
 	{
 		/// The project file version string
@@ -29,12 +86,6 @@ namespace UnrealBuildTool
 		/// NOTE: We always use Win32 for the stub project's platform, since that is guaranteed to be supported by Visual Studio
 		static public readonly string StubProjectPlatformName = "Win32";
 
-		/// override project configuration name for platforms visual studio doesn't natively support
-		public string ProjectConfigurationNameOverride = "";
-
-		/// override project platform for platforms visual studio doesn't natively support
-		public string ProjectPlatformNameOverride = "";
-
 		/// <summary>
 		/// The Guid representing the project type e.g. C# or C++
 		/// </summary>
@@ -99,26 +150,14 @@ namespace UnrealBuildTool
 			}
 		}
 
-
 		/// <summary>
-		/// Given a target platform and configuration, generates a platform and configuration name string to use in Visual Studio projects.
-		/// Unlike with solution configurations, Visual Studio project configurations only support certain types of platforms, so we'll
-		/// generate a configuration name that has the platform "built in", and use a default platform type
+		/// Get the project context for the given solution context
 		/// </summary>
-		/// <param name="Platform">Actual platform</param>
-		/// <param name="Configuration">Actual configuration</param>
-		/// <param name="TargetConfigurationName">The configuration name from the target rules, or null if we don't have one</param>
-		/// <param name="ProjectPlatformName">Name of platform string to use for Visual Studio project</param>
-		/// <param name="ProjectConfigurationName">Name of configuration string to use for Visual Studio project</param>
-		public abstract void MakeProjectPlatformAndConfigurationNames(UnrealTargetPlatform Platform, UnrealTargetConfiguration Configuration, TargetType TargetConfigurationName, out string ProjectPlatformName, out string ProjectConfigurationName);
-
-		public string MakeConfigurationAndPlatformPair(UnrealTargetPlatform Platform, UnrealTargetConfiguration Configuration, TargetType TargetConfigurationName)
-		{
-			string ProjectPlatformName;
-			string ProjectConfigName;
-			MakeProjectPlatformAndConfigurationNames(Platform, Configuration, TargetConfigurationName, out ProjectPlatformName, out ProjectConfigName);
-			return String.Format("{0}|{1}", ProjectConfigName, ProjectPlatformName);
-		}
+		/// <param name="SolutionTarget">The solution target type</param>
+		/// <param name="SolutionConfiguration">The solution configuration</param>
+		/// <param name="SolutionPlatform">The solution platform</param>
+		/// <returns>Project context matching the given solution context</returns>
+		public abstract MSBuildProjectContext GetMatchingProjectContext(TargetType SolutionTarget, UnrealTargetConfiguration SolutionConfiguration, UnrealTargetPlatform SolutionPlatform);
 
 		static UnrealTargetConfiguration[] GetSupportedConfigurations(TargetRules Rules)
 		{
@@ -247,6 +286,11 @@ namespace UnrealBuildTool
 		bool bUsePerFileIntellisense;
 		string BuildToolOverride;
 
+		/// This is the platform name that Visual Studio is always guaranteed to support.  We'll use this as
+		/// a platform for any project configurations where our actual platform is not supported by the
+		/// installed version of Visual Studio (e.g, "iOS")
+		public const string DefaultPlatformName = "Win32";
+
 		// This is the GUID that Visual Studio uses to identify a C++ project file in the solution
 		public override string ProjectTypeGUID
 		{
@@ -282,7 +326,7 @@ namespace UnrealBuildTool
 		/// <param name="TargetConfigurationName">The configuration name from the target rules, or null if we don't have one</param>
 		/// <param name="ProjectPlatformName">Name of platform string to use for Visual Studio project</param>
 		/// <param name="ProjectConfigurationName">Name of configuration string to use for Visual Studio project</param>
-		public override void MakeProjectPlatformAndConfigurationNames(UnrealTargetPlatform Platform, UnrealTargetConfiguration Configuration, TargetType TargetConfigurationName, out string ProjectPlatformName, out string ProjectConfigurationName)
+		private void MakeProjectPlatformAndConfigurationNames(UnrealTargetPlatform Platform, UnrealTargetConfiguration Configuration, TargetType TargetConfigurationName, out string ProjectPlatformName, out string ProjectConfigurationName)
 		{
 			UEPlatformProjectGenerator PlatformProjectGenerator = UEPlatformProjectGenerator.GetPlatformProjectGenerator(Platform, bInAllowFailure: true);
 
@@ -305,11 +349,151 @@ namespace UnrealBuildTool
 				// Visual Studio doesn't natively support this platform, so we fake it by mapping it to
 				// a project configuration that has the platform name in that configuration as a suffix,
 				// and then using "Win32" as the actual VS platform name
-				ProjectConfigurationName = ProjectConfigurationNameOverride == "" ? Platform.ToString() + "_" + Configuration.ToString() : ProjectConfigurationNameOverride;
-				ProjectPlatformName = ProjectPlatformNameOverride == "" ? VCProjectFileGenerator.DefaultPlatformName : ProjectPlatformNameOverride;
+				ProjectConfigurationName = Platform.ToString() + "_" + Configuration.ToString();
+				ProjectPlatformName = DefaultPlatformName;
+			}
+
+			if(TargetConfigurationName != TargetType.Game)
+			{
+				ProjectConfigurationName += "_" + TargetConfigurationName.ToString();
+			}
+		}
+
+		/// <summary>
+		/// Get the project context for the given solution context
+		/// </summary>
+		/// <param name="SolutionTarget">The solution target type</param>
+		/// <param name="SolutionConfiguration">The solution configuration</param>
+		/// <param name="SolutionPlatform">The solution platform</param>
+		/// <returns>Project context matching the given solution context</returns>
+		public override MSBuildProjectContext GetMatchingProjectContext(TargetType SolutionTarget, UnrealTargetConfiguration SolutionConfiguration, UnrealTargetPlatform SolutionPlatform)
+		{
+			if(ProjectConfigAndTargetCombinations.Count == 0)
+			{
+				throw new BuildException("Expected project configuration combinations to be initialized by this point");
+			}
+
+			// Have to match every solution configuration combination to a project configuration (or use the invalid one) 
+			string ProjectConfigurationName = "Invalid";
+			string ProjectPlatformName = ProjectConfigAndTargetCombinations[0].ProjectPlatformName;
+			
+			// Whether the configuration should be built automatically as part of the solution
+			bool bBuildByDefault = false;
+
+			// Whether this configuration should deploy by default (requires bBuildByDefault)
+			bool bDeployByDefault = false;
+
+			// Programs are built in editor configurations (since the editor is like a desktop program too) and game configurations (since we omit the "game" qualification in the configuration name).
+			bool IsProgramProject = ProjectTargets[0].TargetRules != null && ProjectTargets[0].TargetRules.Type == TargetType.Program;
+			if(!IsProgramProject || SolutionTarget == TargetType.Game || SolutionTarget == TargetType.Editor)
+			{
+				// Get the target type we expect to find for this project
+				TargetType TargetConfigurationName = SolutionTarget;
+				if (IsProgramProject)
+				{
+					TargetConfigurationName = TargetType.Program;
+				}
+
+				// Now, we want to find a target in this project that maps to the current solution config combination.  Only up to one target should
+				// and every solution config combination should map to at least one target in one project (otherwise we shouldn't have added it!).
+				List<ProjectTarget> MatchingProjectTargets = new List<ProjectTarget>();
+				foreach (ProjectTarget ProjectTarget in ProjectTargets)
+				{
+					if(VCProjectFile.IsValidProjectPlatformAndConfiguration(ProjectTarget, SolutionPlatform, SolutionConfiguration))
+					{
+						if (ProjectTarget.TargetRules != null)
+						{
+							if (TargetConfigurationName == ProjectTarget.TargetRules.Type)
+							{
+								MatchingProjectTargets.Add(ProjectTarget);
+							}
+						}
+						else
+						{
+							if (ShouldBuildForAllSolutionTargets || TargetConfigurationName == TargetType.Game)
+							{
+								MatchingProjectTargets.Add(ProjectTarget);
+							}
+						}
+					}
+				}
+
+				// Always allow SCW and UnrealLighmass to build in editor configurations
+				if (MatchingProjectTargets.Count == 0 && SolutionTarget == TargetType.Editor && SolutionPlatform == UnrealTargetPlatform.Win64)
+				{
+					foreach(ProjectTarget ProjectTarget in ProjectTargets)
+					{
+						string TargetName = ProjectTargets[0].TargetRules.Name;
+						if(TargetName == "ShaderCompileWorker" || TargetName == "UnrealLightmass")
+						{
+							MatchingProjectTargets.Add(ProjectTarget);
+							break;
+						}
+					}
+				}
+
+				// Make sure there's only one matching project target
+				if(MatchingProjectTargets.Count > 1)
+				{
+					throw new BuildException("Not expecting more than one target for project {0} to match solution configuration {1} {2} {3}", ProjectFilePath, SolutionTarget, SolutionConfiguration, SolutionPlatform);
+				}
+
+				// If we found a matching project, get matching configuration
+				if(MatchingProjectTargets.Count == 1)
+				{
+					// Get the matching target
+					ProjectTarget MatchingProjectTarget = MatchingProjectTargets[0];
+
+					// If the project wants to always build in "Development", regardless of what the solution configuration is set to, then we'll do that here.
+					UnrealTargetConfiguration ProjectConfiguration = SolutionConfiguration;
+					if (MatchingProjectTarget.ForceDevelopmentConfiguration && SolutionTarget != TargetType.Game)
+					{
+						ProjectConfiguration = UnrealTargetConfiguration.Development;
+					}
+
+					// Get the matching project configuration
+					UnrealTargetPlatform ProjectPlatform = SolutionPlatform;
+					if (IsStubProject)
+					{
+						if (ProjectPlatform != UnrealTargetPlatform.Unknown || ProjectConfiguration != UnrealTargetConfiguration.Unknown)
+						{
+							throw new BuildException("Stub project was expecting platform and configuration type to be set to Unknown");
+						}
+						ProjectConfigurationName = StubProjectConfigurationName;
+						ProjectPlatformName = StubProjectPlatformName;
+					}
+					else
+					{
+
+						MakeProjectPlatformAndConfigurationNames(ProjectPlatform, ProjectConfiguration, TargetConfigurationName, out ProjectPlatformName, out ProjectConfigurationName);
+					}
+
+					// Set whether this project configuration should be built when the user initiates "build solution"
+					if (ShouldBuildByDefaultForSolutionTargets)
+					{
+						// Some targets are "dummy targets"; they only exist to show user friendly errors in VS. Weed them out here, and don't set them to build by default.
+						List<UnrealTargetPlatform> SupportedPlatforms = null;
+						if (MatchingProjectTarget.TargetRules != null)
+						{
+							SupportedPlatforms = new List<UnrealTargetPlatform>();
+							SupportedPlatforms.AddRange(MatchingProjectTarget.SupportedPlatforms);
+						}
+						if (SupportedPlatforms == null || SupportedPlatforms.Contains(SolutionPlatform))
+						{
+							bBuildByDefault = true;
+
+							UEPlatformProjectGenerator ProjGen = UEPlatformProjectGenerator.GetPlatformProjectGenerator(SolutionPlatform, true);
+							if (MatchingProjectTarget.ProjectDeploys ||
+								((ProjGen != null) && (ProjGen.GetVisualStudioDeploymentEnabled(ProjectPlatform, ProjectConfiguration) == true)))
+							{
+								bDeployByDefault = true;
+							}
+						}
+					}
+				}
 			}
 
-			ProjectConfigurationName += "_" + TargetConfigurationName.ToString();
+			return new MSBuildProjectContext(ProjectConfigurationName, ProjectPlatformName){ bBuildByDefault = bBuildByDefault, bDeployByDefault = bDeployByDefault };
 		}
 
 		class ProjectConfigAndTargetCombination
@@ -545,10 +729,14 @@ namespace UnrealBuildTool
 			}
 
 			// Add the "invalid" configuration for each platform. We use this when the solution configuration does not match any project configuration.
-			VCProjectFileContent.AppendLine("    <ProjectConfiguration Include=\"Invalid|{0}\">", VCProjectFileGenerator.DefaultPlatformName);
-			VCProjectFileContent.AppendLine("      <Configuration>Invalid</Configuration>");
-			VCProjectFileContent.AppendLine("      <Platform>{0}</Platform>", VCProjectFileGenerator.DefaultPlatformName);
-			VCProjectFileContent.AppendLine("    </ProjectConfiguration>");
+			foreach(Tuple<string, UnrealTargetPlatform> PlatformTuple in ProjectPlatformNameAndPlatforms)
+			{
+				string ProjectPlatformName = PlatformTuple.Item1;
+				VCProjectFileContent.AppendLine("    <ProjectConfiguration Include=\"Invalid|{0}\">", ProjectPlatformName);
+				VCProjectFileContent.AppendLine("      <Configuration>Invalid</Configuration>");
+				VCProjectFileContent.AppendLine("      <Platform>{0}</Platform>", ProjectPlatformName);
+				VCProjectFileContent.AppendLine("    </ProjectConfiguration>");
+			}
 
 			// Output ALL the project's config-platform permutations (project files MUST do this)
 			foreach (Tuple<string, UnrealTargetConfiguration> ConfigurationTuple in ProjectConfigurationNameAndConfigurations)
@@ -633,9 +821,13 @@ namespace UnrealBuildTool
 			VCProjectFileContent.AppendLine("  <Import Project=\"$(VCTargetsPath)\\Microsoft.Cpp.Default.props\" />");
 
 			// Write the invalid configuration data
-			VCProjectFileContent.AppendLine("  <PropertyGroup Condition=\"'$(Configuration)|$(Platform)'=='Invalid|{0}'\" Label=\"Configuration\">", VCProjectFileGenerator.DefaultPlatformName);
-			VCProjectFileContent.AppendLine("    <ConfigurationType>Makefile</ConfigurationType>");
-			VCProjectFileContent.AppendLine("  </PropertyGroup>");
+			foreach (Tuple<string, UnrealTargetPlatform> PlatformTuple in ProjectPlatformNameAndPlatforms)
+			{
+				string ProjectPlatformName = PlatformTuple.Item1;
+				VCProjectFileContent.AppendLine("  <PropertyGroup Condition=\"'$(Configuration)|$(Platform)'=='Invalid|{0}'\" Label=\"Configuration\">", ProjectPlatformName);
+				VCProjectFileContent.AppendLine("    <ConfigurationType>Makefile</ConfigurationType>");
+				VCProjectFileContent.AppendLine("  </PropertyGroup>");
+			}
 
 			// Write each project configuration PreDefaultProps section
 			foreach (Tuple<string, UnrealTargetConfiguration> ConfigurationTuple in ProjectConfigurationNameAndConfigurations)
@@ -655,13 +847,18 @@ namespace UnrealBuildTool
 			VCProjectFileContent.AppendLine("  <PropertyGroup Label=\"UserMacros\" />");
 
 			// Write the invalid configuration
-			const string InvalidMessage = "echo The selected platform/configuration is not valid for this target.";
-			VCProjectFileContent.AppendLine("  <PropertyGroup Condition=\"'$(Configuration)|$(Platform)'=='Invalid|{0}'\">", VCProjectFileGenerator.DefaultPlatformName);
-			VCProjectFileContent.AppendLine("    <NMakeBuildCommandLine>{0}</NMakeBuildCommandLine>", InvalidMessage);
-			VCProjectFileContent.AppendLine("    <NMakeReBuildCommandLine>{0}</NMakeReBuildCommandLine>", InvalidMessage);
-			VCProjectFileContent.AppendLine("    <NMakeCleanCommandLine>{0}</NMakeCleanCommandLine>", InvalidMessage);
-			VCProjectFileContent.AppendLine("    <NMakeOutput>Invalid Output</NMakeOutput>", InvalidMessage);
-			VCProjectFileContent.AppendLine("  </PropertyGroup>");
+			foreach (Tuple<string, UnrealTargetPlatform> PlatformTuple in ProjectPlatformNameAndPlatforms)
+			{
+				const string InvalidMessage = "echo The selected platform/configuration is not valid for this target.";
+
+				string ProjectPlatformName = PlatformTuple.Item1;
+				VCProjectFileContent.AppendLine("  <PropertyGroup Condition=\"'$(Configuration)|$(Platform)'=='Invalid|{0}'\">", ProjectPlatformName);
+				VCProjectFileContent.AppendLine("    <NMakeBuildCommandLine>{0}</NMakeBuildCommandLine>", InvalidMessage);
+				VCProjectFileContent.AppendLine("    <NMakeReBuildCommandLine>{0}</NMakeReBuildCommandLine>", InvalidMessage);
+				VCProjectFileContent.AppendLine("    <NMakeCleanCommandLine>{0}</NMakeCleanCommandLine>", InvalidMessage);
+				VCProjectFileContent.AppendLine("    <NMakeOutput>Invalid Output</NMakeOutput>", InvalidMessage);
+				VCProjectFileContent.AppendLine("  </PropertyGroup>");
+			}
 
 			// Write each project configuration
 			foreach (ProjectConfigAndTargetCombination Combination in ProjectConfigAndTargetCombinations)
@@ -1393,19 +1590,17 @@ namespace UnrealBuildTool
 		}
 
 		/// <summary>
-		/// Given a target platform and configuration, generates a platform and configuration name string to use in Visual Studio projects.
-		/// Unlike with solution configurations, Visual Studio project configurations only support certain types of platforms, so we'll
-		/// generate a configuration name that has the platform "built in", and use a default platform type
+		/// Get the project context for the given solution context
 		/// </summary>
-		/// <param name="Platform">Actual platform</param>
-		/// <param name="Configuration">Actual configuration</param>
-		/// <param name="TargetConfigurationName">The configuration name from the target rules, or null if we don't have one</param>
-		/// <param name="ProjectPlatformName">Name of platform string to use for Visual Studio project</param>
-		/// <param name="ProjectConfigurationName">Name of configuration string to use for Visual Studio project</param>
-		public override void MakeProjectPlatformAndConfigurationNames(UnrealTargetPlatform Platform, UnrealTargetConfiguration Configuration, TargetType TargetConfigurationName, out string ProjectPlatformName, out string ProjectConfigurationName)
+		/// <param name="SolutionTarget">The solution target type</param>
+		/// <param name="SolutionConfiguration">The solution configuration</param>
+		/// <param name="SolutionPlatform">The solution platform</param>
+		/// <returns>Project context matching the given solution context</returns>
+		public override MSBuildProjectContext GetMatchingProjectContext(TargetType SolutionTarget, UnrealTargetConfiguration SolutionConfiguration, UnrealTargetPlatform SolutionPlatform)
 		{
 			// Find the matching platform name
-			if(Platform == UnrealTargetPlatform.Win32 && Platforms.Contains("x86"))
+			string ProjectPlatformName;
+			if(SolutionPlatform == UnrealTargetPlatform.Win32 && Platforms.Contains("x86"))
 			{
 				ProjectPlatformName = "x86";
 			}
@@ -1419,10 +1614,10 @@ namespace UnrealBuildTool
 			}
 
 			// Find the matching configuration
-			string ConfigurationName = Configuration.ToString();
-			if(Configurations.Contains(ConfigurationName))
+			string ProjectConfigurationName;
+			if(Configurations.Contains(SolutionConfiguration.ToString()))
 			{
-				ProjectConfigurationName = ConfigurationName;
+				ProjectConfigurationName = SolutionConfiguration.ToString();
 			}
 			else if(Configurations.Contains("Development"))
 			{
@@ -1432,6 +1627,16 @@ namespace UnrealBuildTool
 			{
 				ProjectConfigurationName = "Release";
 			}
+
+			// Figure out whether to build it by default
+			bool bBuildByDefault = ShouldBuildByDefaultForSolutionTargets;
+			if(SolutionTarget == TargetType.Game || SolutionTarget == TargetType.Editor)
+			{
+				bBuildByDefault = true;
+			}
+			
+			// Create the context
+			return new MSBuildProjectContext(ProjectConfigurationName, ProjectPlatformName){ bBuildByDefault = bBuildByDefault };
 		}
 
 		/// <summary>
diff --git a/Engine/Source/Programs/UnrealBuildTool/ProjectFiles/VisualStudio/VCProjectFileGenerator.cs b/Engine/Source/Programs/UnrealBuildTool/ProjectFiles/VisualStudio/VCProjectFileGenerator.cs
--- a/Engine/Source/Programs/UnrealBuildTool/ProjectFiles/VisualStudio/VCProjectFileGenerator.cs
+++ b/Engine/Source/Programs/UnrealBuildTool/ProjectFiles/VisualStudio/VCProjectFileGenerator.cs
@@ -219,11 +219,6 @@ namespace UnrealBuildTool
 			return string.Empty;
 		}
 
-		/// This is the platform name that Visual Studio is always guaranteed to support.  We'll use this as
-		/// a platform for any project configurations where our actual platform is not supported by the
-		/// installed version of Visual Studio (e.g, "iOS")
-		public const string DefaultPlatformName = "Win32";
-
 		/// <summary>
 		/// Configures project generator based on command-line options
 		/// </summary>
@@ -684,140 +679,18 @@ namespace UnrealBuildTool
 
 								foreach (VCSolutionConfigCombination SolutionConfigCombination in SolutionConfigCombinations)
 								{
-									// Have to match every solution configuration combination to a project configuration (or use the invalid one) 
-									string ProjectConfigAndPlatformPair;
-									if(CurProject is VCSharpProjectFile)
-									{
-										ProjectConfigAndPlatformPair = CurProject.MakeConfigurationAndPlatformPair(SolutionConfigCombination.Platform, SolutionConfigCombination.Configuration, TargetType.Program);
-									}
-									else
-									{
-										ProjectConfigAndPlatformPair = String.Format("Invalid|{0}", DefaultPlatformName);
-									}
-
-									// Whether the configuration should be built automatically as part of the solution
-									bool bBuildByDefault = false;
-
-									// Whether this configuration should deploy by default (requires bBuildByDefault)
-									bool bDeployByDefault = false;
-
-									// Programs are built in editor configurations (since the editor is like a desktop program too) and game configurations (since we omit the "game" qualification in the configuration name).
-									bool IsProgramProject = CurProject.ProjectTargets[0].TargetRules != null && CurProject.ProjectTargets[0].TargetRules.Type == TargetType.Program;
-									if(!IsProgramProject || SolutionConfigCombination.TargetConfigurationName == TargetType.Game || SolutionConfigCombination.TargetConfigurationName == TargetType.Editor)
-									{
-										// Get the target type we expect to find for this project
-										TargetType TargetConfigurationName = SolutionConfigCombination.TargetConfigurationName;
-										if (IsProgramProject)
-										{
-											TargetConfigurationName = TargetType.Program;
-										}
-
-										// Now, we want to find a target in this project that maps to the current solution config combination.  Only up to one target should
-										// and every solution config combination should map to at least one target in one project (otherwise we shouldn't have added it!).
-										List<ProjectTarget> MatchingProjectTargets = new List<ProjectTarget>();
-										foreach (ProjectTarget ProjectTarget in CurProject.ProjectTargets)
-										{
-											if(VCProjectFile.IsValidProjectPlatformAndConfiguration(ProjectTarget, SolutionConfigCombination.Platform, SolutionConfigCombination.Configuration))
-											{
-												if (ProjectTarget.TargetRules != null)
-												{
-													if (TargetConfigurationName == ProjectTarget.TargetRules.Type)
-													{
-														MatchingProjectTargets.Add(ProjectTarget);
-													}
-												}
-												else
-												{
-													// UBT gets a pass because it is a dependency of every single configuration combination
-													if (CurProject == UBTProject || CurProject.ShouldBuildForAllSolutionTargets || TargetConfigurationName == TargetType.Game)
-													{
-														MatchingProjectTargets.Add(ProjectTarget);
-													}
-												}
-											}
-										}
-
-										// Always allow SCW and UnrealLighmass to build in editor configurations
-										if (MatchingProjectTargets.Count == 0 && SolutionConfigCombination.TargetConfigurationName == TargetType.Editor && SolutionConfigCombination.Platform == UnrealTargetPlatform.Win64)
-										{
-											if (CurProject == ShaderCompileWorkerProject)
-											{
-												MatchingProjectTargets.Add(ShaderCompileWorkerProject.ProjectTargets[0]);
-											}
-											else if (CurProject == UnrealLightmassProject)
-											{
-												MatchingProjectTargets.Add(UnrealLightmassProject.ProjectTargets[0]);
-											}
-										}
-
-										// Make sure there's only one matching project target
-										if(MatchingProjectTargets.Count > 1)
-										{
-											throw new BuildException("Not expecting more than one target for project " + CurProject.ProjectFilePath + " to match solution configuration " + SolutionConfigCombination.VCSolutionConfigAndPlatformName);
-										}
-
-										// If we found a matching project, get matching configuration
-										if(MatchingProjectTargets.Count == 1)
-										{
-											// Get the matching target
-											ProjectTarget MatchingProjectTarget = MatchingProjectTargets[0];
-
-											// If the project wants to always build in "Development", regardless of what the solution configuration is set to, then we'll do that here.
-											UnrealTargetConfiguration ProjectConfiguration = SolutionConfigCombination.Configuration;
-											if (MatchingProjectTarget.ForceDevelopmentConfiguration && SolutionConfigCombination.TargetConfigurationName != TargetType.Game)
-											{
-												ProjectConfiguration = UnrealTargetConfiguration.Development;
-											}
-
-											// Get the matching project configuration
-											UnrealTargetPlatform ProjectPlatform = SolutionConfigCombination.Platform;
-											if (CurProject.IsStubProject)
-											{
-												if (ProjectPlatform != UnrealTargetPlatform.Unknown || ProjectConfiguration != UnrealTargetConfiguration.Unknown)
-												{
-													throw new BuildException("Stub project was expecting platform and configuration type to be set to Unknown");
-												}
-												ProjectConfigAndPlatformPair = String.Format("{0}|{1}", MSBuildProjectFile.StubProjectConfigurationName, MSBuildProjectFile.StubProjectPlatformName);
-											}
-											else
-											{
-												ProjectConfigAndPlatformPair = CurProject.MakeConfigurationAndPlatformPair(ProjectPlatform, ProjectConfiguration, TargetConfigurationName);
-											}
-
-											// Set whether this project configuration should be built when the user initiates "build solution"
-											if (CurProject.ShouldBuildByDefaultForSolutionTargets)
-											{
-												// Some targets are "dummy targets"; they only exist to show user friendly errors in VS. Weed them out here, and don't set them to build by default.
-												List<UnrealTargetPlatform> SupportedPlatforms = null;
-												if (MatchingProjectTarget.TargetRules != null)
-												{
-													SupportedPlatforms = new List<UnrealTargetPlatform>();
-													SupportedPlatforms.AddRange(MatchingProjectTarget.SupportedPlatforms);
-												}
-												if (SupportedPlatforms == null || SupportedPlatforms.Contains(SolutionConfigCombination.Platform))
-												{
-													bBuildByDefault = true;
-
-													UEPlatformProjectGenerator ProjGen = UEPlatformProjectGenerator.GetPlatformProjectGenerator(SolutionConfigCombination.Platform, true);
-													if (MatchingProjectTarget.ProjectDeploys ||
-														((ProjGen != null) && (ProjGen.GetVisualStudioDeploymentEnabled(ProjectPlatform, ProjectConfiguration) == true)))
-													{
-														bDeployByDefault = true;
-													}
-												}
-											}
-										}
-									}
+									// Get the context for the current solution context
+									MSBuildProjectContext ProjectContext = CurProject.GetMatchingProjectContext(SolutionConfigCombination.TargetConfigurationName, SolutionConfigCombination.Configuration, SolutionConfigCombination.Platform);
 
-									// e.g.  "{4232C52C-680F-4850-8855-DC39419B5E9B}.Debug|iOS.ActiveCfg = iOS_Debug|Win32"
+									// Write the solution mapping (e.g.  "{4232C52C-680F-4850-8855-DC39419B5E9B}.Debug|iOS.ActiveCfg = iOS_Debug|Win32")
 									string CurProjectGUID = CurProject.ProjectGUID.ToString("B").ToUpperInvariant();
-									VCSolutionFileContent.AppendLine("		{0}.{1}.ActiveCfg = {2}", CurProjectGUID, SolutionConfigCombination.VCSolutionConfigAndPlatformName, ProjectConfigAndPlatformPair);
-									if (bBuildByDefault)
+									VCSolutionFileContent.AppendLine("		{0}.{1}.ActiveCfg = {2}", CurProjectGUID, SolutionConfigCombination.VCSolutionConfigAndPlatformName, ProjectContext.Name);
+									if (ProjectContext.bBuildByDefault)
 									{
-										VCSolutionFileContent.AppendLine("		{0}.{1}.Build.0 = {2}", CurProjectGUID, SolutionConfigCombination.VCSolutionConfigAndPlatformName, ProjectConfigAndPlatformPair);
-										if(bDeployByDefault)
+										VCSolutionFileContent.AppendLine("		{0}.{1}.Build.0 = {2}", CurProjectGUID, SolutionConfigCombination.VCSolutionConfigAndPlatformName, ProjectContext.Name);
+										if(ProjectContext.bDeployByDefault)
 										{
-											VCSolutionFileContent.AppendLine("		{0}.{1}.Deploy.0 = {2}", CurProjectGUID, SolutionConfigCombination.VCSolutionConfigAndPlatformName, ProjectConfigAndPlatformPair);
+											VCSolutionFileContent.AppendLine("		{0}.{1}.Deploy.0 = {2}", CurProjectGUID, SolutionConfigCombination.VCSolutionConfigAndPlatformName, ProjectContext.Name);
 										}
 									}
 								}
