Fix "Invalid" folder being created when building invalid platform/configuration.

Also includes temporary fix for exception when generating project files.

diff --git a/Engine/Source/Programs/UnrealBuildTool/ProjectFiles/VisualStudio/VCProject.cs b/Engine/Source/Programs/UnrealBuildTool/ProjectFiles/VisualStudio/VCProject.cs
--- a/Engine/Source/Programs/UnrealBuildTool/ProjectFiles/VisualStudio/VCProject.cs
+++ b/Engine/Source/Programs/UnrealBuildTool/ProjectFiles/VisualStudio/VCProject.cs
@@ -370,12 +370,12 @@ namespace UnrealBuildTool
 		{
 			if(ProjectConfigAndTargetCombinations.Count == 0)
 			{
-				throw new BuildException("Expected project configuration combinations to be initialized by this point");
+//				throw new BuildException("Expected project configuration combinations to be initialized by this point");
 			}
 
 			// Have to match every solution configuration combination to a project configuration (or use the invalid one) 
 			string ProjectConfigurationName = "Invalid";
-			string ProjectPlatformName = ProjectConfigAndTargetCombinations[0].ProjectPlatformName;
+			string ProjectPlatformName = "Win32";//ProjectConfigAndTargetCombinations[0].ProjectPlatformName;
 			
 			// Whether the configuration should be built automatically as part of the solution
 			bool bBuildByDefault = false;
@@ -852,11 +852,15 @@ namespace UnrealBuildTool
 				const string InvalidMessage = "echo The selected platform/configuration is not valid for this target.";
 
 				string ProjectPlatformName = PlatformTuple.Item1;
+				string ProjectRelativeUnusedDirectory = NormalizeProjectPath(DirectoryReference.Combine(UnrealBuildTool.EngineDirectory, "Intermediate", "Build", "Unused"));
+
 				VCProjectFileContent.AppendLine("  <PropertyGroup Condition=\"'$(Configuration)|$(Platform)'=='Invalid|{0}'\">", ProjectPlatformName);
 				VCProjectFileContent.AppendLine("    <NMakeBuildCommandLine>{0}</NMakeBuildCommandLine>", InvalidMessage);
 				VCProjectFileContent.AppendLine("    <NMakeReBuildCommandLine>{0}</NMakeReBuildCommandLine>", InvalidMessage);
 				VCProjectFileContent.AppendLine("    <NMakeCleanCommandLine>{0}</NMakeCleanCommandLine>", InvalidMessage);
 				VCProjectFileContent.AppendLine("    <NMakeOutput>Invalid Output</NMakeOutput>", InvalidMessage);
+				VCProjectFileContent.AppendLine("    <OutDir>{0}{1}</OutDir>", ProjectRelativeUnusedDirectory, Path.DirectorySeparatorChar);
+				VCProjectFileContent.AppendLine("    <IntDir>{0}{1}</IntDir>", ProjectRelativeUnusedDirectory, Path.DirectorySeparatorChar);
 				VCProjectFileContent.AppendLine("  </PropertyGroup>");
 			}
 
