Proper fix for exception when generating project files, if the user does not have any valid platforms for building a target.

diff --git a/Engine/Source/Programs/UnrealBuildTool/ProjectFiles/VisualStudio/VCProject.cs b/Engine/Source/Programs/UnrealBuildTool/ProjectFiles/VisualStudio/VCProject.cs
--- a/Engine/Source/Programs/UnrealBuildTool/ProjectFiles/VisualStudio/VCProject.cs
+++ b/Engine/Source/Programs/UnrealBuildTool/ProjectFiles/VisualStudio/VCProject.cs
@@ -368,15 +368,12 @@ namespace UnrealBuildTool
 		/// <returns>Project context matching the given solution context</returns>
 		public override MSBuildProjectContext GetMatchingProjectContext(TargetType SolutionTarget, UnrealTargetConfiguration SolutionConfiguration, UnrealTargetPlatform SolutionPlatform)
 		{
-			if(ProjectConfigAndTargetCombinations.Count == 0)
-			{
-//				throw new BuildException("Expected project configuration combinations to be initialized by this point");
-			}
-
 			// Have to match every solution configuration combination to a project configuration (or use the invalid one) 
 			string ProjectConfigurationName = "Invalid";
-			string ProjectPlatformName = "Win32";//ProjectConfigAndTargetCombinations[0].ProjectPlatformName;
-			
+
+			// Get the default platform. If there were not valid platforms for this project, just use one that will always be available in VS.
+			string ProjectPlatformName = InvalidConfigPlatformNames[0];
+
 			// Whether the configuration should be built automatically as part of the solution
 			bool bBuildByDefault = false;
 
@@ -535,64 +532,75 @@ namespace UnrealBuildTool
 			}
 		}
 
-		List<ProjectConfigAndTargetCombination> ProjectConfigAndTargetCombinations = new List<ProjectConfigAndTargetCombination>();
+		List<string> InvalidConfigPlatformNames;
+		List<ProjectConfigAndTargetCombination> ProjectConfigAndTargetCombinations;
 
 		private void BuildProjectConfigAndTargetCombinations(List<UnrealTargetPlatform> InPlatforms, List<UnrealTargetConfiguration> InConfigurations)
 		{
 			//no need to do this more than once
-			if(ProjectConfigAndTargetCombinations.Count > 0)
+			if(ProjectConfigAndTargetCombinations == null)
 			{
-				return;
-			}
-
-			// Build up a list of platforms and configurations this project will support.  In this list, Unknown simply
-			// means that we should use the default "stub" project platform and configuration name.
+				// Build up a list of platforms and configurations this project will support.  In this list, Unknown simply
+				// means that we should use the default "stub" project platform and configuration name.
 
-			// If this is a "stub" project, then only add a single configuration to the project
-			if (IsStubProject)
-			{
-				ProjectConfigAndTargetCombination StubCombination = new ProjectConfigAndTargetCombination(UnrealTargetPlatform.Unknown, UnrealTargetConfiguration.Unknown, StubProjectPlatformName, StubProjectConfigurationName, null);
-				ProjectConfigAndTargetCombinations.Add(StubCombination);
-			}
-			else
-			{
-				// Figure out all the desired configurations
-				foreach (UnrealTargetConfiguration Configuration in InConfigurations)
+				// If this is a "stub" project, then only add a single configuration to the project
+				ProjectConfigAndTargetCombinations = new List<ProjectConfigAndTargetCombination>();
+				if (IsStubProject)
 				{
-					[REDACTED_EMAIL]: Put this in a commonly accessible place?
-					if (InstalledPlatformInfo.IsValidConfiguration(Configuration, EProjectType.Code) == false)
-					{
-						continue;
-					}
-					foreach (UnrealTargetPlatform Platform in InPlatforms)
+					ProjectConfigAndTargetCombination StubCombination = new ProjectConfigAndTargetCombination(UnrealTargetPlatform.Unknown, UnrealTargetConfiguration.Unknown, StubProjectPlatformName, StubProjectConfigurationName, null);
+					ProjectConfigAndTargetCombinations.Add(StubCombination);
+				}
+				else
+				{
+					// Figure out all the desired configurations
+					foreach (UnrealTargetConfiguration Configuration in InConfigurations)
 					{
-						if (InstalledPlatformInfo.IsValidPlatform(Platform, EProjectType.Code) == false)
+						[REDACTED_EMAIL]: Put this in a commonly accessible place?
+						if (InstalledPlatformInfo.IsValidConfiguration(Configuration, EProjectType.Code) == false)
 						{
 							continue;
 						}
-						UEBuildPlatform BuildPlatform = UEBuildPlatform.GetBuildPlatform(Platform, true);
-						if ((BuildPlatform != null) && (BuildPlatform.HasRequiredSDKsInstalled() == SDKStatus.Valid))
+						foreach (UnrealTargetPlatform Platform in InPlatforms)
 						{
-							// Now go through all of the target types for this project
-							if (ProjectTargets.Count == 0)
+							if (InstalledPlatformInfo.IsValidPlatform(Platform, EProjectType.Code) == false)
 							{
-								throw new BuildException("Expecting at least one ProjectTarget to be associated with project '{0}' in the TargetProjects list ", ProjectFilePath);
+								continue;
 							}
-
-							foreach (ProjectTarget ProjectTarget in ProjectTargets)
+							UEBuildPlatform BuildPlatform = UEBuildPlatform.GetBuildPlatform(Platform, true);
+							if ((BuildPlatform != null) && (BuildPlatform.HasRequiredSDKsInstalled() == SDKStatus.Valid))
 							{
-								if (IsValidProjectPlatformAndConfiguration(ProjectTarget, Platform, Configuration))
+								// Now go through all of the target types for this project
+								if (ProjectTargets.Count == 0)
 								{
-									string ProjectPlatformName, ProjectConfigurationName;
-									MakeProjectPlatformAndConfigurationNames(Platform, Configuration, ProjectTarget.TargetRules.Type, out ProjectPlatformName, out ProjectConfigurationName);
+									throw new BuildException("Expecting at least one ProjectTarget to be associated with project '{0}' in the TargetProjects list ", ProjectFilePath);
+								}
 
-									ProjectConfigAndTargetCombination Combination = new ProjectConfigAndTargetCombination(Platform, Configuration, ProjectPlatformName, ProjectConfigurationName, ProjectTarget);
-									ProjectConfigAndTargetCombinations.Add(Combination);
+								foreach (ProjectTarget ProjectTarget in ProjectTargets)
+								{
+									if (IsValidProjectPlatformAndConfiguration(ProjectTarget, Platform, Configuration))
+									{
+										string ProjectPlatformName, ProjectConfigurationName;
+										MakeProjectPlatformAndConfigurationNames(Platform, Configuration, ProjectTarget.TargetRules.Type, out ProjectPlatformName, out ProjectConfigurationName);
+
+										ProjectConfigAndTargetCombination Combination = new ProjectConfigAndTargetCombination(Platform, Configuration, ProjectPlatformName, ProjectConfigurationName, ProjectTarget);
+										ProjectConfigAndTargetCombinations.Add(Combination);
+									}
 								}
 							}
 						}
 					}
 				}
+
+				// Create a list of platforms for the "invalid" configuration. We always require at least one of these.
+				InvalidConfigPlatformNames = new List<string>();
+				if(ProjectConfigAndTargetCombinations.Count == 0)
+				{
+					InvalidConfigPlatformNames.Add(DefaultPlatformName);
+				}
+				else
+				{
+					InvalidConfigPlatformNames.AddRange(ProjectConfigAndTargetCombinations.Select(x => x.ProjectPlatformName));
+				}
 			}
 		}
 
@@ -729,12 +737,11 @@ namespace UnrealBuildTool
 			}
 
 			// Add the "invalid" configuration for each platform. We use this when the solution configuration does not match any project configuration.
-			foreach(Tuple<string, UnrealTargetPlatform> PlatformTuple in ProjectPlatformNameAndPlatforms)
+			foreach(string InvalidConfigPlatformName in InvalidConfigPlatformNames)
 			{
-				string ProjectPlatformName = PlatformTuple.Item1;
-				VCProjectFileContent.AppendLine("    <ProjectConfiguration Include=\"Invalid|{0}\">", ProjectPlatformName);
+				VCProjectFileContent.AppendLine("    <ProjectConfiguration Include=\"Invalid|{0}\">", InvalidConfigPlatformName);
 				VCProjectFileContent.AppendLine("      <Configuration>Invalid</Configuration>");
-				VCProjectFileContent.AppendLine("      <Platform>{0}</Platform>", ProjectPlatformName);
+				VCProjectFileContent.AppendLine("      <Platform>{0}</Platform>", InvalidConfigPlatformName);
 				VCProjectFileContent.AppendLine("    </ProjectConfiguration>");
 			}
 
@@ -821,10 +828,9 @@ namespace UnrealBuildTool
 			VCProjectFileContent.AppendLine("  <Import Project=\"$(VCTargetsPath)\\Microsoft.Cpp.Default.props\" />");
 
 			// Write the invalid configuration data
-			foreach (Tuple<string, UnrealTargetPlatform> PlatformTuple in ProjectPlatformNameAndPlatforms)
+			foreach(string InvalidConfigPlatformName in InvalidConfigPlatformNames)
 			{
-				string ProjectPlatformName = PlatformTuple.Item1;
-				VCProjectFileContent.AppendLine("  <PropertyGroup Condition=\"'$(Configuration)|$(Platform)'=='Invalid|{0}'\" Label=\"Configuration\">", ProjectPlatformName);
+				VCProjectFileContent.AppendLine("  <PropertyGroup Condition=\"'$(Configuration)|$(Platform)'=='Invalid|{0}'\" Label=\"Configuration\">", InvalidConfigPlatformName);
 				VCProjectFileContent.AppendLine("    <ConfigurationType>Makefile</ConfigurationType>");
 				VCProjectFileContent.AppendLine("  </PropertyGroup>");
 			}
@@ -847,14 +853,13 @@ namespace UnrealBuildTool
 			VCProjectFileContent.AppendLine("  <PropertyGroup Label=\"UserMacros\" />");
 
 			// Write the invalid configuration
-			foreach (Tuple<string, UnrealTargetPlatform> PlatformTuple in ProjectPlatformNameAndPlatforms)
+			foreach(string InvalidConfigPlatformName in InvalidConfigPlatformNames)
 			{
 				const string InvalidMessage = "echo The selected platform/configuration is not valid for this target.";
 
-				string ProjectPlatformName = PlatformTuple.Item1;
 				string ProjectRelativeUnusedDirectory = NormalizeProjectPath(DirectoryReference.Combine(UnrealBuildTool.EngineDirectory, "Intermediate", "Build", "Unused"));
 
-				VCProjectFileContent.AppendLine("  <PropertyGroup Condition=\"'$(Configuration)|$(Platform)'=='Invalid|{0}'\">", ProjectPlatformName);
+				VCProjectFileContent.AppendLine("  <PropertyGroup Condition=\"'$(Configuration)|$(Platform)'=='Invalid|{0}'\">", InvalidConfigPlatformName);
 				VCProjectFileContent.AppendLine("    <NMakeBuildCommandLine>{0}</NMakeBuildCommandLine>", InvalidMessage);
 				VCProjectFileContent.AppendLine("    <NMakeReBuildCommandLine>{0}</NMakeReBuildCommandLine>", InvalidMessage);
 				VCProjectFileContent.AppendLine("    <NMakeCleanCommandLine>{0}</NMakeCleanCommandLine>", InvalidMessage);
