fix for asset catalog min version

diff --git a/Engine/Source/Programs/UnrealBuildTool/Platform/IOS/IOSToolChain.cs b/Engine/Source/Programs/UnrealBuildTool/Platform/IOS/IOSToolChain.cs
--- a/Engine/Source/Programs/UnrealBuildTool/Platform/IOS/IOSToolChain.cs
+++ b/Engine/Source/Programs/UnrealBuildTool/Platform/IOS/IOSToolChain.cs
@@ -919,9 +919,9 @@ namespace UnrealBuildTool
 			{
 				Arguments.Append(" --launch-image LaunchImage");
 				Arguments.Append(" --filter-for-device-model AppleTV5,3");
-				Arguments.Append(" --filter-for-device-os-version 9.2");
+				Arguments.Append(" --filter-for-device-os-version 10.0");
 				Arguments.Append(" --target-device tv");
-				Arguments.Append(" --minimum-deployment-target 9.2");
+				Arguments.Append(" --minimum-deployment-target 10.0");
 				Arguments.Append(" --platform appletvos");
 			}
 			else
@@ -929,7 +929,7 @@ namespace UnrealBuildTool
 				Arguments.Append(" --product-type com.apple.product-type.application");
 				Arguments.Append(" --target-device iphone");
 				Arguments.Append(" --target-device ipad");
-				Arguments.Append(" --minimum-deployment-target 9.0");
+				Arguments.Append(" --minimum-deployment-target 10.0");
 				Arguments.Append(" --platform iphoneos");
 			}
             Arguments.Append(" --enable-on-demand-resources YES");
