UBT: Allow specifying the path for the XML config file, and specifically upload it when doing remote builds. When running in an installed engine, this file is stored in the user's settings directory, and will not be uploaded as part of the workspace.

diff --git a/Engine/Source/Programs/UnrealBuildTool/System/PlatformExports.cs b/Engine/Source/Programs/UnrealBuildTool/System/PlatformExports.cs
--- a/Engine/Source/Programs/UnrealBuildTool/System/PlatformExports.cs
+++ b/Engine/Source/Programs/UnrealBuildTool/System/PlatformExports.cs
@@ -146,7 +146,7 @@ namespace UnrealBuildTool
 			UnrealBuildTool.SetIsEngineInstalled(bIsEngineInstalled);
 
 			// Read the XML configuration files
-			if(!XmlConfig.ReadConfigFiles(false))
+			if(!XmlConfig.ReadConfigFiles(null))
 			{
 				return false;
 			}
diff --git a/Engine/Source/Programs/UnrealBuildTool/System/XmlConfig.cs b/Engine/Source/Programs/UnrealBuildTool/System/XmlConfig.cs
--- a/Engine/Source/Programs/UnrealBuildTool/System/XmlConfig.cs
+++ b/Engine/Source/Programs/UnrealBuildTool/System/XmlConfig.cs
@@ -35,6 +35,11 @@ namespace UnrealBuildTool
 			public string FolderName;
 		}
 
+		/// <summary>
+		/// The cache file that is being used
+		/// </summary>
+		public static FileReference CacheFile;
+
 		/// <summary>
 		/// Parsed config values
 		/// </summary>
@@ -48,26 +53,18 @@ namespace UnrealBuildTool
 		/// <summary>
 		/// Initialize the config system with the given types
 		/// </summary>
-		/// <param name="bForceLoadCache">Force use of the cached XML config without checking if it's valid (useful for remote builds)</param>
-		public static bool ReadConfigFiles(bool bForceLoadCache)
+		/// <param name="OverrideCacheFile">Force use of the cached XML config without checking if it's valid (useful for remote builds)</param>
+		public static bool ReadConfigFiles(FileReference OverrideCacheFile)
 		{
 			// Find all the configurable types
 			List<Type> ConfigTypes = FindConfigurableTypes();
 
-			// Get the path to the cache file
-			FileReference CacheFile = FileReference.Combine(UnrealBuildTool.EngineDirectory, "Intermediate", "Build", "XmlConfigCache.bin");
-			if(UnrealBuildTool.IsEngineInstalled())
-			{
-				DirectoryReference UserSettingsDir = Utils.GetUserSettingDirectory();
-				if(UserSettingsDir != null)
-				{
-					CacheFile = FileReference.Combine(UserSettingsDir, "UnrealEngine", String.Format("XmlConfigCache-{0}.bin", UnrealBuildTool.RootDirectory.FullName.Replace(":", "").Replace(Path.DirectorySeparatorChar, '+')));
-				}
-			}
-
 			// Update the cache if necessary
-			if(bForceLoadCache)
+			if(OverrideCacheFile != null)
 			{
+				// Set the cache file to the overriden value
+				CacheFile = OverrideCacheFile;
+
 				// Never rebuild the cache; just try to load it.
 				if(!XmlConfigData.TryRead(CacheFile, ConfigTypes, out Values))
 				{
@@ -76,6 +73,17 @@ namespace UnrealBuildTool
 			}
 			else
 			{
+				// Get the default cache file
+				CacheFile = FileReference.Combine(UnrealBuildTool.EngineDirectory, "Intermediate", "Build", "XmlConfigCache.bin");
+				if(UnrealBuildTool.IsEngineInstalled())
+				{
+					DirectoryReference UserSettingsDir = Utils.GetUserSettingDirectory();
+					if(UserSettingsDir != null)
+					{
+						CacheFile = FileReference.Combine(UserSettingsDir, "UnrealEngine", String.Format("XmlConfigCache-{0}.bin", UnrealBuildTool.RootDirectory.FullName.Replace(":", "").Replace(Path.DirectorySeparatorChar, '+')));
+					}
+				}
+
 				// Find all the input files
 				FileReference[] InputFiles = FindInputFiles().Select(x => x.Location).ToArray();
 
diff --git a/Engine/Source/Programs/UnrealBuildTool/ToolChain/RemoteMac.cs b/Engine/Source/Programs/UnrealBuildTool/ToolChain/RemoteMac.cs
--- a/Engine/Source/Programs/UnrealBuildTool/ToolChain/RemoteMac.cs
+++ b/Engine/Source/Programs/UnrealBuildTool/ToolChain/RemoteMac.cs
@@ -351,7 +351,7 @@ namespace UnrealBuildTool
 			RemoteArguments.Add(TargetDesc.Platform.ToString());
 			RemoteArguments.Add(TargetDesc.Configuration.ToString());
 			RemoteArguments.Add("-SkipRulesCompile"); // Use the rules assembly built locally
-			RemoteArguments.Add("-ForceXmlConfigCache"); // Use the XML config cache built locally, since the remote won't have it
+			RemoteArguments.Add(String.Format("-XmlConfigCache={0}", GetRemotePath(XmlConfig.CacheFile))); // Use the XML config cache built locally, since the remote won't have it
 			RemoteArguments.Add(String.Format("-Log={0}", GetRemotePath(RemoteLogFile)));
 			RemoteArguments.Add(String.Format("-Manifest={0}", GetRemotePath(RemoteManifestFile)));
 
@@ -757,6 +757,10 @@ namespace UnrealBuildTool
 			Log.TraceInformation("[Remote] Uploading scripts...");
 			UploadFiles(TempDir, GetRemotePath(UnrealBuildTool.EngineDirectory), ScriptPathsFileName);
 
+			// Upload the config files
+			Log.TraceInformation("[Remote] Uploading config files...");
+			UploadFile(XmlConfig.CacheFile);
+
 			// Upload the engine files
 			List<FileReference> EngineFilters = new List<FileReference>();
 			EngineFilters.Add(ScriptProtectList);
diff --git a/Engine/Source/Programs/UnrealBuildTool/UnrealBuildTool.cs b/Engine/Source/Programs/UnrealBuildTool/UnrealBuildTool.cs
--- a/Engine/Source/Programs/UnrealBuildTool/UnrealBuildTool.cs
+++ b/Engine/Source/Programs/UnrealBuildTool/UnrealBuildTool.cs
@@ -496,9 +496,20 @@ namespace UnrealBuildTool
 						bIsEngineInstalled = FileReference.Exists(FileReference.Combine(RootDirectory, "Engine", "Build", "InstalledBuild.txt"));
 					}
 
+					// Parse the argument for overriding the XML configuration file location (for remote builds)
+					FileReference XmlConfigCache = null;
+					foreach (string Argument in Arguments)
+					{
+						const string Prefix = "-XmlConfigCache=";
+						if(Argument.StartsWith(Prefix, StringComparison.OrdinalIgnoreCase))
+						{
+							XmlConfigCache = new FileReference(Argument.Substring(Prefix.Length));
+							break;
+						}
+					}
+
 					// Read the XML configuration files
-					bool bForceXmlConfigCache = Arguments.Any(x => x.Equals("-ForceXmlConfigCache", StringComparison.InvariantCultureIgnoreCase));
-					if (!XmlConfig.ReadConfigFiles(bForceXmlConfigCache))
+					if (!XmlConfig.ReadConfigFiles(XmlConfigCache))
 					{
 						return 1;
 					}
