UBT: Fix missing icons for program executables on Windows.

diff --git a/Engine/Source/Programs/UnrealBuildTool/Configuration/UEBuildBinary.cs b/Engine/Source/Programs/UnrealBuildTool/Configuration/UEBuildBinary.cs
--- a/Engine/Source/Programs/UnrealBuildTool/Configuration/UEBuildBinary.cs
+++ b/Engine/Source/Programs/UnrealBuildTool/Configuration/UEBuildBinary.cs
@@ -782,8 +782,12 @@ namespace UnrealBuildTool
 					}
 					else
 					{
+						// Get the intermediate directory
+						DirectoryReference ResourceIntermediateDirectory = ((UEBuildModuleCPP)Modules.First()).IntermediateDirectory;
+
 						// Create a compile environment for resource files
 						CppCompileEnvironment ResourceCompileEnvironment = new CppCompileEnvironment(BinaryCompileEnvironment);
+						WindowsPlatform.SetupResourceCompileEnvironment(ResourceCompileEnvironment, ResourceIntermediateDirectory, Target);
 
 						// @todo: This should be in some Windows code somewhere...
 						// Set the original file name macro; used in PCLaunch.rc to set the binary metadata fields.
@@ -796,7 +800,7 @@ namespace UnrealBuildTool
 
 						// Otherwise compile the default resource file per-binary, so that it gets the correct ORIGINAL_FILE_NAME macro.
 						FileItem DefaultResourceFile = FileItem.GetItemByFileReference(FileReference.Combine(UnrealBuildTool.EngineSourceDirectory, "Runtime", "Launch", "Resources", "Windows", "PCLaunch.rc"));
-						CPPOutput DefaultResourceOutput = ToolChain.CompileRCFiles(BinaryCompileEnvironment, new List<FileItem> { DefaultResourceFile }, ((UEBuildModuleCPP)Modules.First()).IntermediateDirectory, ActionGraph);
+						CPPOutput DefaultResourceOutput = ToolChain.CompileRCFiles(ResourceCompileEnvironment, new List<FileItem> { DefaultResourceFile }, ResourceIntermediateDirectory, ActionGraph);
 						BinaryLinkEnvironment.InputFiles.AddRange(DefaultResourceOutput.ObjectFiles);
 					}
 				}
diff --git a/Engine/Source/Programs/UnrealBuildTool/Configuration/UEBuildModuleCPP.cs b/Engine/Source/Programs/UnrealBuildTool/Configuration/UEBuildModuleCPP.cs
--- a/Engine/Source/Programs/UnrealBuildTool/Configuration/UEBuildModuleCPP.cs
+++ b/Engine/Source/Programs/UnrealBuildTool/Configuration/UEBuildModuleCPP.cs
@@ -565,20 +565,8 @@ namespace UnrealBuildTool
 			// must be specified relative to the resource file itself or Engine/Source.
 			if(SourceFilesToBuild.RCFiles.Count > 0)
 			{
-				// Figure the icon to use. We can only use a custom icon when compiling to a project-specific intemediate directory (and not for the shared editor executable, for example).
-				FileReference IconFile;
-				if(Target.ProjectFile != null && IntermediateDirectory.IsUnderDirectory(Target.ProjectFile.Directory))
-				{
-					IconFile = WindowsPlatform.GetApplicationIcon(Target.ProjectFile);
-				}
-				else
-				{
-					IconFile = WindowsPlatform.GetApplicationIcon(null);
-				}
-
-				// Setup the compile environment, setting the icon to use via a macro. This is used in PCLaunch.rc2.
 				CppCompileEnvironment ResourceCompileEnvironment = new CppCompileEnvironment(BinaryCompileEnvironment);
-				ResourceCompileEnvironment.Definitions.Add(String.Format("BUILD_ICON_FILE_NAME=\"\\\"{0}\\\"\"", IconFile.FullName.Replace("\\", "\\\\")));
+				WindowsPlatform.SetupResourceCompileEnvironment(ResourceCompileEnvironment, IntermediateDirectory, Target);
 				LinkInputFiles.AddRange(ToolChain.CompileRCFiles(ResourceCompileEnvironment, SourceFilesToBuild.RCFiles, IntermediateDirectory, ActionGraph).ObjectFiles);
 			}
 
diff --git a/Engine/Source/Programs/UnrealBuildTool/Configuration/UEBuildTarget.cs b/Engine/Source/Programs/UnrealBuildTool/Configuration/UEBuildTarget.cs
--- a/Engine/Source/Programs/UnrealBuildTool/Configuration/UEBuildTarget.cs
+++ b/Engine/Source/Programs/UnrealBuildTool/Configuration/UEBuildTarget.cs
@@ -2189,8 +2189,10 @@ namespace UnrealBuildTool
 
 						FileItem DefaultResourceFile = FileItem.GetExistingItemByFileReference(FileReference.Combine(UnrealBuildTool.EngineSourceDirectory, "Runtime", "Launch", "Resources", "Windows", "PCLaunch.rc"));
 						DefaultResourceFile.CachedIncludePaths = DefaultResourceCompileEnvironment.IncludePaths;
-						CPPOutput DefaultResourceOutput = TargetToolChain.CompileRCFiles(DefaultResourceCompileEnvironment, new List<FileItem> { DefaultResourceFile }, EngineIntermediateDirectory, ActionGraph);
 
+						WindowsPlatform.SetupResourceCompileEnvironment(DefaultResourceCompileEnvironment, EngineIntermediateDirectory, Rules);
+
+						CPPOutput DefaultResourceOutput = TargetToolChain.CompileRCFiles(DefaultResourceCompileEnvironment, new List<FileItem> { DefaultResourceFile }, EngineIntermediateDirectory, ActionGraph);
 						GlobalLinkEnvironment.DefaultResourceFiles.AddRange(DefaultResourceOutput.ObjectFiles);
 					}
 				}
diff --git a/Engine/Source/Programs/UnrealBuildTool/Platform/Windows/UEBuildWindows.cs b/Engine/Source/Programs/UnrealBuildTool/Platform/Windows/UEBuildWindows.cs
--- a/Engine/Source/Programs/UnrealBuildTool/Platform/Windows/UEBuildWindows.cs
+++ b/Engine/Source/Programs/UnrealBuildTool/Platform/Windows/UEBuildWindows.cs
@@ -1349,6 +1349,29 @@ namespace UnrealBuildTool
 			return FileReference.Combine(UnrealBuildTool.EngineDirectory, "Source", "Runtime", "Launch", "Resources", "Windows", "UE4.ico");
 		}
 
+		/// <summary>
+		/// Configures the resource compile environment for the given target
+		/// </summary>
+		/// <param name="ResourceCompileEnvironment">The compile environment</param>
+		/// <param name="IntermediateDirectory">The output directory for compiled files</param>
+		/// <param name="Target">The target being built</param>
+		public static void SetupResourceCompileEnvironment(CppCompileEnvironment ResourceCompileEnvironment, DirectoryReference IntermediateDirectory, ReadOnlyTargetRules Target)
+		{
+			// Figure the icon to use. We can only use a custom icon when compiling to a project-specific intemediate directory (and not for the shared editor executable, for example).
+			FileReference IconFile;
+			if(Target.ProjectFile != null && IntermediateDirectory.IsUnderDirectory(Target.ProjectFile.Directory))
+			{
+				IconFile = WindowsPlatform.GetApplicationIcon(Target.ProjectFile);
+			}
+			else
+			{
+				IconFile = WindowsPlatform.GetApplicationIcon(null);
+			}
+
+			// Setup the compile environment, setting the icon to use via a macro. This is used in PCLaunch.rc2.
+			ResourceCompileEnvironment.Definitions.Add(String.Format("BUILD_ICON_FILE_NAME=\"\\\"{0}\\\"\"", IconFile.FullName.Replace("\\", "\\\\")));
+		}
+
 		/// <summary>
 		/// Modify the rules for a newly created module, in a target that's being built for this platform.
 		/// This is not required - but allows for hiding details of a particular platform.
