UBT: Use correct directory name for generated code.

diff --git a/Engine/Source/Programs/UnrealBuildTool/Configuration/UEBuildTarget.cs b/Engine/Source/Programs/UnrealBuildTool/Configuration/UEBuildTarget.cs
--- a/Engine/Source/Programs/UnrealBuildTool/Configuration/UEBuildTarget.cs
+++ b/Engine/Source/Programs/UnrealBuildTool/Configuration/UEBuildTarget.cs
@@ -4053,6 +4053,7 @@ namespace UnrealBuildTool
 				DirectoryReference GeneratedCodeDirectory = null;
 				if (RulesObject.Type != ModuleRules.ModuleType.External)
 				{
+					// Get the base directory
 					if (RulesObject.Plugin != null)
 					{
 						GeneratedCodeDirectory = RulesObject.Plugin.Directory;
@@ -4069,7 +4070,18 @@ namespace UnrealBuildTool
 					{
 						GeneratedCodeDirectory = ProjectDirectory;
 					}
-					GeneratedCodeDirectory = DirectoryReference.Combine(GeneratedCodeDirectory, PlatformIntermediateFolder, AppName, "Inc", ModuleName);
+
+					// Get the subfolder containing generated code
+					GeneratedCodeDirectory = DirectoryReference.Combine(GeneratedCodeDirectory, PlatformIntermediateFolder, AppName, "Inc");
+
+					// Append the binaries subfolder, if present. We rely on this to ensure that build products can be filtered correctly.
+					if(RulesObject.BinariesSubFolder != null)
+					{
+						GeneratedCodeDirectory = DirectoryReference.Combine(GeneratedCodeDirectory, RulesObject.BinariesSubFolder);
+					}
+
+					// Finally, append the module name.
+					GeneratedCodeDirectory = DirectoryReference.Combine(GeneratedCodeDirectory, ModuleName);
 				}
 
 				// For legacy modules, add a bunch of default include paths.
