Fix missing files being added to dependency lists, causing failure to make installed builds from Git repos.

diff --git a/Engine/Source/Programs/UnrealBuildTool/Configuration/UEBuildTarget.cs b/Engine/Source/Programs/UnrealBuildTool/Configuration/UEBuildTarget.cs
--- a/Engine/Source/Programs/UnrealBuildTool/Configuration/UEBuildTarget.cs
+++ b/Engine/Source/Programs/UnrealBuildTool/Configuration/UEBuildTarget.cs
@@ -1204,13 +1204,13 @@ namespace UnrealBuildTool
 			foreach(RuntimeDependency RuntimeDependency in RuntimeDependencies)
 			{
 				FileReference SourceFile;
-				if(RuntimeDependencyTargetFileToSourceFile.TryGetValue(RuntimeDependency.Path, out SourceFile))
+				if(!RuntimeDependencyTargetFileToSourceFile.TryGetValue(RuntimeDependency.Path, out SourceFile))
 				{
-					Files.Add(SourceFile);
+					SourceFile = RuntimeDependency.Path;
 				}
-				else
+				if(RuntimeDependency.Type != StagedFileType.DebugNonUFS || FileReference.Exists(SourceFile))
 				{
-					Files.Add(RuntimeDependency.Path);
+					Files.Add(SourceFile);
 				}
 			}
 
