Expose the Windows SDK directory, VC toolchain directory, and VC toolchain version number as properties on the TargetRules and ReadOnlyTargetRules classes (ie. Target.WindowsPlatform.WindowsSdkDir, Target.WindowsPlatform.ToolChainDir, Target.WindowsPlatform.ToolChainVersion). Also set the compiler version and Windows SDK version to the actual versions being used once they are known.

diff --git a/Engine/Source/Programs/UnrealBuildTool/Platform/Windows/PVSToolChain.cs b/Engine/Source/Programs/UnrealBuildTool/Platform/Windows/PVSToolChain.cs
--- a/Engine/Source/Programs/UnrealBuildTool/Platform/Windows/PVSToolChain.cs
+++ b/Engine/Source/Programs/UnrealBuildTool/Platform/Windows/PVSToolChain.cs
@@ -44,7 +44,7 @@ namespace UnrealBuildTool
 
 		public PVSToolChain(CppPlatform Platform, ReadOnlyTargetRules Target) : base(Platform)
 		{
-			EnvVars = VCEnvironment.Create(Target.WindowsPlatform.Compiler, Platform, Target.WindowsPlatform.CompilerVersion, Target.WindowsPlatform.WindowsSdkVersion);
+			EnvVars = Target.WindowsPlatform.Environment;
 
 			AnalyzerFile = FileReference.Combine(new DirectoryReference(Environment.GetFolderPath(Environment.SpecialFolder.ProgramFilesX86)), "PVS-Studio", "x64", "PVS-Studio.exe");
 			if(!FileReference.Exists(AnalyzerFile))
diff --git a/Engine/Source/Programs/UnrealBuildTool/Platform/Windows/UEBuildWindows.cs b/Engine/Source/Programs/UnrealBuildTool/Platform/Windows/UEBuildWindows.cs
--- a/Engine/Source/Programs/UnrealBuildTool/Platform/Windows/UEBuildWindows.cs
+++ b/Engine/Source/Programs/UnrealBuildTool/Platform/Windows/UEBuildWindows.cs
@@ -187,6 +187,35 @@ namespace UnrealBuildTool
 		/// </summary>
 		public bool bBuildLargeAddressAwareBinary = true;
 
+		/// <summary>
+		/// The Visual C++ environment to use for this target. Only initialized after all the target settings are finalized, in ValidateTarget().
+		/// </summary>
+		internal VCEnvironment Environment;
+
+		/// <summary>
+		/// Directory containing the toolchain
+		/// </summary>
+		public string ToolChainDir
+		{
+			get { return (Environment == null)? null : Environment.ToolChainDir.FullName; }
+		}
+
+		/// <summary>
+		/// The version number of the toolchain
+		/// </summary>
+		public string ToolChainVersion
+		{
+			get { return (Environment == null)? null : Environment.ToolChainVersion.ToString(); }
+		}
+
+		/// <summary>
+		/// Root directory containing the Windows Sdk
+		/// </summary>
+		public string WindowsSdkDir
+		{
+			get { return (Environment == null)? null : Environment.WindowsSdkDir.FullName; }
+		}
+
 		/// <summary>
 		/// When using a Visual Studio compiler, returns the version name as a string
 		/// </summary>
@@ -324,6 +353,26 @@ namespace UnrealBuildTool
 			return Inner.GetVisualStudioCompilerVersionName();
 		}
 
+		internal VCEnvironment Environment
+		{
+			get { return Inner.Environment; }
+		}
+
+		public string ToolChainDir
+		{
+			get { return Inner.ToolChainDir; }
+		}
+
+		public string ToolChainVersion
+		{
+			get { return Inner.ToolChainVersion; }
+		}
+
+		public string WindowsSdkDir
+		{
+			get { return Inner.WindowsSdkDir; }
+		}
+
 		#if !__MonoCS__
 		#pragma warning restore CS1591
 		#endif
@@ -459,6 +508,13 @@ namespace UnrealBuildTool
 				Target.bDisableDebugInfoForGeneratedCode = false;
 			}
 
+			// Initialize the VC environment for the target, and set all the version numbers to the concrete values we chose.
+			VCEnvironment Environment = VCEnvironment.Create(Target.WindowsPlatform.Compiler, DefaultCppPlatform, Target.WindowsPlatform.CompilerVersion, Target.WindowsPlatform.WindowsSdkVersion);
+			Target.WindowsPlatform.Environment = Environment;
+			Target.WindowsPlatform.Compiler = Environment.Compiler;
+			Target.WindowsPlatform.CompilerVersion = Environment.CompilerVersion.ToString();
+			Target.WindowsPlatform.WindowsSdkVersion = Environment.WindowsSdkVersion.ToString();
+
 //			@Todo: Still getting reports of frequent OOM issues with this enabled as of 15.7.
 //			// Enable fast PDB linking if we're on VS2017 15.7 or later. Previous versions have OOM issues with large projects.
 //			if(!Target.bFormalBuild && !Target.bUseFastPDBLinking.HasValue && Target.WindowsPlatform.Compiler >= WindowsCompiler.VisualStudio2017)
diff --git a/Engine/Source/Programs/UnrealBuildTool/Platform/Windows/VCToolChain.cs b/Engine/Source/Programs/UnrealBuildTool/Platform/Windows/VCToolChain.cs
--- a/Engine/Source/Programs/UnrealBuildTool/Platform/Windows/VCToolChain.cs
+++ b/Engine/Source/Programs/UnrealBuildTool/Platform/Windows/VCToolChain.cs
@@ -28,7 +28,7 @@ namespace UnrealBuildTool
 			: base(Platform)
 		{
 			this.Target = Target;
-			this.EnvVars = VCEnvironment.Create(Target.WindowsPlatform.Compiler, Platform, Target.WindowsPlatform.CompilerVersion, Target.WindowsPlatform.WindowsSdkVersion);
+			this.EnvVars = Target.WindowsPlatform.Environment;
 
 			if (Target.WindowsPlatform.ObjSrcMapFile != null)
 			{
