Fix hot-reload from IDE when another module is referenced.

diff --git a/Engine/Source/Programs/UnrealBuildTool/System/ActionGraph.cs b/Engine/Source/Programs/UnrealBuildTool/System/ActionGraph.cs
--- a/Engine/Source/Programs/UnrealBuildTool/System/ActionGraph.cs
+++ b/Engine/Source/Programs/UnrealBuildTool/System/ActionGraph.cs
@@ -392,8 +392,8 @@ namespace UnrealBuildTool
 					HashSet<Action> UnlinkedActionsToReadd = new HashSet<Action>(UnlinkedActions.Where(Action => PrerequisiteLinkActions.Contains(Action)));
 
 					// Also re-add any DLL whose import library is being rebuilt. These may be separate actions, and the import library will reference the new DLL even if it isn't being compiled itself, so it must exist.
-					HashSet<FileReference> ProducedItemsToReAdd = new HashSet<FileReference>(ActionsToExecute.SelectMany(x => x.ProducedItems).Select(x => x.Location).Where(x => x.HasExtension(".lib")).Select(x => x.ChangeExtension(".suppressed.lib")));
-					UnlinkedActionsToReadd.UnionWith(UnlinkedActions.Where(x => x.ProducedItems.Any(y => ProducedItemsToReAdd.Contains(y.Location))));
+					HashSet<string> ProducedDllsToReAdd = new HashSet<string>(UnlinkedActionsToReadd.SelectMany(x => x.ProducedItems).Select(x => x.Location).Where(x => x.HasExtension(".lib")).Select(x => x.GetFileNameWithoutExtension() + ".dll"), StringComparer.OrdinalIgnoreCase);
+					UnlinkedActionsToReadd.UnionWith(UnlinkedActions.Where(x => x.ProducedItems.Any(y => ProducedDllsToReAdd.Contains(y.Location.GetFileName()))));
 
 					// Bail if we didn't find anything
 					if (UnlinkedActionsToReadd.Count == 0)
