Fix plugin packaging not executing monolithic compile actions when run from an installed build.

diff --git a/Engine/Source/Programs/UnrealBuildTool/Configuration/UEBuildBinary.cs b/Engine/Source/Programs/UnrealBuildTool/Configuration/UEBuildBinary.cs
--- a/Engine/Source/Programs/UnrealBuildTool/Configuration/UEBuildBinary.cs
+++ b/Engine/Source/Programs/UnrealBuildTool/Configuration/UEBuildBinary.cs
@@ -183,8 +183,8 @@ namespace UnrealBuildTool
 		/// <returns>Set of built products</returns>
 		public IEnumerable<FileItem> Build(ReadOnlyTargetRules Target, UEToolChain ToolChain, CppCompileEnvironment CompileEnvironment, LinkEnvironment LinkEnvironment, List<PrecompiledHeaderTemplate> SharedPCHs, ISourceFileWorkingSet WorkingSet, DirectoryReference ExeDir, ActionGraph ActionGraph)
 		{
-			// Return nothing if we're using precompiled binaries
-			if(bUsePrecompiled)
+			// Return nothing if we're using precompiled binaries. If we're not linking, we might want just one module to be compiled (eg. a foreign plugin), so allow any actions to run.
+			if (bUsePrecompiled && !Target.bDisableLinking)
 			{
 				return new List<FileItem>();
 			}
