Fix installed UE4Editor.modules file being deleted when packaging a plugin from the installed engine.

diff --git a/Engine/Source/Programs/UnrealBuildTool/Configuration/UEBuildTarget.cs b/Engine/Source/Programs/UnrealBuildTool/Configuration/UEBuildTarget.cs
--- a/Engine/Source/Programs/UnrealBuildTool/Configuration/UEBuildTarget.cs
+++ b/Engine/Source/Programs/UnrealBuildTool/Configuration/UEBuildTarget.cs
@@ -1701,7 +1701,7 @@ namespace UnrealBuildTool
 				foreach (FileReference VersionManifestFile in FileReferenceToModuleManifestPairs.Select(x => x.Key))
 				{
 					// Make sure the file (and directory) exists before trying to delete it
-					if(FileReference.Exists(VersionManifestFile))
+					if(FileReference.Exists(VersionManifestFile) && !IsFileInstalled(VersionManifestFile))
 					{
 						FileReference.Delete(VersionManifestFile);
 					}
