Fix Visual Studio prompting to save solution file freshly generated by UBT. Solution configurations for stub projects were not being included.

Also fix natvis file being generated with the incorrect path for foreign projects.

diff --git a/Engine/Source/Programs/UnrealBuildTool/ProjectFiles/VisualStudio/VCProject.cs b/Engine/Source/Programs/UnrealBuildTool/ProjectFiles/VisualStudio/VCProject.cs
--- a/Engine/Source/Programs/UnrealBuildTool/ProjectFiles/VisualStudio/VCProject.cs
+++ b/Engine/Source/Programs/UnrealBuildTool/ProjectFiles/VisualStudio/VCProject.cs
@@ -368,6 +368,12 @@ namespace UnrealBuildTool
 		/// <returns>Project context matching the given solution context</returns>
 		public override MSBuildProjectContext GetMatchingProjectContext(TargetType SolutionTarget, UnrealTargetConfiguration SolutionConfiguration, UnrealTargetPlatform SolutionPlatform)
 		{
+			// Stub projects always build in the same configuration
+			if(IsStubProject)
+			{
+				return new MSBuildProjectContext(StubProjectConfigurationName, StubProjectPlatformName);
+			}
+
 			// Have to match every solution configuration combination to a project configuration (or use the invalid one) 
 			string ProjectConfigurationName = "Invalid";
 
diff --git a/Engine/Source/Programs/UnrealBuildTool/ProjectFiles/VisualStudio/VCProjectFileGenerator.cs b/Engine/Source/Programs/UnrealBuildTool/ProjectFiles/VisualStudio/VCProjectFileGenerator.cs
--- a/Engine/Source/Programs/UnrealBuildTool/ProjectFiles/VisualStudio/VCProjectFileGenerator.cs
+++ b/Engine/Source/Programs/UnrealBuildTool/ProjectFiles/VisualStudio/VCProjectFileGenerator.cs
@@ -546,10 +546,14 @@ namespace UnrealBuildTool
 					VCSolutionFileContent.AppendLine("EndProject");
 				}
 
+				// Get the path to the visualizers file. Try to make it relative to the solution directory, but fall back to a full path if it's a foreign project.
+				FileReference VisualizersFile = FileReference.Combine(UnrealBuildTool.EngineDirectory, "Extras", "VisualStudioDebugging", "UE4.natvis");
+				DirectoryReference SolutionDir = new FileReference(SolutionFileName).Directory;
+
 				// Add the visualizers at the solution level. Doesn't seem to be picked up from a makefile project in VS2017 15.8.5.
-				VCSolutionFileContent.AppendLine(String.Format("Project(\"{{2150E333-8FDC-42A3-9474-1A3956D46DE8}}\") = \"Visualizers\", \"Visualizers\", \"{0}\"", Guid.NewGuid().ToString()));
+				VCSolutionFileContent.AppendLine(String.Format("Project(\"{{2150E333-8FDC-42A3-9474-1A3956D46DE8}}\") = \"Visualizers\", \"Visualizers\", \"{0}\"", Guid.NewGuid().ToString("B").ToUpperInvariant()));
 				VCSolutionFileContent.AppendLine("\tProjectSection(SolutionItems) = preProject");
-				VCSolutionFileContent.AppendLine("\t\tEngine\\Extras\\VisualStudioDebugging\\UE4.natvis = Engine\\Extras\\VisualStudioDebugging\\UE4.natvis");
+				VCSolutionFileContent.AppendLine("\t\t{0} = {0}", VisualizersFile.MakeRelativeTo(SolutionDir));
 				VCSolutionFileContent.AppendLine("\tEndProjectSection");
 				VCSolutionFileContent.AppendLine("EndProject");
 			}
@@ -664,34 +668,20 @@ namespace UnrealBuildTool
 
 						foreach (MSBuildProjectFile CurProject in AllProjectFiles)
 						{
-							// NOTE: We don't emit solution configuration entries for "stub" projects.  Those projects are only
-							// built using UnrealBuildTool and don't require a presence in the solution project list
-
-							// NOTE: We also process projects that were "imported" here, hoping to match those to our solution
-							//       configurations.  In some cases this may not be successful though.  Imported projects
-							//       should always be carefully setup to match our project generator's solution configs.
-							if (!CurProject.IsStubProject)
+							foreach (VCSolutionConfigCombination SolutionConfigCombination in SolutionConfigCombinations)
 							{
-								if (CurProject.ProjectTargets.Count == 0)
-								{
-									throw new BuildException("Expecting project '" + CurProject.ProjectFilePath + "' to have at least one ProjectTarget associated with it!");
-								}
+								// Get the context for the current solution context
+								MSBuildProjectContext ProjectContext = CurProject.GetMatchingProjectContext(SolutionConfigCombination.TargetConfigurationName, SolutionConfigCombination.Configuration, SolutionConfigCombination.Platform);
 
-								foreach (VCSolutionConfigCombination SolutionConfigCombination in SolutionConfigCombinations)
+								// Write the solution mapping (e.g.  "{4232C52C-680F-4850-8855-DC39419B5E9B}.Debug|iOS.ActiveCfg = iOS_Debug|Win32")
+								string CurProjectGUID = CurProject.ProjectGUID.ToString("B").ToUpperInvariant();
+								VCSolutionFileContent.AppendLine("		{0}.{1}.ActiveCfg = {2}", CurProjectGUID, SolutionConfigCombination.VCSolutionConfigAndPlatformName, ProjectContext.Name);
+								if (ProjectContext.bBuildByDefault)
 								{
-									// Get the context for the current solution context
-									MSBuildProjectContext ProjectContext = CurProject.GetMatchingProjectContext(SolutionConfigCombination.TargetConfigurationName, SolutionConfigCombination.Configuration, SolutionConfigCombination.Platform);
-
-									// Write the solution mapping (e.g.  "{4232C52C-680F-4850-8855-DC39419B5E9B}.Debug|iOS.ActiveCfg = iOS_Debug|Win32")
-									string CurProjectGUID = CurProject.ProjectGUID.ToString("B").ToUpperInvariant();
-									VCSolutionFileContent.AppendLine("		{0}.{1}.ActiveCfg = {2}", CurProjectGUID, SolutionConfigCombination.VCSolutionConfigAndPlatformName, ProjectContext.Name);
-									if (ProjectContext.bBuildByDefault)
+									VCSolutionFileContent.AppendLine("		{0}.{1}.Build.0 = {2}", CurProjectGUID, SolutionConfigCombination.VCSolutionConfigAndPlatformName, ProjectContext.Name);
+									if(ProjectContext.bDeployByDefault)
 									{
-										VCSolutionFileContent.AppendLine("		{0}.{1}.Build.0 = {2}", CurProjectGUID, SolutionConfigCombination.VCSolutionConfigAndPlatformName, ProjectContext.Name);
-										if(ProjectContext.bDeployByDefault)
-										{
-											VCSolutionFileContent.AppendLine("		{0}.{1}.Deploy.0 = {2}", CurProjectGUID, SolutionConfigCombination.VCSolutionConfigAndPlatformName, ProjectContext.Name);
-										}
+										VCSolutionFileContent.AppendLine("		{0}.{1}.Deploy.0 = {2}", CurProjectGUID, SolutionConfigCombination.VCSolutionConfigAndPlatformName, ProjectContext.Name);
 									}
 								}
 							}
