UE-65459 - Remote Building for iOS Fails During Build Step: "File for download 'UBT-UE65273Verif-IOS-Development_Remote.txt' is not under the engine or project directory"

diff --git a/Engine/Source/Programs/UnrealBuildTool/ToolChain/RemoteMac.cs b/Engine/Source/Programs/UnrealBuildTool/ToolChain/RemoteMac.cs
--- a/Engine/Source/Programs/UnrealBuildTool/ToolChain/RemoteMac.cs
+++ b/Engine/Source/Programs/UnrealBuildTool/ToolChain/RemoteMac.cs
@@ -336,6 +336,20 @@ namespace UnrealBuildTool
 			DirectoryReference TempDir = DirectoryReference.Combine(BaseDir, "Intermediate", "Remote", TargetDesc.Name, TargetDesc.Platform.ToString(), TargetDesc.Configuration.ToString());
 			DirectoryReference.CreateDirectory(TempDir);
 
+			bool bLogIsMapped = false;
+			foreach (RemoteMapping Mapping in Mappings)
+			{
+				if (RemoteLogFile.Directory.FullName.Equals(Mapping.LocalDirectory.FullName, StringComparison.InvariantCultureIgnoreCase))
+				{
+					bLogIsMapped = true;
+					break;
+				}
+			}
+			if (!bLogIsMapped)
+			{
+				Mappings.Add(new RemoteMapping(RemoteLogFile.Directory, GetRemotePath(RemoteLogFile.Directory)));
+			}
+
 			// Compile the rules assembly
 			RulesCompiler.CreateTargetRulesAssembly(TargetDesc.ProjectFile, TargetDesc.Name, false, false, TargetDesc.ForeignPlugin);
 
