UE-40545 - TVOS images are throwing errors when attempting to deploy to iTunes Connect

diff --git a/Engine/Build/Rsync/RsyncEngine.txt b/Engine/Build/Rsync/RsyncEngine.txt
--- a/Engine/Build/Rsync/RsyncEngine.txt
+++ b/Engine/Build/Rsync/RsyncEngine.txt
@@ -220,6 +220,8 @@ protect /Binaries/ThirdParty/Mono/Mac/lib/mono/monodoc/monodoc.dll
 + /Build/Mac/**
 + /Build/IOS/
 + /Build/IOS/**
++ /Build/TVOS/
++ /Build/TVOS/**
 protect /Build/IOS/*.PackageVersionCounter
 
 # Config: Include everything
diff --git a/Engine/Build/Rsync/RsyncProject.txt b/Engine/Build/Rsync/RsyncProject.txt
--- a/Engine/Build/Rsync/RsyncProject.txt
+++ b/Engine/Build/Rsync/RsyncProject.txt
@@ -9,6 +9,8 @@
 + /Intermediate/Build/BuildRules/*.dll
 + /Intermediate/Source/
 + /Intermediate/Source/**
++ /Intermediate/TVOS/
++ /Intermediate/TVOS/**
 
 protect /Intermediate/**
 
diff --git a/Engine/Build/TVOS/Resources/Assets.xcassets/AppIcon.brandassets/App Icon - Large.imagestack/Back.imagestacklayer/Content.imageset/Contents.json b/Engine/Build/TVOS/Resources/Assets.xcassets/AppIcon.brandassets/App Icon - Large.imagestack/Back.imagestacklayer/Content.imageset/Contents.json
new file mode 100644
--- /dev/null
+++ b/Engine/Build/TVOS/Resources/Assets.xcassets/AppIcon.brandassets/App Icon - Large.imagestack/Back.imagestacklayer/Content.imageset/Contents.json	
@@ -0,0 +1,17 @@
+{
+	"images" : [
+		{
+			"idiom" : "tv",
+			"filename" : "Icon_Large_Back-1280x768.png",
+			"scale" : "1x"
+    },
+    {
+      "idiom" : "tv",
+      "scale" : "2x"
+    }
+	],
+	"info" : {
+		"version" : 1,
+		"author" : "xcode"
+	}
+}
\ No newline at end of file
diff --git a/Engine/Build/TVOS/Resources/Assets.xcassets/AppIcon.brandassets/App Icon - Large.imagestack/Back.imagestacklayer/Contents.json b/Engine/Build/TVOS/Resources/Assets.xcassets/AppIcon.brandassets/App Icon - Large.imagestack/Back.imagestacklayer/Contents.json
new file mode 100644
--- /dev/null
+++ b/Engine/Build/TVOS/Resources/Assets.xcassets/AppIcon.brandassets/App Icon - Large.imagestack/Back.imagestacklayer/Contents.json	
@@ -0,0 +1,6 @@
+{
+	"info" : {
+		"version" : 1,
+		"author" : "xcode"
+	}
+}
\ No newline at end of file
diff --git a/Engine/Build/TVOS/Resources/Assets.xcassets/AppIcon.brandassets/App Icon - Large.imagestack/Contents.json b/Engine/Build/TVOS/Resources/Assets.xcassets/AppIcon.brandassets/App Icon - Large.imagestack/Contents.json
new file mode 100644
--- /dev/null
+++ b/Engine/Build/TVOS/Resources/Assets.xcassets/AppIcon.brandassets/App Icon - Large.imagestack/Contents.json	
@@ -0,0 +1,17 @@
+{
+	"layers" : [
+		{
+			"filename" : "Front.imagestacklayer"
+		},
+		{
+			"filename" : "Middle.imagestacklayer"
+		},
+		{
+			"filename" : "Back.imagestacklayer"
+		}
+	],
+	"info" : {
+		"version" : 1,
+		"author" : "xcode"
+	}
+}
\ No newline at end of file
diff --git a/Engine/Build/TVOS/Resources/Assets.xcassets/AppIcon.brandassets/App Icon - Large.imagestack/Front.imagestacklayer/Content.imageset/Contents.json b/Engine/Build/TVOS/Resources/Assets.xcassets/AppIcon.brandassets/App Icon - Large.imagestack/Front.imagestacklayer/Content.imageset/Contents.json
new file mode 100644
--- /dev/null
+++ b/Engine/Build/TVOS/Resources/Assets.xcassets/AppIcon.brandassets/App Icon - Large.imagestack/Front.imagestacklayer/Content.imageset/Contents.json	
@@ -0,0 +1,17 @@
+{
+	"images" : [
+		{
+			"idiom" : "tv",
+			"filename" : "Icon_Large_Front-1280x768.png",
+			"scale" : "1x"
+    },
+    {
+      "idiom" : "tv",
+      "scale" : "2x"
+    }
+	],
+	"info" : {
+		"version" : 1,
+		"author" : "xcode"
+	}
+}
\ No newline at end of file
diff --git a/Engine/Build/TVOS/Resources/Assets.xcassets/AppIcon.brandassets/App Icon - Large.imagestack/Front.imagestacklayer/Contents.json b/Engine/Build/TVOS/Resources/Assets.xcassets/AppIcon.brandassets/App Icon - Large.imagestack/Front.imagestacklayer/Contents.json
new file mode 100644
--- /dev/null
+++ b/Engine/Build/TVOS/Resources/Assets.xcassets/AppIcon.brandassets/App Icon - Large.imagestack/Front.imagestacklayer/Contents.json	
@@ -0,0 +1,6 @@
+{
+	"info" : {
+		"version" : 1,
+		"author" : "xcode"
+	}
+}
\ No newline at end of file
diff --git a/Engine/Build/TVOS/Resources/Assets.xcassets/AppIcon.brandassets/App Icon - Large.imagestack/Middle.imagestacklayer/Content.imageset/Contents.json b/Engine/Build/TVOS/Resources/Assets.xcassets/AppIcon.brandassets/App Icon - Large.imagestack/Middle.imagestacklayer/Content.imageset/Contents.json
new file mode 100644
--- /dev/null
+++ b/Engine/Build/TVOS/Resources/Assets.xcassets/AppIcon.brandassets/App Icon - Large.imagestack/Middle.imagestacklayer/Content.imageset/Contents.json	
@@ -0,0 +1,17 @@
+{
+	"images" : [
+		{
+			"idiom" : "tv",
+			"filename" : "Icon_Large_Middle-1280x768.png",
+			"scale" : "1x"
+    },
+    {
+      "idiom" : "tv",
+      "scale" : "2x"
+    }
+	],
+	"info" : {
+		"version" : 1,
+		"author" : "xcode"
+	}
+}
\ No newline at end of file
diff --git a/Engine/Build/TVOS/Resources/Assets.xcassets/AppIcon.brandassets/App Icon - Large.imagestack/Middle.imagestacklayer/Contents.json b/Engine/Build/TVOS/Resources/Assets.xcassets/AppIcon.brandassets/App Icon - Large.imagestack/Middle.imagestacklayer/Contents.json
new file mode 100644
--- /dev/null
+++ b/Engine/Build/TVOS/Resources/Assets.xcassets/AppIcon.brandassets/App Icon - Large.imagestack/Middle.imagestacklayer/Contents.json	
@@ -0,0 +1,6 @@
+{
+	"info" : {
+		"version" : 1,
+		"author" : "xcode"
+	}
+}
\ No newline at end of file
diff --git a/Engine/Build/TVOS/Resources/Assets.xcassets/AppIcon.brandassets/App Icon - Small.imagestack/Back.imagestacklayer/Content.imageset/Contents.json b/Engine/Build/TVOS/Resources/Assets.xcassets/AppIcon.brandassets/App Icon - Small.imagestack/Back.imagestacklayer/Content.imageset/Contents.json
new file mode 100644
--- /dev/null
+++ b/Engine/Build/TVOS/Resources/Assets.xcassets/AppIcon.brandassets/App Icon - Small.imagestack/Back.imagestacklayer/Content.imageset/Contents.json	
@@ -0,0 +1,18 @@
+{
+	"images" : [
+		{
+			"idiom" : "tv",
+			"filename" : "Icon_Small_Back-400x240.png",
+			"scale" : "1x"
+		},
+		{
+			"idiom" : "tv",
+			"filename" : "Icon_Small_Back-800x480.png",
+			"scale" : "2x"
+		}
+	],
+	"info" : {
+		"version" : 1,
+		"author" : "xcode"
+	}
+}
\ No newline at end of file
diff --git a/Engine/Build/TVOS/Resources/Assets.xcassets/AppIcon.brandassets/App Icon - Small.imagestack/Back.imagestacklayer/Contents.json b/Engine/Build/TVOS/Resources/Assets.xcassets/AppIcon.brandassets/App Icon - Small.imagestack/Back.imagestacklayer/Contents.json
new file mode 100644
--- /dev/null
+++ b/Engine/Build/TVOS/Resources/Assets.xcassets/AppIcon.brandassets/App Icon - Small.imagestack/Back.imagestacklayer/Contents.json	
@@ -0,0 +1,6 @@
+{
+	"info" : {
+		"version" : 1,
+		"author" : "xcode"
+	}
+}
\ No newline at end of file
diff --git a/Engine/Build/TVOS/Resources/Assets.xcassets/AppIcon.brandassets/App Icon - Small.imagestack/Contents.json b/Engine/Build/TVOS/Resources/Assets.xcassets/AppIcon.brandassets/App Icon - Small.imagestack/Contents.json
new file mode 100644
--- /dev/null
+++ b/Engine/Build/TVOS/Resources/Assets.xcassets/AppIcon.brandassets/App Icon - Small.imagestack/Contents.json	
@@ -0,0 +1,17 @@
+{
+	"layers" : [
+		{
+			"filename" : "Front.imagestacklayer"
+		},
+		{
+			"filename" : "Middle.imagestacklayer"
+		},
+		{
+			"filename" : "Back.imagestacklayer"
+		}
+	],
+	"info" : {
+		"version" : 1,
+		"author" : "xcode"
+	}
+}
\ No newline at end of file
diff --git a/Engine/Build/TVOS/Resources/Assets.xcassets/AppIcon.brandassets/App Icon - Small.imagestack/Front.imagestacklayer/Content.imageset/Contents.json b/Engine/Build/TVOS/Resources/Assets.xcassets/AppIcon.brandassets/App Icon - Small.imagestack/Front.imagestacklayer/Content.imageset/Contents.json
new file mode 100644
--- /dev/null
+++ b/Engine/Build/TVOS/Resources/Assets.xcassets/AppIcon.brandassets/App Icon - Small.imagestack/Front.imagestacklayer/Content.imageset/Contents.json	
@@ -0,0 +1,18 @@
+{
+	"images" : [
+		{
+			"idiom" : "tv",
+			"filename" : "Icon_Small_Front-400x240.png",
+			"scale" : "1x"
+		},
+		{
+			"idiom" : "tv",
+			"filename" : "Icon_Small_Front-800x480.png",
+			"scale" : "2x"
+		}
+	],
+	"info" : {
+		"version" : 1,
+		"author" : "xcode"
+	}
+}
\ No newline at end of file
diff --git a/Engine/Build/TVOS/Resources/Assets.xcassets/AppIcon.brandassets/App Icon - Small.imagestack/Front.imagestacklayer/Contents.json b/Engine/Build/TVOS/Resources/Assets.xcassets/AppIcon.brandassets/App Icon - Small.imagestack/Front.imagestacklayer/Contents.json
new file mode 100644
--- /dev/null
+++ b/Engine/Build/TVOS/Resources/Assets.xcassets/AppIcon.brandassets/App Icon - Small.imagestack/Front.imagestacklayer/Contents.json	
@@ -0,0 +1,6 @@
+{
+	"info" : {
+		"version" : 1,
+		"author" : "xcode"
+	}
+}
\ No newline at end of file
diff --git a/Engine/Build/TVOS/Resources/Assets.xcassets/AppIcon.brandassets/App Icon - Small.imagestack/Middle.imagestacklayer/Content.imageset/Contents.json b/Engine/Build/TVOS/Resources/Assets.xcassets/AppIcon.brandassets/App Icon - Small.imagestack/Middle.imagestacklayer/Content.imageset/Contents.json
new file mode 100644
--- /dev/null
+++ b/Engine/Build/TVOS/Resources/Assets.xcassets/AppIcon.brandassets/App Icon - Small.imagestack/Middle.imagestacklayer/Content.imageset/Contents.json	
@@ -0,0 +1,18 @@
+{
+	"images" : [
+		{
+			"idiom" : "tv",
+			"filename" : "Icon_Small_Middle-400x240.png",
+			"scale" : "1x"
+		},
+		{
+			"idiom" : "tv",
+			"filename" : "Icon_Small_Middle-800x480.png",
+			"scale" : "2x"
+		}
+	],
+	"info" : {
+		"version" : 1,
+		"author" : "xcode"
+	}
+}
\ No newline at end of file
diff --git a/Engine/Build/TVOS/Resources/Assets.xcassets/AppIcon.brandassets/App Icon - Small.imagestack/Middle.imagestacklayer/Contents.json b/Engine/Build/TVOS/Resources/Assets.xcassets/AppIcon.brandassets/App Icon - Small.imagestack/Middle.imagestacklayer/Contents.json
new file mode 100644
--- /dev/null
+++ b/Engine/Build/TVOS/Resources/Assets.xcassets/AppIcon.brandassets/App Icon - Small.imagestack/Middle.imagestacklayer/Contents.json	
@@ -0,0 +1,6 @@
+{
+	"info" : {
+		"version" : 1,
+		"author" : "xcode"
+	}
+}
\ No newline at end of file
diff --git a/Engine/Build/TVOS/Resources/Assets.xcassets/AppIcon.brandassets/Contents.json b/Engine/Build/TVOS/Resources/Assets.xcassets/AppIcon.brandassets/Contents.json
new file mode 100644
--- /dev/null
+++ b/Engine/Build/TVOS/Resources/Assets.xcassets/AppIcon.brandassets/Contents.json
@@ -0,0 +1,32 @@
+{
+  "assets" : [
+    {
+      "size" : "1280x768",
+      "idiom" : "tv",
+      "filename" : "App Icon - Large.imagestack",
+      "role" : "primary-app-icon"
+    },
+    {
+      "size" : "400x240",
+      "idiom" : "tv",
+      "filename" : "App Icon - Small.imagestack",
+      "role" : "primary-app-icon"
+    },
+    {
+      "size" : "2320x720",
+      "idiom" : "tv",
+      "filename" : "Top Shelf Image Wide.imageset",
+      "role" : "top-shelf-image-wide"
+    },
+    {
+      "size" : "1920x720",
+      "idiom" : "tv",
+      "filename" : "Top Shelf Image.imageset",
+      "role" : "top-shelf-image"
+    }
+  ],
+  "info" : {
+    "version" : 1,
+    "author" : "xcode"
+  }
+}
diff --git a/Engine/Build/TVOS/Resources/Assets.xcassets/AppIcon.brandassets/Top Shelf Image Wide.imageset/Contents.json b/Engine/Build/TVOS/Resources/Assets.xcassets/AppIcon.brandassets/Top Shelf Image Wide.imageset/Contents.json
new file mode 100644
--- /dev/null
+++ b/Engine/Build/TVOS/Resources/Assets.xcassets/AppIcon.brandassets/Top Shelf Image Wide.imageset/Contents.json	
@@ -0,0 +1,18 @@
+{
+	"images" : [
+		{
+			"idiom" : "tv",
+			"filename" : "TopShelfWide-1920x720.png",
+			"scale" : "1x"
+		},
+		{
+			"idiom" : "tv",
+			"filename" : "[REDACTED_EMAIL]",
+			"scale" : "2x"
+		}
+	],
+	"info" : {
+		"version" : 1,
+		"author" : "xcode"
+	}
+}
\ No newline at end of file
diff --git a/Engine/Build/TVOS/Resources/Assets.xcassets/AppIcon.brandassets/Top Shelf Image.imageset/Contents.json b/Engine/Build/TVOS/Resources/Assets.xcassets/AppIcon.brandassets/Top Shelf Image.imageset/Contents.json
new file mode 100644
--- /dev/null
+++ b/Engine/Build/TVOS/Resources/Assets.xcassets/AppIcon.brandassets/Top Shelf Image.imageset/Contents.json	
@@ -0,0 +1,18 @@
+{
+	"images" : [
+		{
+			"idiom" : "tv",
+			"filename" : "TopShelf.png",
+			"scale" : "1x"
+		},
+		{
+			"idiom" : "tv",
+			"filename" : "[REDACTED_EMAIL]",
+			"scale" : "2x"
+		}
+	],
+	"info" : {
+		"version" : 1,
+		"author" : "xcode"
+	}
+}
\ No newline at end of file
diff --git a/Engine/Build/TVOS/Resources/Assets.xcassets/Contents.json b/Engine/Build/TVOS/Resources/Assets.xcassets/Contents.json
new file mode 100644
--- /dev/null
+++ b/Engine/Build/TVOS/Resources/Assets.xcassets/Contents.json
@@ -0,0 +1,6 @@
+{
+	"info" : {
+		"version" : 1,
+		"author" : "xcode"
+	}
+}
\ No newline at end of file
diff --git a/Engine/Build/TVOS/Resources/Assets.xcassets/LaunchImage.launchimage/Contents.json b/Engine/Build/TVOS/Resources/Assets.xcassets/LaunchImage.launchimage/Contents.json
new file mode 100644
--- /dev/null
+++ b/Engine/Build/TVOS/Resources/Assets.xcassets/LaunchImage.launchimage/Contents.json
@@ -0,0 +1,16 @@
+{
+	"images" : [
+		{
+			"orientation" : "landscape",
+			"idiom" : "tv",
+			"filename" : "Launch.png",
+			"extent" : "full-screen",
+			"minimum-system-version" : "9.0",
+			"scale" : "1x"
+		}
+	],
+	"info" : {
+		"version" : 1,
+		"author" : "xcode"
+	}
+}
\ No newline at end of file
diff --git a/Engine/Source/Programs/UnrealBuildTool/Platform/IOS/IOSToolChain.cs b/Engine/Source/Programs/UnrealBuildTool/Platform/IOS/IOSToolChain.cs
--- a/Engine/Source/Programs/UnrealBuildTool/Platform/IOS/IOSToolChain.cs
+++ b/Engine/Source/Programs/UnrealBuildTool/Platform/IOS/IOSToolChain.cs
@@ -919,7 +919,7 @@ namespace UnrealBuildTool
 			{
 				Arguments.Append(" --launch-image LaunchImage");
 				Arguments.Append(" --filter-for-device-model AppleTV5,3");
-				Arguments.Append(" --filter-for-device-os-version 10.0");
+				//Arguments.Append(" --filter-for-device-os-version 10.0");
 				Arguments.Append(" --target-device tv");
 				Arguments.Append(" --minimum-deployment-target 10.0");
 				Arguments.Append(" --platform appletvos");
@@ -1168,89 +1168,68 @@ namespace UnrealBuildTool
 			string ResourcesDir = Path.Combine(IntermediateDir, "Resources");
             if (Platform == CppPlatform.TVOS)
             {
-                string[] Directories = { "Assets.xcassets",
-                                            Path.Combine("Assets.xcassets", "AppIcon.brandassets"),
-                                            Path.Combine("Assets.xcassets", "AppIcon.brandassets", "AppIconLarge.imagestack"),
-                                            Path.Combine("Assets.xcassets", "AppIcon.brandassets", "AppIconLarge.imagestack", "Back.imagestacklayer"),
-                                            Path.Combine("Assets.xcassets", "AppIcon.brandassets", "AppIconLarge.imagestack", "Back.imagestacklayer", "Content.imageset"),
-                                            Path.Combine("Assets.xcassets", "AppIcon.brandassets", "AppIconLarge.imagestack", "Front.imagestacklayer"),
-                                            Path.Combine("Assets.xcassets", "AppIcon.brandassets", "AppIconLarge.imagestack", "Front.imagestacklayer", "Content.imageset"),
-                                            Path.Combine("Assets.xcassets", "AppIcon.brandassets", "AppIconLarge.imagestack", "Middle.imagestacklayer"),
-                                            Path.Combine("Assets.xcassets", "AppIcon.brandassets", "AppIconLarge.imagestack", "Middle.imagestacklayer", "Content.imageset"),
-                                            Path.Combine("Assets.xcassets", "AppIcon.brandassets", "AppIconSmall.imagestack"),
-                                            Path.Combine("Assets.xcassets", "AppIcon.brandassets", "AppIconSmall.imagestack", "Back.imagestacklayer"),
-                                            Path.Combine("Assets.xcassets", "AppIcon.brandassets", "AppIconSmall.imagestack", "Back.imagestacklayer", "Content.imageset"),
-                                            Path.Combine("Assets.xcassets", "AppIcon.brandassets", "AppIconSmall.imagestack", "Front.imagestacklayer"),
-                                            Path.Combine("Assets.xcassets", "AppIcon.brandassets", "AppIconSmall.imagestack", "Front.imagestacklayer", "Content.imageset"),
-                                            Path.Combine("Assets.xcassets", "AppIcon.brandassets", "AppIconSmall.imagestack", "Middle.imagestacklayer"),
-                                            Path.Combine("Assets.xcassets", "AppIcon.brandassets", "AppIconSmall.imagestack", "Middle.imagestacklayer", "Content.imageset"),
-                                            Path.Combine("Assets.xcassets", "AppIcon.brandassets", "TopShelf.imageset"),
-                                            Path.Combine("Assets.xcassets", "LaunchImage.launchimage"),
-                };
-                string[] Contents = { "{\n\t\"info\" : {\n\t\t\"version\" : 1,\n\t\t\"author\" : \"xcode\"\n\t}\n}",
-                                    "{\n\t\"assets\" : [\n\t\t{\n\t\t\t\"size\" : \"1280x768\",\n\t\t\t\"idiom\" : \"tv\",\n\t\t\t\"filename\" : \"AppIconLarge.imagestack\",\n\t\t\t\"role\" : \"primary-app-icon\"\n\t\t},\n\t\t{\n\t\t\t\"size\" : \"400x240\",\n\t\t\t\"idiom\" : \"tv\",\n\t\t\t\"filename\" : \"AppIconSmall.imagestack\",\n\t\t\t\"role\" : \"primary-app-icon\"\n\t\t},\n\t\t{\n\t\t\t\"size\" : \"1920x720\",\n\t\t\t\"idiom\" : \"tv\",\n\t\t\t\"filename\" : \"TopShelf.imageset\",\n\t\t\t\"role\" : \"top-shelf-image\"\n\t\t}\n\t],\n\t\"info\" : {\n\t\t\"version\" : 1,\n\t\t\"author\" : \"xcode\"\n\t}\n}",
-                                    "{\n\t\"layers\" : [\n\t\t{\n\t\t\t\"filename\" : \"Front.imagestacklayer\"\n\t\t},\n\t\t{\n\t\t\t\"filename\" : \"Middle.imagestacklayer\"\n\t\t},\n\t\t{\n\t\t\t\"filename\" : \"Back.imagestacklayer\"\n\t\t}\n\t],\n\t\"info\" : {\n\t\t\"version\" : 1,\n\t\t\"author\" : \"xcode\"\n\t}\n}",
-                                    "{\n\t\"info\" : {\n\t\t\"version\" : 1,\n\t\t\"author\" : \"xcode\"\n\t}\n}",
-                                    "{\n\t\"images\" : [\n\t\t{\n\t\t\t\"idiom\" : \"tv\",\n\t\t\t\"filename\" : \"Icon_Large_Back.png\",\n\t\t\t\"scale\" : \"1x\"\n\t\t}\n\t],\n\t\"info\" : {\n\t\t\"version\" : 1,\n\t\t\"author\" : \"xcode\"\n\t}\n}",
-                                    "{\n\t\"info\" : {\n\t\t\"version\" : 1,\n\t\t\"author\" : \"xcode\"\n\t}\n}",
-                                    "{\n\t\"images\" : [\n\t\t{\n\t\t\t\"idiom\" : \"tv\",\n\t\t\t\"filename\" : \"Icon_Large_Front.png\",\n\t\t\t\"scale\" : \"1x\"\n\t\t}\n\t],\n\t\"info\" : {\n\t\t\"version\" : 1,\n\t\t\"author\" : \"xcode\"\n\t}\n}",
-                                    "{\n\t\"info\" : {\n\t\t\"version\" : 1,\n\t\t\"author\" : \"xcode\"\n\t}\n}",
-                                    "{\n\t\"images\" : [\n\t\t{\n\t\t\t\"idiom\" : \"tv\",\n\t\t\t\"filename\" : \"Icon_Large_Middle.png\",\n\t\t\t\"scale\" : \"1x\"\n\t\t}\n\t],\n\t\"info\" : {\n\t\t\"version\" : 1,\n\t\t\"author\" : \"xcode\"\n\t}\n}",
-                                    "{\n\t\"layers\" : [\n\t\t{\n\t\t\t\"filename\" : \"Front.imagestacklayer\"\n\t\t},\n\t\t{\n\t\t\t\"filename\" : \"Middle.imagestacklayer\"\n\t\t},\n\t\t{\n\t\t\t\"filename\" : \"Back.imagestacklayer\"\n\t\t}\n\t],\n\t\"info\" : {\n\t\t\"version\" : 1,\n\t\t\"author\" : \"xcode\"\n\t}\n}",
-                                    "{\n\t\"info\" : {\n\t\t\"version\" : 1,\n\t\t\"author\" : \"xcode\"\n\t}\n}",
-                                    "{\n\t\"images\" : [\n\t\t{\n\t\t\t\"idiom\" : \"tv\",\n\t\t\t\"filename\" : \"Icon_Small_Back.png\",\n\t\t\t\"scale\" : \"1x\"\n\t\t}\n\t],\n\t\"info\" : {\n\t\t\"version\" : 1,\n\t\t\"author\" : \"xcode\"\n\t}\n}",
-                                    "{\n\t\"info\" : {\n\t\t\"version\" : 1,\n\t\t\"author\" : \"xcode\"\n\t}\n}",
-                                    "{\n\t\"images\" : [\n\t\t{\n\t\t\t\"idiom\" : \"tv\",\n\t\t\t\"filename\" : \"Icon_Small_Front.png\",\n\t\t\t\"scale\" : \"1x\"\n\t\t}\n\t],\n\t\"info\" : {\n\t\t\"version\" : 1,\n\t\t\"author\" : \"xcode\"\n\t}\n}",
-                                    "{\n\t\"info\" : {\n\t\t\"version\" : 1,\n\t\t\"author\" : \"xcode\"\n\t}\n}",
-                                    "{\n\t\"images\" : [\n\t\t{\n\t\t\t\"idiom\" : \"tv\",\n\t\t\t\"filename\" : \"Icon_Small_Middle.png\",\n\t\t\t\"scale\" : \"1x\"\n\t\t}\n\t],\n\t\"info\" : {\n\t\t\"version\" : 1,\n\t\t\"author\" : \"xcode\"\n\t}\n}",
-                                    "{\n\t\"images\" : [\n\t\t{\n\t\t\t\"idiom\" : \"tv\",\n\t\t\t\"filename\" : \"TopShelf.png\",\n\t\t\t\"scale\" : \"1x\"\n\t\t}\n\t],\n\t\"info\" : {\n\t\t\"version\" : 1,\n\t\t\"author\" : \"xcode\"\n\t}\n}",
-                                    "{\n\t\"images\" : [\n\t\t{\n\t\t\t\"orientation\" : \"landscape\",\n\t\t\t\"idiom\" : \"tv\",\n\t\t\t\"filename\" : \"Launch.png\",\n\t\t\t\"extent\" : \"full-screen\",\n\t\t\t\"minimum-system-version\" : \"9.0\",\n\t\t\t\"scale\" : \"1x\"\n\t\t}\n\t],\n\t\"info\" : {\n\t\t\"version\" : 1,\n\t\t\"author\" : \"xcode\"\n\t}\n}",
-                };
-                string[] Images = { null,
-                                null,
-                                null,
-                                null,
-                                "Icon_Large_Back.png",
-                                null,
-                                "Icon_Large_Front.png",
-                                null,
-                                "Icon_Large_Middle.png",
-                                null,
-                                null,
-                                "Icon_Small_Back.png",
-                                null,
-                                "Icon_Small_Front.png",
-                                null,
-                                "Icon_Small_Middle.png",
-                                "TopShelf.png",
-                                "Launch.png"
-                };
+				// copy the template asset catalog to the appropriate directory
+				string Dir = Path.Combine(ResourcesDir, "Assets.xcassets");
+				if (!Directory.Exists(Dir))
+				{
+					Directory.CreateDirectory(Dir);
+				}
+				// create the directories
+				foreach (string directory in Directory.EnumerateDirectories(Path.Combine(EngineDir, "Build", "TVOS", "Resources", "Assets.xcassets"), "*", SearchOption.AllDirectories))
+				{
+					Dir = directory.Replace(Path.Combine(EngineDir, "Build", "TVOS"), IntermediateDir);
+					if (!Directory.Exists(Dir))
+					{
+						Directory.CreateDirectory(Dir);
+					}
+				}
+				// copy the default files
+				foreach (string file in Directory.EnumerateFiles(Path.Combine(EngineDir, "Build", "TVOS", "Resources", "Assets.xcassets"), "*", SearchOption.AllDirectories))
+				{
+					Dir = file.Replace(Path.Combine(EngineDir, "Build", "TVOS"), IntermediateDir);
+					File.Copy(file, Dir, true);
+					FileInfo DestFileInfo = new FileInfo(Dir);
+					DestFileInfo.Attributes = DestFileInfo.Attributes & ~FileAttributes.ReadOnly;
+				}
+				// copy the icons from the game directory if it has any
+				string[][] Images = {
+					new string []{ "Icon_Large_Back-1280x768.png", "AppIcon.brandassets/App Icon - Large.imagestack/Back.imagestacklayer/Content.imageset" },
+					new string []{ "Icon_Large_Front-1280x768.png", "AppIcon.brandassets/App Icon - Large.imagestack/Front.imagestacklayer/Content.imageset" },
+					new string []{ "Icon_Large_Middle-1280x768.png", "AppIcon.brandassets/App Icon - Large.imagestack/Middle.imagestacklayer/Content.imageset" },
+					new string []{ "Icon_Small_Back-400x240.png", "AppIcon.brandassets/App Icon - Small.imagestack/Back.imagestacklayer/Content.imageset" },
+					new string []{ "Icon_Small_Back-800x480.png", "AppIcon.brandassets/App Icon - Small.imagestack/Back.imagestacklayer/Content.imageset" },
+					new string []{ "Icon_Small_Front-400x240.png", "AppIcon.brandassets/App Icon - Small.imagestack/Front.imagestacklayer/Content.imageset" },
+					new string []{ "Icon_Small_Front-800x480.png", "AppIcon.brandassets/App Icon - Small.imagestack/Front.imagestacklayer/Content.imageset" },
+					new string []{ "Icon_Small_Middle-400x240.png", "AppIcon.brandassets/App Icon - Small.imagestack/Middle.imagestacklayer/Content.imageset" },
+					new string []{ "Icon_Small_Middle-800x480.png", "AppIcon.brandassets/App Icon - Small.imagestack/Middle.imagestacklayer/Content.imageset" },
+					//new string []{ "TopShelfMarketWide.png", "AppIcon.brandassets/Top Shelf Image Wide.imageset" },
+					//new string []{ "[REDACTED_EMAIL]", "AppIcon.brandassets/Top Shelf Image Wide.imageset" },
+					new string []{ "[REDACTED_EMAIL]", "AppIcon.brandassets/Top Shelf Image Wide.imageset" },
+					new string []{ "TopShelfWide-1920x720.png", "AppIcon.brandassets/Top Shelf Image Wide.imageset" },
+					//new string []{ "TopShelfMarket.png", "AppIcon.brandassets/Top Shelf Image.imageset" },
+					//new string []{ "[REDACTED_EMAIL]", "AppIcon.brandassets/Top Shelf Image.imageset" },
+					new string []{ "TopShelf.png", "AppIcon.brandassets/Top Shelf Image.imageset" },
+					new string []{ "[REDACTED_EMAIL]", "AppIcon.brandassets/Top Shelf Image.imageset" },
+					new string []{ "Launch.png", "LaunchImage.launchimage" },
+					new string []{ "[REDACTED_EMAIL]", "LaunchImage.launchimage" },
+				};
+				Dir = Path.Combine(IntermediateDir, "Resources", "Assets.xcassets");
+
+				string BuildResourcesGraphicsDir = Path.Combine(BuildDir, "Resources", "Graphics");
+				for (int Index = 0; Index < Images.Length; ++Index)
+				{
+					string Image = Path.Combine((Directory.Exists(Path.Combine(BuildDir, "Resources", "Graphics")) ? (BuildDir) : (Path.Combine(EngineDir, "Build", "TVOS"))), "Resources", "Graphics", Images[Index][0]);
+					if (File.Exists(Image))
+					{
+						bUserImagesExist |= Image.StartsWith(BuildResourcesGraphicsDir);
 
-                // create asset catalog for images
-                for (int i = 0; i < Directories.Length; ++i)
-                {
-                    string Dir = Path.Combine(ResourcesDir, Directories[i]);
-                    if (!Directory.Exists(Dir))
-                    {
-                        Directory.CreateDirectory(Dir);
-                    }
-                    File.WriteAllText(Path.Combine(Dir, "Contents.json"), Contents[i]);
-                    if (Images[i] != null)
-                    {
-                        // This assumption might not be true, but we need the asset catalog process to fire anyway.
-                        bUserImagesExist = true;
-
-                        string Image = Path.Combine((Directory.Exists(Path.Combine(BuildDir, "Resource", "Graphics")) ? (BuildDir) : (Path.Combine(EngineDir, "Build", "TVOS"))), "Resources", "Graphics", Images[i]);
-                        if (File.Exists(Image))
-                        {
-                            File.Copy(Image, Path.Combine(Dir, Images[i]), true);
-                            FileInfo DestFileInfo = new FileInfo(Path.Combine(Dir, Images[i]));
-                            DestFileInfo.Attributes = DestFileInfo.Attributes & ~FileAttributes.ReadOnly;
-                        }
-                    }
-                }
-            }
-            else
+						File.Copy(Image, Path.Combine(Dir, Images[Index][1], Images[Index][0]), true);
+						FileInfo DestFileInfo = new FileInfo(Path.Combine(Dir, Images[Index][1], Images[Index][0]));
+						DestFileInfo.Attributes = DestFileInfo.Attributes & ~FileAttributes.ReadOnly;
+					}
+				}
+			}
+			else
             {
                 // copy the template asset catalog to the appropriate directory
                 string Dir = Path.Combine(ResourcesDir, "Assets.xcassets");
diff --git a/Engine/Source/Programs/UnrealBuildTool/Platform/TVOS/UEDeployTVOS.cs b/Engine/Source/Programs/UnrealBuildTool/Platform/TVOS/UEDeployTVOS.cs
--- a/Engine/Source/Programs/UnrealBuildTool/Platform/TVOS/UEDeployTVOS.cs
+++ b/Engine/Source/Programs/UnrealBuildTool/Platform/TVOS/UEDeployTVOS.cs
@@ -143,8 +143,10 @@ namespace UnrealBuildTool
             Text.AppendLine("\t<key>TVTopShelfImage</key>");
             Text.AppendLine("\t<dict>");
             Text.AppendLine("\t\t<key>TVTopShelfPrimaryImage</key>");
-            Text.AppendLine("\t\t<string>TopShelf</string>");
-            Text.AppendLine("\t</dict>");
+            Text.AppendLine("\t\t<string>Top Shelf Image</string>");
+			Text.AppendLine("\t\t<key>TVTopShelfPrimaryImageWide</key>");
+			Text.AppendLine("\t\t<string>Top Shelf Image Wide</string>");
+			Text.AppendLine("\t</dict>");
             Text.AppendLine("\t<key>UILaunchImages</key>");
             Text.AppendLine("\t<array>");
             Text.AppendLine("\t\t<dict>");
@@ -161,7 +163,7 @@ namespace UnrealBuildTool
             Text.AppendLine("\t<key>CFBundleIcons</key>");
             Text.AppendLine("\t<dict>");
             Text.AppendLine("\t\t<key>CFBundlePrimaryIcon</key>");
-            Text.AppendLine("\t\t<string>AppIconSmall</string>");
+            Text.AppendLine("\t\t<string>App Icon - Small</string>");
             Text.AppendLine("\t</dict>");
 
             /*			Text.AppendLine("\t<key>CFBundleIcons</key>");
