PR #5180: Fix wrong .natvis file path in generated Visual Studio solution (Contributed by [REDACTED_PERSON])

diff --git a/Engine/Source/Programs/UnrealBuildTool/ProjectFiles/VisualStudio/VCProjectFileGenerator.cs b/Engine/Source/Programs/UnrealBuildTool/ProjectFiles/VisualStudio/VCProjectFileGenerator.cs
--- a/Engine/Source/Programs/UnrealBuildTool/ProjectFiles/VisualStudio/VCProjectFileGenerator.cs
+++ b/Engine/Source/Programs/UnrealBuildTool/ProjectFiles/VisualStudio/VCProjectFileGenerator.cs
@@ -548,12 +548,11 @@ namespace UnrealBuildTool
 
 				// Get the path to the visualizers file. Try to make it relative to the solution directory, but fall back to a full path if it's a foreign project.
 				FileReference VisualizersFile = FileReference.Combine(UnrealBuildTool.EngineDirectory, "Extras", "VisualStudioDebugging", "UE4.natvis");
-				DirectoryReference SolutionDir = new FileReference(SolutionFileName).Directory;
 
 				// Add the visualizers at the solution level. Doesn't seem to be picked up from a makefile project in VS2017 15.8.5.
 				VCSolutionFileContent.AppendLine(String.Format("Project(\"{{2150E333-8FDC-42A3-9474-1A3956D46DE8}}\") = \"Visualizers\", \"Visualizers\", \"{0}\"", Guid.NewGuid().ToString("B").ToUpperInvariant()));
 				VCSolutionFileContent.AppendLine("\tProjectSection(SolutionItems) = preProject");
-				VCSolutionFileContent.AppendLine("\t\t{0} = {0}", VisualizersFile.MakeRelativeTo(SolutionDir));
+				VCSolutionFileContent.AppendLine("\t\t{0} = {0}", VisualizersFile.MakeRelativeTo(MasterProjectPath));
 				VCSolutionFileContent.AppendLine("\tEndProjectSection");
 				VCSolutionFileContent.AppendLine("EndProject");
 			}
