Allow multiple targets of the same type when generating project files. Secondary targets (as determined by sorting target files alphabetically) are placed in their own generated project file, with the ability to override the name of it specified by the GeneratedProjectName variable in the TargetRules class.

diff --git a/Engine/Source/Programs/UnrealBuildTool/Configuration/TargetRules.cs b/Engine/Source/Programs/UnrealBuildTool/Configuration/TargetRules.cs
--- a/Engine/Source/Programs/UnrealBuildTool/Configuration/TargetRules.cs
+++ b/Engine/Source/Programs/UnrealBuildTool/Configuration/TargetRules.cs
@@ -1062,6 +1062,11 @@ namespace UnrealBuildTool
 		[RequiresUniqueBuildEnvironment]
 		public string AdditionalLinkerArguments;
 
+		/// <summary>
+		/// When generating project files, specifies the name of the project file to use when there are multiple targets of the same type.
+		/// </summary>
+		public string GeneratedProjectName;
+
 		/// <summary>
 		/// Android-specific target settings.
 		/// </summary>
@@ -2075,6 +2080,11 @@ namespace UnrealBuildTool
 			get { return Inner.AdditionalLinkerArguments; }
 		}
 
+		public string GeneratedProjectName
+		{
+			get { return Inner.GeneratedProjectName; }
+		}
+
 		public ReadOnlyAndroidTargetRules AndroidPlatform
 		{
 			get;
diff --git a/Engine/Source/Programs/UnrealBuildTool/ProjectFiles/Project.cs b/Engine/Source/Programs/UnrealBuildTool/ProjectFiles/Project.cs
--- a/Engine/Source/Programs/UnrealBuildTool/ProjectFiles/Project.cs
+++ b/Engine/Source/Programs/UnrealBuildTool/ProjectFiles/Project.cs
@@ -1,4 +1,4 @@
-﻿// Copyright 1998-2018 Epic Games, Inc. All Rights Reserved.
+// Copyright 1998-2018 Epic Games, Inc. All Rights Reserved.
 
 using System;
 using System.Collections.Generic;
@@ -143,6 +143,14 @@ namespace UnrealBuildTool
 			protected set;
 		}
 
+		/// <summary>
+		/// The base directory for files within this project
+		/// </summary>
+		public DirectoryReference BaseDir
+		{
+			get;
+			set;
+		}
 
 		/// Returns true if this is a generated project (as opposed to an imported project)
 		public bool IsGeneratedProject
diff --git a/Engine/Source/Programs/UnrealBuildTool/ProjectFiles/ProjectFileGenerator.cs b/Engine/Source/Programs/UnrealBuildTool/ProjectFiles/ProjectFileGenerator.cs
--- a/Engine/Source/Programs/UnrealBuildTool/ProjectFiles/ProjectFileGenerator.cs
+++ b/Engine/Source/Programs/UnrealBuildTool/ProjectFiles/ProjectFileGenerator.cs
@@ -661,16 +661,17 @@ namespace UnrealBuildTool
 			// we're generating project files for.
 			List<FileReference> AllModuleFiles = DiscoverModules(AllGameProjects);
 
-			ProjectFile EngineProject = null;
-			ProjectFile EnterpriseProject = null;
-			Dictionary<DirectoryReference, ProjectFile> GameProjects = null;
-			Dictionary<DirectoryReference, ProjectFile> ModProjects = null;
-			Dictionary<FileReference, ProjectFile> ProgramProjects = null;
-			Dictionary<DirectoryReference, ProjectFile> TemplateGameProjects = null;
-			Dictionary<DirectoryReference, ProjectFile> SampleGameProjects = null;
+			ProjectFile EngineProject;
+			ProjectFile EnterpriseProject;
+			List<ProjectFile> GameProjects;
+			List<ProjectFile> ModProjects;
+			Dictionary<FileReference, ProjectFile> ProgramProjects;
 			{
 				// Setup buildable projects for all targets
-				AddProjectsForAllTargets( AllGameProjects, out EngineProject, out EnterpriseProject, out GameProjects, out ModProjects, out ProgramProjects, out TemplateGameProjects, out SampleGameProjects );
+				AddProjectsForAllTargets( AllGameProjects, out EngineProject, out EnterpriseProject, out GameProjects, out ProgramProjects );
+
+				// Add projects for mods
+				AddProjectsForMods(GameProjects, out ModProjects);
 
 				// Add all game projects and game config files
 				AddAllGameProjects(GameProjects, SupportedPlatformNames, RootFolder);
@@ -678,11 +679,11 @@ namespace UnrealBuildTool
 				// Set the game to be the default project
 				if(ModProjects.Count > 0)
 				{
-					DefaultProject = ModProjects.Values.First();
+					DefaultProject = ModProjects.First();
 				}
 				else if(bGeneratingGameProjectFiles && GameProjects.Count > 0)
 				{
-					DefaultProject = GameProjects.Values.First();
+					DefaultProject = GameProjects.First();
 				}
 
 				//Related Debug Project Files - Tuple here has the related Debug Project, SolutionFolder
@@ -697,7 +698,7 @@ namespace UnrealBuildTool
 						FileReference ProjectFilePath = FileReference.Combine(IntermediateProjectFilesPath, EngineProjectFileNameBase + ProjectFileExtension);
 
 						bool bAlreadyExisted;
-						EngineProject = FindOrAddProject(ProjectFilePath, true, out bAlreadyExisted);
+						EngineProject = FindOrAddProject(ProjectFilePath, UnrealBuildTool.EngineDirectory, true, out bAlreadyExisted);
 
 						EngineProject.IsForeignProject = false;
 						EngineProject.IsGeneratedProject = true;
@@ -767,19 +768,23 @@ namespace UnrealBuildTool
 						RootFolder.AddSubFolder(UnrealBuildTool.EnterpriseDirectory.GetDirectoryName()).ChildProjects.Add(EnterpriseProject);
 					}
 
-					foreach( ProjectFile CurModProject in ModProjects.Values )
+					foreach( ProjectFile CurModProject in ModProjects )
 					{
 						RootFolder.AddSubFolder("Mods").ChildProjects.Add(CurModProject);
 					}
 
-					foreach( ProjectFile CurGameProject in GameProjects.Values )
+					DirectoryReference TemplatesDirectory = DirectoryReference.Combine(UnrealBuildTool.RootDirectory, "Templates");
+					DirectoryReference SamplesDirectory = DirectoryReference.Combine(UnrealBuildTool.RootDirectory, "Samples");
+					DirectoryReference EnterpriseTemplatesDirectory = DirectoryReference.Combine(UnrealBuildTool.EnterpriseDirectory, "Templates");
+					DirectoryReference EnterpriseSamplesDirectory = DirectoryReference.Combine(UnrealBuildTool.EnterpriseDirectory, "Samples");
+
+					foreach( ProjectFile CurGameProject in GameProjects )
 					{
 						// Templates go under a different solution folder than games
-						DirectoryReference TemplateGameDirectory = TemplateGameProjects.FirstOrDefault( FolderAndProject => FolderAndProject.Value == CurGameProject ).Key;
-						DirectoryReference SampleGameDirectory = SampleGameProjects.FirstOrDefault( FolderAndProject => FolderAndProject.Value == CurGameProject ).Key;
-
-						if( TemplateGameDirectory != null )
+						FileReference UnrealProjectFile = CurGameProject.ProjectTargets.First().UnrealProjectFilePath;
+						if (UnrealProjectFile.IsUnderDirectory(TemplatesDirectory) || UnrealProjectFile.IsUnderDirectory(EnterpriseTemplatesDirectory))
 						{
+							DirectoryReference TemplateGameDirectory = CurGameProject.BaseDir;
 							if( TemplateGameDirectory.IsUnderDirectory((UnrealBuildTool.EnterpriseDirectory)) )
 							{
 								RootFolder.AddSubFolder(UnrealBuildTool.EnterpriseDirectory.GetDirectoryName() + Path.DirectorySeparatorChar + "Templates").ChildProjects.Add(CurGameProject);
@@ -789,8 +794,9 @@ namespace UnrealBuildTool
 								RootFolder.AddSubFolder("Templates").ChildProjects.Add(CurGameProject);
 							}
 						}
-						else if( SampleGameDirectory != null )
+						else if (UnrealProjectFile.IsUnderDirectory(SamplesDirectory) || UnrealProjectFile.IsUnderDirectory(EnterpriseSamplesDirectory))
 						{
+							DirectoryReference SampleGameDirectory = CurGameProject.BaseDir;
 							if( SampleGameDirectory.IsUnderDirectory((UnrealBuildTool.EnterpriseDirectory)) )
                             {
 								RootFolder.AddSubFolder(UnrealBuildTool.EnterpriseDirectory.GetDirectoryName() + Path.DirectorySeparatorChar + "Samples").ChildProjects.Add(CurGameProject);
@@ -928,7 +934,7 @@ namespace UnrealBuildTool
 					}
 
 					// Game/template targets
-					foreach( ProjectFile GameProject in GameProjects.Values )
+					foreach( ProjectFile GameProject in GameProjects )
 					{
 						foreach( ProjectTarget ProjectTarget in GameProject.ProjectTargets )
 						{
@@ -1141,48 +1147,48 @@ namespace UnrealBuildTool
 		/// <summary>
 		/// Adds all game project files, including target projects and config files
 		/// </summary>
-		private void AddAllGameProjects(Dictionary<DirectoryReference, ProjectFile> GameProjects, string SupportedPlatformNames, MasterProjectFolder ProjectsFolder)
+		private void AddAllGameProjects(List<ProjectFile> GameProjects, string SupportedPlatformNames, MasterProjectFolder ProjectsFolder)
 		{
-			foreach( KeyValuePair<DirectoryReference, ProjectFile> GameFolderAndProjectFile in GameProjects )
+			HashSet<DirectoryReference> UniqueGameProjectDirectories = new HashSet<DirectoryReference>();
+			foreach( ProjectFile GameProject in GameProjects )
 			{
-				DirectoryReference GameProjectDirectory = GameFolderAndProjectFile.Key;
-
-				// @todo projectfiles: We have engine localization files, but should we also add GAME localization files?
-
-				// Game config files
-				if( bIncludeConfigFiles )
+				DirectoryReference GameProjectDirectory = GameProject.BaseDir;
+				if(UniqueGameProjectDirectories.Add(GameProjectDirectory))
 				{
-					DirectoryReference GameConfigDirectory = DirectoryReference.Combine(GameProjectDirectory, "Config");
-					if( DirectoryReference.Exists(GameConfigDirectory) )
+					// @todo projectfiles: We have engine localization files, but should we also add GAME localization files?
+
+					// Game config files
+					if( bIncludeConfigFiles )
 					{
-						ProjectFile GameProjectFile = GameFolderAndProjectFile.Value;
-						GameProjectFile.AddFilesToProject( SourceFileSearch.FindFiles( GameConfigDirectory ), GameProjectDirectory );
+						DirectoryReference GameConfigDirectory = DirectoryReference.Combine(GameProjectDirectory, "Config");
+						if( DirectoryReference.Exists(GameConfigDirectory) )
+						{
+							GameProject.AddFilesToProject( SourceFileSearch.FindFiles( GameConfigDirectory ), GameProjectDirectory );
+						}
 					}
-				}
 
-				// Game build files
-				if( bIncludeBuildSystemFiles )
-				{
-					DirectoryReference GameBuildDirectory = DirectoryReference.Combine(GameProjectDirectory, "Build");
-					if( DirectoryReference.Exists(GameBuildDirectory) )
+					// Game build files
+					if( bIncludeBuildSystemFiles )
 					{
-						List<string> SubdirectoryNamesToExclude = new List<string>();
-						SubdirectoryNamesToExclude.Add("Receipts");
-						SubdirectoryNamesToExclude.Add("Scripts");
-						SubdirectoryNamesToExclude.Add("FileOpenOrder");
-						SubdirectoryNamesToExclude.Add("PipelineCaches");
+						DirectoryReference GameBuildDirectory = DirectoryReference.Combine(GameProjectDirectory, "Build");
+						if( DirectoryReference.Exists(GameBuildDirectory) )
+						{
+							List<string> SubdirectoryNamesToExclude = new List<string>();
+							SubdirectoryNamesToExclude.Add("Receipts");
+							SubdirectoryNamesToExclude.Add("Scripts");
+							SubdirectoryNamesToExclude.Add("FileOpenOrder");
+							SubdirectoryNamesToExclude.Add("PipelineCaches");
 
-						ProjectFile GameProjectFile = GameFolderAndProjectFile.Value;
-						GameProjectFile.AddFilesToProject( SourceFileSearch.FindFiles( GameBuildDirectory, SubdirectoryNamesToExclude ), GameProjectDirectory );
+							GameProject.AddFilesToProject( SourceFileSearch.FindFiles( GameBuildDirectory, SubdirectoryNamesToExclude ), GameProjectDirectory );
+						}
 					}
-				}
 
-                DirectoryReference GameShaderDirectory = DirectoryReference.Combine(GameProjectDirectory, "Shaders");
-                if (DirectoryReference.Exists(GameShaderDirectory))
-                {
-                    ProjectFile GameProjectFile = GameFolderAndProjectFile.Value;
-                    GameProjectFile.AddFilesToProject(SourceFileSearch.FindFiles(GameShaderDirectory), GameProjectDirectory);
-                }
+					DirectoryReference GameShaderDirectory = DirectoryReference.Combine(GameProjectDirectory, "Shaders");
+					if (DirectoryReference.Exists(GameShaderDirectory))
+					{
+						GameProject.AddFilesToProject(SourceFileSearch.FindFiles(GameShaderDirectory), GameProjectDirectory);
+					}
+				}
             }
 		}
 
@@ -1838,7 +1844,7 @@ namespace UnrealBuildTool
 		/// <param name="ModProjects">All mod projects</param>
 		/// <param name="AllModuleFiles">List of *.Build.cs files for all engine programs and games</param>
 		/// <param name="bGatherThirdPartySource">True to gather source code from third party projects too</param>
-		protected void AddProjectsForAllModules( List<FileReference> AllGames, Dictionary<FileReference, ProjectFile> ProgramProjects, Dictionary<DirectoryReference, ProjectFile> ModProjects, List<FileReference> AllModuleFiles, bool bGatherThirdPartySource )
+		protected void AddProjectsForAllModules( List<FileReference> AllGames, Dictionary<FileReference, ProjectFile> ProgramProjects, List<ProjectFile> ModProjects, List<FileReference> AllModuleFiles, bool bGatherThirdPartySource )
 		{
 			HashSet<ProjectFile> ProjectsWithPlugins = new HashSet<ProjectFile>();
 			foreach( FileReference CurModuleFile in AllModuleFiles )
@@ -1899,7 +1905,7 @@ namespace UnrealBuildTool
 							// Add all the plugin files for this project
 							foreach(FileReference PluginFileName in Plugins.EnumeratePlugins(PluginFolder))
 							{
-								if(!ModProjects.ContainsKey(PluginFileName.Directory))
+								if(!ModProjects.Any(x => x.BaseDir == PluginFileName.Directory))
 								{
 									AddPluginFilesToProject(PluginFileName, BaseFolder, ProjectFile);
 								}
@@ -1941,7 +1947,7 @@ namespace UnrealBuildTool
             }
         }
 
-		private ProjectFile FindProjectForModule(FileReference CurModuleFile, List<FileReference> AllGames, Dictionary<FileReference, ProjectFile> ProgramProjects, Dictionary<DirectoryReference, ProjectFile> ModProjects, out DirectoryReference BaseFolder)
+		private ProjectFile FindProjectForModule(FileReference CurModuleFile, List<FileReference> AllGames, Dictionary<FileReference, ProjectFile> ProgramProjects, List<ProjectFile> ModProjects, out DirectoryReference BaseFolder)
 		{
 			string ProjectFileNameBase = null;
 
@@ -1967,12 +1973,12 @@ namespace UnrealBuildTool
 			else
 			{
 				// Check if it's a mod
-				foreach(KeyValuePair<DirectoryReference, ProjectFile> ModProject in ModProjects)
+				foreach(ProjectFile ModProject in ModProjects)
 				{
-					if(CurModuleFile.IsUnderDirectory(ModProject.Key))
+					if(CurModuleFile.IsUnderDirectory(ModProject.BaseDir))
 					{
-						BaseFolder = ModProject.Key;
-						return ModProject.Value;
+						BaseFolder = ModProject.BaseDir;
+						return ModProject;
 					}
 				}
 
@@ -1989,7 +1995,7 @@ namespace UnrealBuildTool
 			// Setup a project file entry for this module's project.  Remember, some projects may host multiple modules!
 			FileReference ProjectFileName = FileReference.Combine( IntermediateProjectFilesPath, ProjectFileNameBase + ProjectFileExtension );
 			bool bProjectAlreadyExisted;
-			return FindOrAddProject( ProjectFileName, IncludeInGeneratedProjects:true, bAlreadyExisted:out bProjectAlreadyExisted );
+			return FindOrAddProject( ProjectFileName, BaseFolder, IncludeInGeneratedProjects:true, bAlreadyExisted:out bProjectAlreadyExisted );
 		}
 
 		/// <summary>
@@ -1999,29 +2005,23 @@ namespace UnrealBuildTool
 		/// <param name="EngineProject">The engine project we created</param>
 		/// <param name="EnterpriseProject">The enterprise project we created</param>
 		/// <param name="GameProjects">Map of game folder name to all of the game projects we created</param>
-		/// <param name="ModProjects">Map of mod folder name to all the mod projects we created</param>
 		/// <param name="ProgramProjects">Map of program names to all of the program projects we created</param>
-		/// <param name="TemplateGameProjects">Set of template game projects we found.  These will also be in the GameProjects map</param>
-		/// <param name="SampleGameProjects">Set of sample game projects that were found</param>
-		private void AddProjectsForAllTargets( List<FileReference> AllGames, out ProjectFile EngineProject, out ProjectFile EnterpriseProject,
-			out Dictionary<DirectoryReference, ProjectFile> GameProjects, out Dictionary<DirectoryReference, ProjectFile> ModProjects,
-			out Dictionary<FileReference, ProjectFile> ProgramProjects, out Dictionary<DirectoryReference, ProjectFile> TemplateGameProjects, out Dictionary<DirectoryReference, ProjectFile> SampleGameProjects )
+		private void AddProjectsForAllTargets( 
+			List<FileReference> AllGames, 
+			out ProjectFile EngineProject, 
+			out ProjectFile EnterpriseProject,
+			out List<ProjectFile> GameProjects, 
+			out Dictionary<FileReference, ProjectFile> ProgramProjects)
 		{
 			// As we're creating project files, we'll also keep track of whether we created an "engine" project and return that if we have one
 			EngineProject = null;
 			EnterpriseProject = null;
-			GameProjects = new Dictionary<DirectoryReference, ProjectFile>();
+			GameProjects = new List<ProjectFile>();
 			ProgramProjects = new Dictionary<FileReference, ProjectFile>();
-			TemplateGameProjects = new Dictionary<DirectoryReference, ProjectFile>();
-			SampleGameProjects = new Dictionary<DirectoryReference, ProjectFile>();
 
 			// Get some standard directories
 			DirectoryReference EngineSourceProgramsDirectory = DirectoryReference.Combine(UnrealBuildTool.EngineSourceDirectory, "Programs");
-			DirectoryReference TemplatesDirectory = DirectoryReference.Combine(UnrealBuildTool.EngineDirectory, "..", "Templates");
-			DirectoryReference SamplesDirectory = DirectoryReference.Combine(UnrealBuildTool.EngineDirectory, "..", "Samples");
 			DirectoryReference EnterpriseSourceProgramsDirectory = DirectoryReference.Combine(UnrealBuildTool.EnterpriseSourceDirectory, "Programs");
-			DirectoryReference EnterpriseTemplatesDirectory = DirectoryReference.Combine(UnrealBuildTool.EnterpriseDirectory, "Templates");
-			DirectoryReference EnterpriseSamplesDirectory = DirectoryReference.Combine(UnrealBuildTool.EnterpriseDirectory, "Samples");
 
 			// Read the current engine version
 			ReadOnlyBuildVersion Version = new ReadOnlyBuildVersion(BuildVersion.ReadDefault());
@@ -2029,6 +2029,10 @@ namespace UnrealBuildTool
 			// Find all of the target files.  This will filter out any modules or targets that don't
 			// belong to platforms we're generating project files for.
 			List<FileReference> AllTargetFiles = DiscoverTargets(AllGames);
+
+			// Sort the targets by name. When we have multiple targets of a given type for a project, we'll use the order to determine which goes in the primary project file (so that client names with a suffix will go into their own project).
+			AllTargetFiles = AllTargetFiles.OrderBy(x => x.FullName, StringComparer.OrdinalIgnoreCase).ToList();
+
 			foreach( FileReference TargetFilePath in AllTargetFiles )
 			{
 				string TargetName = TargetFilePath.GetFileNameWithoutAnyExtensions();		// Remove both ".cs" and ".Target"
@@ -2120,10 +2124,23 @@ namespace UnrealBuildTool
 						ProjectFileNameBase = ProjectInfo.GetFileNameWithoutExtension();
 					}
 
-					// @todo projectfiles: We should move all of the Target.cs files out of sub-folders to clean up the project directories a bit (e.g. GameUncooked folder)
-
-					FileReference ProjectFilePath = FileReference.Combine(IntermediateProjectFilesPath, ProjectFileNameBase + ProjectFileExtension);
+					// Get the suffix to use for this project file. If we have multiple targets of the same type, we'll have to split them out into separate IDE project files.
+					string GeneratedProjectName = TargetRulesObject.GeneratedProjectName;
+					if(GeneratedProjectName == null)
+					{
+						ProjectFile ExistingProjectFile;
+						if( ProjectFileMap.TryGetValue( GetProjectLocation(ProjectFileNameBase), out ExistingProjectFile ) && ExistingProjectFile.ProjectTargets.Any(x => x.TargetRules.Type == TargetRulesObject.Type))
+						{
+							GeneratedProjectName = TargetRulesObject.Name;
+						}
+						else
+						{
+							GeneratedProjectName = ProjectFileNameBase;
+						}
+					}
 
+					// @todo projectfiles: We should move all of the Target.cs files out of sub-folders to clean up the project directories a bit (e.g. GameUncooked folder)
+					FileReference ProjectFilePath = GetProjectLocation(GeneratedProjectName);
 					if (TargetRulesObject.Type == TargetType.Game || TargetRulesObject.Type == TargetType.Client || TargetRulesObject.Type == TargetType.Server)
 					{
 						// Allow platforms to generate stub projects here...
@@ -2136,8 +2153,26 @@ namespace UnrealBuildTool
 							InConfigurations: SupportedConfigurations);
 					}
 
+					DirectoryReference BaseFolder;
+					if (IsProgramTarget)
+					{
+						BaseFolder = TargetFilePath.Directory;
+					}
+					else if (IsEngineTarget)
+					{
+						BaseFolder = UnrealBuildTool.EngineDirectory;
+					}
+					else if (IsEnterpriseTarget)
+					{
+						BaseFolder = UnrealBuildTool.EnterpriseDirectory;
+					}
+					else
+					{
+						BaseFolder = GameFolder;
+					}
+
 					bool bProjectAlreadyExisted;
-					ProjectFile ProjectFile = FindOrAddProject(ProjectFilePath, IncludeInGeneratedProjects: true, bAlreadyExisted: out bProjectAlreadyExisted);
+					ProjectFile ProjectFile = FindOrAddProject(ProjectFilePath, BaseFolder, IncludeInGeneratedProjects: true, bAlreadyExisted: out bProjectAlreadyExisted);
 					ProjectFile.IsForeignProject = CheckProjectFile != null && UProjectInfo.IsForeignProject(CheckProjectFile);
 					ProjectFile.IsGeneratedProject = true;
 					ProjectFile.IsStubProject = UnrealBuildTool.IsProjectInstalled();
@@ -2146,20 +2181,14 @@ namespace UnrealBuildTool
 						ProjectFile.ShouldBuildByDefaultForSolutionTargets = TargetRulesObject.bBuildInSolutionByDefault.Value;
 					}
 
-					// Check to see if this is a template target.  That is, the target is located under the "Templates" folder
-					bool IsTemplateTarget = TargetFilePath.IsUnderDirectory(TemplatesDirectory) || TargetFilePath.IsUnderDirectory(EnterpriseTemplatesDirectory);
-					bool IsSampleTarget = TargetFilePath.IsUnderDirectory(SamplesDirectory) || TargetFilePath.IsUnderDirectory(EnterpriseSamplesDirectory);
-
-					DirectoryReference BaseFolder = null;
+					// Add the project to the right output list
 					if (IsProgramTarget)
 					{
 						ProgramProjects[TargetFilePath] = ProjectFile;
-						BaseFolder = TargetFilePath.Directory;
 					}
 					else if (IsEngineTarget)
 					{
 						EngineProject = ProjectFile;
-						BaseFolder = UnrealBuildTool.EngineDirectory;
 						if (UnrealBuildTool.IsEngineInstalled())
 						{
 							// Allow engine projects to be created but not built for Installed Engine builds
@@ -2171,7 +2200,6 @@ namespace UnrealBuildTool
 					else if (IsEnterpriseTarget)
 					{
 						EnterpriseProject = ProjectFile;
-						BaseFolder = UnrealBuildTool.EnterpriseDirectory;
 						if (UnrealBuildTool.IsEnterpriseInstalled())
 						{
 							// Allow enterprise projects to be created but not built for Installed Engine builds
@@ -2182,16 +2210,7 @@ namespace UnrealBuildTool
 					}
 					else
 					{
-						GameProjects[GameFolder] = ProjectFile;
-						if (IsTemplateTarget)
-						{
-							TemplateGameProjects[GameFolder] = ProjectFile;
-						}
-						else if (IsSampleTarget)
-						{
-							SampleGameProjects[GameFolder] = ProjectFile;
-						}
-						BaseFolder = GameFolder;
+						GameProjects.Add(ProjectFile);
 
 						if (!bProjectAlreadyExisted)
 						{
@@ -2206,7 +2225,6 @@ namespace UnrealBuildTool
 								throw new BuildException("Not expecting to find a game with no .uproject file.  File '{0}' doesn't exist", UProjectFilePath);
 							}
 						}
-
 					}
 
 					foreach (ProjectTarget ExistingProjectTarget in ProjectFile.ProjectTargets)
@@ -2248,34 +2266,51 @@ namespace UnrealBuildTool
 					Log.TraceVerbose("Generating target {0} for {1}", TargetRulesObject.Type.ToString(), ProjectFilePath);
 				}
 			}
+		}
 
+		/// <summary>
+		/// Adds separate project files for all mods
+		/// </summary>
+		/// <param name="GameProjects">List of game project files</param>
+		/// <param name="ModProjects">Receives the list of mod projects on success</param>
+		void AddProjectsForMods(List<ProjectFile> GameProjects, out List<ProjectFile> ModProjects)
+		{
 			// Find all the mods for game projects
-			ModProjects = new Dictionary<DirectoryReference, ProjectFile>();
+			ModProjects = new List<ProjectFile>();
 			if(GameProjects.Count == 1)
 			{
-				KeyValuePair<DirectoryReference, ProjectFile> GameProject = GameProjects.First();
-				foreach(PluginInfo PluginInfo in Plugins.ReadProjectPlugins(GameProject.Key))
+				ProjectFile GameProject = GameProjects.First();
+				foreach(PluginInfo PluginInfo in Plugins.ReadProjectPlugins(GameProject.BaseDir))
 				{
 					if(PluginInfo.Descriptor.Modules != null && PluginInfo.Descriptor.Modules.Length > 0 && PluginInfo.Type == PluginType.Mod)
 					{
 						FileReference ModProjectFilePath = FileReference.Combine(PluginInfo.Directory, "Mods", PluginInfo.Name + ProjectFileExtension);
 
 						bool bProjectAlreadyExisted;
-						ProjectFile ModProjectFile = FindOrAddProject(ModProjectFilePath, IncludeInGeneratedProjects: true, bAlreadyExisted: out bProjectAlreadyExisted);
-						ModProjectFile.IsForeignProject = GameProject.Value.IsForeignProject;
+						ProjectFile ModProjectFile = FindOrAddProject(ModProjectFilePath, PluginInfo.Directory, IncludeInGeneratedProjects: true, bAlreadyExisted: out bProjectAlreadyExisted);
+						ModProjectFile.IsForeignProject = GameProject.IsForeignProject;
 						ModProjectFile.IsGeneratedProject = true;
 						ModProjectFile.IsStubProject = false;
 						ModProjectFile.PluginFilePath = PluginInfo.File;
-						ModProjectFile.ProjectTargets.AddRange(GameProject.Value.ProjectTargets);
+						ModProjectFile.ProjectTargets.AddRange(GameProject.ProjectTargets);
 
 						AddPluginFilesToProject(PluginInfo.File, PluginInfo.Directory, ModProjectFile);
 
-						ModProjects.Add(PluginInfo.Directory, ModProjectFile);
+						ModProjects.Add(ModProjectFile);
 					}
 				}
 			}
 		}
 
+		/// <summary>
+		/// Gets the location for a project file
+		/// </summary>
+		/// <param name="BaseName">The base name for the project file</param>
+		/// <returns>Full path to the project file</returns>
+		private FileReference GetProjectLocation(string BaseName)
+		{
+			return FileReference.Combine(IntermediateProjectFilesPath, BaseName + ProjectFileExtension);
+		}
 		
 		/// Adds shader source code to the specified project
 		protected void AddEngineShaderSource( ProjectFile EngineProject )
@@ -2377,10 +2412,11 @@ namespace UnrealBuildTool
 		/// Adds a new project file and returns an object that represents that project file (or if the project file is already known, returns that instead.)
 		/// </summary>
 		/// <param name="FilePath">Full path to the project file</param>
+		/// <param name="BaseDir">Base directory for files in this project</param>
 		/// <param name="IncludeInGeneratedProjects">True if this project should be included in the set of generated projects.  Only matters when actually generating project files.</param>
 		/// <param name="bAlreadyExisted">True if we already had this project file</param>
 		/// <returns>Object that represents this project file in Unreal Build Tool</returns>
-		public ProjectFile FindOrAddProject( FileReference FilePath, bool IncludeInGeneratedProjects, out bool bAlreadyExisted )
+		public ProjectFile FindOrAddProject( FileReference FilePath, DirectoryReference BaseDir, bool IncludeInGeneratedProjects, out bool bAlreadyExisted )
 		{
 			if( FilePath == null )
 			{
@@ -2397,6 +2433,7 @@ namespace UnrealBuildTool
 
 			// Add a new project file for the specified path
 			ProjectFile NewProjectFile = AllocateProjectFile( FilePath );
+			NewProjectFile.BaseDir = BaseDir;
 			ProjectFileMap[ FilePath ] = NewProjectFile;
 
 			if( IncludeInGeneratedProjects )
