UE-65016 - OpenGLES API deprecated warnings when packaging AR Samples

diff --git a/Engine/Source/Programs/UnrealBuildTool/Platform/IOS/UEBuildIOS.cs b/Engine/Source/Programs/UnrealBuildTool/Platform/IOS/UEBuildIOS.cs
--- a/Engine/Source/Programs/UnrealBuildTool/Platform/IOS/UEBuildIOS.cs
+++ b/Engine/Source/Programs/UnrealBuildTool/Platform/IOS/UEBuildIOS.cs
@@ -901,6 +901,7 @@ namespace UnrealBuildTool
 		{
 			CompileEnvironment.Definitions.Add("PLATFORM_IOS=1");
 			CompileEnvironment.Definitions.Add("PLATFORM_APPLE=1");
+			CompileEnvironment.Definitions.Add("GLES_SILENCE_DEPRECATION=1");  // suppress GLES "deprecated" warnings until a proper solution is implemented (see UE-65643)
 
 			CompileEnvironment.Definitions.Add("WITH_TTS=0");
 			CompileEnvironment.Definitions.Add("WITH_SPEECH_RECOGNITION=0");
