Add detection and device profiles for iPhone XS / XS Max / SR
Integrated from CL 4433649

diff --git a/Engine/Config/BaseDeviceProfiles.ini b/Engine/Config/BaseDeviceProfiles.ini
--- a/Engine/Config/BaseDeviceProfiles.ini
+++ b/Engine/Config/BaseDeviceProfiles.ini
@@ -32,6 +32,9 @@
 +DeviceProfileNameAndTypes=iPhone8,IOS
 +DeviceProfileNameAndTypes=iPhone8Plus,IOS
 +DeviceProfileNameAndTypes=iPhoneX,IOS
++DeviceProfileNameAndTypes=iPhoneXS,IOS
++DeviceProfileNameAndTypes=iPhoneXSMax,IOS
++DeviceProfileNameAndTypes=iPhoneXR,IOS
 +DeviceProfileNameAndTypes=iPadPro,IOS
 +DeviceProfileNameAndTypes=iPadPro105,IOS
 +DeviceProfileNameAndTypes=iPadPro129,IOS
@@ -288,6 +291,21 @@ DeviceType=IOS
 BaseProfileName=iPhone8
 +CVars=r.CustomUnsafeZones="(L:free[0,-15][812,15]);(P:fixed[83,0][206,30])"
 
+[iPhoneXS DeviceProfile]
+DeviceType=IOS
+BaseProfileName=iPhone8
++CVars=r.CustomUnsafeZones="(L:free[0,-15][812,15]);(P:fixed[83,0][206,30])"
+
+[iPhoneXSMax DeviceProfile]
+DeviceType=IOS
+BaseProfileName=iPhone8
++CVars=r.CustomUnsafeZones="(L:free[0,-15][812,15]);(P:fixed[83,0][206,30])"
+
+[iPhoneXR DeviceProfile]
+DeviceType=IOS
+BaseProfileName=iPhone8
++CVars=r.CustomUnsafeZones="(L:free[0,-15][812,15]);(P:fixed[83,0][206,30])"
+
 [iPadPro DeviceProfile]
 DeviceType=IOS
 BaseProfileName=IOS
diff --git a/Engine/Source/Programs/UnrealBuildTool/Platform/IOS/UEDeployIOS.cs b/Engine/Source/Programs/UnrealBuildTool/Platform/IOS/UEDeployIOS.cs
--- a/Engine/Source/Programs/UnrealBuildTool/Platform/IOS/UEDeployIOS.cs
+++ b/Engine/Source/Programs/UnrealBuildTool/Platform/IOS/UEDeployIOS.cs
@@ -563,6 +563,10 @@ namespace UnrealBuildTool
 					"Default-568h.png", "Portrait", "{320, 568}", "7.0",
 					"Default-IPhoneX-Landscape.png", "Landscape", "{375, 812}",  "11.0",
 					"Default-IPhoneX-Portrait.png", "Portrait", "{375, 812}",  "11.0",
+					"Default-IPhoneXSMax-Landscape.png", "Landscape", "{414, 896}",  "12.0",
+					"Default-IPhoneXSMax-Portrait.png", "Portrait", "{414, 896}",  "12.0",
+					"Default-IPhoneXR-Landscape.png", "Landscape", "{276, 597}",  "12.0",
+					"Default-IPhoneXR-Portrait.png", "Portrait", "{276, 597}",  "12.0",
 				};
 
 				Text.AppendLine("\t<key>UILaunchImages~iphone</key>");
@@ -797,6 +801,8 @@ namespace UnrealBuildTool
 					//                    CopyFiles(InEngineDir + "/Build/IOS/Resources/Graphics", AppDirectory, "Default-Portrait-1336.png", true);
 					CopyFiles(InEngineDir + "/Build/IOS/Resources/Graphics", AppDirectory, "[REDACTED_EMAIL]", true);
 					CopyFiles(InEngineDir + "/Build/IOS/Resources/Graphics", AppDirectory, "Default-IPhoneX-Portrait.png", true);
+					CopyFiles(InEngineDir + "/Build/IOS/Resources/Graphics", AppDirectory, "Default-IPhoneXSMax-Portrait.png", true);
+					CopyFiles(InEngineDir + "/Build/IOS/Resources/Graphics", AppDirectory, "Default-IPhoneXR-Portrait.png", true);
 				}
 				if (bSupportsLandscape)
 				{
@@ -807,6 +813,8 @@ namespace UnrealBuildTool
 					//                    CopyFiles(InEngineDir + "/Build/IOS/Resources/Graphics", AppDirectory, "Default-Landscape-1336.png", true);
 					CopyFiles(InEngineDir + "/Build/IOS/Resources/Graphics", AppDirectory, "[REDACTED_EMAIL]", true);
 					CopyFiles(InEngineDir + "/Build/IOS/Resources/Graphics", AppDirectory, "Default-IPhoneX-Landscape.png", true);
+					CopyFiles(InEngineDir + "/Build/IOS/Resources/Graphics", AppDirectory, "Default-IPhoneXSMax-Landscape.png", true);
+					CopyFiles(InEngineDir + "/Build/IOS/Resources/Graphics", AppDirectory, "Default-IPhoneXR-Landscape.png", true);
 				}
 				//                CopyFiles(InEngineDir + "/Build/IOS/Resources/Graphics", AppDirectory, "Default.png", true);
 				CopyFiles(InEngineDir + "/Build/IOS/Resources/Graphics", AppDirectory, "[REDACTED_EMAIL]", true);
@@ -829,6 +837,8 @@ namespace UnrealBuildTool
 					//                    CopyFiles(BuildDirectory + "/Resources/Graphics", AppDirectory, "Default-Portrait-1336.png", true);
 					CopyFiles(BuildDirectory + "/Resources/Graphics", AppDirectory, "[REDACTED_EMAIL]", true);
 					CopyFiles(BuildDirectory + "/Resources/Graphics", AppDirectory, "Default-IPhoneX-Portrait.png", true);
+					CopyFiles(BuildDirectory + "/Resources/Graphics", AppDirectory, "Default-IPhoneXSMax-Portrait.png", true);
+					CopyFiles(BuildDirectory + "/Resources/Graphics", AppDirectory, "Default-IPhoneXR-Portrait.png", true);
 				}
 				if (bSupportsLandscape)
 				{
@@ -839,6 +849,8 @@ namespace UnrealBuildTool
 					//                    CopyFiles(BuildDirectory + "/Resources/Graphics", AppDirectory, "Default-Landscape-1336.png", true);
 					CopyFiles(BuildDirectory + "/Resources/Graphics", AppDirectory, "[REDACTED_EMAIL]", true);
 					CopyFiles(BuildDirectory + "/Resources/Graphics", AppDirectory, "Default-IPhoneX-Landscape.png", true);
+					CopyFiles(BuildDirectory + "/Resources/Graphics", AppDirectory, "Default-IPhoneXSMax-Landscape.png", true);
+					CopyFiles(BuildDirectory + "/Resources/Graphics", AppDirectory, "Default-IPhoneXR-Landscape.png", true);
 				}
 				//                CopyFiles(BuildDirectory + "/Resources/Graphics", AppDirectory, "Default.png", true);
 				CopyFiles(BuildDirectory + "/Resources/Graphics", AppDirectory, "[REDACTED_EMAIL]", true);
diff --git a/Engine/Source/Runtime/ApplicationCore/Private/IOS/IOSAppDelegate.cpp b/Engine/Source/Runtime/ApplicationCore/Private/IOS/IOSAppDelegate.cpp
--- a/Engine/Source/Runtime/ApplicationCore/Private/IOS/IOSAppDelegate.cpp
+++ b/Engine/Source/Runtime/ApplicationCore/Private/IOS/IOSAppDelegate.cpp
@@ -794,7 +794,7 @@ static FAutoConsoleVariableRef CVarGEnableThermalsReport(
 			[ImageString appendString : @"-Portrait"];
 		}
 	}
-	else if (Device == FPlatformMisc::IOS_IPhoneX)
+    else if (Device == FPlatformMisc::IOS_IPhoneX || Device == FPlatformMisc::IOS_IPhoneXS)
 	{
 		[ImageString appendString : @"-IPhoneX"];
 		if (!self.bDeviceInPortraitMode)
@@ -806,6 +806,30 @@ static FAutoConsoleVariableRef CVarGEnableThermalsReport(
 			[ImageString appendString : @"-Portrait"];
 		}
 	}
+    else if (Device == FPlatformMisc::IOS_IPhoneXSMax)
+    {
+        [ImageString appendString : @"-IPhoneXSMax"];
+        if (!self.bDeviceInPortraitMode)
+        {
+            [ImageString appendString : @"-Landscape"];
+        }
+        else
+        {
+            [ImageString appendString : @"-Portrait"];
+        }
+    }
+    else if (Device == FPlatformMisc::IOS_IPhoneXR)
+    {
+        [ImageString appendString : @"-IPhoneXR"];
+        if (!self.bDeviceInPortraitMode)
+        {
+            [ImageString appendString : @"-Landscape"];
+        }
+        else
+        {
+            [ImageString appendString : @"-Portrait"];
+        }
+    }
 	else if (Device == FPlatformMisc::IOS_AppleTV)
 	{
 		// @todo tvos: Make an AppleTV one?
diff --git a/Engine/Source/Runtime/ApplicationCore/Private/IOS/IOSPlatformApplicationMisc.cpp b/Engine/Source/Runtime/ApplicationCore/Private/IOS/IOSPlatformApplicationMisc.cpp
--- a/Engine/Source/Runtime/ApplicationCore/Private/IOS/IOSPlatformApplicationMisc.cpp
+++ b/Engine/Source/Runtime/ApplicationCore/Private/IOS/IOSPlatformApplicationMisc.cpp
@@ -224,7 +224,7 @@ void FIOSPlatformApplicationMisc::ClipboardPaste(class FString& Result)
 EScreenPhysicalAccuracy FIOSPlatformApplicationMisc::ComputePhysicalScreenDensity(int32& ScreenDensity)
 {
 	FPlatformMisc::EIOSDevice Device = FPlatformMisc::GetIOSDeviceType();
-	static_assert( FPlatformMisc::EIOSDevice::IOS_Unknown == 32, "Every device needs to be handled here." );
+	static_assert( FPlatformMisc::EIOSDevice::IOS_Unknown == 35, "Every device needs to be handled here." );
 
 	ScreenDensity = 0;
 	EScreenPhysicalAccuracy Accuracy = EScreenPhysicalAccuracy::Unknown;
@@ -245,6 +245,7 @@ EScreenPhysicalAccuracy FIOSPlatformApplicationMisc::ComputePhysicalScreenDensit
 	case FPlatformMisc::IOS_IPhone6S:
 	case FPlatformMisc::IOS_IPhone7:
 	case FPlatformMisc::IOS_IPhone8:
+    case FPlatformMisc::IOS_IPhoneXR:
 		ScreenDensity = 326;
 		Accuracy = EScreenPhysicalAccuracy::Truth;
 		break;
@@ -253,6 +254,8 @@ EScreenPhysicalAccuracy FIOSPlatformApplicationMisc::ComputePhysicalScreenDensit
 	case FPlatformMisc::IOS_IPhone7Plus:
 	case FPlatformMisc::IOS_IPhone8Plus:
 	case FPlatformMisc::IOS_IPhoneX:
+    case FPlatformMisc::IOS_IPhoneXS:
+    case FPlatformMisc::IOS_IPhoneXSMax:
 		ScreenDensity = 401;
 		Accuracy = EScreenPhysicalAccuracy::Truth;
 		break;
diff --git a/Engine/Source/Runtime/Core/Private/IOS/IOSPlatformMisc.cpp b/Engine/Source/Runtime/Core/Private/IOS/IOSPlatformMisc.cpp
--- a/Engine/Source/Runtime/Core/Private/IOS/IOSPlatformMisc.cpp
+++ b/Engine/Source/Runtime/Core/Private/IOS/IOSPlatformMisc.cpp
@@ -536,17 +536,32 @@ FIOSPlatformMisc::EIOSDevice FIOSPlatformMisc::GetIOSDeviceType()
 				DeviceType = IOS_IPhoneX;
 			}
 		}
-		else if (Major >= 10)
+        else if (Major == 11)
+        {
+            if (Minor == 2)
+            {
+                DeviceType = IOS_IPhoneXS;
+            }
+            else if (Minor == 4 || Minor == 6)
+            {
+                DeviceType = IOS_IPhoneXSMax;
+            }
+            else if (Minor == 8)
+            {
+                DeviceType = IOS_IPhoneXR;
+            }
+        }
+		else if (Major >= 12)
 		{
 			// for going forward into unknown devices (like 8/8+?), we can't use Minor,
 			// so treat devices with a scale > 2.5 to be 6SPlus type devices, < 2.5 to be 6S type devices
 			if ([UIScreen mainScreen].scale > 2.5f)
 			{
-				DeviceType = IOS_IPhone8Plus;
+				DeviceType = IOS_IPhoneXSMax;
 			}
 			else
 			{
-				DeviceType = IOS_IPhone8;
+				DeviceType = IOS_IPhoneXS;
 			}
 		}
 	}
diff --git a/Engine/Source/Runtime/Core/Public/IOS/IOSPlatformMisc.h b/Engine/Source/Runtime/Core/Public/IOS/IOSPlatformMisc.h
--- a/Engine/Source/Runtime/Core/Public/IOS/IOSPlatformMisc.h
+++ b/Engine/Source/Runtime/Core/Public/IOS/IOSPlatformMisc.h
@@ -131,6 +131,9 @@ struct CORE_API FIOSPlatformMisc : public FApplePlatformMisc
 		IOS_IPadPro_105,
 		IOS_IPadPro2_129,
 		IOS_IPad5,
+        IOS_IPhoneXS,
+        IOS_IPhoneXSMax,
+        IOS_IPhoneXR,
 		IOS_Unknown,
 	};
 
@@ -172,6 +175,9 @@ struct CORE_API FIOSPlatformMisc : public FApplePlatformMisc
 			TEXT("IPadPro105"),
 			TEXT("IPadPro2_129"),
 			TEXT("IPad5"),
+            TEXT("IPhoneXS"),
+            TEXT("IPhoneXSMax"),
+            TEXT("IPhoneXR"),
 			TEXT("Unknown"),
 		};
 		static_assert((sizeof(IOSDeviceNames) / sizeof(IOSDeviceNames[0])) == ((int32)IOS_Unknown + 1), "Mismatched IOSDeviceNames and EIOSDevice.");
