Changes from Magic Leap for 4.21 release. They've removed the magic leap music service integration and changed the build script to look for a minimum sdk version rather than an exact match.

diff --git a/Engine/Plugins/Lumin/MagicLeap/MagicLeap.uplugin b/Engine/Plugins/Lumin/MagicLeap/MagicLeap.uplugin
--- a/Engine/Plugins/Lumin/MagicLeap/MagicLeap.uplugin
+++ b/Engine/Plugins/Lumin/MagicLeap/MagicLeap.uplugin
@@ -73,17 +73,6 @@
 				"Lumin"
 			]
 		},
-		{
-			"Name" : "MagicLeapMusicService",
-			"Type" : "Runtime",
-			"LoadingPhase" : "PostConfigInit",
-			"WhitelistPlatforms" :
-			[
-				"Win64",
-				"Mac",
-				"Lumin"
-			]
-		},
 		{
 			"Name" : "MagicLeapSecureStorage",
 			"Type" : "Runtime",
diff --git a/Engine/Plugins/Lumin/MagicLeap/Source/MagicLeapMusicService/MagicLeapMusicService.Build.cs b/Engine/Plugins/Lumin/MagicLeap/Source/MagicLeapMusicService/MagicLeapMusicService.Build.cs
deleted file mode 100644
diff --git a/Engine/Plugins/Lumin/MagicLeap/Source/MagicLeapMusicService/Private/MagicLeapMusicService.cpp b/Engine/Plugins/Lumin/MagicLeap/Source/MagicLeapMusicService/Private/MagicLeapMusicService.cpp
deleted file mode 100644
diff --git a/Engine/Plugins/Lumin/MagicLeap/Source/MagicLeapMusicService/Public/IMagicLeapMusicServicePlugin.h b/Engine/Plugins/Lumin/MagicLeap/Source/MagicLeapMusicService/Public/IMagicLeapMusicServicePlugin.h
deleted file mode 100644
diff --git a/Engine/Plugins/Lumin/MagicLeap/Source/MagicLeapMusicService/Public/MagicLeapMusicService.h b/Engine/Plugins/Lumin/MagicLeap/Source/MagicLeapMusicService/Public/MagicLeapMusicService.h
deleted file mode 100644
diff --git a/Engine/Plugins/Lumin/MagicLeap/Source/MagicLeapMusicService/Public/MagicLeapMusicServiceTypes.h b/Engine/Plugins/Lumin/MagicLeap/Source/MagicLeapMusicService/Public/MagicLeapMusicServiceTypes.h
deleted file mode 100644
diff --git a/Engine/Source/Programs/UnrealBuildTool/Platform/Lumin/UEBuildLumin.cs b/Engine/Source/Programs/UnrealBuildTool/Platform/Lumin/UEBuildLumin.cs
--- a/Engine/Source/Programs/UnrealBuildTool/Platform/Lumin/UEBuildLumin.cs
+++ b/Engine/Source/Programs/UnrealBuildTool/Platform/Lumin/UEBuildLumin.cs
@@ -273,19 +273,20 @@ namespace UnrealBuildTool
 	class LuminPlatformSDK : AndroidPlatformSDK
 	{
 		/// <summary>
-		/// This is the SDK version we support
+		/// This is the minimum SDK version we support.
 		/// </summary>
-		static string ExpectedSDKVersion = "0.16";   // now unified for all the architectures
+		static uint MinimumSDKVersionMajor = 0;
+		static uint MinimumSDKVersionMinor = 16;
 
 		public override string GetSDKTargetPlatformName()
 		{
 			return "Lumin";
 		}
-
 		protected override string GetRequiredSDKString()
 		{
-			return ExpectedSDKVersion;
+			return string.Format("{0}.{1}", MinimumSDKVersionMajor, MinimumSDKVersionMinor);
 		}
+
 		protected override String GetRequiredScriptVersionString()
 		{
 			return "Lumin_15";
@@ -325,6 +326,11 @@ namespace UnrealBuildTool
 					// if the folder specified by the config var doesn't exist, fall back to the env var.
 					if (Directory.Exists(path))
 					{
+						if (MLSDKPath != path)
+						{
+							Console.WriteLine("*** MLSDK environment variable differs from config file; overriding environment variable ***");
+						}
+
 						MLSDKPath = path;
 					}
 				}
@@ -336,10 +342,13 @@ namespace UnrealBuildTool
 			// we don't have an MLSDK directory specified
 			if (String.IsNullOrEmpty(MLSDKPath))
 			{
+				Console.WriteLine("*** Unable to determine MLSDK location from MLSDK environment variable or config file ***");
 				return false;
 			}
-			// we don't have the required MLSDK setup
-			String DetectedVersion = "Unknown";
+
+			// detected SDK version is < minimum major/minor
+			uint DetectedMajorVersion;
+			uint DetectedMinorVersion;
 			String VersionFile = string.Format("{0}/include/ml_version.h", MLSDKPath).Replace('/', Path.DirectorySeparatorChar);
 			if (File.Exists(VersionFile))
 			{
@@ -347,14 +356,23 @@ namespace UnrealBuildTool
 
 				String MajorVersion = FindVersionNumber("MLSDK_VERSION_MAJOR", VersionText);
 				String MinorVersion = FindVersionNumber("MLSDK_VERSION_MINOR", VersionText);
-				DetectedVersion = string.Format("{0}.{1}", MajorVersion, MinorVersion);
+				DetectedMajorVersion = Convert.ToUInt32(MajorVersion);
+				DetectedMinorVersion = Convert.ToUInt32(MinorVersion);
 			}
-			if (!DetectedVersion.Equals(GetRequiredSDKString()))
+			else
+			{
+				Console.WriteLine("*** Unable to locate MLSDK version file ml_version.h ***");
+				return false;
+			}
+
+			if (DetectedMajorVersion < MinimumSDKVersionMajor || (DetectedMajorVersion == MinimumSDKVersionMajor && DetectedMinorVersion < MinimumSDKVersionMinor))
 			{
-				Console.WriteLine("*** Found installed MLSDK version {0} but require {1} ***", DetectedVersion, GetRequiredSDKString());
+				Console.WriteLine("*** Found installed MLSDK version {0}.{1} but require at least {2}.{3} ***",
+					DetectedMajorVersion, DetectedMinorVersion, MinimumSDKVersionMajor, MinimumSDKVersionMinor);
 				return false;
 			}
 
+			Console.WriteLine("*** Found installed MLSDK version {0}.{1} ***", DetectedMajorVersion, DetectedMinorVersion);
 			return true;
 		}
 
