Always allow nativized project rules to be compiled on remote Mac. The initiating PC doesn't have enough context (eg. which target type is being built) to be able to compile it in advance.

diff --git a/Engine/Source/Programs/UnrealBuildTool/Configuration/UEBuildTarget.cs b/Engine/Source/Programs/UnrealBuildTool/Configuration/UEBuildTarget.cs
--- a/Engine/Source/Programs/UnrealBuildTool/Configuration/UEBuildTarget.cs
+++ b/Engine/Source/Programs/UnrealBuildTool/Configuration/UEBuildTarget.cs
@@ -513,7 +513,7 @@ namespace UnrealBuildTool
 				FileReference PluginFile = FileReference.Combine(RulesObject.ProjectFile.Directory, "Intermediate", "Plugins", "NativizedAssets", PlatformName, ProjectTargetType, "NativizedAssets.uplugin");
 				if (FileReference.Exists(PluginFile))
 				{
-					RulesAssembly = RulesCompiler.CreatePluginRulesAssembly(PluginFile, bSkipRulesCompile, RulesAssembly, false);
+					RulesAssembly = RulesCompiler.CreatePluginRulesAssembly(PluginFile, false, RulesAssembly, false);
 				}
 				else
 				{
