Fix case issue one iPhone X splash screens and renamed file to solve P4 issues, added Project Settings UI for iPhone XSMax/XR screenshots and replaced them with the updated version

diff --git a/Engine/Source/Developer/IOS/IOSPlatformEditor/Private/IOSTargetSettingsCustomization.cpp b/Engine/Source/Developer/IOS/IOSPlatformEditor/Private/IOSTargetSettingsCustomization.cpp
--- a/Engine/Source/Developer/IOS/IOSPlatformEditor/Private/IOSTargetSettingsCustomization.cpp
+++ b/Engine/Source/Developer/IOS/IOSPlatformEditor/Private/IOSTargetSettingsCustomization.cpp
@@ -89,10 +89,14 @@ FIOSTargetSettingsCustomization::FIOSTargetSettingsCustomization()
     new (LaunchImageNames) FPlatformIconInfo(TEXT("[REDACTED_EMAIL]"), LOCTEXT("LaunchImage_iPhone5_Portrait", "iPhone 5/5S/SE in Portrait"), FText::GetEmpty(), 640, 1136, FPlatformIconInfo::Required);
     new (LaunchImageNames) FPlatformIconInfo(TEXT("Default-IPhone6-Landscape.png"), LOCTEXT("LaunchImage_iPhone6_Landscape", "iPhone 6/7/8 in Landscape"), FText::GetEmpty(), 1334, 750, FPlatformIconInfo::Required);
     new (LaunchImageNames) FPlatformIconInfo(TEXT("Default-IPhone6.png"), LOCTEXT("LaunchImage_iPhone6", "iPhone 6/7/8 in Portrait"), FText::GetEmpty(), 750, 1334, FPlatformIconInfo::Required);
-    new (LaunchImageNames) FPlatformIconInfo(TEXT("Default-IPhone6Plus-Landscape.png"), LOCTEXT("LaunchImage_iPhone6Plus_Landscape", "iPhone 6+/7+/8+ in Landscape"), FText::GetEmpty(), 2208, 1242, FPlatformIconInfo::Required);
-    new (LaunchImageNames) FPlatformIconInfo(TEXT("Default-IPhone6Plus-Portrait.png"), LOCTEXT("LaunchImage_iPhone6Plus_Portrait", "iPhone 6+/7+/8+ in Portrait"), FText::GetEmpty(), 1242, 2208, FPlatformIconInfo::Required);
-    new (LaunchImageNames) FPlatformIconInfo(TEXT("Default-IPhoneX-Landscape.png"), LOCTEXT("LaunchImage_iPhoneX_Landscape", "iPhone X in Landscape"), FText::GetEmpty(), 2436, 1125, FPlatformIconInfo::Required);
-    new (LaunchImageNames) FPlatformIconInfo(TEXT("Default-IPhoneX-Portrait.png"), LOCTEXT("LaunchImage_iPhoneX_Portrait", "iPhone X in Portrait"), FText::GetEmpty(), 1125, 2436, FPlatformIconInfo::Required);
+    new (LaunchImageNames) FPlatformIconInfo(TEXT("Default-IPhone6Plus-Landscape.png"), LOCTEXT("LaunchImage_iPhone6Plus_Landscape", "iPhone 6+/7+/8+/9 in Landscape"), FText::GetEmpty(), 2208, 1242, FPlatformIconInfo::Required);
+    new (LaunchImageNames) FPlatformIconInfo(TEXT("Default-IPhone6Plus-Portrait.png"), LOCTEXT("LaunchImage_iPhone6Plus_Portrait", "iPhone 6+/7+/8+/9 in Portrait"), FText::GetEmpty(), 1242, 2208, FPlatformIconInfo::Required);
+    new (LaunchImageNames) FPlatformIconInfo(TEXT("Default-IPhoneXS-Landscape.png"), LOCTEXT("LaunchImage_iPhoneX_Landscape", "iPhone X/XS in Landscape"), FText::GetEmpty(), 2436, 1125, FPlatformIconInfo::Required);
+    new (LaunchImageNames) FPlatformIconInfo(TEXT("Default-IPhoneXS-Portrait.png"), LOCTEXT("LaunchImage_iPhoneX_Portrait", "iPhone X/XS in Portrait"), FText::GetEmpty(), 1125, 2436, FPlatformIconInfo::Required);
+	new (LaunchImageNames) FPlatformIconInfo(TEXT("Default-IPhoneXSMax-Landscape.png"), LOCTEXT("LaunchImage_iPhoneXSMax_Landscape", "iPhone XS Max in Landscape"), FText::GetEmpty(), 2688, 1242, FPlatformIconInfo::Required);
+	new (LaunchImageNames) FPlatformIconInfo(TEXT("Default-IPhoneXSMax-Portrait.png"), LOCTEXT("LaunchImage_iPhoneXMax_Portrait", "iPhone XS Max in Portrait"), FText::GetEmpty(), 1242, 2688, FPlatformIconInfo::Required);
+	new (LaunchImageNames) FPlatformIconInfo(TEXT("Default-IPhoneXR-Landscape.png"), LOCTEXT("LaunchImage_iPhoneXR_Landscape", "iPhone XR in Landscape"), FText::GetEmpty(), 1792, 828, FPlatformIconInfo::Required);
+	new (LaunchImageNames) FPlatformIconInfo(TEXT("Default-IPhoneXR-Portrait.png"), LOCTEXT("LaunchImage_iPhoneXR_Portrait", "iPhone XR in Portrait"), FText::GetEmpty(), 828, 1792, FPlatformIconInfo::Required);
 	new (LaunchImageNames) FPlatformIconInfo(TEXT("[REDACTED_EMAIL]"), LOCTEXT("LaunchImage_iPadRetina_Landscape", "iPad in Landscape"), FText::GetEmpty(), 2048, 1536, FPlatformIconInfo::Required);
 	new (LaunchImageNames) FPlatformIconInfo(TEXT("[REDACTED_EMAIL]"), LOCTEXT("LaunchImage_iPadRetina_Portrait", "iPad in Portrait"), FText::GetEmpty(), 1536, 2048, FPlatformIconInfo::Required);
 	new (LaunchImageNames) FPlatformIconInfo(TEXT("Default-IPadPro10_5_Landscape.png"), LOCTEXT("LaunchImage_iPadPro_Landscape", "iPad Pro 10.5 in Landscape"), FText::GetEmpty(), 2224, 1668, FPlatformIconInfo::Required);
diff --git a/Engine/Source/Programs/AutomationTool/IOS/IOSPlatform.Automation.cs b/Engine/Source/Programs/AutomationTool/IOS/IOSPlatform.Automation.cs
--- a/Engine/Source/Programs/AutomationTool/IOS/IOSPlatform.Automation.cs
+++ b/Engine/Source/Programs/AutomationTool/IOS/IOSPlatform.Automation.cs
@@ -894,7 +894,7 @@ public class IOSPlatform : Platform
 				ImageFileNames.Add("Default-IPhone6Plus-Portrait.png");
 				ImageFileNames.Add("[REDACTED_EMAIL]");
 				ImageFileNames.Add("[REDACTED_EMAIL]");
-				ImageFileNames.Add("Default-IPhoneX-Portrait.png");
+				ImageFileNames.Add("Default-IPhoneXS-Portrait.png");
 			}
 			if (bSupportsLandscape)
 			{
@@ -902,7 +902,7 @@ public class IOSPlatform : Platform
 				ImageFileNames.Add("Default-IPhone6Plus-Landscape.png");
 				ImageFileNames.Add("[REDACTED_EMAIL]");
 				ImageFileNames.Add("[REDACTED_EMAIL]");
-				ImageFileNames.Add("Default-IPhoneX-Landscape.png");
+				ImageFileNames.Add("Default-IPhoneXS-Landscape.png");
 			}
 			ImageFileNames.Add("[REDACTED_EMAIL]");
 			ImageFileNames.Add("[REDACTED_EMAIL]");
diff --git a/Engine/Source/Programs/IOS/iPhonePackager/CompileTime.cs b/Engine/Source/Programs/IOS/iPhonePackager/CompileTime.cs
--- a/Engine/Source/Programs/IOS/iPhonePackager/CompileTime.cs
+++ b/Engine/Source/Programs/IOS/iPhonePackager/CompileTime.cs
@@ -269,8 +269,8 @@ namespace iPhonePackager
 					"Default", "Portrait", "{320, 480}",
 					"Default-568h", "Landscape", "{320, 568}",
 					"Default-568h", "Portrait", "{320, 568}",
-					"Default-IPhoneX-Landscape", "Landscape", "{375, 812}",
-					"Default-IPhoneX-Portrait", "Portrait", "{375, 812}",
+					"Default-IPhoneXS-Landscape", "Landscape", "{375, 812}",
+					"Default-IPhoneXS-Portrait", "Portrait", "{375, 812}",
 				};
 
 			StringBuilder NewLaunchImagesString = new StringBuilder("<key>UILaunchImages~iphone</key>\n\t\t<array>\n");
diff --git a/Engine/Source/Programs/UnrealBuildTool/Platform/IOS/UEDeployIOS.cs b/Engine/Source/Programs/UnrealBuildTool/Platform/IOS/UEDeployIOS.cs
--- a/Engine/Source/Programs/UnrealBuildTool/Platform/IOS/UEDeployIOS.cs
+++ b/Engine/Source/Programs/UnrealBuildTool/Platform/IOS/UEDeployIOS.cs
@@ -561,8 +561,8 @@ namespace UnrealBuildTool
 					"Default-IPhone6Plus-Portrait.png", "Portrait", "{414, 736}",  "8.0",
 					"Default.png", "Portrait", "{320, 480}", "7.0",
 					"Default-568h.png", "Portrait", "{320, 568}", "7.0",
-					"Default-IPhoneX-Landscape.png", "Landscape", "{375, 812}",  "11.0",
-					"Default-IPhoneX-Portrait.png", "Portrait", "{375, 812}",  "11.0",
+					"Default-IPhoneXS-Landscape.png", "Landscape", "{375, 812}",  "11.0",
+					"Default-IPhoneXS-Portrait.png", "Portrait", "{375, 812}",  "11.0",
 					"Default-IPhoneXSMax-Landscape.png", "Landscape", "{414, 896}",  "12.0",
 					"Default-IPhoneXSMax-Portrait.png", "Portrait", "{414, 896}",  "12.0",
 					"Default-IPhoneXR-Landscape.png", "Landscape", "{276, 597}",  "12.0",
@@ -800,7 +800,7 @@ namespace UnrealBuildTool
 					CopyFiles(InEngineDir + "/Build/IOS/Resources/Graphics", AppDirectory, "[REDACTED_EMAIL]", true);
 					//                    CopyFiles(InEngineDir + "/Build/IOS/Resources/Graphics", AppDirectory, "Default-Portrait-1336.png", true);
 					CopyFiles(InEngineDir + "/Build/IOS/Resources/Graphics", AppDirectory, "[REDACTED_EMAIL]", true);
-					CopyFiles(InEngineDir + "/Build/IOS/Resources/Graphics", AppDirectory, "Default-IPhoneX-Portrait.png", true);
+					CopyFiles(InEngineDir + "/Build/IOS/Resources/Graphics", AppDirectory, "Default-IPhoneXS-Portrait.png", true);
 					CopyFiles(InEngineDir + "/Build/IOS/Resources/Graphics", AppDirectory, "Default-IPhoneXSMax-Portrait.png", true);
 					CopyFiles(InEngineDir + "/Build/IOS/Resources/Graphics", AppDirectory, "Default-IPhoneXR-Portrait.png", true);
 				}
@@ -812,7 +812,7 @@ namespace UnrealBuildTool
 					CopyFiles(InEngineDir + "/Build/IOS/Resources/Graphics", AppDirectory, "[REDACTED_EMAIL]", true);
 					//                    CopyFiles(InEngineDir + "/Build/IOS/Resources/Graphics", AppDirectory, "Default-Landscape-1336.png", true);
 					CopyFiles(InEngineDir + "/Build/IOS/Resources/Graphics", AppDirectory, "[REDACTED_EMAIL]", true);
-					CopyFiles(InEngineDir + "/Build/IOS/Resources/Graphics", AppDirectory, "Default-IPhoneX-Landscape.png", true);
+					CopyFiles(InEngineDir + "/Build/IOS/Resources/Graphics", AppDirectory, "Default-IPhoneXS-Landscape.png", true);
 					CopyFiles(InEngineDir + "/Build/IOS/Resources/Graphics", AppDirectory, "Default-IPhoneXSMax-Landscape.png", true);
 					CopyFiles(InEngineDir + "/Build/IOS/Resources/Graphics", AppDirectory, "Default-IPhoneXR-Landscape.png", true);
 				}
@@ -836,7 +836,7 @@ namespace UnrealBuildTool
 					CopyFiles(BuildDirectory + "/Resources/Graphics", AppDirectory, "[REDACTED_EMAIL]", true);
 					//                    CopyFiles(BuildDirectory + "/Resources/Graphics", AppDirectory, "Default-Portrait-1336.png", true);
 					CopyFiles(BuildDirectory + "/Resources/Graphics", AppDirectory, "[REDACTED_EMAIL]", true);
-					CopyFiles(BuildDirectory + "/Resources/Graphics", AppDirectory, "Default-IPhoneX-Portrait.png", true);
+					CopyFiles(BuildDirectory + "/Resources/Graphics", AppDirectory, "Default-IPhoneXS-Portrait.png", true);
 					CopyFiles(BuildDirectory + "/Resources/Graphics", AppDirectory, "Default-IPhoneXSMax-Portrait.png", true);
 					CopyFiles(BuildDirectory + "/Resources/Graphics", AppDirectory, "Default-IPhoneXR-Portrait.png", true);
 				}
@@ -848,7 +848,7 @@ namespace UnrealBuildTool
 					CopyFiles(BuildDirectory + "/Resources/Graphics", AppDirectory, "[REDACTED_EMAIL]", true);
 					//                    CopyFiles(BuildDirectory + "/Resources/Graphics", AppDirectory, "Default-Landscape-1336.png", true);
 					CopyFiles(BuildDirectory + "/Resources/Graphics", AppDirectory, "[REDACTED_EMAIL]", true);
-					CopyFiles(BuildDirectory + "/Resources/Graphics", AppDirectory, "Default-IPhoneX-Landscape.png", true);
+					CopyFiles(BuildDirectory + "/Resources/Graphics", AppDirectory, "Default-IPhoneXS-Landscape.png", true);
 					CopyFiles(BuildDirectory + "/Resources/Graphics", AppDirectory, "Default-IPhoneXSMax-Landscape.png", true);
 					CopyFiles(BuildDirectory + "/Resources/Graphics", AppDirectory, "Default-IPhoneXR-Landscape.png", true);
 				}
