UBT: Always compile the nativized plugin rules assembly on the host machine in a remote build.

diff --git a/Engine/Build/Rsync/RsyncProject.txt b/Engine/Build/Rsync/RsyncProject.txt
--- a/Engine/Build/Rsync/RsyncProject.txt
+++ b/Engine/Build/Rsync/RsyncProject.txt
@@ -8,7 +8,11 @@
 + /Intermediate/Build/BuildRules/
 + /Intermediate/Build/BuildRules/*.dll
 + /Intermediate/Plugins/
-protect /Intermediate/Plugins/**/Intermediate/
++ /Intermediate/Plugins/**/Intermediate/
++ /Intermediate/Plugins/**/Intermediate/Build/
++ /Intermediate/Plugins/**/Intermediate/Build/BuildRules/
++ /Intermediate/Plugins/**/Intermediate/Build/BuildRules/*.dll
+protect /Intermediate/Plugins/**/Intermediate/**
 + /Intermediate/Plugins/**
 + /Intermediate/Source/
 + /Intermediate/Source/**
diff --git a/Engine/Source/Programs/UnrealBuildTool/Configuration/TargetRules.cs b/Engine/Source/Programs/UnrealBuildTool/Configuration/TargetRules.cs
--- a/Engine/Source/Programs/UnrealBuildTool/Configuration/TargetRules.cs
+++ b/Engine/Source/Programs/UnrealBuildTool/Configuration/TargetRules.cs
@@ -1309,6 +1309,58 @@ namespace UnrealBuildTool
 			}
 		}
 
+		/// <summary>
+		/// Checks whether nativization is enabled for this target, and determine the path for the nativized plugin
+		/// </summary>
+		/// <returns>The nativized plugin file, or null if nativization is not enabled</returns>
+		internal FileReference GetNativizedPlugin()
+		{
+			if (ProjectFile != null && (Type == TargetType.Game || Type == TargetType.Client || Type == TargetType.Server))
+			{
+				// Read the config files for this project
+				ConfigHierarchy Config = ConfigCache.ReadHierarchy(ConfigHierarchyType.Game, ProjectFile.Directory, BuildHostPlatform.Current.Platform);
+				if (Config != null)
+				{
+					// Determine whether or not the user has enabled nativization of Blueprint assets at cook time (default is 'Disabled')
+					string NativizationMethod;
+					if (Config.TryGetValue("/Script/UnrealEd.ProjectPackagingSettings", "BlueprintNativizationMethod", out NativizationMethod) && NativizationMethod != "Disabled")
+					{
+						string PlatformName;
+						if (Platform == UnrealTargetPlatform.Win32 || Platform == UnrealTargetPlatform.Win64)
+						{
+							PlatformName = "Windows";
+						}
+						else
+						{
+							PlatformName = Platform.ToString();
+						}
+
+						// Temp fix to force platforms that only support "Game" configurations at cook time to the correct path.
+						string ProjectTargetType;
+						if (Platform == UnrealTargetPlatform.Win32 || Platform == UnrealTargetPlatform.Win64 || Platform == UnrealTargetPlatform.Linux || Platform == UnrealTargetPlatform.Mac)
+						{
+							ProjectTargetType = Type.ToString();
+						}
+						else
+						{
+							ProjectTargetType = "Game";
+						}
+
+						FileReference PluginFile = FileReference.Combine(ProjectFile.Directory, "Intermediate", "Plugins", "NativizedAssets", PlatformName, ProjectTargetType, "NativizedAssets.uplugin");
+						if (FileReference.Exists(PluginFile))
+						{
+							return PluginFile;
+						}
+						else
+						{
+							Log.TraceWarning("{0} is configured for nativization, but is missing the generated code plugin at \"{1}\". Make sure to cook {2} data before attempting to build the {3} target. If data was cooked with nativization enabled, this can also mean there were no Blueprint assets that required conversion, in which case this warning can be safely ignored.", Name, PluginFile.FullName, Type.ToString(), Platform.ToString());
+						}
+					}
+				}
+			}
+			return null;
+		}
+
 		/// <summary>
 		/// Finds all the subobjects which can be configured by command line options and config files
 		/// </summary>
diff --git a/Engine/Source/Programs/UnrealBuildTool/Configuration/UEBuildTarget.cs b/Engine/Source/Programs/UnrealBuildTool/Configuration/UEBuildTarget.cs
--- a/Engine/Source/Programs/UnrealBuildTool/Configuration/UEBuildTarget.cs
+++ b/Engine/Source/Programs/UnrealBuildTool/Configuration/UEBuildTarget.cs
@@ -484,41 +484,10 @@ namespace UnrealBuildTool
 			}
 
 			// Include generated code plugin if not building an editor target and project is configured for nativization
-			if (RulesObject.ProjectFile != null
-				&& (RulesObject.Type == TargetType.Game || RulesObject.Type == TargetType.Client || RulesObject.Type == TargetType.Server)
-				&& ShouldIncludeNativizedAssets(RulesObject.ProjectFile.Directory))
+			FileReference NativizedPluginFile = RulesObject.GetNativizedPlugin();
+			if(NativizedPluginFile != null)
 			{
-				string PlatformName;
-				if (RulesObject.Platform == UnrealTargetPlatform.Win32 || RulesObject.Platform == UnrealTargetPlatform.Win64)
-				{
-					PlatformName = "Windows";
-				}
-				else
-				{
-					PlatformName = RulesObject.Platform.ToString();
-				}
-
-				// Temp fix to force platforms that only support "Game" configurations at cook time to the correct path.
-				string ProjectTargetType;
-				if (RulesObject.Platform == UnrealTargetPlatform.Win32 || RulesObject.Platform == UnrealTargetPlatform.Win64
-					|| RulesObject.Platform == UnrealTargetPlatform.Linux || RulesObject.Platform == UnrealTargetPlatform.Mac)
-				{
-					ProjectTargetType = RulesObject.Type.ToString();
-				}
-				else
-				{
-					ProjectTargetType = "Game";
-				}
-
-				FileReference PluginFile = FileReference.Combine(RulesObject.ProjectFile.Directory, "Intermediate", "Plugins", "NativizedAssets", PlatformName, ProjectTargetType, "NativizedAssets.uplugin");
-				if (FileReference.Exists(PluginFile))
-				{
-					RulesAssembly = RulesCompiler.CreatePluginRulesAssembly(PluginFile, false, RulesAssembly, false);
-				}
-				else
-				{
-					Log.TraceWarning("{0} is configured for nativization, but is missing the generated code plugin at \"{1}\". Make sure to cook {2} data before attempting to build the {3} target. If data was cooked with nativization enabled, this can also mean there were no Blueprint assets that required conversion, in which case this warning can be safely ignored.", RulesObject.Name, PluginFile.FullName, RulesObject.Type.ToString(), RulesObject.Platform.ToString());
-				}
+				RulesAssembly = RulesCompiler.CreatePluginRulesAssembly(NativizedPluginFile, bSkipRulesCompile, RulesAssembly, false);
 			}
 
 			// Generate a build target from this rules module
@@ -4350,23 +4319,5 @@ namespace UnrealBuildTool
 
 			return FilteredFileItems;
 		}
-
-        static bool ShouldIncludeNativizedAssets(DirectoryReference GameProjectDirectory)
-        {
-            ConfigHierarchy Config = ConfigCache.ReadHierarchy(ConfigHierarchyType.Game, GameProjectDirectory, BuildHostPlatform.Current.Platform);
-            if (Config != null)
-            {
-                // Determine whether or not the user has enabled nativization of Blueprint assets at cook time (default is 'Disabled')
-                string BlueprintNativizationMethod;
-                if (!Config.TryGetValue("/Script/UnrealEd.ProjectPackagingSettings", "BlueprintNativizationMethod", out BlueprintNativizationMethod))
-                {
-                    BlueprintNativizationMethod = "Disabled";
-                }
-
-                return BlueprintNativizationMethod != "Disabled";
-            }
-
-            return false;
-        }
 	}
 }
diff --git a/Engine/Source/Programs/UnrealBuildTool/ToolChain/RemoteMac.cs b/Engine/Source/Programs/UnrealBuildTool/ToolChain/RemoteMac.cs
--- a/Engine/Source/Programs/UnrealBuildTool/ToolChain/RemoteMac.cs
+++ b/Engine/Source/Programs/UnrealBuildTool/ToolChain/RemoteMac.cs
@@ -351,7 +351,17 @@ namespace UnrealBuildTool
 			}
 
 			// Compile the rules assembly
-			RulesCompiler.CreateTargetRulesAssembly(TargetDesc.ProjectFile, TargetDesc.Name, false, false, TargetDesc.ForeignPlugin);
+			RulesAssembly RulesAssembly = RulesCompiler.CreateTargetRulesAssembly(TargetDesc.ProjectFile, TargetDesc.Name, false, false, TargetDesc.ForeignPlugin);
+
+			// Create the target rules
+			TargetRules Rules = RulesAssembly.CreateTargetRules(TargetDesc.Name, TargetDesc.Platform, TargetDesc.Configuration, TargetDesc.Architecture, TargetDesc.ProjectFile, new ReadOnlyBuildVersion(BuildVersion.ReadDefault()), new string[0]);
+
+			// Check if we need to enable a nativized plugin, and compile the assembly for that if we do
+			FileReference NativizedPluginFile = Rules.GetNativizedPlugin();
+			if(NativizedPluginFile != null)
+			{
+				RulesAssembly = RulesCompiler.CreatePluginRulesAssembly(NativizedPluginFile, false, RulesAssembly, false);
+			}
 
 			// Path to the local manifest file. This has to be translated from the remote format after the build is complete.
 			List<FileReference> LocalManifestFiles = new List<FileReference>();
