Linux: fix RPATHs for projects located on a different drive (UE-65913).

- Based on PR #5226 by dorgonman.

diff --git a/Engine/Source/Programs/UnrealBuildTool/Platform/Linux/LinuxToolChain.cs b/Engine/Source/Programs/UnrealBuildTool/Platform/Linux/LinuxToolChain.cs
--- a/Engine/Source/Programs/UnrealBuildTool/Platform/Linux/LinuxToolChain.cs
+++ b/Engine/Source/Programs/UnrealBuildTool/Platform/Linux/LinuxToolChain.cs
@@ -1478,6 +1478,15 @@ namespace UnrealBuildTool
 				if ((AdditionalLibrary.Contains("Plugins") || AdditionalLibrary.Contains("Binaries/ThirdParty") || AdditionalLibrary.Contains("Binaries\\ThirdParty")) && Path.GetDirectoryName(AdditionalLibrary) != Path.GetDirectoryName(OutputFile.AbsolutePath))
 				{
 					string RelativePath = new FileReference(AdditionalLibrary).Directory.MakeRelativeTo(OutputFile.Location.Directory);
+					// On Windows, MakeRelativeTo can silently fail if the engine and the project are located on different drives
+					if (CrossCompiling() && RelativePath.StartsWith(UnrealBuildTool.EngineDirectory.FullName))
+					{
+						// do not replace directly, but take care to avoid potential double slashes or missed slashes
+						string PathFromEngineDir = RelativePath.Replace(UnrealBuildTool.EngineDirectory.FullName, "");
+						// Path.Combine doesn't combine these properly
+						RelativePath = ((PathFromEngineDir.StartsWith("\\") || PathFromEngineDir.StartsWith("/")) ? "..\\..\\.." : "..\\..\\..\\") + PathFromEngineDir;
+					}
+
 					if (!RPaths.Contains(RelativePath))
 					{
 						RPaths.Add(RelativePath);
