UBT: Prevent enterprise files being overwritten in installed builds. Suspect this allows version files to be overwritten, and allows loading incompatible modules.

diff --git a/Engine/Source/Programs/UnrealBuildTool/Configuration/UEBuildTarget.cs b/Engine/Source/Programs/UnrealBuildTool/Configuration/UEBuildTarget.cs
--- a/Engine/Source/Programs/UnrealBuildTool/Configuration/UEBuildTarget.cs
+++ b/Engine/Source/Programs/UnrealBuildTool/Configuration/UEBuildTarget.cs
@@ -1818,6 +1818,10 @@ namespace UnrealBuildTool
 			{
 				return true;
 			}
+			if(UnrealBuildTool.IsEnterpriseInstalled() && File.IsUnderDirectory(UnrealBuildTool.EnterpriseDirectory))
+			{
+				return true;
+			}
 			if(UnrealBuildTool.IsProjectInstalled() && ProjectFile != null && File.IsUnderDirectory(ProjectFile.Directory))
 			{
 				return true;
