Fix remote build failure from installed builds if the installation directory contains a space.

diff --git a/Engine/Source/Programs/UnrealBuildTool/ToolChain/RemoteMac.cs b/Engine/Source/Programs/UnrealBuildTool/ToolChain/RemoteMac.cs
--- a/Engine/Source/Programs/UnrealBuildTool/ToolChain/RemoteMac.cs
+++ b/Engine/Source/Programs/UnrealBuildTool/ToolChain/RemoteMac.cs
@@ -653,7 +653,7 @@ namespace UnrealBuildTool
 			List<string> Arguments = new List<string>(CommonRsyncArguments);
 			Arguments.Add(String.Format("--rsync-path=\"mkdir -p {0} && rsync\"", RemoteDirectory));
 			Arguments.Add(String.Format("\"{0}\"", GetLocalCygwinPath(LocalFile)));
-			Arguments.Add(String.Format("\"{0}@{1}\":'{2}/'", UserName, ServerName, RemoteDirectory));
+			Arguments.Add(String.Format("\"{0}@{1}\":'{2}'", UserName, ServerName, RemoteFile));
 			Arguments.Add("-q");
 
 			int Result = Rsync(String.Join(" ", Arguments));
