Linux: check against root path, not engine path, when fixing RPATHs (UE-65913).

diff --git a/Engine/Source/Programs/UnrealBuildTool/Platform/Linux/LinuxToolChain.cs b/Engine/Source/Programs/UnrealBuildTool/Platform/Linux/LinuxToolChain.cs
--- a/Engine/Source/Programs/UnrealBuildTool/Platform/Linux/LinuxToolChain.cs
+++ b/Engine/Source/Programs/UnrealBuildTool/Platform/Linux/LinuxToolChain.cs
@@ -1479,12 +1479,12 @@ namespace UnrealBuildTool
 				{
 					string RelativePath = new FileReference(AdditionalLibrary).Directory.MakeRelativeTo(OutputFile.Location.Directory);
 					// On Windows, MakeRelativeTo can silently fail if the engine and the project are located on different drives
-					if (CrossCompiling() && RelativePath.StartsWith(UnrealBuildTool.EngineDirectory.FullName))
+					if (CrossCompiling() && RelativePath.StartsWith(UnrealBuildTool.RootDirectory.FullName))
 					{
 						// do not replace directly, but take care to avoid potential double slashes or missed slashes
-						string PathFromEngineDir = RelativePath.Replace(UnrealBuildTool.EngineDirectory.FullName, "");
+						string PathFromRootDir = RelativePath.Replace(UnrealBuildTool.RootDirectory.FullName, "");
 						// Path.Combine doesn't combine these properly
-						RelativePath = ((PathFromEngineDir.StartsWith("\\") || PathFromEngineDir.StartsWith("/")) ? "..\\..\\.." : "..\\..\\..\\") + PathFromEngineDir;
+						RelativePath = ((PathFromRootDir.StartsWith("\\") || PathFromRootDir.StartsWith("/")) ? "..\\..\\.." : "..\\..\\..\\") + PathFromRootDir;
 					}
 
 					if (!RPaths.Contains(RelativePath))
