Fix for spammy MLSDK warnings

diff --git a/Engine/Source/Programs/UnrealBuildTool/Platform/Lumin/UEBuildLumin.cs b/Engine/Source/Programs/UnrealBuildTool/Platform/Lumin/UEBuildLumin.cs
--- a/Engine/Source/Programs/UnrealBuildTool/Platform/Lumin/UEBuildLumin.cs
+++ b/Engine/Source/Programs/UnrealBuildTool/Platform/Lumin/UEBuildLumin.cs
@@ -342,7 +342,6 @@ namespace UnrealBuildTool
 			// we don't have an MLSDK directory specified
 			if (String.IsNullOrEmpty(MLSDKPath))
 			{
-				Console.WriteLine("*** Unable to determine MLSDK location from MLSDK environment variable or config file ***");
 				return false;
 			}
 
@@ -372,7 +371,6 @@ namespace UnrealBuildTool
 				return false;
 			}
 
-			Console.WriteLine("*** Found installed MLSDK version {0}.{1} ***", DetectedMajorVersion, DetectedMinorVersion);
 			return true;
 		}
 
