addition of iPad 6 device profile
addition of detection, launch images, and profiles for new iPad Pros

diff --git a/Engine/Config/BaseDeviceProfiles.ini b/Engine/Config/BaseDeviceProfiles.ini
--- a/Engine/Config/BaseDeviceProfiles.ini
+++ b/Engine/Config/BaseDeviceProfiles.ini
@@ -41,6 +41,9 @@
 +DeviceProfileNameAndTypes=iPadPro97,IOS
 +DeviceProfileNameAndTypes=iPadPro2_129,IOS
 +DeviceProfileNameAndTypes=iPad5,IOS
++DeviceProfileNameAndTypes=iPad6,IOS
++DeviceProfileNameAndTypes=iPadPro11,IOS
++DeviceProfileNameAndTypes=iPadPro3_129,IOS
 +DeviceProfileNameAndTypes=AppleTV,IOS
 +DeviceProfileNameAndTypes=AppleTV4K,IOS
 +DeviceProfileNameAndTypes=HTML5,HTML5
@@ -345,10 +348,26 @@ BaseProfileName=IPadPro
 DeviceType=IOS
 BaseProfileName=IPadPro
 
+[iPadPro11 DeviceProfile]
+DeviceType=IOS
+BaseProfileName=IPadPro
+
 [iPad5 DeviceProfile]
 DeviceType=IOS
 BaseProfileName=IPadPro97
 
+[iPad6 DeviceProfile]
+DeviceType=IOS
+BaseProfileName=IPadPro97
+
+[iPadPro2_129 DeviceProfile]
+DeviceType=IOS
+BaseProfileName=IPadPro
+
+[iPadPro3_129 DeviceProfile]
+DeviceType=IOS
+BaseProfileName=IPadPro
+
 [/Script/AndroidDeviceProfileSelector.AndroidDeviceProfileMatchingRules]
 MatchProfile=(Profile="Android_PowerVR54x",Match=((SourceType=SRC_GpuFamily,CompareType=CMP_Regex,MatchString="PowerVR SGX 54[0-9]")))
 +MatchProfile=(Profile="Android_PowerVRMid",Match=((SourceType=SRC_GpuFamily,CompareType=CMP_Regex,MatchString="PowerVR"),(SourceType=SRC_GlVersion,CompareType=CMP_Regex,MatchString="^OpenGL ES 3\\.")))
diff --git a/Engine/Config/BaseEditorPerProjectUserSettings.ini b/Engine/Config/BaseEditorPerProjectUserSettings.ini
--- a/Engine/Config/BaseEditorPerProjectUserSettings.ini
+++ b/Engine/Config/BaseEditorPerProjectUserSettings.ini
@@ -266,10 +266,14 @@ LastExecutedLaunchPlatform=
 +PhoneScreenResolutions=(Description="Razer Phone",Width=1080,Height=1920,AspectRatio="16:9",bCanSwapAspectRatio=true,ProfileName="Android_Mid")
 
 ; common screen resolutions for tablet devices
++TabletScreenResolutions=(Description="iPad Pro 12.9-inch (3rd gen.)",Width=1024,Height=1366,AspectRatio="~3:4",bCanSwapAspectRatio=true,ProfileName="iPadPro3_129")
 +TabletScreenResolutions=(Description="iPad Pro 12.9-inch (2nd gen.)",Width=1024,Height=1366,AspectRatio="~3:4",bCanSwapAspectRatio=true,ProfileName="iPadPro2_129")
++TabletScreenResolutions=(Description="iPad Pro 11-inch",Width=834,Height=1194,AspectRatio="5:7",bCanSwapAspectRatio=true,ProfileName="iPadPro11")
 +TabletScreenResolutions=(Description="iPad Pro 10.5-inch",Width=834,Height=1112,AspectRatio="3:4",bCanSwapAspectRatio=true,ProfileName="iPadPro105")
 +TabletScreenResolutions=(Description="iPad Pro 12.9-inch",Width=1024,Height=1366,AspectRatio="3:4",bCanSwapAspectRatio=true,ProfileName="iPadPro129")
 +TabletScreenResolutions=(Description="iPad Pro 9.7-inch",Width=768,Height=1024,AspectRatio="3:4",bCanSwapAspectRatio=true,ProfileName="iPadPro97")
++TabletScreenResolutions=(Description="iPad (6th gen.)",Width=768,Height=1024,AspectRatio="3:4",bCanSwapAspectRatio=true,ProfileName="iPad6")
++TabletScreenResolutions=(Description="iPad (5th gen.)",Width=768,Height=1024,AspectRatio="3:4",bCanSwapAspectRatio=true,ProfileName="iPad5")
 +TabletScreenResolutions=(Description="iPad Air 2",Width=768,Height=1024,AspectRatio="3:4",bCanSwapAspectRatio=true,ProfileName="iPadAir2")
 +TabletScreenResolutions=(Description="iPad Mini 4",Width=768,Height=1024,AspectRatio="3:4",bCanSwapAspectRatio=true,ProfileName="iPadMini4")
 +TabletScreenResolutions=(Description="LG G Pad X 8.0",Width=768,Height=1366,AspectRatio="9:16",bCanSwapAspectRatio=true)
diff --git a/Engine/Source/Developer/IOS/IOSPlatformEditor/Private/IOSTargetSettingsCustomization.cpp b/Engine/Source/Developer/IOS/IOSPlatformEditor/Private/IOSTargetSettingsCustomization.cpp
--- a/Engine/Source/Developer/IOS/IOSPlatformEditor/Private/IOSTargetSettingsCustomization.cpp
+++ b/Engine/Source/Developer/IOS/IOSPlatformEditor/Private/IOSTargetSettingsCustomization.cpp
@@ -89,18 +89,18 @@ FIOSTargetSettingsCustomization::FIOSTargetSettingsCustomization()
     new (LaunchImageNames) FPlatformIconInfo(TEXT("[REDACTED_EMAIL]"), LOCTEXT("LaunchImage_iPhone5_Portrait", "iPhone 5/5S/SE in Portrait"), FText::GetEmpty(), 640, 1136, FPlatformIconInfo::Required);
     new (LaunchImageNames) FPlatformIconInfo(TEXT("Default-IPhone6-Landscape.png"), LOCTEXT("LaunchImage_iPhone6_Landscape", "iPhone 6/7/8 in Landscape"), FText::GetEmpty(), 1334, 750, FPlatformIconInfo::Required);
     new (LaunchImageNames) FPlatformIconInfo(TEXT("Default-IPhone6.png"), LOCTEXT("LaunchImage_iPhone6", "iPhone 6/7/8 in Portrait"), FText::GetEmpty(), 750, 1334, FPlatformIconInfo::Required);
-    new (LaunchImageNames) FPlatformIconInfo(TEXT("Default-IPhone6Plus-Landscape.png"), LOCTEXT("LaunchImage_iPhone6Plus_Landscape", "iPhone 6+/7+/8+/9 in Landscape"), FText::GetEmpty(), 2208, 1242, FPlatformIconInfo::Required);
-    new (LaunchImageNames) FPlatformIconInfo(TEXT("Default-IPhone6Plus-Portrait.png"), LOCTEXT("LaunchImage_iPhone6Plus_Portrait", "iPhone 6+/7+/8+/9 in Portrait"), FText::GetEmpty(), 1242, 2208, FPlatformIconInfo::Required);
+    new (LaunchImageNames) FPlatformIconInfo(TEXT("Default-IPhone6Plus-Landscape.png"), LOCTEXT("LaunchImage_iPhone6Plus_Landscape", "iPhone 6+/7+/8+ in Landscape"), FText::GetEmpty(), 2208, 1242, FPlatformIconInfo::Required);
+    new (LaunchImageNames) FPlatformIconInfo(TEXT("Default-IPhone6Plus-Portrait.png"), LOCTEXT("LaunchImage_iPhone6Plus_Portrait", "iPhone 6+/7+/8+ in Portrait"), FText::GetEmpty(), 1242, 2208, FPlatformIconInfo::Required);
     new (LaunchImageNames) FPlatformIconInfo(TEXT("Default-IPhoneXS-Landscape.png"), LOCTEXT("LaunchImage_iPhoneX_Landscape", "iPhone X/XS in Landscape"), FText::GetEmpty(), 2436, 1125, FPlatformIconInfo::Required);
     new (LaunchImageNames) FPlatformIconInfo(TEXT("Default-IPhoneXS-Portrait.png"), LOCTEXT("LaunchImage_iPhoneX_Portrait", "iPhone X/XS in Portrait"), FText::GetEmpty(), 1125, 2436, FPlatformIconInfo::Required);
 	new (LaunchImageNames) FPlatformIconInfo(TEXT("Default-IPhoneXSMax-Landscape.png"), LOCTEXT("LaunchImage_iPhoneXSMax_Landscape", "iPhone XS Max in Landscape"), FText::GetEmpty(), 2688, 1242, FPlatformIconInfo::Required);
-	new (LaunchImageNames) FPlatformIconInfo(TEXT("Default-IPhoneXSMax-Portrait.png"), LOCTEXT("LaunchImage_iPhoneXMax_Portrait", "iPhone XS Max in Portrait"), FText::GetEmpty(), 1242, 2688, FPlatformIconInfo::Required);
-	new (LaunchImageNames) FPlatformIconInfo(TEXT("Default-IPhoneXR-Landscape.png"), LOCTEXT("LaunchImage_iPhoneXR_Landscape", "iPhone XR in Landscape"), FText::GetEmpty(), 1792, 828, FPlatformIconInfo::Required);
-	new (LaunchImageNames) FPlatformIconInfo(TEXT("Default-IPhoneXR-Portrait.png"), LOCTEXT("LaunchImage_iPhoneXR_Portrait", "iPhone XR in Portrait"), FText::GetEmpty(), 828, 1792, FPlatformIconInfo::Required);
+	new (LaunchImageNames) FPlatformIconInfo(TEXT("Default-IPhoneXSMax-Portrait.png"), LOCTEXT("LaunchImage_iPhoneXSMax_Portrait", "iPhone XS Max in Portrait"), FText::GetEmpty(), 1242, 2688, FPlatformIconInfo::Required);
 	new (LaunchImageNames) FPlatformIconInfo(TEXT("[REDACTED_EMAIL]"), LOCTEXT("LaunchImage_iPadRetina_Landscape", "iPad in Landscape"), FText::GetEmpty(), 2048, 1536, FPlatformIconInfo::Required);
 	new (LaunchImageNames) FPlatformIconInfo(TEXT("[REDACTED_EMAIL]"), LOCTEXT("LaunchImage_iPadRetina_Portrait", "iPad in Portrait"), FText::GetEmpty(), 1536, 2048, FPlatformIconInfo::Required);
-	new (LaunchImageNames) FPlatformIconInfo(TEXT("Default-IPadPro10_5_Landscape.png"), LOCTEXT("LaunchImage_iPadPro_Landscape", "iPad Pro 10.5 in Landscape"), FText::GetEmpty(), 2224, 1668, FPlatformIconInfo::Required);
-	new (LaunchImageNames) FPlatformIconInfo(TEXT("Default-IPadPro10_5_Portrait.png"), LOCTEXT("LaunchImage_iPadPro_Portrait", "iPad Pro 10.5 in Portrait"), FText::GetEmpty(), 1668, 2224, FPlatformIconInfo::Required);
+	new (LaunchImageNames)FPlatformIconInfo(TEXT("[REDACTED_EMAIL]"), LOCTEXT("LaunchImage_iPadPro105Retina_Landscape", "iPad Pro 10.5 in Landscape"), FText::GetEmpty(), 2224, 1668, FPlatformIconInfo::Required);
+	new (LaunchImageNames)FPlatformIconInfo(TEXT("[REDACTED_EMAIL]"), LOCTEXT("LaunchImage_iPadPro105Retina_Portrait", "iPad Pro 10.5 in Portrait"), FText::GetEmpty(), 1668, 2224, FPlatformIconInfo::Required);
+	new (LaunchImageNames)FPlatformIconInfo(TEXT("[REDACTED_EMAIL]"), LOCTEXT("LaunchImage_iPadPro11Retina_Landscape", "iPad Pro 11 in Landscape"), FText::GetEmpty(), 2388, 1668, FPlatformIconInfo::Required);
+	new (LaunchImageNames)FPlatformIconInfo(TEXT("[REDACTED_EMAIL]"), LOCTEXT("LaunchImage_iPadPro11Retina_Portrait", "iPad Pro 11 in Portrait"), FText::GetEmpty(), 1668, 2388, FPlatformIconInfo::Required);
 	new (LaunchImageNames)FPlatformIconInfo(TEXT("[REDACTED_EMAIL]"), LOCTEXT("LaunchImage_iPadProRetina_Landscape", "iPad Pro 12.9 in Landscape"), FText::GetEmpty(), 2732, 2048, FPlatformIconInfo::Required);
 	new (LaunchImageNames)FPlatformIconInfo(TEXT("[REDACTED_EMAIL]"), LOCTEXT("LaunchImage_iPadProRetina_Portrait", "iPad Pro 12.9 in Portrait"), FText::GetEmpty(), 2048, 2732, FPlatformIconInfo::Required);
 
@@ -1817,6 +1817,12 @@ void FIOSTargetSettingsCustomization::SetShaderStandard(int32 Value)
             SetMinVersion((int32)EIOSVersion::IOS_12);
         }
 
+		// make sure we never set the min version to less than current supported
+		if (((EIOSVersion)EnumValue < EIOSVersion::IOS_10))
+		{
+			SetMinVersion((int32)EIOSVersion::IOS_10);
+		}
+
 		
 		ShaderVersionWarningTextBox->SetError(Message);
 	}
diff --git a/Engine/Source/Programs/UnrealBuildTool/Platform/IOS/UEDeployIOS.cs b/Engine/Source/Programs/UnrealBuildTool/Platform/IOS/UEDeployIOS.cs
--- a/Engine/Source/Programs/UnrealBuildTool/Platform/IOS/UEDeployIOS.cs
+++ b/Engine/Source/Programs/UnrealBuildTool/Platform/IOS/UEDeployIOS.cs
@@ -598,6 +598,10 @@ namespace UnrealBuildTool
 				{
 					"Default-Landscape.png", "Landscape", "{768, 1024}", "7.0",
 					"Default-Portrait.png", "Portrait", "{768, 1024}",  "7.0",
+					"Default-Landscape-1112.png", "Landscape", "{834, 1112}",  "9.0",
+					"Default-Portrait-1112.png", "Portrait", "{834, 1112}",  "9.0",
+					"Default-Landscape-1194.png", "Landscape", "{834, 1194}",  "9.0",
+					"Default-Portrait-1194.png", "Portrait", "{834, 1194}",  "9.0",
 					"Default-Landscape-1336.png", "Landscape", "{1024, 1366}",  "9.0",
 					"Default-Portrait-1336.png", "Portrait", "{1024, 1366}",  "9.0",
 				};
@@ -796,9 +800,9 @@ namespace UnrealBuildTool
 				{
 					CopyFiles(InEngineDir + "/Build/IOS/Resources/Graphics", AppDirectory, "Default-IPhone6.png", true);
 					CopyFiles(InEngineDir + "/Build/IOS/Resources/Graphics", AppDirectory, "Default-IPhone6Plus-Portrait.png", true);
-					//                   CopyFiles(InEngineDir + "/Build/IOS/Resources/Graphics", AppDirectory, "Default-Portrait.png", true);
 					CopyFiles(InEngineDir + "/Build/IOS/Resources/Graphics", AppDirectory, "[REDACTED_EMAIL]", true);
-					//                    CopyFiles(InEngineDir + "/Build/IOS/Resources/Graphics", AppDirectory, "Default-Portrait-1336.png", true);
+					CopyFiles(InEngineDir + "/Build/IOS/Resources/Graphics", AppDirectory, "[REDACTED_EMAIL]", true);
+					CopyFiles(InEngineDir + "/Build/IOS/Resources/Graphics", AppDirectory, "[REDACTED_EMAIL]", true);
 					CopyFiles(InEngineDir + "/Build/IOS/Resources/Graphics", AppDirectory, "[REDACTED_EMAIL]", true);
 					CopyFiles(InEngineDir + "/Build/IOS/Resources/Graphics", AppDirectory, "Default-IPhoneXS-Portrait.png", true);
 					CopyFiles(InEngineDir + "/Build/IOS/Resources/Graphics", AppDirectory, "Default-IPhoneXSMax-Portrait.png", true);
@@ -808,15 +812,14 @@ namespace UnrealBuildTool
 				{
 					CopyFiles(InEngineDir + "/Build/IOS/Resources/Graphics", AppDirectory, "Default-IPhone6-Landscape.png", true);
 					CopyFiles(InEngineDir + "/Build/IOS/Resources/Graphics", AppDirectory, "Default-IPhone6Plus-Landscape.png", true);
-					//                    CopyFiles(InEngineDir + "/Build/IOS/Resources/Graphics", AppDirectory, "Default-Landscape.png", true);
 					CopyFiles(InEngineDir + "/Build/IOS/Resources/Graphics", AppDirectory, "[REDACTED_EMAIL]", true);
-					//                    CopyFiles(InEngineDir + "/Build/IOS/Resources/Graphics", AppDirectory, "Default-Landscape-1336.png", true);
+					CopyFiles(InEngineDir + "/Build/IOS/Resources/Graphics", AppDirectory, "[REDACTED_EMAIL]", true);
+					CopyFiles(InEngineDir + "/Build/IOS/Resources/Graphics", AppDirectory, "[REDACTED_EMAIL]", true);
 					CopyFiles(InEngineDir + "/Build/IOS/Resources/Graphics", AppDirectory, "[REDACTED_EMAIL]", true);
 					CopyFiles(InEngineDir + "/Build/IOS/Resources/Graphics", AppDirectory, "Default-IPhoneXS-Landscape.png", true);
 					CopyFiles(InEngineDir + "/Build/IOS/Resources/Graphics", AppDirectory, "Default-IPhoneXSMax-Landscape.png", true);
 					CopyFiles(InEngineDir + "/Build/IOS/Resources/Graphics", AppDirectory, "Default-IPhoneXR-Landscape.png", true);
 				}
-				//                CopyFiles(InEngineDir + "/Build/IOS/Resources/Graphics", AppDirectory, "Default.png", true);
 				CopyFiles(InEngineDir + "/Build/IOS/Resources/Graphics", AppDirectory, "[REDACTED_EMAIL]", true);
 				CopyFiles(InEngineDir + "/Build/IOS/Resources/Graphics", AppDirectory, "[REDACTED_EMAIL]", true);
 			}
@@ -832,9 +835,9 @@ namespace UnrealBuildTool
 				{
 					CopyFiles(BuildDirectory + "/Resources/Graphics", AppDirectory, "Default-IPhone6.png", true);
 					CopyFiles(BuildDirectory + "/Resources/Graphics", AppDirectory, "Default-IPhone6Plus-Portrait.png", true);
-					//                    CopyFiles(BuildDirectory + "/Resources/Graphics", AppDirectory, "Default-Portrait.png", true);
 					CopyFiles(BuildDirectory + "/Resources/Graphics", AppDirectory, "[REDACTED_EMAIL]", true);
-					//                    CopyFiles(BuildDirectory + "/Resources/Graphics", AppDirectory, "Default-Portrait-1336.png", true);
+					CopyFiles(BuildDirectory + "/Resources/Graphics", AppDirectory, "[REDACTED_EMAIL]", true);
+					CopyFiles(BuildDirectory + "/Resources/Graphics", AppDirectory, "[REDACTED_EMAIL]", true);
 					CopyFiles(BuildDirectory + "/Resources/Graphics", AppDirectory, "[REDACTED_EMAIL]", true);
 					CopyFiles(BuildDirectory + "/Resources/Graphics", AppDirectory, "Default-IPhoneXS-Portrait.png", true);
 					CopyFiles(BuildDirectory + "/Resources/Graphics", AppDirectory, "Default-IPhoneXSMax-Portrait.png", true);
@@ -844,15 +847,14 @@ namespace UnrealBuildTool
 				{
 					CopyFiles(BuildDirectory + "/Resources/Graphics", AppDirectory, "Default-IPhone6-Landscape.png", true);
 					CopyFiles(BuildDirectory + "/Resources/Graphics", AppDirectory, "Default-IPhone6Plus-Landscape.png", true);
-					//                    CopyFiles(BuildDirectory + "/Resources/Graphics", AppDirectory, "Default-Landscape.png", true);
 					CopyFiles(BuildDirectory + "/Resources/Graphics", AppDirectory, "[REDACTED_EMAIL]", true);
-					//                    CopyFiles(BuildDirectory + "/Resources/Graphics", AppDirectory, "Default-Landscape-1336.png", true);
+					CopyFiles(BuildDirectory + "/Resources/Graphics", AppDirectory, "[REDACTED_EMAIL]", true);
+					CopyFiles(BuildDirectory + "/Resources/Graphics", AppDirectory, "[REDACTED_EMAIL]", true);
 					CopyFiles(BuildDirectory + "/Resources/Graphics", AppDirectory, "[REDACTED_EMAIL]", true);
 					CopyFiles(BuildDirectory + "/Resources/Graphics", AppDirectory, "Default-IPhoneXS-Landscape.png", true);
 					CopyFiles(BuildDirectory + "/Resources/Graphics", AppDirectory, "Default-IPhoneXSMax-Landscape.png", true);
 					CopyFiles(BuildDirectory + "/Resources/Graphics", AppDirectory, "Default-IPhoneXR-Landscape.png", true);
 				}
-				//                CopyFiles(BuildDirectory + "/Resources/Graphics", AppDirectory, "Default.png", true);
 				CopyFiles(BuildDirectory + "/Resources/Graphics", AppDirectory, "[REDACTED_EMAIL]", true);
 				CopyFiles(BuildDirectory + "/Resources/Graphics", AppDirectory, "[REDACTED_EMAIL]", true);
 			}
diff --git a/Engine/Source/Runtime/ApplicationCore/Private/IOS/IOSAppDelegate.cpp b/Engine/Source/Runtime/ApplicationCore/Private/IOS/IOSAppDelegate.cpp
--- a/Engine/Source/Runtime/ApplicationCore/Private/IOS/IOSAppDelegate.cpp
+++ b/Engine/Source/Runtime/ApplicationCore/Private/IOS/IOSAppDelegate.cpp
@@ -796,7 +796,7 @@ static FAutoConsoleVariableRef CVarGEnableThermalsReport(
 	}
     else if (Device == FPlatformMisc::IOS_IPhoneX || Device == FPlatformMisc::IOS_IPhoneXS)
 	{
-		[ImageString appendString : @"-IPhoneX"];
+		[ImageString appendString : @"-IPhoneXS"];
 		if (!self.bDeviceInPortraitMode)
 		{
 			[ImageString appendString : @"-Landscape"];
@@ -836,7 +836,7 @@ static FAutoConsoleVariableRef CVarGEnableThermalsReport(
 		// use IPhone6 image for now
 		[ImageString appendString : @"-IPhone6Plus-Landscape"];
 	}
-	else if (Device == FPlatformMisc::IOS_IPadPro_129 || Device == FPlatformMisc::IOS_IPadPro2_129)
+	else if (Device == FPlatformMisc::IOS_IPadPro_129 || Device == FPlatformMisc::IOS_IPadPro2_129 || Device == FPlatformMisc::IOS_IPadPro3_129)
 	{
 		if (!self.bDeviceInPortraitMode)
 		{
@@ -852,6 +852,38 @@ static FAutoConsoleVariableRef CVarGEnableThermalsReport(
             [ImageString appendString:@"@2x"];
         }
 	}
+	else if (Device == FPlatformMisc::IOS_IPadPro_105)
+	{
+		if (!self.bDeviceInPortraitMode)
+		{
+			[ImageString appendString : @"-Landscape-1112"];
+		}
+		else
+		{
+			[ImageString appendString : @"-Portrait-1112"];
+		}
+
+		if (NativeScale > 1.0f)
+		{
+			[ImageString appendString : @"@2x"];
+		}
+	}
+	else if (Device == FPlatformMisc::IOS_IPadPro_11)
+	{
+		if (!self.bDeviceInPortraitMode)
+		{
+			[ImageString appendString : @"-Landscape-1194"];
+		}
+		else
+		{
+			[ImageString appendString : @"-Portrait-1194"];
+		}
+
+		if (NativeScale > 1.0f)
+		{
+			[ImageString appendString : @"@2x"];
+		}
+	}
 	else
 	{
 		if (MainFrame.size.height == 320 && MainFrame.size.width != 480 && !self.bDeviceInPortraitMode)
diff --git a/Engine/Source/Runtime/ApplicationCore/Private/IOS/IOSPlatformApplicationMisc.cpp b/Engine/Source/Runtime/ApplicationCore/Private/IOS/IOSPlatformApplicationMisc.cpp
--- a/Engine/Source/Runtime/ApplicationCore/Private/IOS/IOSPlatformApplicationMisc.cpp
+++ b/Engine/Source/Runtime/ApplicationCore/Private/IOS/IOSPlatformApplicationMisc.cpp
@@ -224,7 +224,7 @@ void FIOSPlatformApplicationMisc::ClipboardPaste(class FString& Result)
 EScreenPhysicalAccuracy FIOSPlatformApplicationMisc::ComputePhysicalScreenDensity(int32& ScreenDensity)
 {
 	FPlatformMisc::EIOSDevice Device = FPlatformMisc::GetIOSDeviceType();
-	static_assert( FPlatformMisc::EIOSDevice::IOS_Unknown == 35, "Every device needs to be handled here." );
+	static_assert( FPlatformMisc::EIOSDevice::IOS_Unknown == 38, "Every device needs to be handled here." );
 
 	ScreenDensity = 0;
 	EScreenPhysicalAccuracy Accuracy = EScreenPhysicalAccuracy::Unknown;
@@ -269,6 +269,7 @@ EScreenPhysicalAccuracy FIOSPlatformApplicationMisc::ComputePhysicalScreenDensit
 	case FPlatformMisc::IOS_IPad3:
 	case FPlatformMisc::IOS_IPad4:
 	case FPlatformMisc::IOS_IPad5:
+	case FPlatformMisc::IOS_IPad6:
 	case FPlatformMisc::IOS_IPadAir:
 	case FPlatformMisc::IOS_IPadAir2:
 	case FPlatformMisc::IOS_IPadPro_97:
@@ -278,7 +279,9 @@ EScreenPhysicalAccuracy FIOSPlatformApplicationMisc::ComputePhysicalScreenDensit
 	case FPlatformMisc::IOS_IPadPro:
 	case FPlatformMisc::IOS_IPadPro_129:
 	case FPlatformMisc::IOS_IPadPro_105:
+	case FPlatformMisc::IOS_IPadPro_11:
 	case FPlatformMisc::IOS_IPadPro2_129:
+	case FPlatformMisc::IOS_IPadPro3_129:
 		ScreenDensity = 264;
 		Accuracy = EScreenPhysicalAccuracy::Truth;
 		break;
diff --git a/Engine/Source/Runtime/Core/Private/IOS/IOSPlatformMisc.cpp b/Engine/Source/Runtime/Core/Private/IOS/IOSPlatformMisc.cpp
--- a/Engine/Source/Runtime/Core/Private/IOS/IOSPlatformMisc.cpp
+++ b/Engine/Source/Runtime/Core/Private/IOS/IOSPlatformMisc.cpp
@@ -448,16 +448,31 @@ FIOSPlatformMisc::EIOSDevice FIOSPlatformMisc::GetIOSDeviceType()
 			{
 				DeviceType = IOS_IPadPro_105;
 			}
+			else if (Minor == 5 || Minor == 6)
+			{
+				DeviceType = IOS_IPad6;
+			}
 			else
 			{
 				DeviceType = IOS_IPadPro2_129;
 			}
 		}
+		else if (Major == 8)
+		{
+			if (Minor == 3 || Minor == 4)
+			{
+				DeviceType = IOS_IPadPro_11;
+			}
+			else
+			{
+				DeviceType = IOS_IPadPro3_129;
+			}
+		}
 
 		// Default to highest settings currently available for any future device
-		else if (Major >= 8)
+		else if (Major >= 9)
 		{
-			DeviceType = IOS_IPadPro2_129;
+			DeviceType = IOS_IPadPro3_129;
 		}
 	}
 	// iPhones
diff --git a/Engine/Source/Runtime/Core/Public/IOS/IOSPlatformMisc.h b/Engine/Source/Runtime/Core/Public/IOS/IOSPlatformMisc.h
--- a/Engine/Source/Runtime/Core/Public/IOS/IOSPlatformMisc.h
+++ b/Engine/Source/Runtime/Core/Public/IOS/IOSPlatformMisc.h
@@ -134,6 +134,9 @@ struct CORE_API FIOSPlatformMisc : public FApplePlatformMisc
         IOS_IPhoneXS,
         IOS_IPhoneXSMax,
         IOS_IPhoneXR,
+		IOS_IPad6,
+		IOS_IPadPro_11,
+		IOS_IPadPro3_129,
 		IOS_Unknown,
 	};
 
@@ -178,6 +181,9 @@ struct CORE_API FIOSPlatformMisc : public FApplePlatformMisc
             TEXT("IPhoneXS"),
             TEXT("IPhoneXSMax"),
             TEXT("IPhoneXR"),
+			TEXT("IPad6"),
+			TEXT("IPadPro11"),
+			TEXT("IPadPro3_129"),
 			TEXT("Unknown"),
 		};
 		static_assert((sizeof(IOSDeviceNames) / sizeof(IOSDeviceNames[0])) == ((int32)IOS_Unknown + 1), "Mismatched IOSDeviceNames and EIOSDevice.");
