UE-65854 - TVOS package fails from Windows, IPP ERROR: Stub IPA does not appear to be signed correctly (missing mobileprovision or CodeResources)

diff --git a/Engine/Source/Programs/UnrealBuildTool/ProjectFiles/Xcode/XcodeProject.cs b/Engine/Source/Programs/UnrealBuildTool/ProjectFiles/Xcode/XcodeProject.cs
--- a/Engine/Source/Programs/UnrealBuildTool/ProjectFiles/Xcode/XcodeProject.cs
+++ b/Engine/Source/Programs/UnrealBuildTool/ProjectFiles/Xcode/XcodeProject.cs
@@ -1050,7 +1050,7 @@ namespace UnrealBuildTool
 				{
 					foreach (UnrealTargetPlatform Platform in Platforms)
 					{
-						if (InstalledPlatformInfo.IsValidPlatform(Platform, EProjectType.Code) && (Platform == UnrealTargetPlatform.Mac || Platform == UnrealTargetPlatform.IOS)) // @todo support other platforms
+						if (InstalledPlatformInfo.IsValidPlatform(Platform, EProjectType.Code) && (Platform == UnrealTargetPlatform.Mac || Platform == UnrealTargetPlatform.IOS || Platform == UnrealTargetPlatform.TVOS)) // @todo support other platforms
 						{
 							UEBuildPlatform BuildPlatform = UEBuildPlatform.GetBuildPlatform(Platform, true);
 							if ((BuildPlatform != null) && (BuildPlatform.HasRequiredSDKsInstalled() == SDKStatus.Valid))
@@ -1114,7 +1114,7 @@ namespace UnrealBuildTool
                                                 string TVOSExecutableName = MacExecutableName.Replace("-Mac-", "-TVOS-");
                                                 BuildConfigs.Add(new XcodeBuildConfig(ConfigName, TargetName, FileReference.Combine(OutputDirectory, "Mac", MacExecutableName), FileReference.Combine(OutputDirectory, "IOS", IOSExecutableName), FileReference.Combine(OutputDirectory, "TVOS", TVOSExecutableName), ProjectTarget, Configuration));
                                             }
-                                            else if (BuildPlatform.Platform == UnrealTargetPlatform.IOS)
+                                            else if (BuildPlatform.Platform == UnrealTargetPlatform.IOS || BuildPlatform.Platform == UnrealTargetPlatform.TVOS)
                                             {
                                                 string IOSExecutableName = MakeExecutableFileName(ExeName, UnrealTargetPlatform.IOS, Configuration, ProjectTarget.TargetRules.Architecture, ProjectTarget.TargetRules.UndecoratedConfiguration);
                                                 string TVOSExecutableName = IOSExecutableName.Replace("-IOS-", "-TVOS-");
