UBT: Fix CachedIncludePaths not being set correctly for modules which do not use unity builds and do not have IWYU enabled. Set the cached include paths for all source files, regardless of the code path used to compile them.

Causes files to not be rebuilt when headers change, potentially causing memory corruption, etc...

diff --git a/Engine/Source/Programs/UnrealBuildTool/Configuration/UEBuildModuleCPP.cs b/Engine/Source/Programs/UnrealBuildTool/Configuration/UEBuildModuleCPP.cs
--- a/Engine/Source/Programs/UnrealBuildTool/Configuration/UEBuildModuleCPP.cs
+++ b/Engine/Source/Programs/UnrealBuildTool/Configuration/UEBuildModuleCPP.cs
@@ -345,6 +345,24 @@ namespace UnrealBuildTool
 				return LinkInputFiles;
 			}
 
+			// Store the module compile environment along with the source file.  This is so that we can use it later on when looking for header dependencies
+			foreach (FileItem CFile in SourceFilesFound.CFiles)
+			{
+				CFile.CachedIncludePaths = ModuleCompileEnvironment.IncludePaths;
+			}
+			foreach (FileItem CCFile in SourceFilesFound.CCFiles)
+			{
+				CCFile.CachedIncludePaths = ModuleCompileEnvironment.IncludePaths;
+			}
+			foreach (FileItem CPPFile in SourceFilesFound.CPPFiles)
+			{
+				CPPFile.CachedIncludePaths = ModuleCompileEnvironment.IncludePaths;
+			}
+			foreach (FileItem MMFile in SourceFilesFound.MMFiles)
+			{
+				MMFile.CachedIncludePaths = ModuleCompileEnvironment.IncludePaths;
+			}
+
 			// Process all of the header file dependencies for this module
 			CheckFirstIncludeMatchesEachCppFile(Target, ModuleCompileEnvironment);
 
