On Mac add paths for frameworks specifed in PublicAdditionalFrameworks to framework search paths passed to the compiler

diff --git a/Engine/Source/Programs/UnrealBuildTool/Platform/Mac/MacToolChain.cs b/Engine/Source/Programs/UnrealBuildTool/Platform/Mac/MacToolChain.cs
--- a/Engine/Source/Programs/UnrealBuildTool/Platform/Mac/MacToolChain.cs
+++ b/Engine/Source/Programs/UnrealBuildTool/Platform/Mac/MacToolChain.cs
@@ -424,6 +424,19 @@ namespace UnrealBuildTool
 				Arguments.Append("\"");
 			}
 
+			List<string> FrameworksSearchPaths = new List<string>();
+			foreach (UEBuildFramework Framework in CompileEnvironment.AdditionalFrameworks)
+			{
+				string FrameworkPath = Path.GetDirectoryName(Path.GetFullPath(Framework.FrameworkName));
+				if (!FrameworksSearchPaths.Contains(FrameworkPath))
+				{
+					Arguments.Append(" -F \"");
+					Arguments.Append(FrameworkPath);
+					Arguments.Append("\"");
+					FrameworksSearchPaths.Add(FrameworkPath);
+				}
+			}
+
 			CPPOutput Result = new CPPOutput();
 			// Create a compile action for each source file.
 			foreach (FileItem SourceFile in InputFiles)
