Fixed a crash in MacToolChain caused by an attempt to access null bundle path when linking a plugin during hot reload

diff --git a/Engine/Source/Programs/UnrealBuildTool/Platform/Mac/MacToolChain.cs b/Engine/Source/Programs/UnrealBuildTool/Platform/Mac/MacToolChain.cs
--- a/Engine/Source/Programs/UnrealBuildTool/Platform/Mac/MacToolChain.cs
+++ b/Engine/Source/Programs/UnrealBuildTool/Platform/Mac/MacToolChain.cs
@@ -1365,7 +1365,7 @@ namespace UnrealBuildTool
 							BuildProducts.Add(FileReference.Combine(BundleContentsDirectory, Resource.BundleContentsSubdir, ResourceFile.Substring(Path.GetDirectoryName(Resource.ResourcePath).Length + 1)), BuildProductType.RequiredResource);
 						}
 					}
-					else
+					else if (BundleContentsDirectory != null)
 					{
 						BuildProducts.Add(FileReference.Combine(BundleContentsDirectory, Resource.BundleContentsSubdir, Path.GetFileName(Resource.ResourcePath)), BuildProductType.RequiredResource);
 					}
