UBT: Use the 6.00 version of Clang by default on Windows, until fixes for newer versions have been merged into Main.

diff --git a/Engine/Source/Programs/UnrealBuildTool/Platform/Windows/UEBuildWindows.cs b/Engine/Source/Programs/UnrealBuildTool/Platform/Windows/UEBuildWindows.cs
--- a/Engine/Source/Programs/UnrealBuildTool/Platform/Windows/UEBuildWindows.cs
+++ b/Engine/Source/Programs/UnrealBuildTool/Platform/Windows/UEBuildWindows.cs
@@ -384,7 +384,12 @@ namespace UnrealBuildTool
 		/// <summary>
 		/// The default compiler version to be used, if installed. 
 		/// </summary>
-		static readonly VersionNumber DefaultToolChainVersion = VersionNumber.Parse("14.13.26128");
+		static readonly VersionNumber DefaultClangToolChainVersion = VersionNumber.Parse("6.0.0");
+
+		/// <summary>
+		/// The default compiler version to be used, if installed. 
+		/// </summary>
+		static readonly VersionNumber DefaultVisualStudioToolChainVersion = VersionNumber.Parse("14.13.26128");
 
 		/// <summary>
 		/// The default Windows SDK version to be used, if installed.
@@ -912,6 +917,16 @@ namespace UnrealBuildTool
 			}
 			else
 			{
+				VersionNumber DefaultToolChainVersion;
+				if(Compiler == WindowsCompiler.Clang)
+				{
+					DefaultToolChainVersion = DefaultClangToolChainVersion;
+				}
+				else
+				{
+					DefaultToolChainVersion = DefaultVisualStudioToolChainVersion;
+				}
+
 				if(ToolChainVersionToDir.ContainsKey(DefaultToolChainVersion))
 				{
 					ToolChainVersion = DefaultToolChainVersion;
