Avoid passing Objective-C exception flags to .c files as they emit an unused argument warning

diff --git a/Engine/Source/Programs/UnrealBuildTool/Platform/IOS/IOSToolChain.cs b/Engine/Source/Programs/UnrealBuildTool/Platform/IOS/IOSToolChain.cs
--- a/Engine/Source/Programs/UnrealBuildTool/Platform/IOS/IOSToolChain.cs
+++ b/Engine/Source/Programs/UnrealBuildTool/Platform/IOS/IOSToolChain.cs
@@ -233,15 +233,6 @@ namespace UnrealBuildTool
 				Result += " -fno-exceptions";
 			}
 
-			if (CompileEnvironment.bEnableObjCExceptions)
-			{
-				Result += " -fobjc-exceptions";
-			}
-			else
-			{
-				Result += " -fno-objc-exceptions";
-			}
-
 			string SanitizerMode = Environment.GetEnvironmentVariable("ENABLE_ADDRESS_SANITIZER");
 			if(SanitizerMode != null && SanitizerMode == "YES")
 			{
@@ -417,6 +408,23 @@ namespace UnrealBuildTool
 			return Result;
 		}
 
+		// Conditionally enable (default disabled) Objective-C exceptions
+		static string GetObjCExceptionsFlag(CppCompileEnvironment CompileEnvironment)
+		{
+			string Result = "";
+
+			if (CompileEnvironment.bEnableObjCExceptions)
+			{
+				Result += " -fobjc-exceptions";
+			}
+			else
+			{
+				Result += " -fno-objc-exceptions";
+			}
+
+			return Result;
+		}
+
 		static FileReference GetFrameworkZip(UEBuildFramework Framework)
 		{
 			if (Framework.OwningModule == null)
@@ -623,6 +631,7 @@ namespace UnrealBuildTool
 					// Compile the file as a C++ PCH.
 					FileArguments += GetCompileArguments_PCH();
 					FileArguments += GetRTTIFlag(CompileEnvironment);
+					FileArguments += GetObjCExceptionsFlag(CompileEnvironment);
 				}
 				else if (Extension == ".C")
 				{
@@ -634,17 +643,20 @@ namespace UnrealBuildTool
 					// Compile the file as Objective-C++ code.
 					FileArguments += GetCompileArguments_MM();
 					FileArguments += GetRTTIFlag(CompileEnvironment);
+					FileArguments += GetObjCExceptionsFlag(CompileEnvironment);
 				}
 				else if (Extension == ".M")
 				{
-					// Compile the file as Objective-C++ code.
+					// Compile the file as Objective-C code.
 					FileArguments += GetCompileArguments_M();
+					FileArguments += GetObjCExceptionsFlag(CompileEnvironment);
 				}
 				else
 				{
 					// Compile the file as C++ code.
 					FileArguments += GetCompileArguments_CPP();
 					FileArguments += GetRTTIFlag(CompileEnvironment);
+					FileArguments += GetObjCExceptionsFlag(CompileEnvironment);
 
 					// only use PCH for .cpp files
 					FilePCHArguments = PCHArguments;
