Fix CIS

#rnx

diff --git a/Engine/Source/Programs/UnrealBuildTool/Modes/WriteMetadataMode.cs b/Engine/Source/Programs/UnrealBuildTool/Modes/WriteMetadataMode.cs
--- a/Engine/Source/Programs/UnrealBuildTool/Modes/WriteMetadataMode.cs
+++ b/Engine/Source/Programs/UnrealBuildTool/Modes/WriteMetadataMode.cs
@@ -122,15 +122,8 @@ namespace UnrealBuildTool
 						// Merge files from the existing manifests with the new ones
 						foreach(KeyValuePair<FileReference, ModuleManifest> Pair in PreviousFileToManifest)
 						{
-							ModuleManifest OldManifest = Pair.Value;
-							ModuleManifest NewManifest = TargetInfo.FileToManifest[Pair.Key];
-							foreach(KeyValuePair<string, string> ModulePair in OldManifest.ModuleNameToFileName)
-							{
-								if(!NewManifest.ModuleNameToFileName.ContainsKey(ModulePair.Key))
-								{
-									NewManifest.ModuleNameToFileName.Add(ModulePair.Key, ModulePair.Value);
-								}
-							}
+							ModuleManifest TargetManifest = TargetInfo.FileToManifest[Pair.Key];
+							MergeManifests(Pair.Value, TargetManifest);
 						}
 
 						// Update the build id to use the current one
@@ -144,6 +137,18 @@ namespace UnrealBuildTool
 					Receipt.Version.BuildId = Guid.NewGuid().ToString();
 				}
 			}
+			else
+			{
+				// Read all the manifests and merge them into the new ones, if they have the same build id
+				foreach(KeyValuePair<FileReference, ModuleManifest> Pair in TargetInfo.FileToManifest)
+				{
+					ModuleManifest SourceManifest;
+					if(TryReadManifest(Pair.Key, out SourceManifest) && SourceManifest.BuildId == Receipt.Version.BuildId)
+					{
+						MergeManifests(SourceManifest, Pair.Value);
+					}
+				}
+			}
 
 			// Update the build id in all the manifests, and write them out
 			foreach (KeyValuePair<FileReference, ModuleManifest> Pair in TargetInfo.FileToManifest)
@@ -224,31 +229,65 @@ namespace UnrealBuildTool
 			bool bCanRecycleManifests = true;
 			foreach(FileReference ManifestFileName in ManifestFiles)
 			{
-				if(ManifestFileName.IsUnderDirectory(UnrealBuildTool.EngineDirectory) && FileReference.Exists(ManifestFileName))
+				ModuleManifest Manifest;
+				if(ManifestFileName.IsUnderDirectory(UnrealBuildTool.EngineDirectory) && TryReadManifest(ManifestFileName, out Manifest))
 				{
-					try
+					if(Manifest.BuildId == BuildId)
 					{
-						ModuleManifest Manifest = ModuleManifest.Read(ManifestFileName);
-						if(Manifest.BuildId == BuildId)
+						if(IsOutOfDate(ManifestFileName, Manifest))
 						{
-							if(IsOutOfDate(ManifestFileName, Manifest))
-							{
-								bCanRecycleManifests = false;
-								break;
-							}
-							RecycleFileToManifest.Add(ManifestFileName, Manifest);
+							bCanRecycleManifests = false;
+							break;
 						}
-					}
-					catch(Exception Ex)
-					{
-						Log.TraceWarning("Unable to read '{0}'; ignoring.", ManifestFileName);
-						Log.TraceLog(ExceptionUtils.FormatExceptionDetails(Ex));
+						RecycleFileToManifest.Add(ManifestFileName, Manifest);
 					}
 				}
 			}
 			return bCanRecycleManifests;
 		}
 
+		/// <summary>
+		/// Attempts to read a manifest from the given location
+		/// </summary>
+		/// <param name="ManifestFileName">Path to the manifest</param>
+		/// <param name="Manifest">If successful, receives the manifest that was read</param>
+		/// <returns>True if the manifest was read correctly, false otherwise</returns>
+		public static bool TryReadManifest(FileReference ManifestFileName, out ModuleManifest Manifest)
+		{
+			if(FileReference.Exists(ManifestFileName))
+			{
+				try
+				{
+					Manifest = ModuleManifest.Read(ManifestFileName);
+					return true;
+				}
+				catch(Exception Ex)
+				{
+					Log.TraceWarning("Unable to read '{0}'; ignoring.", ManifestFileName);
+					Log.TraceLog(ExceptionUtils.FormatExceptionDetails(Ex));
+				}
+			}
+
+			Manifest = null;
+			return false;
+		}
+
+		/// <summary>
+		/// Merge a manifest into another manifest
+		/// </summary>
+		/// <param name="SourceManifest">The source manifest</param>
+		/// <param name="TargetManifest">The target manifest to merge into</param>
+		static void MergeManifests(ModuleManifest SourceManifest, ModuleManifest TargetManifest)
+		{
+			foreach(KeyValuePair<string, string> ModulePair in SourceManifest.ModuleNameToFileName)
+			{
+				if(!TargetManifest.ModuleNameToFileName.ContainsKey(ModulePair.Key))
+				{
+					TargetManifest.ModuleNameToFileName.Add(ModulePair.Key, ModulePair.Value);
+				}
+			}
+		}
+
 		/// <summary>
 		/// Checks whether a module manifest on disk is out of date (whether any of the binaries it references are newer than it is)
 		/// </summary>
