Add more debug info to exceptions thrown when a target is invalid.

diff --git a/Engine/Source/Programs/UnrealBuildTool/System/RulesAssembly.cs b/Engine/Source/Programs/UnrealBuildTool/System/RulesAssembly.cs
--- a/Engine/Source/Programs/UnrealBuildTool/System/RulesAssembly.cs
+++ b/Engine/Source/Programs/UnrealBuildTool/System/RulesAssembly.cs
@@ -518,7 +518,7 @@ namespace UnrealBuildTool
 			if (Rules.bDisableUnverifiedCertificates)
 			{
 				Rules.GlobalDefinitions.Add("DISABLE_UNVERIFIED_CERTIFICATE_LOADING=1");
-		}
+			}
 
 			// Allow the platform to finalize the settings
 			UEBuildPlatform Platform = UEBuildPlatform.GetBuildPlatform(Rules.Platform);
@@ -526,8 +526,8 @@ namespace UnrealBuildTool
 
 			// Some platforms may *require* monolithic compilation...
 			if (Rules.LinkType != TargetLinkType.Monolithic && UEBuildPlatform.PlatformRequiresMonolithicBuilds(Rules.Platform, Rules.Configuration))
-		{
-				throw new BuildException(String.Format("{0} does not support modular builds", Rules.Platform));
+			{
+				throw new BuildException(String.Format("{0}: {1} does not support modular builds", Rules.Name, Rules.Platform));
 			}
 
 			return Rules;
