Fix errors building foreign projects with a space in the path.

#rnx

diff --git a/Engine/Source/Programs/UnrealBuildTool/ProjectFiles/VisualStudio/VCProject.cs b/Engine/Source/Programs/UnrealBuildTool/ProjectFiles/VisualStudio/VCProject.cs
--- a/Engine/Source/Programs/UnrealBuildTool/ProjectFiles/VisualStudio/VCProject.cs
+++ b/Engine/Source/Programs/UnrealBuildTool/ProjectFiles/VisualStudio/VCProject.cs
@@ -1384,7 +1384,7 @@ namespace UnrealBuildTool
 					BuildArguments.AppendFormat("{0} {1} {2}", TargetName, UBTPlatformName, UBTConfigurationName);
 					if (IsForeignProject)
 					{
-						BuildArguments.AppendFormat(" -Project=\"{0}\"", UProjectPath);
+						BuildArguments.AppendFormat(" -Project={0}", UProjectPath);
 					}
 
 					if (bUsePrecompiled)
