Remove SDK missing error if Gradle enabled since it can download it if necessary

#ue4
#android

diff --git a/Engine/Source/Programs/UnrealBuildTool/Platform/Android/UEDeployAndroid.cs b/Engine/Source/Programs/UnrealBuildTool/Platform/Android/UEDeployAndroid.cs
--- a/Engine/Source/Programs/UnrealBuildTool/Platform/Android/UEDeployAndroid.cs
+++ b/Engine/Source/Programs/UnrealBuildTool/Platform/Android/UEDeployAndroid.cs
@@ -248,7 +248,16 @@ namespace UnrealBuildTool
 					SDKLevelInt = GetApiLevelInt(SDKLevel);
 					if (SDKLevelInt < MinimumSDKLevel)
 					{
-						throw new BuildException("Can't make an APK without SDK API 'android-" + MinimumSDKLevel.ToString() + "' minimum installed");
+						if (bGradleEnabled)
+						{
+							SDKLevelInt = MinimumSDKLevel;
+							SDKLevel = "android-" + MinimumSDKLevel.ToString();
+							Log.TraceInformation("Gradle will attempt to download SDK API level {0}", SDKLevelInt);
+						}
+						else
+						{
+							throw new BuildException("Can't make an APK without SDK API 'android-" + MinimumSDKLevel.ToString() + "' minimum installed");
+						}
 					}
 				}
 
