Disable warnings about unused arguments when generating project files. We do additional argument parsing after the point that the warning is output.

diff --git a/Engine/Source/Programs/UnrealBuildTool/Modes/GenerateProjectFilesMode.cs b/Engine/Source/Programs/UnrealBuildTool/Modes/GenerateProjectFilesMode.cs
--- a/Engine/Source/Programs/UnrealBuildTool/Modes/GenerateProjectFilesMode.cs
+++ b/Engine/Source/Programs/UnrealBuildTool/Modes/GenerateProjectFilesMode.cs
@@ -164,7 +164,8 @@ namespace UnrealBuildTool
 			}
 
 			// Check there are no superfluous command line arguments
-			Arguments.CheckAllArgumentsUsed();
+			// TODO (still pass raw arguments below)
+			// Arguments.CheckAllArgumentsUsed();
 
 			// Now generate project files
 			ProjectFileGenerator.bGenerateProjectFiles = true;
