Fix issue with -Architectures and GPUArchitectures not being split properly

#ue4
#android

diff --git a/Engine/Source/Programs/UnrealBuildTool/Platform/Android/AndroidToolChain.cs b/Engine/Source/Programs/UnrealBuildTool/Platform/Android/AndroidToolChain.cs
--- a/Engine/Source/Programs/UnrealBuildTool/Platform/Android/AndroidToolChain.cs
+++ b/Engine/Source/Programs/UnrealBuildTool/Platform/Android/AndroidToolChain.cs
@@ -143,12 +143,18 @@ namespace UnrealBuildTool
 
 			if (InAdditionalArches != null)
 			{
-				AdditionalArches.AddRange(InAdditionalArches);
+				foreach (string Arch in InAdditionalArches)
+				{
+					AdditionalArches.AddRange(Arch.Split('+'));
+				}
 			}
 
 			if (InAdditionalGPUArches != null)
 			{
-				AdditionalGPUArches.AddRange(InAdditionalGPUArches);
+				foreach (string GPU in InAdditionalGPUArches)
+				{
+					AdditionalGPUArches.AddRange(GPU.Split('+'));
+				}
 			}
 
 			// by default tools chains don't parse arguments, but we want to be able to check the -architectures flag defined above. This is
