Fix issue with compiler executable paths with spaces on Android

#ue4
#android

diff --git a/Engine/Source/Programs/UnrealBuildTool/Platform/Android/AndroidToolChain.cs b/Engine/Source/Programs/UnrealBuildTool/Platform/Android/AndroidToolChain.cs
--- a/Engine/Source/Programs/UnrealBuildTool/Platform/Android/AndroidToolChain.cs
+++ b/Engine/Source/Programs/UnrealBuildTool/Platform/Android/AndroidToolChain.cs
@@ -333,20 +333,20 @@ namespace UnrealBuildTool
 			if (NDKDefineInt >= 140200)
 			{
 				ToolchainParamsArm = " -target armv7-none-linux-androideabi" +
-										" --sysroot=\"" + Path.Combine(NDKPath, "sysroot") + "\"" +
-										" -isystem " + Path.Combine(NDKPath, "sysroot/usr/include/arm-linux-androideabi/") +
+										" --sysroot='" + Path.Combine(NDKPath, "sysroot") + "'" +
+										" -isystem '" + Path.Combine(NDKPath, "sysroot/usr/include/arm-linux-androideabi/") + "'" +
 										" -D__ANDROID_API__=" + NDKApiLevel32Int;
 				ToolchainParamsArm64 = " -target aarch64-none-linux-android" +
-										" --sysroot=\"" + Path.Combine(NDKPath, "sysroot") + "\"" +
-										" -isystem " + Path.Combine(NDKPath, "sysroot/usr/include/aarch64-linux-android/") +
+										" --sysroot='" + Path.Combine(NDKPath, "sysroot") + "'" +
+										" -isystem '" + Path.Combine(NDKPath, "sysroot/usr/include/aarch64-linux-android/") + "'" +
 										" -D__ANDROID_API__=" + NDKApiLevel64Int;
 				ToolchainParamsx86 = " -target i686-none-linux-android" +
-										" --sysroot=\"" + Path.Combine(NDKPath, "sysroot") + "\"" +
-										" -isystem " + Path.Combine(NDKPath, "sysroot/usr/include/i686-linux-android/") +
+										" --sysroot='" + Path.Combine(NDKPath, "sysroot") + "'" +
+										" -isystem '" + Path.Combine(NDKPath, "sysroot/usr/include/i686-linux-android/") + "'" +
 										" -D__ANDROID_API__=" + NDKApiLevel32Int;
 				ToolchainParamsx64 = " -target x86_64-none-linux-android" +
-										" --sysroot=\"" + Path.Combine(NDKPath, "sysroot") + "\"" +
-										" -isystem " + Path.Combine(NDKPath, "sysroot/usr/include/x86_64-linux-android/") +
+										" --sysroot='" + Path.Combine(NDKPath, "sysroot") + "'" +
+										" -isystem '" + Path.Combine(NDKPath, "sysroot/usr/include/x86_64-linux-android/") + "'" +
 										" -D__ANDROID_API__=" + NDKApiLevel64Int;
 			}
 			else
@@ -1454,14 +1454,20 @@ namespace UnrealBuildTool
 						CompileAction.WorkingDirectory = UnrealBuildTool.EngineSourceDirectory;
 						if(bExecuteCompilerThroughShell)
 						{
+							string FixedClangPath = ClangPath;
+							if (FixedClangPath.Contains(' '))
+							{
+								FixedClangPath = "'" + FixedClangPath + "'";
+							}
+					
 							CompileAction.CommandPath = BuildHostPlatform.Current.Shell;
 							if (BuildHostPlatform.Current.ShellType == ShellType.Cmd)
 							{
-								CompileAction.CommandArguments = String.Format("/c \"{0} {1}\"", Utils.MakePathSafeToUseWithCommandLine(ClangPath), ResponseArgument);
+								CompileAction.CommandArguments = String.Format("/c \"{0} {1}\"", FixedClangPath, ResponseArgument);
 							}
 							else
 							{
-								CompileAction.CommandArguments = String.Format("-c \'{0} {1}\'", Utils.MakePathSafeToUseWithCommandLine(ClangPath), ResponseArgument);
+								CompileAction.CommandArguments = String.Format("-c \'{0} {1}\'", FixedClangPath, ResponseArgument);
 							}
 							CompileAction.CommandDescription = "Compile";
 						}
@@ -1687,13 +1693,19 @@ namespace UnrealBuildTool
 
 					if(bExecuteCompilerThroughShell)
 					{
+						string LinkCommandPath = LinkAction.CommandPath.FullName;
+						if (LinkCommandPath.Contains(' '))
+						{
+							LinkCommandPath = "'" + LinkCommandPath + "'";
+						}
+						
 						if (BuildHostPlatform.Current.ShellType == ShellType.Cmd)
 						{
-							LinkAction.CommandArguments = String.Format("/c \"{0} {1}\"", LinkAction.CommandPath, LinkAction.CommandArguments);
+							LinkAction.CommandArguments = String.Format("/c \"{0} {1}\"", LinkCommandPath, LinkAction.CommandArguments);
 						}
 						else
 						{
-							LinkAction.CommandArguments = String.Format("-c \'{0} {1}\'", LinkAction.CommandPath, LinkAction.CommandArguments);
+							LinkAction.CommandArguments = String.Format("-c \'{0} {1}\'", LinkCommandPath, LinkAction.CommandArguments);
 						}
 						LinkAction.CommandPath = BuildHostPlatform.Current.Shell;
 						LinkAction.CommandDescription = "Link";
