fix for minimum version we support for TVOS

diff --git a/Engine/Source/Programs/UnrealBuildTool/Platform/TVOS/UEBuildTVOS.cs b/Engine/Source/Programs/UnrealBuildTool/Platform/TVOS/UEBuildTVOS.cs
--- a/Engine/Source/Programs/UnrealBuildTool/Platform/TVOS/UEBuildTVOS.cs
+++ b/Engine/Source/Programs/UnrealBuildTool/Platform/TVOS/UEBuildTVOS.cs
@@ -17,7 +17,7 @@ namespace UnrealBuildTool
 		/// </summary>
 		public override string RuntimeVersion
 		{
-			get { return "9.0"; }
+			get { return "10.0"; }
 		}
 
 		/// <summary>
