UBT: Add support for Aras-P's -ftime-trace PR for Clang, which outputs a JSON file containing a timeline for the compile. See http://aras-p.info/blog/2019/01/16/time-trace-timeline-flame-chart-profiler-for-Clang/ for more info.

diff --git a/Engine/Source/Programs/UnrealBuildTool/Platform/Windows/UEBuildWindows.cs b/Engine/Source/Programs/UnrealBuildTool/Platform/Windows/UEBuildWindows.cs
--- a/Engine/Source/Programs/UnrealBuildTool/Platform/Windows/UEBuildWindows.cs
+++ b/Engine/Source/Programs/UnrealBuildTool/Platform/Windows/UEBuildWindows.cs
@@ -203,6 +203,13 @@ namespace UnrealBuildTool
 		/// </summary>
 		public bool bOptimizeGlobalData = true;
 
+		/// <summary>
+		/// (Experimental) Appends the -ftime-trace argument to the command line for Clang, to output a JSON file containing a timeline for the compile. 
+		/// See http://aras-p.info/blog/2019/01/16/time-trace-timeline-flame-chart-profiler-for-Clang/ for more info.
+		/// </summary>
+		[XmlConfigFile(Category = "WindowsPlatform")]
+		public bool bClangTimeTrace = false;
+
 		/// <summary>
 		/// The Visual C++ environment to use for this target. Only initialized after all the target settings are finalized, in ValidateTarget().
 		/// </summary>
@@ -375,6 +382,11 @@ namespace UnrealBuildTool
 			get { return Inner.bOptimizeGlobalData; }
 		}
 
+		public bool bClangTimeTrace
+		{
+			get { return Inner.bClangTimeTrace; }
+		}
+
 		public string GetVisualStudioCompilerVersionName()
 		{
 			return Inner.GetVisualStudioCompilerVersionName();
diff --git a/Engine/Source/Programs/UnrealBuildTool/Platform/Windows/VCEnvironment.cs b/Engine/Source/Programs/UnrealBuildTool/Platform/Windows/VCEnvironment.cs
--- a/Engine/Source/Programs/UnrealBuildTool/Platform/Windows/VCEnvironment.cs
+++ b/Engine/Source/Programs/UnrealBuildTool/Platform/Windows/VCEnvironment.cs
@@ -502,7 +502,7 @@ namespace UnrealBuildTool
 			if(Compiler == WindowsCompiler.Clang || Compiler == WindowsCompiler.Intel)
 			{
 				ToolChain = WindowsCompiler.VisualStudio2017;
-				if(!WindowsPlatform.TryGetToolChainDir(ToolChain, CompilerVersion, out SelectedToolChainVersion, out SelectedToolChainDir))
+				if(!WindowsPlatform.TryGetToolChainDir(ToolChain, null, out SelectedToolChainVersion, out SelectedToolChainDir))
 				{
 					throw new BuildException("{0}{1} must be installed in order to build this target.", WindowsPlatform.GetCompilerName(Compiler), String.IsNullOrEmpty(CompilerVersion)? "" : String.Format(" ({0})", CompilerVersion));
 				}
diff --git a/Engine/Source/Programs/UnrealBuildTool/Platform/Windows/VCToolChain.cs b/Engine/Source/Programs/UnrealBuildTool/Platform/Windows/VCToolChain.cs
--- a/Engine/Source/Programs/UnrealBuildTool/Platform/Windows/VCToolChain.cs
+++ b/Engine/Source/Programs/UnrealBuildTool/Platform/Windows/VCToolChain.cs
@@ -1020,6 +1020,13 @@ namespace UnrealBuildTool
 					CompileAction.ProducedItems.Add(ObjectFile);
 					Result.ObjectFiles.Add(ObjectFile);
 					FileArguments.Add(String.Format("/Fo\"{0}\"", ObjectFile.AbsolutePath));
+
+					// Experimental: support for JSON output of timing data
+					if(Target.WindowsPlatform.Compiler == WindowsCompiler.Clang && Target.WindowsPlatform.bClangTimeTrace)
+					{
+						SharedArguments.Add("-Xclang -ftime-trace");
+						CompileAction.ProducedItems.Add(FileItem.GetItemByFileReference(ObjectFile.Location.ChangeExtension(".json")));
+					}
 				}
 
 				// Don't farm out creation of precompiled headers as it is the critical path task.
