Fix hot reload from IDE for DebugGame targets. Now that UE4Editor has a separate executable in DebugGame, we have to check for it running when starting a build/

diff --git a/Engine/Source/Programs/UnrealBuildTool/System/HotReload.cs b/Engine/Source/Programs/UnrealBuildTool/System/HotReload.cs
--- a/Engine/Source/Programs/UnrealBuildTool/System/HotReload.cs
+++ b/Engine/Source/Programs/UnrealBuildTool/System/HotReload.cs
@@ -144,7 +144,7 @@ namespace UnrealBuildTool
 			if (TargetDesc.Name.EndsWith("Editor", StringComparison.OrdinalIgnoreCase))
 			{
 				string EditorBaseFileName = "UE4Editor";
-				if (TargetDesc.Configuration != UnrealTargetConfiguration.Development && TargetDesc.Configuration != UnrealTargetConfiguration.DebugGame)
+				if (TargetDesc.Configuration != UnrealTargetConfiguration.Development)
 				{
 					EditorBaseFileName = String.Format("{0}-{1}-{2}", EditorBaseFileName, TargetDesc.Platform, TargetDesc.Configuration);
 				}
