UBT: Add support for using adaptive non-unity builds when the engine and project are in separate repositories.

diff --git a/Engine/Source/Programs/UnrealBuildTool/System/SourceFileWorkingSet.cs b/Engine/Source/Programs/UnrealBuildTool/System/SourceFileWorkingSet.cs
--- a/Engine/Source/Programs/UnrealBuildTool/System/SourceFileWorkingSet.cs
+++ b/Engine/Source/Programs/UnrealBuildTool/System/SourceFileWorkingSet.cs
@@ -1,4 +1,4 @@
-﻿// Copyright 1998-2018 Epic Games, Inc. All Rights Reserved.
+// Copyright 1998-2018 Epic Games, Inc. All Rights Reserved.
 
 using System;
 using System.Collections.Generic;
@@ -94,20 +94,23 @@ namespace UnrealBuildTool
 		HashSet<FileReference> Files;
 		List<DirectoryReference> Directories;
 		List<string> ErrorOutput;
+		GitSourceFileWorkingSet Inner;
 
 		/// <summary>
 		/// Constructor
 		/// </summary>
 		/// <param name="GitPath">Path to the Git executable</param>
 		/// <param name="RootDir">Root directory to run queries from (typically the directory containing the .git folder, to ensure all subfolders can be searched)</param>
-		public GitSourceFileWorkingSet(string GitPath, DirectoryReference RootDir)
+		/// <param name="Inner">An inner working set. This allows supporting multiple Git repositories (one containing the engine, another containing the project, for example)</param>
+		public GitSourceFileWorkingSet(string GitPath, DirectoryReference RootDir, GitSourceFileWorkingSet Inner)
 		{
 			this.RootDir = RootDir;
 			this.Files = new HashSet<FileReference>();
 			this.Directories = new List<DirectoryReference>();
 			this.ErrorOutput = new List<string>();
+			this.Inner = Inner;
 
-			Log.WriteLine(LogEventType.Console, "Using 'git status' to determine working set for adaptive non-unity build.");
+			Log.WriteLine(LogEventType.Console, "Using 'git status' to determine working set for adaptive non-unity build ({0}).", RootDir);
 
 			BackgroundProcess = new Process();
 			BackgroundProcess.StartInfo.FileName = GitPath;
@@ -186,6 +189,11 @@ namespace UnrealBuildTool
 		public void Dispose()
 		{
 			TerminateBackgroundProcess();
+
+			if(Inner != null)
+			{
+				Inner.Dispose();
+			}
 		}
 
 		/// <summary>
@@ -196,7 +204,15 @@ namespace UnrealBuildTool
 		public bool Contains(FileReference File)
 		{
 			WaitForBackgroundProcess();
-			return Files.Contains(File) || Directories.Any(x => File.IsUnderDirectory(x));
+			if(Files.Contains(File) || Directories.Any(x => File.IsUnderDirectory(x)))
+			{
+				return true;
+			}
+			if(Inner != null && Inner.Contains(File))
+			{
+				return true;
+			}
+			return false;
 		}
 
 		/// <summary>
@@ -306,16 +322,16 @@ namespace UnrealBuildTool
 			}
 			else if (Provider == ProviderType.Git)
 			{
-				DirectoryReference RepoDir;
+				GitSourceFileWorkingSet WorkingSet;
 				if (!String.IsNullOrEmpty(RepositoryPath))
 				{
-					RepoDir = DirectoryReference.Combine(RootDir, RepositoryPath);
+					WorkingSet = new GitSourceFileWorkingSet(GitPath, DirectoryReference.Combine(RootDir, RepositoryPath), null);
 				}
-				else if (!TryGetGitRepositoryPath(RootDir, ProjectDir, out RepoDir))
+				else if(!TryCreateGitWorkingSet(RootDir, ProjectDir, out WorkingSet))
 				{
-					RepoDir = RootDir;
+					WorkingSet = new GitSourceFileWorkingSet(GitPath, RootDir, null);
 				}
-				return new GitSourceFileWorkingSet(GitPath, RepoDir);
+				return WorkingSet;
 			}
 			else if (Provider == ProviderType.Perforce)
 			{
@@ -323,10 +339,10 @@ namespace UnrealBuildTool
 			}
 			else
 			{
-				DirectoryReference RepoDir;
-				if(TryGetGitRepositoryPath(RootDir, ProjectDir, out RepoDir))
+				GitSourceFileWorkingSet WorkingSet;
+				if(TryCreateGitWorkingSet(RootDir, ProjectDir, out WorkingSet))
 				{
-					return new GitSourceFileWorkingSet(GitPath, RepoDir);
+					return WorkingSet;
 				}
 				else
 				{
@@ -335,39 +351,35 @@ namespace UnrealBuildTool
 			}
 		}
 
-		/// <summary>
-		/// Tries to find the directory containing a git repository
-		/// </summary>
-		/// <param name="RootDir">Root directory containing the engine folder</param>
-		/// <param name="ProjectDir">Directory containing the project</param>
-		/// <param name="RepoDir">On success, receives the path to the repository</param>
-		/// <returns>True if a Git repository was found, false otherwise</returns>
-		static bool TryGetGitRepositoryPath(DirectoryReference RootDir, DirectoryReference ProjectDir, out DirectoryReference RepoDir)
+		static bool TryCreateGitWorkingSet(DirectoryReference RootDir, DirectoryReference ProjectDir, out GitSourceFileWorkingSet OutWorkingSet)
 		{
+			GitSourceFileWorkingSet WorkingSet  = null;
+
+			// Create the working set for the engine directory
 			if (DirectoryReference.Exists(DirectoryReference.Combine(RootDir, ".git")))
 			{
-				RepoDir = RootDir;
-				return true;
+				WorkingSet = new GitSourceFileWorkingSet(GitPath, RootDir, WorkingSet);
 			}
 
-			if (ProjectDir != null)
+			// Try to create a working set for the project directory
+			if(ProjectDir != null)
 			{
-				if (DirectoryReference.Exists(DirectoryReference.Combine(ProjectDir, ".git")))
-				{
-					RepoDir = ProjectDir;
-					return true;
-				}
-
-				DirectoryReference ParentProjectDir = ProjectDir.ParentDirectory;
-				if (DirectoryReference.Exists(DirectoryReference.Combine(ParentProjectDir, ".git")))
+				if(WorkingSet == null || !ProjectDir.IsUnderDirectory(RootDir))
 				{
-					RepoDir = ParentProjectDir;
-					return true;
+					if (DirectoryReference.Exists(DirectoryReference.Combine(ProjectDir, ".git")))
+					{
+						WorkingSet = new GitSourceFileWorkingSet(GitPath, ProjectDir, WorkingSet);
+					}
+					else if (DirectoryReference.Exists(DirectoryReference.Combine(ProjectDir.ParentDirectory, ".git")))
+					{
+						WorkingSet = new GitSourceFileWorkingSet(GitPath, ProjectDir.ParentDirectory, WorkingSet);
+					}
 				}
 			}
 
-			RepoDir = null;
-			return false;
+			// Set the output value
+			OutWorkingSet = WorkingSet;
+			return WorkingSet != null;
 		}
 	}
 }
