UBT: Prevent blank lines with XGE timing info when running with -progress. Some tools don't have any output, but our markup triggers writing to stdout, and XGE appends a timing string to it (eg. "[REDACTED_HANDLE] (0:01:23 at +0:34)").

diff --git a/Engine/Source/Programs/UnrealBuildTool/System/XGE.cs b/Engine/Source/Programs/UnrealBuildTool/System/XGE.cs
--- a/Engine/Source/Programs/UnrealBuildTool/System/XGE.cs
+++ b/Engine/Source/Programs/UnrealBuildTool/System/XGE.cs
@@ -545,7 +545,15 @@ namespace UnrealBuildTool
 					if (Args.Data != null && Args.Data.StartsWith(ProgressMarkupPrefix))
 					{
 						Writer.Write(++NumCompletedActions, NumActions);
-						Args = ConstructDataReceivedEventArgs(Args.Data.Substring(ProgressMarkupPrefix.Length));
+
+						// Strip out anything that is just an XGE timer. Some programs don't output anything except the progress text.
+						string RemainingData = Args.Data.Substring(ProgressMarkupPrefix.Length);
+						if(RemainingData.StartsWith(" (") && RemainingData.EndsWith(")"))
+						{
+							RemainingData = null;
+						}
+
+						Args = ConstructDataReceivedEventArgs(RemainingData);
 					}
 					if (OutputEventHandler != null)
 					{
