UBT: Remove some unnecessary code that uses reflection to construct internal DataReceivedEventArgs objects.

#rnx

diff --git a/Engine/Source/Programs/UnrealBuildTool/System/XGE.cs b/Engine/Source/Programs/UnrealBuildTool/System/XGE.cs
--- a/Engine/Source/Programs/UnrealBuildTool/System/XGE.cs
+++ b/Engine/Source/Programs/UnrealBuildTool/System/XGE.cs
@@ -165,32 +165,11 @@ namespace UnrealBuildTool
 					}
 				}
 
-				DataReceivedEventHandler OutputHandler = ((Sender, Args) => { if(Args.Data != null) { DefaultOutputHandler(Args.Data); } });
-				XGEResult = ExecuteTaskFileWithProgressMarkup(XGETaskFilePath, Actions.Count, OutputHandler);
+				XGEResult = ExecuteTaskFileWithProgressMarkup(XGETaskFilePath, Actions.Count);
 			}
 			return XGEResult;
 		}
 
-		private static void DefaultOutputHandler(string Message)
-		{
-			Log.TraceInformation(Message);
-		}
-
-		private static DataReceivedEventArgs ConstructDataReceivedEventArgs(string NewData)
-		{
-			// NOTE: This is a pretty nasty hack with C# reflection, however it saves us from having to replace all the
-			// handlers in various Toolchains with a wrapper handler that takes a string - it is certainly doable, but 
-			// touches code outside of this class that I don't want to touch right now
-
-			// DataReceivedEventArgs is not normally constructable, so work around it with creating a scratch Args object
-			DataReceivedEventArgs EventArgs = (DataReceivedEventArgs)System.Runtime.Serialization.FormatterServices.GetUninitializedObject(typeof(DataReceivedEventArgs));
-
-			// now we need to set the Data field using reflection, since it is read only
-			FieldInfo[] ArgFields = typeof(DataReceivedEventArgs).GetFields(BindingFlags.NonPublic | BindingFlags.Instance | BindingFlags.DeclaredOnly);
-			ArgFields[0].SetValue(EventArgs, NewData);
-			return EventArgs;
-		}
-
 		/// <summary>
 		/// Writes a XGE task file containing the specified actions to the specified file path.
 		/// </summary>
@@ -533,7 +512,7 @@ namespace UnrealBuildTool
 		/// <summary>
 		/// Executes the tasks in the specified file, parsing progress markup as part of the output.
 		/// </summary>
-		bool ExecuteTaskFileWithProgressMarkup(string TaskFilePath, int NumActions, DataReceivedEventHandler OutputEventHandler)
+		bool ExecuteTaskFileWithProgressMarkup(string TaskFilePath, int NumActions)
 		{
 			using (ProgressWriter Writer = new ProgressWriter("Compiling C++ source files...", false))
 			{
@@ -542,22 +521,21 @@ namespace UnrealBuildTool
 				// Create a wrapper delegate that will parse the output actions
 				DataReceivedEventHandler EventHandlerWrapper = (Sender, Args) =>
 				{
-					if (Args.Data != null && Args.Data.StartsWith(ProgressMarkupPrefix))
+					if(Args.Data != null)
 					{
-						Writer.Write(++NumCompletedActions, NumActions);
-
-						// Strip out anything that is just an XGE timer. Some programs don't output anything except the progress text.
-						string RemainingData = Args.Data.Substring(ProgressMarkupPrefix.Length);
-						if(RemainingData.StartsWith(" (") && RemainingData.EndsWith(")"))
+						string Text = Args.Data;
+						if (Text.StartsWith(ProgressMarkupPrefix))
 						{
-							RemainingData = null;
-						}
+							Writer.Write(++NumCompletedActions, NumActions);
 
-						Args = ConstructDataReceivedEventArgs(RemainingData);
-					}
-					if (OutputEventHandler != null)
-					{
-						OutputEventHandler(Sender, Args);
+							// Strip out anything that is just an XGE timer. Some programs don't output anything except the progress text.
+							Text = Args.Data.Substring(ProgressMarkupPrefix.Length);
+							if(Text.StartsWith(" (") && Text.EndsWith(")"))
+							{
+								return;
+							}
+						}
+						Log.TraceInformation(Text);
 					}
 				};
 
