PR #4775: New build param PCHMemoryAllocationFactor to set /Zm VS build param. (Contributed by [REDACTED_PERSON])

diff --git a/Engine/Source/Programs/UnrealBuildTool/Platform/Windows/UEBuildWindows.cs b/Engine/Source/Programs/UnrealBuildTool/Platform/Windows/UEBuildWindows.cs
--- a/Engine/Source/Programs/UnrealBuildTool/Platform/Windows/UEBuildWindows.cs
+++ b/Engine/Source/Programs/UnrealBuildTool/Platform/Windows/UEBuildWindows.cs
@@ -174,6 +174,12 @@ namespace UnrealBuildTool
 		[ConfigFile(ConfigHierarchyType.Engine, "/Script/WindowsTargetPlatform.WindowsTargetSettings")]
 		public int DefaultStackSizeCommit;
 
+		/// <summary>
+		/// Determines the amount of memory that the compiler allocates to construct precompiled headers (/Zm).
+		/// </summary>
+		[XmlConfigFile(Category = "WindowsPlatform")]
+		public int PCHMemoryAllocationFactor = 0;
+
 		/// <summary>
 		/// When using a Visual Studio compiler, returns the version name as a string
 		/// </summary>
@@ -295,6 +301,11 @@ namespace UnrealBuildTool
 			get { return Inner.DefaultStackSizeCommit; }
 		}
 
+		public int PCHMemoryAllocationFactor
+		{
+			get { return Inner.PCHMemoryAllocationFactor; }
+		}
+
 		public string GetVisualStudioCompilerVersionName()
 		{
 			return Inner.GetVisualStudioCompilerVersionName();
diff --git a/Engine/Source/Programs/UnrealBuildTool/Platform/Windows/VCToolChain.cs b/Engine/Source/Programs/UnrealBuildTool/Platform/Windows/VCToolChain.cs
--- a/Engine/Source/Programs/UnrealBuildTool/Platform/Windows/VCToolChain.cs
+++ b/Engine/Source/Programs/UnrealBuildTool/Platform/Windows/VCToolChain.cs
@@ -269,13 +269,23 @@ namespace UnrealBuildTool
 				Arguments.Add("/Gy");
 
 				// Allow 750% of the default memory allocation limit when using the static analyzer, and 850% at other times.
-				if (Target.WindowsPlatform.StaticAnalyzer == WindowsStaticAnalyzer.VisualCpp)
+				if(Target.WindowsPlatform.PCHMemoryAllocationFactor == 0)
 				{
-					Arguments.Add("/Zm750");
+					if (Target.WindowsPlatform.StaticAnalyzer == WindowsStaticAnalyzer.VisualCpp)
+					{
+						Arguments.Add("/Zm750");
+					}
+					else
+					{
+						Arguments.Add("/Zm850");
+					}
 				}
 				else
 				{
-					Arguments.Add("/Zm850");
+					if(Target.WindowsPlatform.PCHMemoryAllocationFactor > 0)
+					{
+						Arguments.Add(String.Format("/Zm{0}", Target.WindowsPlatform.PCHMemoryAllocationFactor));
+					}
 				}
 
 				// Disable "The file contains a character that cannot be represented in the current code page" warning for non-US windows.
