UBT: Remove some unnecessary indirection for generated code paths.

#rnx

diff --git a/Engine/Source/Programs/UnrealBuildTool/Configuration/UEBuildModuleCPP.cs b/Engine/Source/Programs/UnrealBuildTool/Configuration/UEBuildModuleCPP.cs
--- a/Engine/Source/Programs/UnrealBuildTool/Configuration/UEBuildModuleCPP.cs
+++ b/Engine/Source/Programs/UnrealBuildTool/Configuration/UEBuildModuleCPP.cs
@@ -1,4 +1,4 @@
-﻿// Copyright 1998-2018 Epic Games, Inc. All Rights Reserved.
+// Copyright 1998-2018 Epic Games, Inc. All Rights Reserved.
 
 using System;
 using System.Collections.Generic;
@@ -16,39 +16,6 @@ namespace UnrealBuildTool
 	/// </summary>
 	class UEBuildModuleCPP : UEBuildModule
 	{
-		public class AutoGenerateCppInfoClass
-		{
-			public class BuildInfoClass
-			{
-				/// <summary>
-				/// The wildcard of the *.gen.cpp file which was generated for the module
-				/// </summary>
-				public readonly string FileWildcard;
-
-				public BuildInfoClass(string InWildcard)
-				{
-					Debug.Assert(InWildcard != null);
-
-					FileWildcard = InWildcard;
-				}
-			}
-
-			/// <summary>
-			/// Information about how to build the .gen.cpp files. If this is null, then we're not building .gen.cpp files for this module.
-			/// </summary>
-			public BuildInfoClass BuildInfo;
-
-			public AutoGenerateCppInfoClass(BuildInfoClass InBuildInfo)
-			{
-				BuildInfo = InBuildInfo;
-			}
-		}
-
-		/// <summary>
-		/// Information about the .gen.cpp file.  If this is null then this module doesn't have any UHT-produced code.
-		/// </summary>
-		public AutoGenerateCppInfoClass AutoGenerateCppInfo = null;
-
 		public class SourceFilesClass
 		{
 			public readonly List<FileItem> MissingFiles = new List<FileItem>();
@@ -100,16 +67,6 @@ namespace UnrealBuildTool
 			}
 		}
 
-		/// <summary>
-		/// Adds additional source cpp files for this module.
-		/// </summary>
-		/// <param name="Files">Files to add.</param>
-		public void AddAdditionalCPPFiles(IEnumerable<FileItem> Files)
-		{
-			SourceFiles.AddRange(Files);
-			SourceFilesToBuild.CPPFiles.AddRange(Files);
-		}
-
 		/// <summary>
 		/// All the source files for this module
 		/// </summary>
@@ -135,6 +92,11 @@ namespace UnrealBuildTool
 		/// </summary>
 		public readonly DirectoryReference GeneratedCodeDirectory;
 
+		/// <summary>
+		/// Wildcard matching the *.gen.cpp files for this module.  If this is null then this module doesn't have any UHT-produced code.
+		/// </summary>
+		public string GeneratedCodeWildcard;
+
 		public List<DirectoryReference> IncludeSearchPaths = new List<DirectoryReference>();
 
 		public class ProcessedDependenciesClass
@@ -405,7 +367,7 @@ namespace UnrealBuildTool
 			bool bLegacyPublicIncludePaths
 			)
 		{
-			if(AutoGenerateCppInfo != null)
+			if(GeneratedCodeWildcard != null)
 			{
 				IncludePaths.Add(GeneratedCodeDirectory);
 			}
@@ -616,9 +578,9 @@ namespace UnrealBuildTool
 			}
 
 			// Compile all the generated CPP files
-			if (AutoGenerateCppInfo != null && AutoGenerateCppInfo.BuildInfo != null && !CompileEnvironment.bHackHeaderGenerator)
+			if (GeneratedCodeWildcard != null && !CompileEnvironment.bHackHeaderGenerator)
 			{
-				string[] GeneratedFiles = Directory.GetFiles(Path.GetDirectoryName(AutoGenerateCppInfo.BuildInfo.FileWildcard), Path.GetFileName(AutoGenerateCppInfo.BuildInfo.FileWildcard));
+				string[] GeneratedFiles = Directory.GetFiles(Path.GetDirectoryName(GeneratedCodeWildcard), Path.GetFileName(GeneratedCodeWildcard));
 				if(GeneratedFiles.Length > 0)
 				{
 					// Create a compile environment for the generated files. We can disable creating debug info here to improve link times.
diff --git a/Engine/Source/Programs/UnrealBuildTool/System/ExternalExecution.cs b/Engine/Source/Programs/UnrealBuildTool/System/ExternalExecution.cs
--- a/Engine/Source/Programs/UnrealBuildTool/System/ExternalExecution.cs
+++ b/Engine/Source/Programs/UnrealBuildTool/System/ExternalExecution.cs
@@ -499,15 +499,12 @@ namespace UnrealBuildTool
 				{
 					// If we've got this far and there are no source files then it's likely we're installed and ignoring
 					// engine files, so we don't need a .gen.cpp either
-					UEBuildModuleCPP.AutoGenerateCppInfoClass.BuildInfoClass BuildInfo = null;
 					UHTModuleInfo.Info.GeneratedCPPFilenameBase = Path.Combine(Module.GeneratedCodeDirectory.FullName, UHTModuleInfo.Info.ModuleName) + ".gen";
 					if (Module.SourceFilesToBuild.Count != 0)
 					{
-						BuildInfo = new UEBuildModuleCPP.AutoGenerateCppInfoClass.BuildInfoClass(Path.Combine(Module.GeneratedCodeDirectory.FullName, "*.gen.cpp"));
+						Module.GeneratedCodeWildcard = Path.Combine(Module.GeneratedCodeDirectory.FullName, "*.gen.cpp");
 					}
 
-					Module.AutoGenerateCppInfo = new UEBuildModuleCPP.AutoGenerateCppInfoClass(BuildInfo);
-
 					// If we're running in "gather" mode only, we'll go ahead and cache PCH information for each module right now, so that we don't
 					// have to do it in the assembling phase.  It's OK for gathering to take a bit longer, even if UObject headers are not out of
 					// date in order to save a lot of time in the assembling runs.
