UBT: Remove another unused variable from UEBuildModuleCPP.

#rnx

diff --git a/Engine/Source/Programs/UnrealBuildTool/Configuration/UEBuildModuleCPP.cs b/Engine/Source/Programs/UnrealBuildTool/Configuration/UEBuildModuleCPP.cs
--- a/Engine/Source/Programs/UnrealBuildTool/Configuration/UEBuildModuleCPP.cs
+++ b/Engine/Source/Programs/UnrealBuildTool/Configuration/UEBuildModuleCPP.cs
@@ -97,8 +97,6 @@ namespace UnrealBuildTool
 		/// </summary>
 		public string GeneratedCodeWildcard;
 
-		public List<DirectoryReference> IncludeSearchPaths = new List<DirectoryReference>();
-
 		public class ProcessedDependenciesClass
 		{
 			/// <summary>
@@ -383,8 +381,6 @@ namespace UnrealBuildTool
 			List<FileItem> LinkInputFiles = new List<FileItem>();
 
 			CppCompileEnvironment ModuleCompileEnvironment = CreateModuleCompileEnvironment(Target, BinaryCompileEnvironment);
-			IncludeSearchPaths = ModuleCompileEnvironment.IncludePaths.UserIncludePaths.ToList();
-			IncludeSearchPaths.AddRange(ModuleCompileEnvironment.IncludePaths.SystemIncludePaths.ToList());
 
 			// If the module is precompiled, read the object files from the manifest
 			if(Rules.bUsePrecompiled)
