Merging //UE4/Dev-Main @ 4131780 to Dev-Core (//UE4/Dev-Core)

diff --git a/Engine/Shaders/Private/ReflectionEnvironmentPixelShader.usf b/Engine/Shaders/Private/ReflectionEnvironmentPixelShader.usf
--- a/Engine/Shaders/Private/ReflectionEnvironmentPixelShader.usf
+++ b/Engine/Shaders/Private/ReflectionEnvironmentPixelShader.usf
@@ -298,7 +298,7 @@ void ReflectionEnvironmentSkyLighting(
 {
 	float2 ScreenUV = UVAndScreenPos.xy;
 	FScreenSpaceData ScreenSpaceData = GetScreenSpaceData(ScreenUV);
-	uint ShadingModelID = GetShadingModelId(ScreenUV);
+	uint ShadingModelID = ScreenSpaceData.GBuffer.ShadingModelID;
 
 	OutColor = 0.0f;
 
diff --git a/Engine/Source/Editor/UnrealEd/Private/ObjectTools.cpp b/Engine/Source/Editor/UnrealEd/Private/ObjectTools.cpp
--- a/Engine/Source/Editor/UnrealEd/Private/ObjectTools.cpp
+++ b/Engine/Source/Editor/UnrealEd/Private/ObjectTools.cpp
@@ -785,15 +785,15 @@ namespace ObjectTools
 			}
 		}
 
-		for (int CurrentIndex = 0; CurrentIndex < ReferencingPropertiesMapKeys.Num(); CurrentIndex++)
 		{
-			int CheckForCircles = 0;
-			FFindReferencersArchive FindDependentArchive(ReferencingPropertiesMapKeys[CurrentIndex], ReferencingPropertiesMapKeys);
-			for (int DependentIndex = CurrentIndex + 1; DependentIndex < ReferencingPropertiesMapKeys.Num(); DependentIndex++)
+			TBitArray<> TouchedThisItteration(false, ReferencingPropertiesMapKeys.Num());
+			for (int CurrentIndex = 0; CurrentIndex < ReferencingPropertiesMapKeys.Num(); CurrentIndex++)
 			{
-				if (FindDependentArchive.GetReferenceCount(ReferencingPropertiesMapKeys[DependentIndex]) > 0)
+				TouchedThisItteration.Init(false, ReferencingPropertiesMapKeys.Num());
+				FFindReferencersArchive FindDependentArchive(ReferencingPropertiesMapKeys[CurrentIndex], ReferencingPropertiesMapKeys);
+				for (int DependentIndex = CurrentIndex + 1; DependentIndex < ReferencingPropertiesMapKeys.Num(); DependentIndex++)
 				{
-					if (CheckForCircles != DependentIndex)
+					if (!TouchedThisItteration[DependentIndex] && FindDependentArchive.GetReferenceCount(ReferencingPropertiesMapKeys[DependentIndex]) > 0)
 					{
 						UObject* Key = ReferencingPropertiesMapKeys[CurrentIndex];
 						PropertyArrayType Value = ReferencingPropertiesMapValues[CurrentIndex];
@@ -803,12 +803,7 @@ namespace ObjectTools
 						ReferencingPropertiesMapValues[DependentIndex] = Value;
 
 						FindDependentArchive.ResetPotentialReferencer(ReferencingPropertiesMapKeys[CurrentIndex]);
-						
-						if(CheckForCircles == 0)
-						{
-							CheckForCircles = DependentIndex;
-						}
-						
+						TouchedThisItteration[DependentIndex] = true;
 						DependentIndex = CurrentIndex;
 					}
 				}
diff --git a/Engine/Source/Programs/AutomationTool/BuildGraph/Script.cs b/Engine/Source/Programs/AutomationTool/BuildGraph/Script.cs
--- a/Engine/Source/Programs/AutomationTool/BuildGraph/Script.cs
+++ b/Engine/Source/Programs/AutomationTool/BuildGraph/Script.cs
@@ -794,12 +794,10 @@ namespace AutomationTool
 				{
 					if (Types.Length > 0 && Agent.PossibleTypes.Length > 0)
 					{
-						string[] NewTypes = Agent.PossibleTypes.Intersect(Types, StringComparer.InvariantCultureIgnoreCase).ToArray();
-						if (NewTypes.Length == 0)
+						if (Types.Length != Agent.PossibleTypes.Length || !Types.SequenceEqual(Agent.PossibleTypes, StringComparer.InvariantCultureIgnoreCase))
 						{
-							LogError(Element, "No common agent types with previous agent definition");
+							LogError(Element, "Agent types ({0}) were different than previous agent definition with types ({1}). Must either be empty or match exactly.", string.Join(",", Types), string.Join(",", Agent.PossibleTypes));
 						}
-						Agent.PossibleTypes = NewTypes;
 					}
 				}
 				else
diff --git a/Engine/Source/Programs/UnrealBuildTool/Platform/Android/AndroidToolChain.cs b/Engine/Source/Programs/UnrealBuildTool/Platform/Android/AndroidToolChain.cs
--- a/Engine/Source/Programs/UnrealBuildTool/Platform/Android/AndroidToolChain.cs
+++ b/Engine/Source/Programs/UnrealBuildTool/Platform/Android/AndroidToolChain.cs
@@ -1665,12 +1665,69 @@ namespace UnrealBuildTool
 						LinkAction.CommandArguments = String.Format("/c \"{0} {1}\"", LinkAction.CommandPath, LinkAction.CommandArguments);
 						LinkAction.CommandPath = "cmd.exe";
 					}
+
+					// Windows can run into an issue with too long of a commandline when clang tries to call ld to link.
+					// To work around this we call clang to just get the command it would execute and generate a
+					// second response file to directly call ld with the right arguments instead of calling through clang.
+					if (!Utils.IsRunningOnMono)
+					{
+						// capture the actual link command without running it
+						ProcessStartInfo StartInfo = new ProcessStartInfo();
+						StartInfo.WorkingDirectory = LinkEnvironment.IntermediateDirectory.FullName;
+						StartInfo.FileName = LinkAction.CommandPath;
+						StartInfo.Arguments = "-### " + LinkAction.CommandArguments;
+						StartInfo.UseShellExecute = false;
+						StartInfo.CreateNoWindow = true;
+						StartInfo.RedirectStandardError = true;
+
+						LinkerCommandline = "";
+
+						Process Proc = new Process();
+						Proc.StartInfo = StartInfo;
+						Proc.ErrorDataReceived += new DataReceivedEventHandler(OutputReceivedForLinker);
+						Proc.Start();
+						Proc.BeginErrorReadLine();
+						Proc.WaitForExit(5000);
+
+						LinkerCommandline = LinkerCommandline.Trim();
+
+						// the command should be in quotes; if not we'll just use clang to link as usual
+						int FirstQuoteIndex = LinkerCommandline.IndexOf('"');
+						if (FirstQuoteIndex >= 0)
+						{
+							int SecondQuoteIndex = LinkerCommandline.Substring(FirstQuoteIndex + 1).IndexOf('"');
+							if (SecondQuoteIndex >= 0)
+							{
+								LinkAction.CommandPath = LinkerCommandline.Substring(FirstQuoteIndex + 1, SecondQuoteIndex - FirstQuoteIndex);
+								LinkAction.CommandArguments = LinkerCommandline.Substring(FirstQuoteIndex + SecondQuoteIndex + 3);
+
+								// replace double backslashes
+								LinkAction.CommandPath = LinkAction.CommandPath.Replace("\\\\", "/");
+
+								// now create a response file for the full command using ld directly
+								FileReference FinalResponseFileName = FileReference.Combine(LinkEnvironment.IntermediateDirectory, OutputFile.Location.GetFileName() + ".responseFinal");
+								FileItem FinalResponseFileItem = FileItem.CreateIntermediateTextFile(FinalResponseFileName, LinkAction.CommandArguments);
+								LinkAction.CommandArguments = string.Format("@\"{0}\"", FinalResponseFileName);
+								LinkAction.PrerequisiteItems.Add(FinalResponseFileItem);
+							}
+						}
+					}
 				}
 			}
 
 			return Outputs.ToArray();
 		}
 
+		// captures stderr from clang
+		private static string LinkerCommandline = "";
+		static public void OutputReceivedForLinker(Object Sender, DataReceivedEventArgs Line)
+		{
+			if ((Line != null) && (Line.Data != null) && (Line.Data.Contains("--sysroot")))
+			{
+				LinkerCommandline += Line.Data;
+			}
+		}
+
 		private void ExportObjectFilePaths(LinkEnvironment LinkEnvironment, string FileName)
 		{
 			// Write the list of object file directories
diff --git a/Engine/Source/Runtime/Core/Private/HAL/LowLevelMemTracker.cpp b/Engine/Source/Runtime/Core/Private/HAL/LowLevelMemTracker.cpp
--- a/Engine/Source/Runtime/Core/Private/HAL/LowLevelMemTracker.cpp
+++ b/Engine/Source/Runtime/Core/Private/HAL/LowLevelMemTracker.cpp
@@ -61,6 +61,7 @@ DECLARE_LLM_MEMORY_STAT(TEXT("Networking"), STAT_NetworkingLLM, STATGROUP_LLMFUL
 DECLARE_LLM_MEMORY_STAT(TEXT("Meshes"), STAT_MeshesLLM, STATGROUP_LLMFULL);
 DECLARE_LLM_MEMORY_STAT(TEXT("Stats"), STAT_StatsLLM, STATGROUP_LLMFULL);
 DECLARE_LLM_MEMORY_STAT(TEXT("Shaders"), STAT_ShadersLLM, STATGROUP_LLMFULL);
+DECLARE_LLM_MEMORY_STAT(TEXT("PSO"), STAT_PSOLLM, STATGROUP_LLMFULL);
 DECLARE_LLM_MEMORY_STAT(TEXT("Textures"), STAT_TexturesLLM, STATGROUP_LLMFULL);
 DECLARE_LLM_MEMORY_STAT(TEXT("Render Targets"), STAT_RenderTargetsLLM, STATGROUP_LLMFULL);
 DECLARE_LLM_MEMORY_STAT(TEXT("RHIMisc"), STAT_RHIMiscLLM, STATGROUP_LLMFULL);
diff --git a/Engine/Source/Runtime/Core/Public/HAL/LowLevelMemTracker.h b/Engine/Source/Runtime/Core/Public/HAL/LowLevelMemTracker.h
--- a/Engine/Source/Runtime/Core/Public/HAL/LowLevelMemTracker.h
+++ b/Engine/Source/Runtime/Core/Public/HAL/LowLevelMemTracker.h
@@ -138,6 +138,7 @@ enum class ELLMTagSet : uint8
 	macro(Meshes,								"Meshes",						GET_STATFNAME(STAT_MeshesLLM),								GET_STATFNAME(STAT_MeshesSummaryLLM))		\
 	macro(Stats,								"Stats",						GET_STATFNAME(STAT_StatsLLM),								GET_STATFNAME(STAT_EngineSummaryLLM))		\
 	macro(Shaders,								"Shaders",						GET_STATFNAME(STAT_ShadersLLM),								GET_STATFNAME(STAT_EngineSummaryLLM))		\
+	macro(PSO,									"PSO",							GET_STATFNAME(STAT_PSOLLM),									GET_STATFNAME(STAT_EngineSummaryLLM))		\
 	macro(Textures,								"Textures",						GET_STATFNAME(STAT_TexturesLLM),							GET_STATFNAME(STAT_TexturesSummaryLLM))		\
 	macro(RenderTargets,						"RenderTargets",				GET_STATFNAME(STAT_RenderTargetsLLM),						GET_STATFNAME(STAT_EngineSummaryLLM))		\
 	macro(RHIMisc,								"RHIMisc",						GET_STATFNAME(STAT_RHIMiscLLM),								GET_STATFNAME(STAT_EngineSummaryLLM))		\
diff --git a/Engine/Source/Runtime/RHI/Private/PipelineStateCache.cpp b/Engine/Source/Runtime/RHI/Private/PipelineStateCache.cpp
--- a/Engine/Source/Runtime/RHI/Private/PipelineStateCache.cpp
+++ b/Engine/Source/Runtime/RHI/Private/PipelineStateCache.cpp
@@ -483,6 +483,8 @@ public:
 
 	void DoTask(ENamedThreads::Type CurrentThread, const FGraphEventRef& MyCompletionGraphEvent)
 	{
+		LLM_SCOPE(ELLMTag::PSO);
+
 		if (Pipeline->IsCompute())
 		{
 			FComputePipelineState* ComputePipeline = static_cast<FComputePipelineState*>(Pipeline);
@@ -713,7 +715,7 @@ FRHIComputePipelineState* ExecuteSetComputePipelineState(FComputePipelineState*
 
 FGraphicsPipelineState* PipelineStateCache::GetAndOrCreateGraphicsPipelineState(FRHICommandList& RHICmdList, const FGraphicsPipelineStateInitializer& OriginalInitializer, EApplyRendertargetOption ApplyFlags)
 {
-
+	LLM_SCOPE(ELLMTag::PSO);
 	SCOPE_CYCLE_COUNTER(STAT_GetOrCreatePSO);
 
 	FGraphicsPipelineStateInitializer NewInitializer;
diff --git a/Engine/Source/Runtime/Renderer/Private/PostProcess/PostProcessing.cpp b/Engine/Source/Runtime/Renderer/Private/PostProcess/PostProcessing.cpp
--- a/Engine/Source/Runtime/Renderer/Private/PostProcess/PostProcessing.cpp
+++ b/Engine/Source/Runtime/Renderer/Private/PostProcess/PostProcessing.cpp
@@ -932,11 +932,11 @@ static void AddTemporalAA( FPostprocessContext& Context, FRenderingCompositeOutp
 
 	FSceneViewState* ViewState = Context.View.ViewState;
 
-	FRCPassPostProcessTemporalAA* TemporalAAPass = new(FMemStack::Get()) FRCPassPostProcessTemporalAA(
+	FRCPassPostProcessTemporalAA* TemporalAAPass = Context.Graph.RegisterPass(new(FMemStack::Get()) FRCPassPostProcessTemporalAA(
 		Context,
 		Parameters,
 		Context.View.PrevViewInfo.TemporalAAHistory,
-		&ViewState->PendingPrevFrameViewInfo.TemporalAAHistory);
+		&ViewState->PendingPrevFrameViewInfo.TemporalAAHistory));
 
 	TemporalAAPass->SetInput( ePId_Input0, Context.FinalOutput );
 	TemporalAAPass->SetInput( ePId_Input2, VelocityInput );
@@ -1726,7 +1726,6 @@ void FPostProcessing::Process(FRHICommandListImmediate& RHICmdList, const FViewI
 			// some views don't have a state (thumbnail rendering)
 			if(!GIsHighResScreenshot && View.State && IStereoRendering::IsAPrimaryView(StereoPass))
 			{
-				
 				const bool bUseBasicEyeAdaptation = (AutoExposure.MethodId == EAutoExposureMethod::AEM_Basic);
 
 				if (bUseBasicEyeAdaptation) // log average ps reduction ( non histogram ) 
