UBT: Remove legacy code to handle disabling optional plugins; now that this is compiled into the target, it will work for any plugins we choose.

#rnx

diff --git a/Engine/Source/Programs/UnrealBuildTool/Configuration/UEBuildTarget.cs b/Engine/Source/Programs/UnrealBuildTool/Configuration/UEBuildTarget.cs
--- a/Engine/Source/Programs/UnrealBuildTool/Configuration/UEBuildTarget.cs
+++ b/Engine/Source/Programs/UnrealBuildTool/Configuration/UEBuildTarget.cs
@@ -3374,7 +3374,7 @@ namespace UnrealBuildTool
 			// Map of plugin names to instances of that plugin
 			Dictionary<string, UEBuildPlugin> NameToInstance = new Dictionary<string, UEBuildPlugin>(StringComparer.InvariantCultureIgnoreCase);
 
-			// Configure plugins explicitly referenced via target settings
+			// Configure plugins explicitly enabled via target settings
 			foreach(string PluginName in Rules.EnablePlugins)
 			{
 				if(ReferencedNames.Add(PluginName))
@@ -3384,6 +3384,16 @@ namespace UnrealBuildTool
 				}
 			}
 
+			// Configure plugins explicitly disabled via target settings
+			foreach(string PluginName in Rules.DisablePlugins)
+			{
+				if(ReferencedNames.Add(PluginName))
+				{
+					PluginReferenceDescriptor PluginReference = new PluginReferenceDescriptor(PluginName, null, false);
+					AddPlugin(PluginReference, "target settings", ExcludeFolders, NameToInstance, NameToInfo);
+				}
+			}
+
 			// Find a map of plugins which are explicitly referenced in the project file
 			if(ProjectDescriptor != null && ProjectDescriptor.Plugins != null)
 			{
@@ -3468,19 +3478,6 @@ namespace UnrealBuildTool
 				return null;
 			}
 
-			// If this plugin is listed to be excluded, do so here. The reference must be optional in this case.
-			if(Rules.DisablePlugins.Contains(Reference.Name, StringComparer.InvariantCultureIgnoreCase))
-			{
-				if(!Reference.bOptional)
-				{
-					throw new BuildException("The '{0}' plugin is listed in the target's ExcludePlugins list, but is not referenced (via {1}) without the bOptional flag set. This will cause load failures at runtime.", Reference.Name, ReferenceChain);
-				}
-				else
-				{
-					return null;
-				}
-			}
-
 			// Try to get an existing reference to this plugin
 			UEBuildPlugin Instance;
 			if(NameToInstance.TryGetValue(Reference.Name, out Instance))
