Fix warning due to plugins disabled in .target.cs file.

#rnx

diff --git a/Engine/Source/Programs/UnrealBuildTool/Configuration/UEBuildTarget.cs b/Engine/Source/Programs/UnrealBuildTool/Configuration/UEBuildTarget.cs
--- a/Engine/Source/Programs/UnrealBuildTool/Configuration/UEBuildTarget.cs
+++ b/Engine/Source/Programs/UnrealBuildTool/Configuration/UEBuildTarget.cs
@@ -3400,7 +3400,7 @@ namespace UnrealBuildTool
 				string ProjectReferenceChain = ProjectFile.GetFileName();
 				foreach(PluginReferenceDescriptor PluginReference in ProjectDescriptor.Plugins)
 				{
-					if(!Rules.EnablePlugins.Contains(PluginReference.Name, StringComparer.InvariantCultureIgnoreCase))
+					if(!Rules.EnablePlugins.Contains(PluginReference.Name, StringComparer.InvariantCultureIgnoreCase) && !Rules.DisablePlugins.Contains(PluginReference.Name, StringComparer.InvariantCultureIgnoreCase))
 					{
 						// Make sure we don't have multiple references to the same plugin
 						if(!ReferencedNames.Add(PluginReference.Name))
