UBT: Add a check that the UnrealHeaderTool target file exists, rather than throwing an exception when reading it fails.

#rnx

diff --git a/Engine/Source/Programs/UnrealBuildTool/System/ExternalExecution.cs b/Engine/Source/Programs/UnrealBuildTool/System/ExternalExecution.cs
--- a/Engine/Source/Programs/UnrealBuildTool/System/ExternalExecution.cs
+++ b/Engine/Source/Programs/UnrealBuildTool/System/ExternalExecution.cs
@@ -562,6 +562,11 @@ namespace UnrealBuildTool
 		static string GetHeaderToolPath()
 		{
 			FileReference ReceiptFileName = GetHeaderToolReceiptFile();
+			if(!FileReference.Exists(ReceiptFileName))
+			{
+				throw new BuildException("{0} is missing. Please rebuild UnrealHeaderTool manually.", ReceiptFileName);
+			}
+
 			TargetReceipt Receipt = TargetReceipt.Read(ReceiptFileName, UnrealBuildTool.EngineDirectory, null);
 
 			string HeaderToolPath = Receipt.BuildProducts[0].Path.FullName;
