UBT: Move bBuildLargeAddressAwareBinary into an exposed setting.

#rnx

diff --git a/Engine/Source/Programs/UnrealBuildTool/Platform/Windows/UEBuildWindows.cs b/Engine/Source/Programs/UnrealBuildTool/Platform/Windows/UEBuildWindows.cs
--- a/Engine/Source/Programs/UnrealBuildTool/Platform/Windows/UEBuildWindows.cs
+++ b/Engine/Source/Programs/UnrealBuildTool/Platform/Windows/UEBuildWindows.cs
@@ -190,6 +190,11 @@ namespace UnrealBuildTool
 		[XmlConfigFile(Category = "WindowsPlatform")]
 		public int PCHMemoryAllocationFactor = 0;
 
+		/// <summary>
+		/// True if we allow using addresses larger than 2GB on 32 bit builds
+		/// </summary>
+		public bool bBuildLargeAddressAwareBinary = true;
+
 		/// <summary>
 		/// When using a Visual Studio compiler, returns the version name as a string
 		/// </summary>
@@ -317,6 +322,11 @@ namespace UnrealBuildTool
 			get { return Inner.PCHMemoryAllocationFactor; }
 		}
 
+		public bool bBuildLargeAddressAwareBinary
+		{
+			get { return Inner.bBuildLargeAddressAwareBinary; }
+		}
+
 		public string GetVisualStudioCompilerVersionName()
 		{
 			return Inner.GetVisualStudioCompilerVersionName();
@@ -370,11 +380,6 @@ namespace UnrealBuildTool
 		/// </summary>
 		public static readonly bool bAllowICLLinker = true;
 
-		/// <summary>
-		/// True if we allow using addresses larger than 2GB on 32 bit builds
-		/// </summary>
-		public static readonly bool bBuildLargeAddressAwareBinary = true;
-
 		WindowsPlatformSDK SDK;
 
 		/// <summary>
diff --git a/Engine/Source/Programs/UnrealBuildTool/Platform/Windows/VCToolChain.cs b/Engine/Source/Programs/UnrealBuildTool/Platform/Windows/VCToolChain.cs
--- a/Engine/Source/Programs/UnrealBuildTool/Platform/Windows/VCToolChain.cs
+++ b/Engine/Source/Programs/UnrealBuildTool/Platform/Windows/VCToolChain.cs
@@ -693,7 +693,7 @@ namespace UnrealBuildTool
 				Arguments.Add("/FIXED:No");
 
 				// Option is only relevant with 32 bit toolchain.
-				if ((LinkEnvironment.Platform == CppPlatform.Win32) && WindowsPlatform.bBuildLargeAddressAwareBinary)
+				if ((LinkEnvironment.Platform == CppPlatform.Win32) && Target.WindowsPlatform.bBuildLargeAddressAwareBinary)
 				{
 					// Disables the 2GB address space limit on 64-bit Windows and 32-bit Windows with /3GB specified in boot.ini
 					Arguments.Add("/LARGEADDRESSAWARE");
