UBT: Remove static config variable for controlling which configuration of UHT to use.

#rnx

diff --git a/Engine/Source/Programs/UnrealBuildTool/Configuration/BuildConfiguration.cs b/Engine/Source/Programs/UnrealBuildTool/Configuration/BuildConfiguration.cs
--- a/Engine/Source/Programs/UnrealBuildTool/Configuration/BuildConfiguration.cs
+++ b/Engine/Source/Programs/UnrealBuildTool/Configuration/BuildConfiguration.cs
@@ -201,7 +201,7 @@ namespace UnrealBuildTool
 		/// If true, the Debug version of UnrealHeaderTool will be build and run instead of the Development version.
 		/// </summary>
 		[XmlConfigFile(Category = "UEBuildConfiguration")]
-		public static bool bForceDebugUnrealHeaderTool = false;
+		public bool bForceDebugUnrealHeaderTool = false;
 
 		/// <summary>
 		/// When true, the targets won't execute their link actions if there was nothing to compile
diff --git a/Engine/Source/Programs/UnrealBuildTool/System/ExternalExecution.cs b/Engine/Source/Programs/UnrealBuildTool/System/ExternalExecution.cs
--- a/Engine/Source/Programs/UnrealBuildTool/System/ExternalExecution.cs
+++ b/Engine/Source/Programs/UnrealBuildTool/System/ExternalExecution.cs
@@ -550,27 +550,24 @@ namespace UnrealBuildTool
 		/// Gets the path to the receipt for UHT
 		/// </summary>
 		/// <returns>Path to the UHT receipt</returns>
-		static FileReference GetHeaderToolReceiptFile()
+		static FileReference GetHeaderToolReceiptFile(UnrealTargetConfiguration Configuration)
 		{
-			UnrealTargetConfiguration Config = BuildConfiguration.bForceDebugUnrealHeaderTool ? UnrealTargetConfiguration.Debug : UnrealTargetConfiguration.Development;
-			return TargetReceipt.GetDefaultPath(UnrealBuildTool.EngineDirectory, "UnrealHeaderTool", BuildHostPlatform.Current.Platform, Config, "");
+			return TargetReceipt.GetDefaultPath(UnrealBuildTool.EngineDirectory, "UnrealHeaderTool", BuildHostPlatform.Current.Platform, Configuration, "");
 		}
 
 		/// <summary>
 		/// Gets UnrealHeaderTool.exe path. Does not care if UnrealheaderTool was build as a monolithic exe or not.
 		/// </summary>
-		static string GetHeaderToolPath()
+		static FileReference GetHeaderToolPath(UnrealTargetConfiguration Configuration)
 		{
-			FileReference ReceiptFileName = GetHeaderToolReceiptFile();
+			FileReference ReceiptFileName = GetHeaderToolReceiptFile(Configuration);
 			if(!FileReference.Exists(ReceiptFileName))
 			{
 				throw new BuildException("{0} is missing. Please rebuild UnrealHeaderTool manually.", ReceiptFileName);
 			}
 
 			TargetReceipt Receipt = TargetReceipt.Read(ReceiptFileName, UnrealBuildTool.EngineDirectory, null);
-
-			string HeaderToolPath = Receipt.BuildProducts[0].Path.FullName;
-			return HeaderToolPath;
+			return Receipt.BuildProducts[0].Path;
 		}
 
 		/// <summary>
@@ -579,14 +576,12 @@ namespace UnrealBuildTool
 		/// <returns>
 		/// Latest timestamp of UHT binaries or DateTime.MaxValue if UnrealHeaderTool is out of date and needs to be rebuilt.
 		/// </returns>
-		static bool GetHeaderToolTimestamp(out DateTime Timestamp)
+		static bool GetHeaderToolTimestamp(UnrealTargetConfiguration Configuration, out DateTime Timestamp)
 		{
 			using (ScopedTimer TimestampTimer = new ScopedTimer("GetHeaderToolTimestamp"))
 			{
 				// Try to read the receipt for UHT.
-
-
-				FileReference ReceiptPath = GetHeaderToolReceiptFile();
+				FileReference ReceiptPath = GetHeaderToolReceiptFile(Configuration);
 				if (!FileReference.Exists(ReceiptPath))
 				{
 					Timestamp = DateTime.MaxValue;
@@ -973,10 +968,10 @@ namespace UnrealBuildTool
 		/// <summary>
 		/// Run an external native executable (and capture the output), given the executable path and the commandline.
 		/// </summary>
-		public static int RunExternalNativeExecutable(string ExePath, string Commandline)
+		public static int RunExternalNativeExecutable(FileReference ExePath, string Commandline)
 		{
-			ProcessStartInfo ExeInfo = new ProcessStartInfo(ExePath, Commandline);
-			Log.TraceVerbose("RunExternalExecutable {0} {1}", ExePath, Commandline);
+			ProcessStartInfo ExeInfo = new ProcessStartInfo(ExePath.FullName, Commandline);
+			Log.TraceVerbose("RunExternalExecutable {0} {1}", ExePath.FullName, Commandline);
 			ExeInfo.UseShellExecute = false;
 			ExeInfo.RedirectStandardOutput = true;
 			using (Process GameProcess = Process.Start(ExeInfo))
@@ -1019,9 +1014,11 @@ namespace UnrealBuildTool
 
 				string RootLocalPath = UnrealBuildTool.RootDirectory.FullName;
 
+				UnrealTargetConfiguration UHTConfig = BuildConfiguration.bForceDebugUnrealHeaderTool ? UnrealTargetConfiguration.Debug : UnrealTargetConfiguration.Development;
+
 				// check if UHT is out of date
 				DateTime HeaderToolTimestamp = DateTime.MaxValue;
-				bool bHaveHeaderTool = !bIsBuildingUHT && GetHeaderToolTimestamp(out HeaderToolTimestamp);
+				bool bHaveHeaderTool = !bIsBuildingUHT && GetHeaderToolTimestamp(UHTConfig, out HeaderToolTimestamp);
 
 				// ensure the headers are up to date
 				bool bUHTNeedsToRun = (BuildConfiguration.bForceHeaderGeneration || !bHaveHeaderTool || AreGeneratedCodeFilesOutOfDate(BuildConfiguration, UObjectModules, HeaderToolTimestamp, Target.bUsePrecompiled, HotReload, bIsGatheringBuild, bIsAssemblingBuild));
@@ -1132,10 +1129,10 @@ namespace UnrealBuildTool
 					string ActualTargetName = String.IsNullOrEmpty(Target.GetTargetName()) ? "UE4" : Target.GetTargetName();
 					Log.TraceInformation("Parsing headers for {0}", ActualTargetName);
 
-					string HeaderToolPath = GetHeaderToolPath();
-					if (!File.Exists(HeaderToolPath))
+					FileReference HeaderToolPath = GetHeaderToolPath(UHTConfig);
+					if (!FileReference.Exists(HeaderToolPath))
 					{
-						throw new BuildException("Unable to generate headers because UnrealHeaderTool binary was not found ({0}).", Path.GetFullPath(HeaderToolPath));
+						throw new BuildException("Unable to generate headers because UnrealHeaderTool binary was not found ({0}).", HeaderToolPath);
 					}
 
 					// Disable extensions when serializing to remove the $type fields
@@ -1163,7 +1160,7 @@ namespace UnrealBuildTool
 
 					Stopwatch s = new Stopwatch();
 					s.Start();
-					UHTResult = (ECompilationResult)RunExternalNativeExecutable(ExternalExecution.GetHeaderToolPath(), CmdLine);
+					UHTResult = (ECompilationResult)RunExternalNativeExecutable(ExternalExecution.GetHeaderToolPath(UHTConfig), CmdLine);
 					s.Stop();
 
 					if (UHTResult != ECompilationResult.Succeeded)
