UBT: Fix exception if BuildConfiguration.xml contains an unknown category.

diff --git a/Engine/Source/Programs/UnrealBuildTool/System/XmlConfig.cs b/Engine/Source/Programs/UnrealBuildTool/System/XmlConfig.cs
--- a/Engine/Source/Programs/UnrealBuildTool/System/XmlConfig.cs
+++ b/Engine/Source/Programs/UnrealBuildTool/System/XmlConfig.cs
@@ -526,31 +526,34 @@ namespace UnrealBuildTool
 			// Parse the document
 			foreach(XmlElement CategoryElement in ConfigFile.DocumentElement.ChildNodes.OfType<XmlElement>())
 			{
-				Dictionary<string, FieldInfo> NameToField = CategoryToFields[CategoryElement.Name];
-				foreach(XmlElement KeyElement in CategoryElement.ChildNodes.OfType<XmlElement>())
+				Dictionary<string, FieldInfo> NameToField;
+				if(CategoryToFields.TryGetValue(CategoryElement.Name, out NameToField))
 				{
-					FieldInfo Field;
-					if(NameToField.TryGetValue(KeyElement.Name, out Field))
+					foreach(XmlElement KeyElement in CategoryElement.ChildNodes.OfType<XmlElement>())
 					{
-						// Parse the corresponding value
-						object Value;
-						if(Field.FieldType == typeof(string[]))
+						FieldInfo Field;
+						if(NameToField.TryGetValue(KeyElement.Name, out Field))
 						{
-							Value = KeyElement.ChildNodes.OfType<XmlElement>().Where(x => x.Name == "Item").Select(x => x.InnerText).ToArray();
-						}
-						else
-						{
-							Value = ParseValue(Field.FieldType, KeyElement.InnerText);
-						}
+							// Parse the corresponding value
+							object Value;
+							if(Field.FieldType == typeof(string[]))
+							{
+								Value = KeyElement.ChildNodes.OfType<XmlElement>().Where(x => x.Name == "Item").Select(x => x.InnerText).ToArray();
+							}
+							else
+							{
+								Value = ParseValue(Field.FieldType, KeyElement.InnerText);
+							}
 
-						// Add it to the set of values for the type containing this field
-						Dictionary<FieldInfo, object> FieldToValue;
-						if(!TypeToValues.TryGetValue(Field.DeclaringType, out FieldToValue))
-						{
-							FieldToValue = new Dictionary<FieldInfo, object>();
-							TypeToValues.Add(Field.DeclaringType, FieldToValue);
+							// Add it to the set of values for the type containing this field
+							Dictionary<FieldInfo, object> FieldToValue;
+							if(!TypeToValues.TryGetValue(Field.DeclaringType, out FieldToValue))
+							{
+								FieldToValue = new Dictionary<FieldInfo, object>();
+								TypeToValues.Add(Field.DeclaringType, FieldToValue);
+							}
+							FieldToValue[Field] = Value;
 						}
-						FieldToValue[Field] = Value;
 					}
 				}
 			}
