Fixes to support incremental linking on Windows. Does not seem to have any net benefit right now; may improve once minimal rebuild is enabled.

* Incremental linking no longer forces PDB files to be enabled for source files.
* Actions can specify specific files to be deleted before each build. Code to forcibly delete PDB files has been moved to the MSVC toolchain.
* Unused libraries produced by the cross-referenced link are no longer added as build products, since (a) deleting them breaks dependency checking for incremental linking and causes a full link, and (b) not deleting them breaks UBT dependency checking and causes actions to be run over and over again.
* Icon update is disabled for Windows when incremental linking is enabled.
* Removed rarely-used setting to always delete produced items before each build.

diff --git a/Engine/Source/Programs/UnrealBuildTool/Configuration/BuildConfiguration.cs b/Engine/Source/Programs/UnrealBuildTool/Configuration/BuildConfiguration.cs
--- a/Engine/Source/Programs/UnrealBuildTool/Configuration/BuildConfiguration.cs
+++ b/Engine/Source/Programs/UnrealBuildTool/Configuration/BuildConfiguration.cs
@@ -69,12 +69,6 @@ namespace UnrealBuildTool
 		[XmlConfigFile]
 		public bool bAllowSNDBS = true;
 
-		/// <summary>
-		/// Whether or not to delete outdated produced items.
-		/// </summary>
-		[XmlConfigFile]
-		public bool bShouldDeleteAllOutdatedProducedItems = false;
-
 		/// <summary>
 		/// Whether we should export a JSON file containing detailed target information.
 		/// </summary>
diff --git a/Engine/Source/Programs/UnrealBuildTool/Configuration/UEBuildDeployTarget.cs b/Engine/Source/Programs/UnrealBuildTool/Configuration/UEBuildDeployTarget.cs
--- a/Engine/Source/Programs/UnrealBuildTool/Configuration/UEBuildDeployTarget.cs
+++ b/Engine/Source/Programs/UnrealBuildTool/Configuration/UEBuildDeployTarget.cs
@@ -1,4 +1,4 @@
-﻿// Copyright 1998-2018 Epic Games, Inc. All Rights Reserved.
+// Copyright 1998-2018 Epic Games, Inc. All Rights Reserved.
 
 using System;
 using System.Collections.Generic;
@@ -85,6 +85,11 @@ namespace UnrealBuildTool
 		/// </summary>
 		public readonly FileReference BuildReceiptFileName;
 
+		/// <summary>
+		/// Whether incremental linking is enabled (disables icon update on Windows)
+		/// </summary>
+		public readonly bool bUseIncrementalLinking;
+
 		/// <summary>
 		/// If true, then a stub IPA will be generated when compiling is done (minimal files needed for a valid IPA)
 		/// </summary>
@@ -122,6 +127,7 @@ namespace UnrealBuildTool
 			this.EngineIntermediateDirectory = Target.EngineIntermediateDirectory;
 			this.ProjectDirectory = Target.ProjectDirectory;
 			this.BuildReceiptFileName = Target.ReceiptFileName;
+			this.bUseIncrementalLinking = Target.Rules.bUseIncrementalLinking;
 			this.bCreateStubIPA = Target.Rules.bCreateStubIPA;
 			this.AndroidArchitectures = Target.Rules.AndroidPlatform.Architectures.ToArray();
 			this.AndroidGPUArchitectures = Target.Rules.AndroidPlatform.GPUArchitectures.ToArray();
@@ -147,6 +153,7 @@ namespace UnrealBuildTool
 				EngineIntermediateDirectory = Reader.ReadDirectoryReference();
 				ProjectDirectory = Reader.ReadDirectoryReference();
 				BuildReceiptFileName = Reader.ReadFileReference();
+				bUseIncrementalLinking = Reader.ReadBoolean();
 				bCreateStubIPA = Reader.ReadBoolean();
 				AndroidArchitectures = Reader.ReadStringArray();
 				AndroidGPUArchitectures = Reader.ReadStringArray();
@@ -174,6 +181,7 @@ namespace UnrealBuildTool
 				Writer.Write(EngineIntermediateDirectory);
 				Writer.Write(ProjectDirectory);
 				Writer.Write(BuildReceiptFileName);
+				Writer.Write(bUseIncrementalLinking);
 				Writer.Write(bCreateStubIPA);
 				Writer.Write(AndroidArchitectures);
 				Writer.Write(AndroidGPUArchitectures);
diff --git a/Engine/Source/Programs/UnrealBuildTool/Platform/Windows/PVSToolChain.cs b/Engine/Source/Programs/UnrealBuildTool/Platform/Windows/PVSToolChain.cs
--- a/Engine/Source/Programs/UnrealBuildTool/Platform/Windows/PVSToolChain.cs
+++ b/Engine/Source/Programs/UnrealBuildTool/Platform/Windows/PVSToolChain.cs
@@ -1,4 +1,4 @@
-﻿// Copyright 1998-2018 Epic Games, Inc. All Rights Reserved.
+// Copyright 1998-2018 Epic Games, Inc. All Rights Reserved.
 
 using System;
 using System.Collections.Generic;
@@ -218,7 +218,7 @@ namespace UnrealBuildTool
 				AnalyzeAction.PrerequisiteItems.Add(ConfigFileItem);
 				AnalyzeAction.PrerequisiteItems.Add(PreprocessedFileItem);
 				AnalyzeAction.ProducedItems.Add(OutputFileItem);
-				AnalyzeAction.bShouldDeleteProducedItems = true; // PVS Studio will append by default, so need to delete produced items
+				AnalyzeAction.DeleteItems.Add(OutputFileItem); // PVS Studio will append by default, so need to delete produced items
 
 				Result.ObjectFiles.AddRange(AnalyzeAction.ProducedItems);
 			}
@@ -252,7 +252,7 @@ namespace UnrealBuildTool
 			AnalyzeAction.PrerequisiteItems.AddRange(OutputItems);
 			AnalyzeAction.ProducedItems.Add(FileItem.GetItemByFileReference(OutputFile));
 			AnalyzeAction.ActionHandler = (Action Action, out int ExitCode, out string Output) => WriteResults(OutputFile, InputFiles, out ExitCode, out Output);
-			AnalyzeAction.bShouldDeleteProducedItems = true;
+			AnalyzeAction.DeleteItems.AddRange(AnalyzeAction.ProducedItems);
 
 			OutputItems.AddRange(AnalyzeAction.ProducedItems);
 		}
diff --git a/Engine/Source/Programs/UnrealBuildTool/Platform/Windows/UEBuildWindows.cs b/Engine/Source/Programs/UnrealBuildTool/Platform/Windows/UEBuildWindows.cs
--- a/Engine/Source/Programs/UnrealBuildTool/Platform/Windows/UEBuildWindows.cs
+++ b/Engine/Source/Programs/UnrealBuildTool/Platform/Windows/UEBuildWindows.cs
@@ -457,12 +457,6 @@ namespace UnrealBuildTool
 				Target.bUseIncrementalLinking = true;
 			}
 
-			// Incremental linking.
-			if (Target.bUseIncrementalLinking && !Target.bDisableDebugInfo)
-			{
-				Target.bUsePDBFiles = true;
-			}
-
 //			@Todo: Still getting reports of frequent OOM issues with this enabled as of 15.7.
 //			// Enable fast PDB linking if we're on VS2017 15.7 or later. Previous versions have OOM issues with large projects.
 //			if(!Target.bFormalBuild && !Target.bUseFastPDBLinking.HasValue && Target.WindowsPlatform.Compiler >= WindowsCompiler.VisualStudio2017)
diff --git a/Engine/Source/Programs/UnrealBuildTool/Platform/Windows/VCToolChain.cs b/Engine/Source/Programs/UnrealBuildTool/Platform/Windows/VCToolChain.cs
--- a/Engine/Source/Programs/UnrealBuildTool/Platform/Windows/VCToolChain.cs
+++ b/Engine/Source/Programs/UnrealBuildTool/Platform/Windows/VCToolChain.cs
@@ -887,7 +887,7 @@ namespace UnrealBuildTool
 			{
 				Action CompileAction = ActionGraph.Add(ActionType.Compile);
 				CompileAction.CommandDescription = "Compile";
-				CompileAction.bUseIncrementalLinking = CompileEnvironment.bUseIncrementalLinking;
+
 				// ensure compiler timings are captured when we execute the action.
 				if (Target.WindowsPlatform.Compiler != WindowsCompiler.Clang && CompileEnvironment.bPrintTimingInfo)
 				{
@@ -1196,11 +1196,11 @@ namespace UnrealBuildTool
 				return FileItem.GetItemByFileReference(LinkEnvironment.OutputFilePath);
 			}
 
-			bool bIsBuildingLibrary = LinkEnvironment.bIsBuildingLibrary || bBuildImportLibraryOnly;
+			bool bIsBuildingLibraryOrImportLibrary = LinkEnvironment.bIsBuildingLibrary || bBuildImportLibraryOnly;
 
 			// Get link arguments.
 			List<string> Arguments = new List<string>();
-			if (bIsBuildingLibrary)
+			if (bIsBuildingLibraryOrImportLibrary)
 			{
 				AppendLibArguments(LinkEnvironment, Arguments);
 			}
@@ -1228,7 +1228,7 @@ namespace UnrealBuildTool
 			}
 
 
-			if (!bIsBuildingLibrary)
+			if (!bIsBuildingLibraryOrImportLibrary)
 			{
 				// Delay-load these DLLs.
 				foreach (string DelayLoadDLL in LinkEnvironment.DelayLoadDLLs.Distinct())
@@ -1316,7 +1316,7 @@ namespace UnrealBuildTool
 				}
 			}
 
-			if (!bIsBuildingLibrary)
+			if (!bIsBuildingLibraryOrImportLibrary)
 			{
 				foreach (string AdditionalLibrary in LinkEnvironment.AdditionalLibraries)
 				{
@@ -1337,25 +1337,30 @@ namespace UnrealBuildTool
 			// Add the output file to the command-line.
 			Arguments.Add(String.Format("/OUT:\"{0}\"", OutputFile.AbsolutePath));
 
-			if (bBuildImportLibraryOnly || (LinkEnvironment.bHasExports && !bIsBuildingLibrary))
+			// For import libraries and exports generated by cross-referenced builds, we don't track output files. VS 15.3+ doesn't touch timestamps for libs
+			// and exp files with no modifications, breaking our dependency checking, but incremental linking will fall back to a full link if we delete it.
+			// Since all DLLs are typically marked as cross referenced now anyway, we can just ignore this file to allow incremental linking to work.
+			if(LinkEnvironment.bHasExports && !LinkEnvironment.bIsBuildingLibrary && !LinkEnvironment.bIsCrossReferenced)
 			{
-				// An export file is written to the output directory implicitly; add it to the produced items list.
 				FileReference ExportFilePath = ImportLibraryFilePath.ChangeExtension(".exp");
 				FileItem ExportFile = FileItem.GetItemByFileReference(ExportFilePath);
 				ProducedItems.Add(ExportFile);
 			}
 
-			if (!bIsBuildingLibrary)
+			if (!bIsBuildingLibraryOrImportLibrary)
 			{
 				// There is anything to export
-				if (LinkEnvironment.bHasExports
-					// Shipping monolithic builds don't need exports
-					&& (!((LinkEnvironment.Configuration == CppConfiguration.Shipping) && (LinkEnvironment.bShouldCompileMonolithic != false))))
+				if (LinkEnvironment.bHasExports)
 				{
 					// Write the import library to the output directory for nFringe support.
 					FileItem ImportLibraryFile = FileItem.GetItemByFileReference(ImportLibraryFilePath);
 					Arguments.Add(String.Format("/IMPLIB:\"{0}\"", ImportLibraryFilePath));
-					ProducedItems.Add(ImportLibraryFile);
+
+					// Like the export file above, don't add the import library as a produced item when it's cross referenced.
+					if(!LinkEnvironment.bIsCrossReferenced)
+					{
+					    ProducedItems.Add(ImportLibraryFile);
+					}
 				}
 
 				if (LinkEnvironment.bCreateDebugInfo)
@@ -1400,12 +1405,11 @@ namespace UnrealBuildTool
 			Action LinkAction = ActionGraph.Add(ActionType.Link);
 			LinkAction.CommandDescription = "Link";
 			LinkAction.WorkingDirectory = UnrealBuildTool.EngineSourceDirectory.FullName;
-			LinkAction.CommandPath = bIsBuildingLibrary ? EnvVars.LibraryManagerPath.FullName : EnvVars.LinkerPath.FullName;
+			LinkAction.CommandPath = bIsBuildingLibraryOrImportLibrary ? EnvVars.LibraryManagerPath.FullName : EnvVars.LinkerPath.FullName;
 			LinkAction.CommandArguments = String.Format("@\"{0}\"", ResponseFileName);
 			LinkAction.ProducedItems.AddRange(ProducedItems);
 			LinkAction.PrerequisiteItems.AddRange(PrerequisiteItems);
 			LinkAction.StatusDescription = Path.GetFileName(OutputFile.AbsolutePath);
-			LinkAction.bUseIncrementalLinking = LinkEnvironment.bUseIncrementalLinking;
 
 			// ensure compiler timings are captured when we execute the action.
 			if (Target.WindowsPlatform.Compiler != WindowsCompiler.Clang && LinkEnvironment.bPrintTimingInfo)
@@ -1414,9 +1418,15 @@ namespace UnrealBuildTool
 			}
 
 			// VS 15.3+ does not touch lib files if they do not contain any modifications, but we need to ensure the timestamps are updated to avoid repeatedly building them.
-			if (bBuildImportLibraryOnly || (LinkEnvironment.bHasExports && !bIsBuildingLibrary))
+			if (bBuildImportLibraryOnly || (LinkEnvironment.bHasExports && !bIsBuildingLibraryOrImportLibrary))
+			{
+				LinkAction.DeleteItems.AddRange(LinkAction.ProducedItems.Where(x => x.Location.HasExtension(".lib") || x.Location.HasExtension(".exp")));
+			}
+
+			// Delete PDB files for all produced items, since incremental updates are slower than full ones.
+			if (!LinkEnvironment.bUseIncrementalLinking)
 			{
-				LinkAction.bShouldDeleteProducedItems = true; 
+				LinkAction.DeleteItems.AddRange(LinkAction.ProducedItems.Where(x => x.Location.HasExtension(".pdb")));
 			}
 
 			// Tell the action that we're building an import library here and it should conditionally be
diff --git a/Engine/Source/Programs/UnrealBuildTool/Platform/Windows/WindowsDeploy.cs b/Engine/Source/Programs/UnrealBuildTool/Platform/Windows/WindowsDeploy.cs
--- a/Engine/Source/Programs/UnrealBuildTool/Platform/Windows/WindowsDeploy.cs
+++ b/Engine/Source/Programs/UnrealBuildTool/Platform/Windows/WindowsDeploy.cs
@@ -1,4 +1,4 @@
-﻿// Copyright 1998-2018 Epic Games, Inc. All Rights Reserved.
+// Copyright 1998-2018 Epic Games, Inc. All Rights Reserved.
 
 using System;
 using System.Collections.Generic;
@@ -59,13 +59,20 @@ namespace UnrealBuildTool
 			if ((InTarget.TargetType != TargetType.Editor && InTarget.TargetType != TargetType.Program) && (BuildHostPlatform.Current.Platform == UnrealTargetPlatform.Win32 || BuildHostPlatform.Current.Platform == UnrealTargetPlatform.Win64))
 			{
 				string InAppName = InTarget.AppName;
-				Log.TraceInformation("Prepping {0} for deployment to {1}", InAppName, InTarget.Platform.ToString());
-				System.DateTime PrepDeployStartTime = DateTime.UtcNow;
+				if(InTarget.bUseIncrementalLinking)
+				{
+					Log.TraceInformation("Skipping icon update for {0} due to incremental linking being enabled.", InAppName);
+				}
+				else
+				{
+					Log.TraceInformation("Prepping {0} for deployment to {1}", InAppName, InTarget.Platform.ToString());
+					System.DateTime PrepDeployStartTime = DateTime.UtcNow;
 
-                List<UnrealTargetConfiguration> TargetConfigs = new List<UnrealTargetConfiguration> { InTarget.Configuration };
-                List<string> ExePaths = new List<string> { InTarget.OutputPath.FullName };
-				string RelativeEnginePath = UnrealBuildTool.EngineDirectory.MakeRelativeTo(DirectoryReference.GetCurrentDirectory());
-                SetApplicationIcon(InTarget.ProjectFile, InAppName, InTarget.ProjectDirectory.FullName, TargetConfigs, ExePaths, RelativeEnginePath);
+					List<UnrealTargetConfiguration> TargetConfigs = new List<UnrealTargetConfiguration> { InTarget.Configuration };
+					List<string> ExePaths = new List<string> { InTarget.OutputPath.FullName };
+					string RelativeEnginePath = UnrealBuildTool.EngineDirectory.MakeRelativeTo(DirectoryReference.GetCurrentDirectory());
+					SetApplicationIcon(InTarget.ProjectFile, InAppName, InTarget.ProjectDirectory.FullName, TargetConfigs, ExePaths, RelativeEnginePath);
+				}
 			}
 			return true;
 		}
diff --git a/Engine/Source/Programs/UnrealBuildTool/System/ActionGraph.cs b/Engine/Source/Programs/UnrealBuildTool/System/ActionGraph.cs
--- a/Engine/Source/Programs/UnrealBuildTool/System/ActionGraph.cs
+++ b/Engine/Source/Programs/UnrealBuildTool/System/ActionGraph.cs
@@ -51,6 +51,11 @@ namespace UnrealBuildTool
 		/// </summary>
 		public List<FileItem> ProducedItems = new List<FileItem>();
 
+		/// <summary>
+		/// Items that should be deleted before running this action
+		/// </summary>
+		public List<FileItem> DeleteItems = new List<FileItem>();
+
 		/// <summary>
 		/// Directory from which to execute the program to create produced items
 		/// </summary>
@@ -101,22 +106,12 @@ namespace UnrealBuildTool
 		/// </summary>
 		public bool bIsUsingPCH = false;
 
-		/// <summary>
-		/// Whether the files in ProducedItems should be deleted before executing this action, when the action is outdated
-		/// </summary>
-		public bool bShouldDeleteProducedItems = false;
-
 		/// <summary>
 		/// Whether we should log this action, whether executed locally or remotely.  This is useful for actions that take time
 		/// but invoke tools without any console output.
 		/// </summary>
 		public bool bShouldOutputStatusDescription = true;
 
-		/// <summary>
-		/// Whether this action includes incremental linking.
-		/// </summary>
-		public bool bUseIncrementalLinking = false;
-
 		/// <summary>
 		/// True if any libraries produced by this action should be considered 'import libraries'
 		/// </summary>
@@ -184,12 +179,11 @@ namespace UnrealBuildTool
 			bCanExecuteRemotelyWithSNDBS = SerializationInfo.GetBoolean("cs");
 			bIsGCCCompiler = SerializationInfo.GetBoolean("ig");
 			bIsUsingPCH = SerializationInfo.GetBoolean("iu");
-			bShouldDeleteProducedItems = SerializationInfo.GetBoolean("dp");
 			bShouldOutputStatusDescription = SerializationInfo.GetBoolean("os");
-			bUseIncrementalLinking = SerializationInfo.GetBoolean("in");
 			bProducesImportLibrary = SerializationInfo.GetBoolean("il");
 			PrerequisiteItems = (List<FileItem>)SerializationInfo.GetValue("pr", typeof(List<FileItem>));
 			ProducedItems = (List<FileItem>)SerializationInfo.GetValue("pd", typeof(List<FileItem>));
+			DeleteItems = (List<FileItem>)SerializationInfo.GetValue("df", typeof(List<FileItem>));
 		}
 
 		/// <summary>
@@ -208,12 +202,11 @@ namespace UnrealBuildTool
 			SerializationInfo.AddValue("cs", bCanExecuteRemotelyWithSNDBS);
 			SerializationInfo.AddValue("ig", bIsGCCCompiler);
 			SerializationInfo.AddValue("iu", bIsUsingPCH);
-			SerializationInfo.AddValue("dp", bShouldDeleteProducedItems);
 			SerializationInfo.AddValue("os", bShouldOutputStatusDescription);
-			SerializationInfo.AddValue("in", bUseIncrementalLinking);
 			SerializationInfo.AddValue("il", bProducesImportLibrary);
 			SerializationInfo.AddValue("pr", PrerequisiteItems);
 			SerializationInfo.AddValue("pd", ProducedItems);
+			SerializationInfo.AddValue("df", DeleteItems);
 		}
 
 		/// <summary>
@@ -369,7 +362,7 @@ namespace UnrealBuildTool
 			}
 
 			// Delete produced items that are outdated.
-			DeleteOutdatedProducedItems(OutdatedActionDictionary, BuildConfiguration.bShouldDeleteAllOutdatedProducedItems);
+			DeleteOutdatedProducedItems(OutdatedActionDictionary);
 
 			// Save the action history.
 			// This must happen after deleting outdated produced items to ensure that the action history on disk doesn't have
@@ -385,19 +378,6 @@ namespace UnrealBuildTool
 			// Build a list of actions that are both needed for this target and outdated.
 			HashSet<Action> ActionsToExecute = AllActions.Where(Action => Action.CommandPath != null && IsActionOutdatedMap.ContainsKey(Action) && OutdatedActionDictionary[Action]).ToHashSet();
 
-			// Delete PDB files for all produced items, since incremental updates are slower than full ones.
-			foreach (Action ActionToExecute in ActionsToExecute)
-			{
-				foreach (FileItem ProducedItem in ActionToExecute.ProducedItems)
-				{
-					if(ProducedItem.bExists && !ActionToExecute.bUseIncrementalLinking && ProducedItem.Location.HasExtension(".pdb"))
-					{
-						Log.TraceVerbose("Deleting outdated pdb: {0}", ProducedItem.AbsolutePath);
-						ProducedItem.Delete();
-					}
-				}
-			}
-
 			// Remove link actions if asked to
 			if (BuildConfiguration.bSkipLinkingWhenNothingToCompile)
 			{
@@ -1177,20 +1157,19 @@ namespace UnrealBuildTool
 		/// Deletes all the items produced by actions in the provided outdated action dictionary.
 		/// </summary>
 		/// <param name="OutdatedActionDictionary">Dictionary of outdated actions</param>
-		/// <param name="bShouldDeleteAllFiles"> Whether to delete all files associated with outdated items or just ones required</param>
-		static void DeleteOutdatedProducedItems(Dictionary<Action, bool> OutdatedActionDictionary, bool bShouldDeleteAllFiles)
+		static void DeleteOutdatedProducedItems(Dictionary<Action, bool> OutdatedActionDictionary)
 		{
 			foreach (KeyValuePair<Action, bool> OutdatedActionInfo in OutdatedActionDictionary)
 			{
 				if (OutdatedActionInfo.Value)
 				{
 					Action OutdatedAction = OutdatedActionInfo.Key;
-					foreach (FileItem ProducedItem in OutdatedActionInfo.Key.ProducedItems)
+					foreach (FileItem DeleteItem in OutdatedActionInfo.Key.DeleteItems)
 					{
-						if (ProducedItem.bExists && (bShouldDeleteAllFiles || OutdatedAction.bShouldDeleteProducedItems))
+						if (DeleteItem.bExists)
 						{
-							Log.TraceLog("Deleting outdated item: {0}", ProducedItem.AbsolutePath);
-							ProducedItem.Delete();
+							Log.TraceLog("Deleting outdated item: {0}", DeleteItem.AbsolutePath);
+							DeleteItem.Delete();
 						}
 					}
 				}
