UBT: Format XML validation errors in a format that allows double-clicking on the message in Visual Studio.

#rnx

diff --git a/Engine/Source/Programs/UnrealBuildTool/System/XmlConfigFile.cs b/Engine/Source/Programs/UnrealBuildTool/System/XmlConfigFile.cs
--- a/Engine/Source/Programs/UnrealBuildTool/System/XmlConfigFile.cs
+++ b/Engine/Source/Programs/UnrealBuildTool/System/XmlConfigFile.cs
@@ -124,7 +124,7 @@ namespace UnrealBuildTool
 		/// <param name="Args">Standard argument for ValidationEventHandler</param>
 		void ValidationEvent(object Sender, ValidationEventArgs Args)
 		{
-			Log.TraceWarning("{0}({1}): {2}", File.FullName, Args.Exception.LineNumber, Args.Message);
+			Log.WriteLine(LogEventType.Warning, LogFormatOptions.NoSeverityPrefix, "{0}({1}): warning: {2}", File.FullName, Args.Exception.LineNumber, Args.Message);
 		}
 	}
 
