UBT: Add the MakePathSafeToUseWithCommandLine() function to UBT.

#rnx

diff --git a/Engine/Source/Programs/AutomationTool/AutomationUtils/CommandUtils.cs b/Engine/Source/Programs/AutomationTool/AutomationUtils/CommandUtils.cs
--- a/Engine/Source/Programs/AutomationTool/AutomationUtils/CommandUtils.cs
+++ b/Engine/Source/Programs/AutomationTool/AutomationUtils/CommandUtils.cs
@@ -1859,11 +1859,7 @@ namespace AutomationTool
 		/// <returns></returns>
 		public static string MakePathSafeToUseWithCommandLine(string InPath)
 		{
-			if (InPath.Contains(' ') && InPath[0] != '\"')
-			{
-				InPath = "\"" + InPath + "\"";
-			}
-			return InPath;
+			return UnrealBuildTool.Utils.MakePathSafeToUseWithCommandLine(InPath);
 		}
 
 		#endregion
diff --git a/Engine/Source/Programs/UnrealBuildTool/System/Utils.cs b/Engine/Source/Programs/UnrealBuildTool/System/Utils.cs
--- a/Engine/Source/Programs/UnrealBuildTool/System/Utils.cs
+++ b/Engine/Source/Programs/UnrealBuildTool/System/Utils.cs
@@ -128,6 +128,30 @@ namespace UnrealBuildTool
 			return Result;
 		}
 
+		/// <summary>
+		/// Makes sure path can be used as a command line param (adds quotes if it contains spaces)
+		/// </summary>
+		/// <param name="InPath">Path to convert</param>
+		/// <returns></returns>
+		public static string MakePathSafeToUseWithCommandLine(string InPath)
+		{
+			if (InPath.Contains(' ') && InPath[0] != '\"')
+			{
+				InPath = "\"" + InPath + "\"";
+			}
+			return InPath;
+		}
+
+		/// <summary>
+		/// Escapes whitespace in the given command line argument with a backslash. Used on Unix-like platforms for command line arguments in shell commands.
+		/// </summary>
+		/// <param name="Argument">The argument to escape </param>
+		/// <returns>Escaped shell argument</returns>
+		public static string EscapeShellArgument(string Argument)
+		{
+			return Argument.Replace(" ", "\\ ");
+		}
+
 		/// <summary>
 		/// This is a faster replacement of File.ReadAllText. Code snippet based on code
 		/// and analysis by Sam Allen
