UBT: Fix exception in target receipt when architecture is null.

#rnx

diff --git a/Engine/Source/Programs/UnrealBuildTool/System/TargetReceipt.cs b/Engine/Source/Programs/UnrealBuildTool/System/TargetReceipt.cs
--- a/Engine/Source/Programs/UnrealBuildTool/System/TargetReceipt.cs
+++ b/Engine/Source/Programs/UnrealBuildTool/System/TargetReceipt.cs
@@ -445,7 +445,7 @@ namespace UnrealBuildTool
 		{
 			// Get the architecture suffix. Platforms have the option of overriding whether to include this string in filenames.
 			string ArchitectureSuffix = "";
-			if(UEBuildPlatform.GetBuildPlatform(Platform).RequiresArchitectureSuffix())
+			if(!String.IsNullOrEmpty(BuildArchitecture) && UEBuildPlatform.GetBuildPlatform(Platform).RequiresArchitectureSuffix())
 			{
 				ArchitectureSuffix = BuildArchitecture;
 			}
