Fixes to allow enabling edit and continue for Windows builds. Have experienced quite a few VS crashes when testing it in editor; not yet recommended for general use.

- Allow edit and continue for any configuration, not just debug.
- Fixed PDB errors compiling files that use a shared PCH with edit and continue enabled. Path to the generated PDB file was using the wrong directory.
- Removed code that tracks PDB output files, since they're modified multiple times during a build.
- Enable debug information when compiling generated CPP files, since it causes errors if the shared PCH PDB doesn't have the same option.
- Disable support for remote execution of steps that modify the PDB, since the same file has to be modified many times. Remote execution causes the PDB files to be corrupted. Unfortunately, this makes E&C builds significantly slower.

diff --git a/Engine/Source/Programs/UnrealBuildTool/Configuration/TargetRules.cs b/Engine/Source/Programs/UnrealBuildTool/Configuration/TargetRules.cs
--- a/Engine/Source/Programs/UnrealBuildTool/Configuration/TargetRules.cs
+++ b/Engine/Source/Programs/UnrealBuildTool/Configuration/TargetRules.cs
@@ -860,7 +860,7 @@ namespace UnrealBuildTool
 		public bool bAllowASLRInShipping = true;
 
 		/// <summary>
-		/// Whether to support edit and continue.  Only works on Microsoft compilers in 32-bit compiles.
+		/// Whether to support edit and continue.  Only works on Microsoft compilers.
 		/// </summary>
 		[XmlConfigFile(Category = "BuildConfiguration")]
 		public bool bSupportEditAndContinue = false;
diff --git a/Engine/Source/Programs/UnrealBuildTool/Platform/Windows/UEBuildWindows.cs b/Engine/Source/Programs/UnrealBuildTool/Platform/Windows/UEBuildWindows.cs
--- a/Engine/Source/Programs/UnrealBuildTool/Platform/Windows/UEBuildWindows.cs
+++ b/Engine/Source/Programs/UnrealBuildTool/Platform/Windows/UEBuildWindows.cs
@@ -452,9 +452,16 @@ namespace UnrealBuildTool
 			}
 
 			// E&C support.
-			if (Target.bSupportEditAndContinue && Target.Configuration == UnrealTargetConfiguration.Debug)
+			if (Target.bSupportEditAndContinue)
 			{
 				Target.bUseIncrementalLinking = true;
+				Target.bUsePDBFiles = true;
+			}
+
+			// If we're using PDB files and PCHs, the generated code needs to be compiled with the same options as the PCH.
+			if(Target.bUsePDBFiles && Target.bUsePCHFiles)
+			{
+				Target.bDisableDebugInfoForGeneratedCode = false;
 			}
 
 //			@Todo: Still getting reports of frequent OOM issues with this enabled as of 15.7.
diff --git a/Engine/Source/Programs/UnrealBuildTool/Platform/Windows/VCToolChain.cs b/Engine/Source/Programs/UnrealBuildTool/Platform/Windows/VCToolChain.cs
--- a/Engine/Source/Programs/UnrealBuildTool/Platform/Windows/VCToolChain.cs
+++ b/Engine/Source/Programs/UnrealBuildTool/Platform/Windows/VCToolChain.cs
@@ -369,9 +369,7 @@ namespace UnrealBuildTool
 				if (CompileEnvironment.bUsePDBFiles)
 				{
 					// Create debug info suitable for E&C if wanted.
-					if (CompileEnvironment.bSupportEditAndContinue &&
-						// We only need to do this in debug as that's the only configuration that supports E&C.
-						CompileEnvironment.Configuration == CppConfiguration.Debug)
+					if (CompileEnvironment.bSupportEditAndContinue)
 					{
 						Arguments.Add("/ZI");
 					}
@@ -993,55 +991,53 @@ namespace UnrealBuildTool
 					FileArguments.Add(String.Format("/Fo\"{0}\"", ObjectFile.AbsolutePath));
 				}
 
+				// Don't farm out creation of precompiled headers as it is the critical path task.
+				CompileAction.bCanExecuteRemotely =
+					CompileEnvironment.PrecompiledHeaderAction != PrecompiledHeaderAction.Create ||
+					CompileEnvironment.bAllowRemotelyCompiledPCHs
+					;
+
 				// Create PDB files if we were configured to do that.
 				if (CompileEnvironment.bUsePDBFiles)
 				{
-					string PDBFileName;
-					bool bActionProducesPDB = false;
-
-					// All files using the same PCH are required to share the same PDB that was used when compiling the PCH
+					FileReference PDBLocation;
 					if (CompileEnvironment.PrecompiledHeaderAction == PrecompiledHeaderAction.Include)
 					{
-						PDBFileName = CompileEnvironment.PrecompiledHeaderIncludeFilename.GetFileName();
+						// All files using the same PCH are required to share the same PDB that was used when compiling the PCH
+						PDBLocation = CompileEnvironment.PrecompiledHeaderFile.Location.ChangeExtension(".pdb");
+
+						// Enable synchronous file writes, since we'll be modifying the existing PDB
+						FileArguments.Add("/FS");
 					}
-					// Files creating a PCH use a PDB per file.
 					else if (CompileEnvironment.PrecompiledHeaderAction == PrecompiledHeaderAction.Create)
 					{
-						PDBFileName = CompileEnvironment.PrecompiledHeaderIncludeFilename.GetFileName();
-						bActionProducesPDB = true;
+						// Files creating a PCH use a PDB per file.
+						PDBLocation = FileReference.Combine(OutputDir, CompileEnvironment.PrecompiledHeaderIncludeFilename.GetFileName() + ".pdb");
+
+						// Enable synchronous file writes, since we'll be modifying the existing PDB
+						FileArguments.Add("/FS");
 					}
-					// Ungrouped C++ files use a PDB per file.
 					else if (!bIsPlainCFile)
 					{
-						PDBFileName = Path.GetFileName(SourceFile.AbsolutePath);
-						bActionProducesPDB = true;
+						// Ungrouped C++ files use a PDB per file.
+						PDBLocation = FileReference.Combine(OutputDir, SourceFile.Location.GetFileName() + ".pdb");
 					}
-					// Group all plain C files that doesn't use PCH into the same PDB
 					else
 					{
-						PDBFileName = "MiscPlainC";
+						// Group all plain C files that doesn't use PCH into the same PDB
+						PDBLocation = FileReference.Combine(OutputDir, "MiscPlainC.pdb");
 					}
 
-					{
-						// Specify the PDB file that the compiler should write to.
-						FileItem PDBFile = FileItem.GetItemByFileReference(
-								FileReference.Combine(
-									OutputDir,
-									PDBFileName + ".pdb"
-									)
-								);
-						FileArguments.Add(String.Format("/Fd\"{0}\"", PDBFile.AbsolutePath));
-
-						// Only use the PDB as an output file if we want PDBs and this particular action is
-						// the one that produces the PDB (as opposed to no debug info, where the above code
-						// is needed, but not the output PDB, or when multiple files share a single PDB, so
-						// only the action that generates it should count it as output directly)
-						if (bActionProducesPDB)
-						{
-							CompileAction.ProducedItems.Add(PDBFile);
-							Result.DebugDataFiles.Add(PDBFile);
-						}
-					}
+					// Specify the PDB file that the compiler should write to.
+					FileArguments.Add(String.Format("/Fd\"{0}\"", PDBLocation));
+
+					// Don't add the PDB as an output file because it's modified multiple times. This will break timestamp dependency tracking.
+					FileItem PDBFile = FileItem.GetItemByFileReference(PDBLocation);
+					Result.DebugDataFiles.Add(PDBFile);
+
+					// Don't allow remote execution when PDB files are enabled; we need to modify the same files. XGE works around this by generating separate
+					// PDB files per agent, but this functionality is only available with the Visual C++ extension package (via the VCCompiler=true tool option).
+					CompileAction.bCanExecuteRemotely = false;
 				}
 
 				// Add C or C++ specific compiler arguments.
@@ -1095,12 +1091,6 @@ namespace UnrealBuildTool
 					CompileAction.bShouldOutputStatusDescription = false;
 				}
 
-				// Don't farm out creation of precompiled headers as it is the critical path task.
-				CompileAction.bCanExecuteRemotely =
-					CompileEnvironment.PrecompiledHeaderAction != PrecompiledHeaderAction.Create ||
-					CompileEnvironment.bAllowRemotelyCompiledPCHs
-					;
-
 				// When compiling with SN-DBS, modules that contain a #import must be built locally
 				if (CompileEnvironment.bBuildLocallyWithSNDBS == true)
 				{
