UBT: Minor tidy-up (merging UEBuildBinary.Build and UEBuildBinary.SetupOutputFiles)

#rnx

diff --git a/Engine/Source/Programs/UnrealBuildTool/Configuration/UEBuildBinary.cs b/Engine/Source/Programs/UnrealBuildTool/Configuration/UEBuildBinary.cs
--- a/Engine/Source/Programs/UnrealBuildTool/Configuration/UEBuildBinary.cs
+++ b/Engine/Source/Programs/UnrealBuildTool/Configuration/UEBuildBinary.cs
@@ -207,8 +207,43 @@
 				return BinaryLinkEnvironment.InputFiles;
 			}
 
-			// Return linked files.
-			return SetupOutputFiles(ToolChain, CompileEnvironment, BinaryLinkEnvironment, ActionGraph);
+			// Generate import libraries as a separate step
+			List<FileItem> OutputFiles = new List<FileItem>();
+			if (bCreateImportLibrarySeparately)
+			{
+				// Mark the link environment as cross-referenced.
+				BinaryLinkEnvironment.bIsCrossReferenced = true;
+
+				if (BinaryLinkEnvironment.Platform != CppPlatform.Mac && BinaryLinkEnvironment.Platform != CppPlatform.Linux)
+				{
+					// Create the import library.
+					OutputFiles.AddRange(ToolChain.LinkAllFiles(BinaryLinkEnvironment, true, ActionGraph));
+				}
+			}
+
+			// Link the binary.
+			FileItem[] Executables = ToolChain.LinkAllFiles(BinaryLinkEnvironment, false, ActionGraph);
+			OutputFiles.AddRange(Executables);
+
+			// Produce additional console app if requested
+			if (bBuildAdditionalConsoleApp)
+			{
+				// Produce additional binary but link it as a console app
+				LinkEnvironment ConsoleAppLinkEvironment = new LinkEnvironment(BinaryLinkEnvironment);
+				ConsoleAppLinkEvironment.bIsBuildingConsoleApplication = true;
+				ConsoleAppLinkEvironment.WindowsEntryPointOverride = "WinMainCRTStartup";		// For WinMain() instead of "main()" for Launch module
+				ConsoleAppLinkEvironment.OutputFilePaths = ConsoleAppLinkEvironment.OutputFilePaths.Select(Path => GetAdditionalConsoleAppPath(Path)).ToList();
+
+				// Link the console app executable
+				OutputFiles.AddRange(ToolChain.LinkAllFiles(ConsoleAppLinkEvironment, false, ActionGraph));
+			}
+
+			foreach (FileItem Executable in Executables)
+			{
+				OutputFiles.AddRange(ToolChain.PostBuild(Executable, BinaryLinkEnvironment, ActionGraph));
+			}
+
+			return OutputFiles;
 		}
 
 
@@ -679,49 +714,6 @@
 			return BinaryLinkEnvironment;
 		}
 
-		private List<FileItem> SetupOutputFiles(UEToolChain ToolChain, CppCompileEnvironment BinaryCompileEnvironment, LinkEnvironment BinaryLinkEnvironment, ActionGraph ActionGraph)
-		{
-			//
-			// Regular linking action.
-			//
-			List<FileItem> OutputFiles = new List<FileItem>();
-			if (bCreateImportLibrarySeparately)
-			{
-				// Mark the link environment as cross-referenced.
-				BinaryLinkEnvironment.bIsCrossReferenced = true;
-
-				if (BinaryLinkEnvironment.Platform != CppPlatform.Mac && BinaryLinkEnvironment.Platform != CppPlatform.Linux)
-				{
-					// Create the import library.
-					OutputFiles.AddRange(ToolChain.LinkAllFiles(BinaryLinkEnvironment, true, ActionGraph));
-				}
-			}
-
-			// Link the binary.
-			FileItem[] Executables = ToolChain.LinkAllFiles(BinaryLinkEnvironment, false, ActionGraph);
-			OutputFiles.AddRange(Executables);
-
-			// Produce additional console app if requested
-			if (bBuildAdditionalConsoleApp)
-			{
-				// Produce additional binary but link it as a console app
-				LinkEnvironment ConsoleAppLinkEvironment = new LinkEnvironment(BinaryLinkEnvironment);
-				ConsoleAppLinkEvironment.bIsBuildingConsoleApplication = true;
-				ConsoleAppLinkEvironment.WindowsEntryPointOverride = "WinMainCRTStartup";		// For WinMain() instead of "main()" for Launch module
-				ConsoleAppLinkEvironment.OutputFilePaths = ConsoleAppLinkEvironment.OutputFilePaths.Select(Path => GetAdditionalConsoleAppPath(Path)).ToList();
-
-				// Link the console app executable
-				OutputFiles.AddRange(ToolChain.LinkAllFiles(ConsoleAppLinkEvironment, false, ActionGraph));
-			}
-
-			foreach (FileItem Executable in Executables)
-			{
-				OutputFiles.AddRange(ToolChain.PostBuild(Executable, BinaryLinkEnvironment, ActionGraph));
-			}
-
-			return OutputFiles;
-		}
-
 		/// <summary>
 		/// ToString implementation
 		/// </summary>
