Fix formatting.

#rnx

diff --git a/Engine/Source/Programs/UnrealBuildTool/UnrealBuildTool.cs b/Engine/Source/Programs/UnrealBuildTool/UnrealBuildTool.cs
--- a/Engine/Source/Programs/UnrealBuildTool/UnrealBuildTool.cs
+++ b/Engine/Source/Programs/UnrealBuildTool/UnrealBuildTool.cs
@@ -164,7 +164,7 @@ namespace UnrealBuildTool
 		/// <returns>True if running using installed Enterprise components</returns>
 		static public bool IsEnterpriseInstalled()
 		{
-			if(!bIsEnterpriseInstalled.HasValue)
+			if (!bIsEnterpriseInstalled.HasValue)
 			{
 				bIsEnterpriseInstalled = FileReference.Exists(FileReference.Combine(EnterpriseDirectory, "Build", "InstalledBuild.txt"));
 			}
@@ -216,8 +216,8 @@ namespace UnrealBuildTool
 		static public bool IsUnderAnEngineDirectory(DirectoryReference InDirectory)
 		{
 			// Enterprise modules are considered as engine modules
-			return InDirectory.IsUnderDirectory( UnrealBuildTool.EngineDirectory ) || InDirectory.IsUnderDirectory( UnrealBuildTool.EnterpriseSourceDirectory ) ||
-				InDirectory.IsUnderDirectory( UnrealBuildTool.EnterprisePluginsDirectory ) || InDirectory.IsUnderDirectory( UnrealBuildTool.EnterpriseIntermediateDirectory );
+			return InDirectory.IsUnderDirectory(UnrealBuildTool.EngineDirectory) || InDirectory.IsUnderDirectory(UnrealBuildTool.EnterpriseSourceDirectory) ||
+				InDirectory.IsUnderDirectory(UnrealBuildTool.EnterprisePluginsDirectory) || InDirectory.IsUnderDirectory(UnrealBuildTool.EnterpriseIntermediateDirectory);
 		}
 
 		/// <summary>
@@ -228,9 +228,9 @@ namespace UnrealBuildTool
 		static public bool IsUnderAnInstalledDirectory(DirectoryReference InDirectory)
 		{
 			// Enterprise modules are considered as engine modules
-			bool bIsUnderEngine = InDirectory.IsUnderDirectory( UnrealBuildTool.EngineDirectory );
+			bool bIsUnderEngine = InDirectory.IsUnderDirectory(UnrealBuildTool.EngineDirectory);
 
-			bool bIsUnderEnterprise = InDirectory.IsUnderDirectory( UnrealBuildTool.EnterpriseSourceDirectory ) ||
+			bool bIsUnderEnterprise = InDirectory.IsUnderDirectory(UnrealBuildTool.EnterpriseSourceDirectory) ||
 				InDirectory.IsUnderDirectory(UnrealBuildTool.EnterprisePluginsDirectory) || InDirectory.IsUnderDirectory(UnrealBuildTool.EnterpriseIntermediateDirectory);
 
 			return (IsEngineInstalled() && bIsUnderEngine) || (IsEnterpriseInstalled() && bIsUnderEnterprise);
@@ -247,15 +247,15 @@ namespace UnrealBuildTool
 				List<System.Type> ProjectGeneratorList = new List<System.Type>();
 				Type[] AllTypes = UBTAssembly.GetTypes();
 
-					// register all build platforms first, since they implement SDK-switching logic that can set environment variables
-					foreach (Type CheckType in AllTypes)
+				// register all build platforms first, since they implement SDK-switching logic that can set environment variables
+				foreach (Type CheckType in AllTypes)
+				{
+					if (CheckType.IsClass && !CheckType.IsAbstract)
 					{
-						if (CheckType.IsClass && !CheckType.IsAbstract)
+						if (CheckType.IsSubclassOf(typeof(UEBuildPlatformFactory)))
 						{
-							if (CheckType.IsSubclassOf(typeof(UEBuildPlatformFactory)))
-							{
-								Log.TraceVerbose("    Registering build platform: {0}", CheckType.ToString());
-								UEBuildPlatformFactory TempInst = (UEBuildPlatformFactory)(UBTAssembly.CreateInstance(CheckType.FullName, true));
+							Log.TraceVerbose("    Registering build platform: {0}", CheckType.ToString());
+							UEBuildPlatformFactory TempInst = (UEBuildPlatformFactory)(UBTAssembly.CreateInstance(CheckType.FullName, true));
 							TempInst.TryRegisterBuildPlatforms(OutputLevel, bValidatingPlatforms);
 						}
 					}
@@ -301,7 +301,7 @@ namespace UnrealBuildTool
 				RemoteIniPath = InArg.Substring(11);
 			}
 			else
-            {
+			{
 				return false;
 			}
 
@@ -361,15 +361,15 @@ namespace UnrealBuildTool
 			StartupTraceListener StartupListener = new StartupTraceListener();
 			bool bLogProgramNameWithSeverity = Arguments.Any(x => x.Equals("-FromMsBuild", StringComparison.InvariantCultureIgnoreCase));
 			Log.InitLogging(
-                bLogTimestamps: Arguments.Any(x => x.Equals("-Timestamps", StringComparison.InvariantCultureIgnoreCase)),
+				bLogTimestamps: Arguments.Any(x => x.Equals("-Timestamps", StringComparison.InvariantCultureIgnoreCase)),
 				InLogLevel: LogLevel,
-                bLogSeverity: true,
-				bLogProgramNameWithSeverity: bLogProgramNameWithSeverity, 
-                bLogSources: true,
+				bLogSeverity: true,
+				bLogProgramNameWithSeverity: bLogProgramNameWithSeverity,
+				bLogSources: true,
 				bLogSourcesToConsole: false,
-                bColorConsoleOutput: true,
-                TraceListeners: new TraceListener[]
-                {
+				bColorConsoleOutput: true,
+				TraceListeners: new TraceListener[]
+				{
 					StartupListener
 				}
 			);
@@ -831,7 +831,7 @@ namespace UnrealBuildTool
 
 							// Read from the editor config
 							ProjectFileFormat PreferredSourceCodeAccessor;
-							if(ProjectFileGenerator.GetPreferredSourceCodeAccessor(ProjectFile, out PreferredSourceCodeAccessor))
+							if (ProjectFileGenerator.GetPreferredSourceCodeAccessor(ProjectFile, out PreferredSourceCodeAccessor))
 							{
 								ProjectFileFormats.Add(PreferredSourceCodeAccessor);
 							}
@@ -889,7 +889,7 @@ namespace UnrealBuildTool
 								case ProjectFileFormat.VisualStudioCode:
 									Generator = new VSCodeProjectFileGenerator(ProjectFile);
 									break;
-﻿								case ProjectFileFormat.CLion:
+								case ProjectFileFormat.CLion:
 									Generator = new CLionGenerator(ProjectFile);
 									break;
 								default:
@@ -971,9 +971,9 @@ namespace UnrealBuildTool
 		public static int ExtendedErrorCode = 0;
 		private static int Main(string[] Arguments)
 		{
-            // make sure we catch any exceptions and return an appropriate error code.
-            // Some inner code already does this (to ensure the Mutex is released),
-            // but we need something to cover all outer code as well.
+			// make sure we catch any exceptions and return an appropriate error code.
+			// Some inner code already does this (to ensure the Mutex is released),
+			// but we need something to cover all outer code as well.
 			try
 			{
 				return GuardedMain(Arguments);
@@ -1289,444 +1289,444 @@ namespace UnrealBuildTool
 
 				using (ISourceFileWorkingSet WorkingSet = SourceFileWorkingSet.Create(UnrealBuildTool.RootDirectory, DirectoryReference.FromFile(ProjectFile)))
 				{
-				UBTMakefile UBTMakefile = null;
-				{
-					// If we're generating project files, then go ahead and wipe out the existing UBTMakefile for every target, to make sure that
-					// it gets a full dependency scan next time.
-					// NOTE: This is just a safeguard and doesn't have to be perfect.  We also check for newer project file timestamps in LoadUBTMakefile()
+					UBTMakefile UBTMakefile = null;
+					{
+						// If we're generating project files, then go ahead and wipe out the existing UBTMakefile for every target, to make sure that
+						// it gets a full dependency scan next time.
+						// NOTE: This is just a safeguard and doesn't have to be perfect.  We also check for newer project file timestamps in LoadUBTMakefile()
 						FileReference UBTMakefilePath = UBTMakefile.GetUBTMakefilePath(TargetDescs, HotReload);
 						if (ProjectFileGenerator.bGenerateProjectFiles) // @todo ubtmake: This is only hit when generating IntelliSense for project files.  Probably should be done right inside ProjectFileGenerator.bat
 						{                                                   // @todo ubtmake: Won't catch multi-target cases as GPF always builds one target at a time for Intellisense
-						// Delete the UBTMakefile
-						if (FileReference.Exists(UBTMakefilePath))
-						{
+																			// Delete the UBTMakefile
+							if (FileReference.Exists(UBTMakefilePath))
+							{
 								FileReference.Delete(UBTMakefilePath);
+							}
 						}
-					}
 
-					// Make sure the gather phase is executed if we're not actually building anything
-					if (ProjectFileGenerator.bGenerateProjectFiles || BuildConfiguration.bGenerateManifest || BuildConfiguration.bCleanProject || BuildConfiguration.bXGEExport || GeneratingActionGraph)
-					{
+						// Make sure the gather phase is executed if we're not actually building anything
+						if (ProjectFileGenerator.bGenerateProjectFiles || BuildConfiguration.bGenerateManifest || BuildConfiguration.bCleanProject || BuildConfiguration.bXGEExport || GeneratingActionGraph)
+						{
 							bIsGatheringBuild = true;
-					}
+						}
 
-					// Were we asked to run in 'assembler only' mode?  If so, let's check to see if that's even possible by seeing if
-					// we have a valid UBTMakefile already saved to disk, ready for us to load.
+						// Were we asked to run in 'assembler only' mode?  If so, let's check to see if that's even possible by seeing if
+						// we have a valid UBTMakefile already saved to disk, ready for us to load.
 						if (bIsAssemblingBuild && !bIsGatheringBuild)
-					{
-						// @todo ubtmake: Mildly terrified of BuildConfiguration/UEBuildConfiguration globals that were set during the Prepare phase but are not set now.  We may need to save/load all of these, otherwise
-						//		we'll need to call SetupGlobalEnvironment on all of the targets (maybe other stuff, too.  See PreBuildStep())
+						{
+							// @todo ubtmake: Mildly terrified of BuildConfiguration/UEBuildConfiguration globals that were set during the Prepare phase but are not set now.  We may need to save/load all of these, otherwise
+							//		we'll need to call SetupGlobalEnvironment on all of the targets (maybe other stuff, too.  See PreBuildStep())
 
-						// Try to load the UBTMakefile.  It will only be loaded if it has valid content and is not determined to be out of date.    
-						string ReasonNotLoaded;
+							// Try to load the UBTMakefile.  It will only be loaded if it has valid content and is not determined to be out of date.    
+							string ReasonNotLoaded;
 							UBTMakefile = UBTMakefile.LoadUBTMakefile(UBTMakefilePath, ProjectFile, WorkingSet, out ReasonNotLoaded);
 
-						// Invalid makefile if only modules have changed
-						if (UBTMakefile != null && !TargetDescs.SelectMany(x => x.OnlyModules)
-								.Select(x => new Tuple<string, bool>(x.OnlyModuleName.ToLower(), string.IsNullOrWhiteSpace(x.OnlyModuleSuffix)))
-								.SequenceEqual(
-									UBTMakefile.Targets.SelectMany(x => x.OnlyModules)
+							// Invalid makefile if only modules have changed
+							if (UBTMakefile != null && !TargetDescs.SelectMany(x => x.OnlyModules)
 									.Select(x => new Tuple<string, bool>(x.OnlyModuleName.ToLower(), string.IsNullOrWhiteSpace(x.OnlyModuleSuffix)))
+									.SequenceEqual(
+										UBTMakefile.Targets.SelectMany(x => x.OnlyModules)
+										.Select(x => new Tuple<string, bool>(x.OnlyModuleName.ToLower(), string.IsNullOrWhiteSpace(x.OnlyModuleSuffix)))
+									)
 								)
-							)
-						{
-							UBTMakefile = null;
-							ReasonNotLoaded = "modules to compile have changed";
-						}
+							{
+								UBTMakefile = null;
+								ReasonNotLoaded = "modules to compile have changed";
+							}
 
-						if (UBTMakefile != null)
-						{
-							// Check if ini files are newer. Ini files contain build settings too.
-							FileInfo UBTMakefileInfo = new FileInfo(UBTMakefilePath.FullName);
-							foreach (TargetDescriptor Desc in TargetDescs)
+							if (UBTMakefile != null)
 							{
-								DirectoryReference ProjectDirectory = DirectoryReference.FromFile(ProjectFile);
-									foreach (ConfigHierarchyType IniType in (ConfigHierarchyType[])Enum.GetValues(typeof(ConfigHierarchyType)))
+								// Check if ini files are newer. Ini files contain build settings too.
+								FileInfo UBTMakefileInfo = new FileInfo(UBTMakefilePath.FullName);
+								foreach (TargetDescriptor Desc in TargetDescs)
 								{
-									foreach (FileReference IniFilename in ConfigHierarchy.EnumerateConfigFileLocations(IniType, ProjectDirectory, Desc.Platform))
+									DirectoryReference ProjectDirectory = DirectoryReference.FromFile(ProjectFile);
+									foreach (ConfigHierarchyType IniType in (ConfigHierarchyType[])Enum.GetValues(typeof(ConfigHierarchyType)))
 									{
-										FileInfo IniFileInfo = new FileInfo(IniFilename.FullName);
-										if (UBTMakefileInfo.LastWriteTime.CompareTo(IniFileInfo.LastWriteTime) < 0)
+										foreach (FileReference IniFilename in ConfigHierarchy.EnumerateConfigFileLocations(IniType, ProjectDirectory, Desc.Platform))
+										{
+											FileInfo IniFileInfo = new FileInfo(IniFilename.FullName);
+											if (UBTMakefileInfo.LastWriteTime.CompareTo(IniFileInfo.LastWriteTime) < 0)
+											{
+												// Ini files are newer than UBTMakefile
+												UBTMakefile = null;
+												ReasonNotLoaded = "ini files are newer than UBTMakefile";
+												break;
+											}
+										}
+
+										if (UBTMakefile == null)
 										{
-											// Ini files are newer than UBTMakefile
-											UBTMakefile = null;
-											ReasonNotLoaded = "ini files are newer than UBTMakefile";
 											break;
 										}
 									}
-
 									if (UBTMakefile == null)
 									{
 										break;
 									}
 								}
-								if (UBTMakefile == null)
-								{
-									break;
-								}
 							}
-						}
 
-						if (UBTMakefile == null)
-						{
-							// If the Makefile couldn't be loaded, then we're not going to be able to continue in "assembler only" mode.  We'll do both
-							// a 'gather' and 'assemble' in the same run.  This will take a while longer, but subsequent runs will be fast!
+							if (UBTMakefile == null)
+							{
+								// If the Makefile couldn't be loaded, then we're not going to be able to continue in "assembler only" mode.  We'll do both
+								// a 'gather' and 'assemble' in the same run.  This will take a while longer, but subsequent runs will be fast!
 								bIsGatheringBuild = true;
 
-							FileItem.ClearCaches();
+								FileItem.ClearCaches();
 
-							Log.TraceInformation("Creating makefile for {0}{1}{2} ({3})",
-								(HotReload != EHotReload.Disabled) ? "hot reloading " : "",
-									TargetDescs[0].Name,
-								TargetDescs.Count > 1 ? (" (and " + (TargetDescs.Count - 1).ToString() + " more)") : "",
-								ReasonNotLoaded);
+								Log.TraceInformation("Creating makefile for {0}{1}{2} ({3})",
+									(HotReload != EHotReload.Disabled) ? "hot reloading " : "",
+										TargetDescs[0].Name,
+									TargetDescs.Count > 1 ? (" (and " + (TargetDescs.Count - 1).ToString() + " more)") : "",
+									ReasonNotLoaded);
+							}
 						}
 					}
-				}
 
 
 					if (UBTMakefile != null && !bIsGatheringBuild && bIsAssemblingBuild)
-				{
-					// If we've loaded a makefile, then we can fill target information from this file!
-					Targets = UBTMakefile.Targets;
-				}
-				else
-				{
-					DateTime TargetInitStartTime = DateTime.UtcNow;
-
-					ReadOnlyBuildVersion Version = new ReadOnlyBuildVersion(BuildVersion.ReadDefault());
-					
-					Targets = new List<UEBuildTarget>();
-					foreach (TargetDescriptor TargetDesc in TargetDescs)
 					{
-						UEBuildTarget Target = UEBuildTarget.CreateTarget(TargetDesc, Arguments, BuildConfiguration.SingleFileToCompile != null, Version);
-						if ((Target == null) && (BuildConfiguration.bCleanProject))
+						// If we've loaded a makefile, then we can fill target information from this file!
+						Targets = UBTMakefile.Targets;
+					}
+					else
+					{
+						DateTime TargetInitStartTime = DateTime.UtcNow;
+
+						ReadOnlyBuildVersion Version = new ReadOnlyBuildVersion(BuildVersion.ReadDefault());
+
+						Targets = new List<UEBuildTarget>();
+						foreach (TargetDescriptor TargetDesc in TargetDescs)
 						{
-							continue;
+							UEBuildTarget Target = UEBuildTarget.CreateTarget(TargetDesc, Arguments, BuildConfiguration.SingleFileToCompile != null, Version);
+							if ((Target == null) && (BuildConfiguration.bCleanProject))
+							{
+								continue;
+							}
+							Targets.Add(Target);
 						}
-						Targets.Add(Target);
-					}
 
-					if (UnrealBuildTool.bPrintPerformanceInfo)
-					{
-						double TargetInitTime = (DateTime.UtcNow - TargetInitStartTime).TotalSeconds;
-						Log.TraceInformation("Target init took " + TargetInitTime + "s");
+						if (UnrealBuildTool.bPrintPerformanceInfo)
+						{
+							double TargetInitTime = (DateTime.UtcNow - TargetInitStartTime).TotalSeconds;
+							Log.TraceInformation("Target init took " + TargetInitTime + "s");
+						}
 					}
-				}
 
-				// Build action lists for all passed in targets.
-				List<FileItem> OutputItemsForAllTargets = new List<FileItem>();
-				Dictionary<string, List<UHTModuleInfo>> TargetNameToUObjectModules = new Dictionary<string, List<UHTModuleInfo>>(StringComparer.InvariantCultureIgnoreCase);
-				foreach (UEBuildTarget Target in Targets)
-				{
+					// Build action lists for all passed in targets.
+					List<FileItem> OutputItemsForAllTargets = new List<FileItem>();
+					Dictionary<string, List<UHTModuleInfo>> TargetNameToUObjectModules = new Dictionary<string, List<UHTModuleInfo>>(StringComparer.InvariantCultureIgnoreCase);
+					foreach (UEBuildTarget Target in Targets)
+					{
 						if (HotReload != EHotReload.Disabled)
 						{
-						// Don't produce new DLLs if there's been no code changes
-						BuildConfiguration.bSkipLinkingWhenNothingToCompile = true;
-						Log.TraceInformation("Compiling game modules for hot reload");
-					}
+							// Don't produce new DLLs if there's been no code changes
+							BuildConfiguration.bSkipLinkingWhenNothingToCompile = true;
+							Log.TraceInformation("Compiling game modules for hot reload");
+						}
 
-					// Create the header cache for this target
-					FileReference DependencyCacheFile = DependencyCache.GetDependencyCachePathForTarget(ProjectFile, Target.Platform, Target.TargetName);
+						// Create the header cache for this target
+						FileReference DependencyCacheFile = DependencyCache.GetDependencyCachePathForTarget(ProjectFile, Target.Platform, Target.TargetName);
 						bool bUseFlatCPPIncludeDependencyCache = BuildConfiguration.bUseUBTMakefiles && bIsAssemblingBuild;
-					CPPHeaders Headers = new CPPHeaders(ProjectFile, DependencyCacheFile, bUseFlatCPPIncludeDependencyCache, BuildConfiguration.bUseUBTMakefiles, BuildConfiguration.bUseIncludeDependencyResolveCache, BuildConfiguration.bTestIncludeDependencyResolveCache);
-					TargetToHeaders[Target] = Headers;
+						CPPHeaders Headers = new CPPHeaders(ProjectFile, DependencyCacheFile, bUseFlatCPPIncludeDependencyCache, BuildConfiguration.bUseUBTMakefiles, BuildConfiguration.bUseIncludeDependencyResolveCache, BuildConfiguration.bTestIncludeDependencyResolveCache);
+						TargetToHeaders[Target] = Headers;
 
-					// Make sure the appropriate executor is selected
-					UEBuildPlatform BuildPlatform = UEBuildPlatform.GetBuildPlatform(Target.Platform);
-					BuildConfiguration.bAllowXGE &= BuildPlatform.CanUseXGE();
-					BuildConfiguration.bAllowDistcc &= BuildPlatform.CanUseDistcc();
-					BuildConfiguration.bAllowSNDBS &= BuildPlatform.CanUseSNDBS();
+						// Make sure the appropriate executor is selected
+						UEBuildPlatform BuildPlatform = UEBuildPlatform.GetBuildPlatform(Target.Platform);
+						BuildConfiguration.bAllowXGE &= BuildPlatform.CanUseXGE();
+						BuildConfiguration.bAllowDistcc &= BuildPlatform.CanUseDistcc();
+						BuildConfiguration.bAllowSNDBS &= BuildPlatform.CanUseSNDBS();
 
-					// When in 'assembler only' mode, we'll load this cache later on a worker thread.  It takes a long time to load!
+						// When in 'assembler only' mode, we'll load this cache later on a worker thread.  It takes a long time to load!
 						if (!(!bIsGatheringBuild && bIsAssemblingBuild))
-					{
-						// Load the direct include dependency cache.
-						Headers.IncludeDependencyCache = DependencyCache.Create(DependencyCache.GetDependencyCachePathForTarget(Target.ProjectFile, Target.Platform, Target.GetTargetName()));
-					}
+						{
+							// Load the direct include dependency cache.
+							Headers.IncludeDependencyCache = DependencyCache.Create(DependencyCache.GetDependencyCachePathForTarget(Target.ProjectFile, Target.Platform, Target.GetTargetName()));
+						}
 
-					// We don't need this dependency cache in 'gather only' mode
-					if (BuildConfiguration.bUseUBTMakefiles &&
-							!(bIsGatheringBuild && !bIsAssemblingBuild))
-					{
-						// Load the cache that contains the list of flattened resolved includes for resolved source files
-						// @todo ubtmake: Ideally load this asynchronously at startup and only block when it is first needed and not finished loading
-						FileReference CacheFile = FlatCPPIncludeDependencyCache.GetDependencyCachePathForTarget(Target);
-							if (!FlatCPPIncludeDependencyCache.TryRead(CacheFile, out Headers.FlatCPPIncludeDependencyCache))
+						// We don't need this dependency cache in 'gather only' mode
+						if (BuildConfiguration.bUseUBTMakefiles &&
+								!(bIsGatheringBuild && !bIsAssemblingBuild))
 						{
-								if (!bNeedsFullCPPIncludeRescan)
+							// Load the cache that contains the list of flattened resolved includes for resolved source files
+							// @todo ubtmake: Ideally load this asynchronously at startup and only block when it is first needed and not finished loading
+							FileReference CacheFile = FlatCPPIncludeDependencyCache.GetDependencyCachePathForTarget(Target);
+							if (!FlatCPPIncludeDependencyCache.TryRead(CacheFile, out Headers.FlatCPPIncludeDependencyCache))
 							{
-									if (!ProjectFileGenerator.bGenerateProjectFiles && !BuildConfiguration.bXGEExport && !BuildConfiguration.bGenerateManifest && !BuildConfiguration.bCleanProject)
+								if (!bNeedsFullCPPIncludeRescan)
 								{
+									if (!ProjectFileGenerator.bGenerateProjectFiles && !BuildConfiguration.bXGEExport && !BuildConfiguration.bGenerateManifest && !BuildConfiguration.bCleanProject)
+									{
 										bNeedsFullCPPIncludeRescan = true;
-									Log.TraceInformation("Performing full C++ include scan (no include cache file)");
+										Log.TraceInformation("Performing full C++ include scan (no include cache file)");
+									}
 								}
+								Headers.FlatCPPIncludeDependencyCache = new FlatCPPIncludeDependencyCache(CacheFile);
 							}
-							Headers.FlatCPPIncludeDependencyCache = new FlatCPPIncludeDependencyCache(CacheFile);
 						}
-					}
 
 						if (bIsGatheringBuild)
-					{
-						List<FileItem> TargetOutputItems = new List<FileItem>();
-						List<UHTModuleInfo> TargetUObjectModules = new List<UHTModuleInfo>();
-						if (BuildConfiguration.bCleanProject)
-						{
-							BuildResult = Target.Clean(!BuildConfiguration.bDoNotBuildUHT && HotReload != EHotReload.FromIDE);
-						}
-						else
 						{
+							List<FileItem> TargetOutputItems = new List<FileItem>();
+							List<UHTModuleInfo> TargetUObjectModules = new List<UHTModuleInfo>();
+							if (BuildConfiguration.bCleanProject)
+							{
+								BuildResult = Target.Clean(!BuildConfiguration.bDoNotBuildUHT && HotReload != EHotReload.FromIDE);
+							}
+							else
+							{
 								BuildResult = Target.Build(BuildConfiguration, TargetToHeaders[Target], TargetOutputItems, TargetUObjectModules, WorkingSet, ActionGraph, HotReload, bIsAssemblingBuild);
-						}
-						if (BuildResult != ECompilationResult.Succeeded)
-						{
-							break;
-						}
+							}
+							if (BuildResult != ECompilationResult.Succeeded)
+							{
+								break;
+							}
 
-						// Export a JSON dump of this target
-						if (BuildConfiguration.JsonExportFile != null)
-						{
-							Directory.CreateDirectory(Path.GetDirectoryName(Path.GetFullPath(BuildConfiguration.JsonExportFile)));
-							Target.ExportJson(BuildConfiguration.JsonExportFile);
-						}
+							// Export a JSON dump of this target
+							if (BuildConfiguration.JsonExportFile != null)
+							{
+								Directory.CreateDirectory(Path.GetDirectoryName(Path.GetFullPath(BuildConfiguration.JsonExportFile)));
+								Target.ExportJson(BuildConfiguration.JsonExportFile);
+							}
 
-						OutputItemsForAllTargets.AddRange(TargetOutputItems);
+							OutputItemsForAllTargets.AddRange(TargetOutputItems);
 
-						// Update mapping of the target name to the list of UObject modules in this target
-						TargetNameToUObjectModules[Target.GetTargetName()] = TargetUObjectModules;
+							// Update mapping of the target name to the list of UObject modules in this target
+							TargetNameToUObjectModules[Target.GetTargetName()] = TargetUObjectModules;
 
-						if ((BuildConfiguration.bXGEExport && BuildConfiguration.bGenerateManifest) || (!ProjectFileGenerator.bGenerateProjectFiles && !BuildConfiguration.bGenerateManifest && !BuildConfiguration.bCleanProject))
-						{
-							// Generate an action graph if we were asked to do that.  The graph generation needs access to the include dependency cache, so
-							// we generate it before saving and cleaning that up.
-							if (GeneratingActionGraph)
+							if ((BuildConfiguration.bXGEExport && BuildConfiguration.bGenerateManifest) || (!ProjectFileGenerator.bGenerateProjectFiles && !BuildConfiguration.bGenerateManifest && !BuildConfiguration.bCleanProject))
 							{
-								// The graph generation feature currently only works with a single target at a time.  This is because of how we need access
-								// to include dependencies for the target, but those aren't kept around as we process multiple targets
-								if (TargetSettings.Count != 1)
+								// Generate an action graph if we were asked to do that.  The graph generation needs access to the include dependency cache, so
+								// we generate it before saving and cleaning that up.
+								if (GeneratingActionGraph)
 								{
-									throw new BuildException("ERROR: The '-graph' option only works with a single target at a time");
-								}
-								////ActionGraph.FinalizeActionGraph();
-								List<Action> ActionsToExecute = ActionGraph.AllActions;
+									// The graph generation feature currently only works with a single target at a time.  This is because of how we need access
+									// to include dependencies for the target, but those aren't kept around as we process multiple targets
+									if (TargetSettings.Count != 1)
+									{
+										throw new BuildException("ERROR: The '-graph' option only works with a single target at a time");
+									}
+									////ActionGraph.FinalizeActionGraph();
+									List<Action> ActionsToExecute = ActionGraph.AllActions;
 
-								ActionGraph.ActionGraphVisualizationType VisualizationType = ActionGraph.ActionGraphVisualizationType.OnlyCPlusPlusFilesAndHeaders;
-								ActionGraph.SaveActionGraphVisualization(TargetToHeaders[Target], FileReference.Combine(UnrealBuildTool.EngineDirectory, "Intermediate", "Build", Target.GetTargetName() + ".gexf").FullName, Target.GetTargetName(), VisualizationType, ActionsToExecute);
+									ActionGraph.ActionGraphVisualizationType VisualizationType = ActionGraph.ActionGraphVisualizationType.OnlyCPlusPlusFilesAndHeaders;
+									ActionGraph.SaveActionGraphVisualization(TargetToHeaders[Target], FileReference.Combine(UnrealBuildTool.EngineDirectory, "Intermediate", "Build", Target.GetTargetName() + ".gexf").FullName, Target.GetTargetName(), VisualizationType, ActionsToExecute);
+								}
 							}
 						}
 					}
-				}
 
-				if (BuildResult == ECompilationResult.Succeeded && !BuildConfiguration.bSkipBuild &&
-					(
-						(BuildConfiguration.bXGEExport && BuildConfiguration.bGenerateManifest) ||
-						(!GeneratingActionGraph && !ProjectFileGenerator.bGenerateProjectFiles && !BuildConfiguration.bGenerateManifest && !BuildConfiguration.bCleanProject)
-					))
-				{
-						if (bIsGatheringBuild)
+					if (BuildResult == ECompilationResult.Succeeded && !BuildConfiguration.bSkipBuild &&
+						(
+							(BuildConfiguration.bXGEExport && BuildConfiguration.bGenerateManifest) ||
+							(!GeneratingActionGraph && !ProjectFileGenerator.bGenerateProjectFiles && !BuildConfiguration.bGenerateManifest && !BuildConfiguration.bCleanProject)
+						))
 					{
-						ActionGraph.FinalizeActionGraph();
+						if (bIsGatheringBuild)
+						{
+							ActionGraph.FinalizeActionGraph();
 
-						UBTMakefile = new UBTMakefile();
-						UBTMakefile.AllActions = ActionGraph.AllActions;
+							UBTMakefile = new UBTMakefile();
+							UBTMakefile.AllActions = ActionGraph.AllActions;
 
-						// For now simply treat all object files as the root target.
-						{
-							HashSet<Action> PrerequisiteActionsSet = new HashSet<Action>();
-							foreach (FileItem OutputItem in OutputItemsForAllTargets)
+							// For now simply treat all object files as the root target.
 							{
-								ActionGraph.GatherPrerequisiteActions(OutputItem, ref PrerequisiteActionsSet);
+								HashSet<Action> PrerequisiteActionsSet = new HashSet<Action>();
+								foreach (FileItem OutputItem in OutputItemsForAllTargets)
+								{
+									ActionGraph.GatherPrerequisiteActions(OutputItem, ref PrerequisiteActionsSet);
+								}
+								UBTMakefile.PrerequisiteActions = PrerequisiteActionsSet.ToArray();
 							}
-							UBTMakefile.PrerequisiteActions = PrerequisiteActionsSet.ToArray();
-						}
 
-						foreach (System.Collections.DictionaryEntry EnvironmentVariable in Environment.GetEnvironmentVariables())
-						{
-							UBTMakefile.EnvironmentVariables.Add(Tuple.Create((string)EnvironmentVariable.Key, (string)EnvironmentVariable.Value));
-						}
-						UBTMakefile.TargetNameToUObjectModules = TargetNameToUObjectModules;
-						UBTMakefile.Targets = Targets;
-						UBTMakefile.bUseAdaptiveUnityBuild = Targets.Any(x => x.Rules.bUseAdaptiveUnityBuild);
-						UBTMakefile.SourceFileWorkingSet = Unity.SourceFileWorkingSet;
-						UBTMakefile.CandidateSourceFilesForWorkingSet = Unity.CandidateSourceFilesForWorkingSet;
+							foreach (System.Collections.DictionaryEntry EnvironmentVariable in Environment.GetEnvironmentVariables())
+							{
+								UBTMakefile.EnvironmentVariables.Add(Tuple.Create((string)EnvironmentVariable.Key, (string)EnvironmentVariable.Value));
+							}
+							UBTMakefile.TargetNameToUObjectModules = TargetNameToUObjectModules;
+							UBTMakefile.Targets = Targets;
+							UBTMakefile.bUseAdaptiveUnityBuild = Targets.Any(x => x.Rules.bUseAdaptiveUnityBuild);
+							UBTMakefile.SourceFileWorkingSet = Unity.SourceFileWorkingSet;
+							UBTMakefile.CandidateSourceFilesForWorkingSet = Unity.CandidateSourceFilesForWorkingSet;
 
-						if (BuildConfiguration.bUseUBTMakefiles && !UBTMakefile.PrerequisiteActions.Any(x => x.ActionHandler != null))
-						{
-							// We've been told to prepare to build, so let's go ahead and save out our action graph so that we can use in a later invocation 
-							// to assemble the build.  Even if we are configured to assemble the build in this same invocation, we want to save out the
-							// Makefile so that it can be used on subsequent 'assemble only' runs, for the fastest possible iteration times
-							// @todo ubtmake: Optimization: We could make 'gather + assemble' mode slightly faster by saving this while busy compiling (on our worker thread)
+							if (BuildConfiguration.bUseUBTMakefiles && !UBTMakefile.PrerequisiteActions.Any(x => x.ActionHandler != null))
+							{
+								// We've been told to prepare to build, so let's go ahead and save out our action graph so that we can use in a later invocation 
+								// to assemble the build.  Even if we are configured to assemble the build in this same invocation, we want to save out the
+								// Makefile so that it can be used on subsequent 'assemble only' runs, for the fastest possible iteration times
+								// @todo ubtmake: Optimization: We could make 'gather + assemble' mode slightly faster by saving this while busy compiling (on our worker thread)
 								UBTMakefile.SaveUBTMakefile(TargetDescs, HotReload, UBTMakefile);
+							}
 						}
-					}
 
 						if (bIsAssemblingBuild)
-					{
-						// If we didn't build the graph in this session, then we'll need to load a cached one
-							if (!bIsGatheringBuild && BuildResult.Succeeded())
 						{
-							ActionGraph.AllActions = UBTMakefile.AllActions;
-
-							// Patch action history for hot reload when running in assembler mode.  In assembler mode, the suffix on the output file will be
-							// the same for every invocation on that makefile, but we need a new suffix each time.
-							if (HotReload != EHotReload.Disabled)
+							// If we didn't build the graph in this session, then we'll need to load a cached one
+							if (!bIsGatheringBuild && BuildResult.Succeeded())
 							{
-								PatchActionHistoryForHotReloadAssembling(ActionGraph, HotReloadTargetDesc.OnlyModules, Targets);
-							}
+								ActionGraph.AllActions = UBTMakefile.AllActions;
 
+								// Patch action history for hot reload when running in assembler mode.  In assembler mode, the suffix on the output file will be
+								// the same for every invocation on that makefile, but we need a new suffix each time.
+								if (HotReload != EHotReload.Disabled)
+								{
+									PatchActionHistoryForHotReloadAssembling(ActionGraph, HotReloadTargetDesc.OnlyModules, Targets);
+								}
 
-							foreach (Tuple<string, string> EnvironmentVariable in UBTMakefile.EnvironmentVariables)
-							{
-								// @todo ubtmake: There may be some variables we do NOT want to clobber.
-								Environment.SetEnvironmentVariable(EnvironmentVariable.Item1, EnvironmentVariable.Item2);
-							}
 
-							// Execute all the pre-build steps
+								foreach (Tuple<string, string> EnvironmentVariable in UBTMakefile.EnvironmentVariables)
+								{
+									// @todo ubtmake: There may be some variables we do NOT want to clobber.
+									Environment.SetEnvironmentVariable(EnvironmentVariable.Item1, EnvironmentVariable.Item2);
+								}
+
+								// Execute all the pre-build steps
 								if (!ProjectFileGenerator.bGenerateProjectFiles && !BuildConfiguration.bXGEExport)
-							{
-									foreach (UEBuildTarget Target in Targets)
 								{
-										if (!Target.ExecuteCustomPreBuildSteps())
+									foreach (UEBuildTarget Target in Targets)
 									{
-										BuildResult = ECompilationResult.OtherCompilationError;
-										break;
+										if (!Target.ExecuteCustomPreBuildSteps())
+										{
+											BuildResult = ECompilationResult.OtherCompilationError;
+											break;
+										}
 									}
 								}
-							}
 
-							// If any of the targets need UHT to be run, we'll go ahead and do that now
-							if (BuildResult.Succeeded())
-							{
-								foreach (UEBuildTarget Target in Targets)
+								// If any of the targets need UHT to be run, we'll go ahead and do that now
+								if (BuildResult.Succeeded())
 								{
-									List<UHTModuleInfo> TargetUObjectModules;
-									if (UBTMakefile.TargetNameToUObjectModules.TryGetValue(Target.GetTargetName(), out TargetUObjectModules))
+									foreach (UEBuildTarget Target in Targets)
 									{
-										if (TargetUObjectModules.Count > 0)
+										List<UHTModuleInfo> TargetUObjectModules;
+										if (UBTMakefile.TargetNameToUObjectModules.TryGetValue(Target.GetTargetName(), out TargetUObjectModules))
 										{
-											// Execute the header tool
-											FileReference ModuleInfoFileName = FileReference.Combine(Target.ProjectIntermediateDirectory, Target.GetTargetName() + ".uhtmanifest");
-											ECompilationResult UHTResult = ECompilationResult.OtherCompilationError;
-												if (!ExternalExecution.ExecuteHeaderToolIfNecessary(BuildConfiguration, Target, GlobalCompileEnvironment: null, UObjectModules: TargetUObjectModules, ModuleInfoFileName: ModuleInfoFileName, UHTResult: ref UHTResult, HotReload: HotReload, bIsGatheringBuild: bIsGatheringBuild, bIsAssemblingBuild: bIsAssemblingBuild))
+											if (TargetUObjectModules.Count > 0)
 											{
-												Log.TraceInformation("UnrealHeaderTool failed for target '" + Target.GetTargetName() + "' (platform: " + Target.Platform.ToString() + ", module info: " + ModuleInfoFileName + ").");
-												BuildResult = UHTResult;
-												break;
+												// Execute the header tool
+												FileReference ModuleInfoFileName = FileReference.Combine(Target.ProjectIntermediateDirectory, Target.GetTargetName() + ".uhtmanifest");
+												ECompilationResult UHTResult = ECompilationResult.OtherCompilationError;
+												if (!ExternalExecution.ExecuteHeaderToolIfNecessary(BuildConfiguration, Target, GlobalCompileEnvironment: null, UObjectModules: TargetUObjectModules, ModuleInfoFileName: ModuleInfoFileName, UHTResult: ref UHTResult, HotReload: HotReload, bIsGatheringBuild: bIsGatheringBuild, bIsAssemblingBuild: bIsAssemblingBuild))
+												{
+													Log.TraceInformation("UnrealHeaderTool failed for target '" + Target.GetTargetName() + "' (platform: " + Target.Platform.ToString() + ", module info: " + ModuleInfoFileName + ").");
+													BuildResult = UHTResult;
+													break;
+												}
 											}
 										}
 									}
 								}
 							}
-						}
 
-						if (BuildResult.Succeeded())
-						{
-							// Make sure any old DLL files from in-engine recompiles aren't lying around.  Must be called after the action graph is finalized.
-							ActionGraph.DeleteStaleHotReloadDLLs();
+							if (BuildResult.Succeeded())
+							{
+								// Make sure any old DLL files from in-engine recompiles aren't lying around.  Must be called after the action graph is finalized.
+								ActionGraph.DeleteStaleHotReloadDLLs();
 
 								foreach (UEBuildTarget Target in Targets)
 								{
-								UEBuildPlatform.GetBuildPlatform(Target.Platform).PreBuildSync();
-							}
+									UEBuildPlatform.GetBuildPlatform(Target.Platform).PreBuildSync();
+								}
 
-                            // Plan the actions to execute for the build.
-                            Dictionary<UEBuildTarget, List<FileItem>> TargetToOutdatedPrerequisitesMap;
+								// Plan the actions to execute for the build.
+								Dictionary<UEBuildTarget, List<FileItem>> TargetToOutdatedPrerequisitesMap;
 								List<Action> ActionsToExecute = ActionGraph.GetActionsToExecute(BuildConfiguration, UBTMakefile.PrerequisiteActions, Targets, TargetToHeaders, bIsAssemblingBuild, bNeedsFullCPPIncludeRescan, out TargetToOutdatedPrerequisitesMap);
 
-                            // Display some stats to the user.
-                            Log.TraceVerbose(
-                                    "{0} actions, {1} outdated and requested actions",
-                                    ActionGraph.AllActions.Count,
-                                    ActionsToExecute.Count
-                                    );
-
-							// Cache indirect includes for all outdated C++ files.  We kick this off as a background thread so that it can
-							// perform the scan while we're compiling.  It usually only takes up to a few seconds, but we don't want to hurt
-							// our best case UBT iteration times for this task which can easily be performed asynchronously
-							if (BuildConfiguration.bUseUBTMakefiles && TargetToOutdatedPrerequisitesMap.Count > 0)
-							{
-								CppIncludeThread = new CppIncludeBackgroundThread(TargetToOutdatedPrerequisitesMap, TargetToHeaders);
-							}
+								// Display some stats to the user.
+								Log.TraceVerbose(
+										"{0} actions, {1} outdated and requested actions",
+										ActionGraph.AllActions.Count,
+										ActionsToExecute.Count
+										);
+
+								// Cache indirect includes for all outdated C++ files.  We kick this off as a background thread so that it can
+								// perform the scan while we're compiling.  It usually only takes up to a few seconds, but we don't want to hurt
+								// our best case UBT iteration times for this task which can easily be performed asynchronously
+								if (BuildConfiguration.bUseUBTMakefiles && TargetToOutdatedPrerequisitesMap.Count > 0)
+								{
+									CppIncludeThread = new CppIncludeBackgroundThread(TargetToOutdatedPrerequisitesMap, TargetToHeaders);
+								}
 
-							// If we're not touching any shared files (ie. anything under Engine), allow the build ids to be recycled between applications.
-							HashSet<FileReference> OutputFiles = new HashSet<FileReference>(ActionsToExecute.SelectMany(x => x.ProducedItems.Select(y => y.Location)));
-							foreach (UEBuildTarget Target in Targets)
-							{
-									if (!Target.TryRecycleVersionManifests(OutputFiles))
+								// If we're not touching any shared files (ie. anything under Engine), allow the build ids to be recycled between applications.
+								HashSet<FileReference> OutputFiles = new HashSet<FileReference>(ActionsToExecute.SelectMany(x => x.ProducedItems.Select(y => y.Location)));
+								foreach (UEBuildTarget Target in Targets)
 								{
-									Target.InvalidateVersionManifests();
+									if (!Target.TryRecycleVersionManifests(OutputFiles))
+									{
+										Target.InvalidateVersionManifests();
+									}
 								}
-							}
 
-							// Execute the actions.
-							string TargetInfoForTelemetry = String.Join("|", Targets.Select(x => String.Format("{0} {1} {2}{3}", x.TargetName, x.Platform, x.Configuration, "")));
-							DateTime ExecutorStartTime = DateTime.UtcNow;
+								// Execute the actions.
+								string TargetInfoForTelemetry = String.Join("|", Targets.Select(x => String.Format("{0} {1} {2}{3}", x.TargetName, x.Platform, x.Configuration, "")));
+								DateTime ExecutorStartTime = DateTime.UtcNow;
 								if (BuildConfiguration.bXGEExport)
-							{
-								XGE.ExportActions(ActionsToExecute);
-								bSuccess = true;
-							}
-							else
-							{
-								bool bIsRemoteCompile = BuildHostPlatform.Current.Platform != UnrealTargetPlatform.Mac && TargetDescs.Any(x => x.Platform == UnrealTargetPlatform.Mac || x.Platform == UnrealTargetPlatform.IOS);
-								bSuccess = ActionGraph.ExecuteActions(BuildConfiguration, ActionsToExecute, bIsRemoteCompile, out ExecutorName, TargetInfoForTelemetry, HotReload);
-							}
-							TotalExecutorTime += (DateTime.UtcNow - ExecutorStartTime).TotalSeconds;
+								{
+									XGE.ExportActions(ActionsToExecute);
+									bSuccess = true;
+								}
+								else
+								{
+									bool bIsRemoteCompile = BuildHostPlatform.Current.Platform != UnrealTargetPlatform.Mac && TargetDescs.Any(x => x.Platform == UnrealTargetPlatform.Mac || x.Platform == UnrealTargetPlatform.IOS);
+									bSuccess = ActionGraph.ExecuteActions(BuildConfiguration, ActionsToExecute, bIsRemoteCompile, out ExecutorName, TargetInfoForTelemetry, HotReload);
+								}
+								TotalExecutorTime += (DateTime.UtcNow - ExecutorStartTime).TotalSeconds;
 
-							// if the build succeeded, write the receipts and do any needed syncing
-							if (bSuccess)
-							{
-								foreach (UEBuildTarget Target in Targets)
+								// if the build succeeded, write the receipts and do any needed syncing
+								if (bSuccess)
 								{
-									Target.WriteReceipts();
-									UEBuildPlatform.GetBuildPlatform(Target.Platform).PostBuildSync(Target);
+									foreach (UEBuildTarget Target in Targets)
+									{
+										Target.WriteReceipts();
+										UEBuildPlatform.GetBuildPlatform(Target.Platform).PostBuildSync(Target);
+									}
+									if (ActionsToExecute.Count == 0 && BuildConfiguration.bSkipLinkingWhenNothingToCompile)
+									{
+										BuildResult = ECompilationResult.UpToDate;
+									}
 								}
-								if (ActionsToExecute.Count == 0 && BuildConfiguration.bSkipLinkingWhenNothingToCompile)
+								else
 								{
-									BuildResult = ECompilationResult.UpToDate;
+									BuildResult = ECompilationResult.OtherCompilationError;
 								}
 							}
-							else
-							{
-								BuildResult = ECompilationResult.OtherCompilationError;
-							}
-						}
 
-						// Execute all the post-build steps
+							// Execute all the post-build steps
 							if (BuildResult.Succeeded() && !ProjectFileGenerator.bGenerateProjectFiles && !BuildConfiguration.bXGEExport)
-						{
-								foreach (UEBuildTarget Target in Targets)
 							{
-									if (!Target.ExecuteCustomPostBuildSteps())
+								foreach (UEBuildTarget Target in Targets)
 								{
-									BuildResult = ECompilationResult.OtherCompilationError;
-									break;
+									if (!Target.ExecuteCustomPostBuildSteps())
+									{
+										BuildResult = ECompilationResult.OtherCompilationError;
+										break;
+									}
 								}
 							}
-						}
 
-						// Run the deployment steps
+							// Run the deployment steps
 							if (BuildResult.Succeeded()
 								&& String.IsNullOrEmpty(BuildConfiguration.SingleFileToCompile)
 								&& !ProjectFileGenerator.bGenerateProjectFiles
 								&& !BuildConfiguration.bGenerateManifest
 								&& !BuildConfiguration.bCleanProject
 							&& !BuildConfiguration.bXGEExport)
-						{
-							foreach (UEBuildTarget Target in Targets)
 							{
-									if (Target.DeployTargetFile != null && Target.OnlyModules.Count == 0)
+								foreach (UEBuildTarget Target in Targets)
 								{
-									Log.WriteLine(LogEventType.Console, "Deploying {0} {1} {2}...", Target.TargetName, Target.Platform, Target.Configuration);
-									UEBuildDeployTarget DeployTarget = new UEBuildDeployTarget(Target.DeployTargetFile);
-									UEBuildPlatform.GetBuildPlatform(Target.Platform).Deploy(DeployTarget);
+									if (Target.DeployTargetFile != null && Target.OnlyModules.Count == 0)
+									{
+										Log.WriteLine(LogEventType.Console, "Deploying {0} {1} {2}...", Target.TargetName, Target.Platform, Target.Configuration);
+										UEBuildDeployTarget DeployTarget = new UEBuildDeployTarget(Target.DeployTargetFile);
+										UEBuildPlatform.GetBuildPlatform(Target.Platform).Deploy(DeployTarget);
+									}
 								}
 							}
 						}
 					}
 				}
 			}
-			}
 			catch (Exception Ex)
 			{
 				if (!bCatchExceptions)
@@ -1840,7 +1840,7 @@ namespace UnrealBuildTool
 				{
 					throw new BuildException("Unknown editor filename for this platform");
 				}
-					
+
 				BuildHostPlatform.ProcessInfo[] Processes = BuildHostPlatform.Current.GetProcesses();
 				string EditorRunsDir = Path.Combine(UnrealBuildTool.EngineDirectory.FullName, "Intermediate", "EditorRuns");
 
@@ -1917,7 +1917,7 @@ namespace UnrealBuildTool
 			Exception CaughtException;
 
 			public CppIncludeBackgroundThread(Dictionary<UEBuildTarget, List<FileItem>> TargetToOutdatedPrerequisitesMap, Dictionary<UEBuildTarget, CPPHeaders> TargetToHeaders)
-		{
+			{
 				BackgroundThread = new Thread(() => Run(TargetToOutdatedPrerequisitesMap, TargetToHeaders));
 				BackgroundThread.Start();
 			}
@@ -1950,7 +1950,7 @@ namespace UnrealBuildTool
 							Headers.FindAndCacheAllIncludedFiles(PrerequisiteItem, PrerequisiteItem.CachedIncludePaths, bOnlyCachedDependencies: false);
 						}
 					}
-		}
+				}
 				catch (Exception Ex)
 				{
 					CaughtException = Ex;
