UBT: Allow generating manifests to any arbitrary locations with the -Manifest=<Path> argument.

diff --git a/Engine/Source/Programs/UnrealBuildTool/Configuration/TargetRules.cs b/Engine/Source/Programs/UnrealBuildTool/Configuration/TargetRules.cs
--- a/Engine/Source/Programs/UnrealBuildTool/Configuration/TargetRules.cs
+++ b/Engine/Source/Programs/UnrealBuildTool/Configuration/TargetRules.cs
@@ -1090,6 +1090,12 @@ namespace UnrealBuildTool
 		/// </summary>
 		public List<string> ExtraModuleNames = new List<string>();
 
+		/// <summary>
+		/// Path to a manifest to output for this target
+		/// </summary>
+		[CommandLine("-Manifest")]
+		public List<FileReference> ManifestFileNames = new List<FileReference>();
+
 		/// <summary>
 		/// Specifies the build environment for this target. See TargetBuildEnvironment for more information on the available options.
 		/// </summary>
@@ -2064,6 +2070,11 @@ namespace UnrealBuildTool
 			get { return Inner.ExtraModuleNames.AsReadOnly(); }
 		}
 
+		public IReadOnlyList<FileReference> ManifestFileNames
+		{
+			get { return Inner.ManifestFileNames.AsReadOnly(); }
+		}
+
 		public TargetBuildEnvironment BuildEnvironment
 		{
 			get { return Inner.BuildEnvironment; }
diff --git a/Engine/Source/Programs/UnrealBuildTool/Configuration/UEBuildTarget.cs b/Engine/Source/Programs/UnrealBuildTool/Configuration/UEBuildTarget.cs
--- a/Engine/Source/Programs/UnrealBuildTool/Configuration/UEBuildTarget.cs
+++ b/Engine/Source/Programs/UnrealBuildTool/Configuration/UEBuildTarget.cs
@@ -1092,7 +1092,7 @@ namespace UnrealBuildTool
 			}
 			else
 			{
-			ProjectIntermediateDirectory = DirectoryReference.Combine(ProjectDirectory, PlatformIntermediateFolder, GetTargetName(), Configuration.ToString());
+				ProjectIntermediateDirectory = DirectoryReference.Combine(ProjectDirectory, PlatformIntermediateFolder, GetTargetName(), Configuration.ToString());
 			}
 
 			// Build the engine intermediate directory. If we're building agnostic engine binaries, we can use the engine intermediates folder. Otherwise we need to use the project intermediates directory.
@@ -1445,7 +1445,14 @@ namespace UnrealBuildTool
 			{
 				ManifestPath = FileReference.Combine(UnrealBuildTool.EngineDirectory, "Intermediate", "Build", "Manifest.xml");
 			}
+			GenerateManifest(ManifestPath, BuildProducts, PrecompileOnlyBuildProducts);
+		}
 
+		/// <summary>
+		/// Generates a public manifest file for writing out
+		/// </summary>
+		public void GenerateManifest(FileReference ManifestPath, List<KeyValuePair<FileReference, BuildProductType>> BuildProducts, List<KeyValuePair<FileReference, BuildProductType>> PrecompileOnlyBuildProducts)
+		{
 			BuildManifest Manifest = new BuildManifest();
 
 			if (!Rules.bDisableLinking)
@@ -1517,6 +1524,7 @@ namespace UnrealBuildTool
 			Manifest.DeployTargetFiles.Sort();
 			Manifest.PostBuildScripts.Sort();
 
+			Log.TraceInformation("Writing manifest to {0}", ManifestPath);
 			Utils.WriteClass<BuildManifest>(Manifest, ManifestPath.FullName, "");
 		}
 
@@ -2326,9 +2334,9 @@ namespace UnrealBuildTool
 			}
 			else
 			{
-			// Build the target's binaries.
+				// Build the target's binaries.
 				foreach (UEBuildBinary Binary in Binaries)
-			{
+				{
 					OutputItems.AddRange(Binary.Build(Rules, TargetToolChain, GlobalCompileEnvironment, GlobalLinkEnvironment, SharedPCHs, WorkingSet, ActionGraph));
 				}
 
@@ -2387,6 +2395,10 @@ namespace UnrealBuildTool
 			{
 				GenerateManifest(BuildProducts, PrecompileOnlyBuildProducts);
 			}
+			foreach(FileReference ManifestFileName in Rules.ManifestFileNames)
+			{
+				GenerateManifest(ManifestFileName, BuildProducts, PrecompileOnlyBuildProducts);
+			}
 
 			// Clean any stale modules which exist in multiple output directories. This can lead to the wrong DLL being loaded on Windows.
 			CleanStaleModules();
@@ -3149,7 +3161,7 @@ namespace UnrealBuildTool
 				// Otherwise create a new module for it
 				Module.Binary = CreateDynamicLibraryForModule(Module);
 				Binaries.Add(Module.Binary);
-            }
+			}
 		}
 
 		/// <summary>
@@ -3247,8 +3259,8 @@ namespace UnrealBuildTool
 				ModuleConfiguration = UnrealTargetConfiguration.Development;
 			}
 
-			// Get the output filenames
-			FileReference BaseBinaryPath = FileReference.Combine(OutputDirectory, MakeBinaryFileName(AppName + "-" + Module.Name, Platform, ModuleConfiguration, Architecture, Rules.UndecoratedConfiguration, UEBuildBinaryType.DynamicLinkLibrary));
+            // Get the output filenames
+            FileReference BaseBinaryPath = FileReference.Combine(OutputDirectory, MakeBinaryFileName(AppName + "-" + Module.Name, Platform, ModuleConfiguration, Architecture, Rules.UndecoratedConfiguration, UEBuildBinaryType.DynamicLinkLibrary));
 			List<FileReference> OutputFilePaths = UEBuildPlatform.GetBuildPlatform(Platform).FinalizeBinaryPaths(BaseBinaryPath, ProjectFile, Rules);
 
 			// Create the binary
@@ -3262,17 +3274,17 @@ namespace UnrealBuildTool
 			);
 		}
 
-		/// <summary>
-		/// Makes a filename (without path) for a compiled binary (e.g. "Core-Win64-Debug.lib") */
-		/// </summary>
-		/// <param name="BinaryName">The name of this binary</param>
-		/// <param name="Platform">The platform being built for</param>
-		/// <param name="Configuration">The configuration being built</param>
+        /// <summary>
+        /// Makes a filename (without path) for a compiled binary (e.g. "Core-Win64-Debug.lib") */
+        /// </summary>
+        /// <param name="BinaryName">The name of this binary</param>
+        /// <param name="Platform">The platform being built for</param>
+        /// <param name="Configuration">The configuration being built</param>
 		/// <param name="Architecture">The target architecture being built</param>
-		/// <param name="UndecoratedConfiguration">The target configuration which doesn't require a platform and configuration suffix. Development by default.</param>
-		/// <param name="BinaryType">Type of binary</param>
-		/// <returns>Name of the binary</returns>
-		public static string MakeBinaryFileName(string BinaryName, UnrealTargetPlatform Platform, UnrealTargetConfiguration Configuration, string Architecture, UnrealTargetConfiguration UndecoratedConfiguration, UEBuildBinaryType BinaryType)
+        /// <param name="UndecoratedConfiguration">The target configuration which doesn't require a platform and configuration suffix. Development by default.</param>
+        /// <param name="BinaryType">Type of binary</param>
+        /// <returns>Name of the binary</returns>
+        public static string MakeBinaryFileName(string BinaryName, UnrealTargetPlatform Platform, UnrealTargetConfiguration Configuration, string Architecture, UnrealTargetConfiguration UndecoratedConfiguration, UEBuildBinaryType BinaryType)
 		{
 			StringBuilder Result = new StringBuilder();
 
@@ -3294,27 +3306,27 @@ namespace UnrealBuildTool
 				Result.Append(Architecture);
 			}
 
-				Result.Append(BuildPlatform.GetBinaryExtension(BinaryType));
+			Result.Append(BuildPlatform.GetBinaryExtension(BinaryType));
 
-			return Result.ToString();
+            return Result.ToString();
 		}
 
-		/// <summary>
-		/// Determine the output path for a target's executable
-		/// </summary>
-		/// <param name="BaseDirectory">The base directory for the executable; typically either the engine directory or project directory.</param>
-		/// <param name="BinaryName">Name of the binary</param>
-		/// <param name="Platform">Target platform to build for</param>
-		/// <param name="Configuration">Target configuration being built</param>
+        /// <summary>
+        /// Determine the output path for a target's executable
+        /// </summary>
+        /// <param name="BaseDirectory">The base directory for the executable; typically either the engine directory or project directory.</param>
+        /// <param name="BinaryName">Name of the binary</param>
+        /// <param name="Platform">Target platform to build for</param>
+        /// <param name="Configuration">Target configuration being built</param>
 		/// <param name="Architecture">Architecture being built</param>
-		/// <param name="BinaryType">The type of binary we're compiling</param>
-		/// <param name="UndecoratedConfiguration">The configuration which doesn't have a "-{Platform}-{Configuration}" suffix added to the binary</param>
-		/// <param name="bIncludesGameModules">Whether this executable contains game modules</param>
-		/// <param name="ExeSubFolder">Subfolder for executables. May be null.</param>
+        /// <param name="BinaryType">The type of binary we're compiling</param>
+        /// <param name="UndecoratedConfiguration">The configuration which doesn't have a "-{Platform}-{Configuration}" suffix added to the binary</param>
+        /// <param name="bIncludesGameModules">Whether this executable contains game modules</param>
+        /// <param name="ExeSubFolder">Subfolder for executables. May be null.</param>
 		/// <param name="ProjectFile">The project file containing the target being built</param>
 		/// <param name="Rules">Rules for the target being built</param>
-		/// <returns>List of executable paths for this target</returns>
-		public static List<FileReference> MakeBinaryPaths(DirectoryReference BaseDirectory, string BinaryName, UnrealTargetPlatform Platform, UnrealTargetConfiguration Configuration, UEBuildBinaryType BinaryType, string Architecture, UnrealTargetConfiguration UndecoratedConfiguration, bool bIncludesGameModules, string ExeSubFolder, FileReference ProjectFile, ReadOnlyTargetRules Rules)
+        /// <returns>List of executable paths for this target</returns>
+        public static List<FileReference> MakeBinaryPaths(DirectoryReference BaseDirectory, string BinaryName, UnrealTargetPlatform Platform, UnrealTargetConfiguration Configuration, UEBuildBinaryType BinaryType, string Architecture, UnrealTargetConfiguration UndecoratedConfiguration, bool bIncludesGameModules, string ExeSubFolder, FileReference ProjectFile, ReadOnlyTargetRules Rules)
 		{
 			// Get the configuration for the executable. If we're building DebugGame, and this executable only contains engine modules, use the same name as development.
 			UnrealTargetConfiguration ExeConfiguration = Configuration;
@@ -3396,18 +3408,18 @@ namespace UnrealBuildTool
 				{
 					if(!Rules.EnablePlugins.Contains(PluginReference.Name, StringComparer.InvariantCultureIgnoreCase) && !Rules.DisablePlugins.Contains(PluginReference.Name, StringComparer.InvariantCultureIgnoreCase))
 					{
-						// Make sure we don't have multiple references to the same plugin
-						if(!ReferencedNames.Add(PluginReference.Name))
-						{
-							Log.TraceWarning("Plugin '{0}' is listed multiple times in project file '{1}'.", PluginReference.Name, ProjectFile);
-						}
-						else
-						{
-							AddPlugin(PluginReference, ProjectReferenceChain, ExcludeFolders, NameToInstance, NameToInfo);
-						}
+					// Make sure we don't have multiple references to the same plugin
+					if(!ReferencedNames.Add(PluginReference.Name))
+					{
+						Log.TraceWarning("Plugin '{0}' is listed multiple times in project file '{1}'.", PluginReference.Name, ProjectFile);
+					}
+					else
+					{
+						AddPlugin(PluginReference, ProjectReferenceChain, ExcludeFolders, NameToInstance, NameToInfo);
 					}
 				}
 			}
+			}
 
 			// Also synthesize references for plugins which are enabled by default
 			if (Rules.bCompileAgainstEngine)
@@ -3416,15 +3428,15 @@ namespace UnrealBuildTool
 				{
 					if(Plugin.EnabledByDefault && !ReferencedNames.Contains(Plugin.Name))
 					{
-						ReferencedNames.Add(Plugin.Name);
+							ReferencedNames.Add(Plugin.Name);
 
-						PluginReferenceDescriptor PluginReference = new PluginReferenceDescriptor(Plugin.Name, null, true);
-						PluginReference.bOptional = true;
+							PluginReferenceDescriptor PluginReference = new PluginReferenceDescriptor(Plugin.Name, null, true);
+							PluginReference.bOptional = true;
 
-						AddPlugin(PluginReference, "default plugins", ExcludeFolders, NameToInstance, NameToInfo);
+							AddPlugin(PluginReference, "default plugins", ExcludeFolders, NameToInstance, NameToInfo);
+						}
 					}
 				}
-			}
 
 			// If this is a program, synthesize references for plugins which are enabled via the config file
 			if(TargetType == TargetType.Program)
@@ -3635,7 +3647,7 @@ namespace UnrealBuildTool
 				bAllowExports: Rules.bHasExports,
 				PrimaryModule: LaunchModule,
 				bUsePrecompiled: LaunchModule.Rules.bUsePrecompiled && OutputPaths[0].IsUnderDirectory(UnrealBuildTool.EngineDirectory)
-						);
+			);
 			Binaries.Add(Binary);
 
 			// Add the launch module to it
@@ -3643,8 +3655,8 @@ namespace UnrealBuildTool
 			Binary.AddModule(LaunchModule);
 
 			// Create an additional console app for the editor
-					if (Platform == UnrealTargetPlatform.Win64 && Configuration != UnrealTargetConfiguration.Shipping && TargetType == TargetType.Editor)
-					{
+			if (Platform == UnrealTargetPlatform.Win64 && Configuration != UnrealTargetConfiguration.Shipping && TargetType == TargetType.Editor)
+			{
 				Binary.bBuildAdditionalConsoleApp = true;
 			}
 		}
@@ -3699,26 +3711,26 @@ namespace UnrealBuildTool
 			GlobalLinkEnvironment.BundleDirectory = BuildPlatform.GetBundleDirectory(Rules, OutputPaths);
 			GlobalLinkEnvironment.BundleVersion = Rules.BundleVersion;
 			GlobalLinkEnvironment.bAllowLTCG = Rules.bAllowLTCG;
-			GlobalLinkEnvironment.bPGOOptimize = Rules.bPGOOptimize;
-			GlobalLinkEnvironment.bPGOProfile = Rules.bPGOProfile;
+            GlobalLinkEnvironment.bPGOOptimize = Rules.bPGOOptimize;
+            GlobalLinkEnvironment.bPGOProfile = Rules.bPGOProfile;
 			GlobalLinkEnvironment.bUseIncrementalLinking = Rules.bUseIncrementalLinking;
 			GlobalLinkEnvironment.bUseFastPDBLinking = Rules.bUseFastPDBLinking ?? false;
 			GlobalLinkEnvironment.bPrintTimingInfo = Rules.bPrintToolChainTimingInfo;
 			GlobalLinkEnvironment.AdditionalArguments = Rules.AdditionalLinkerArguments;
+		
+            if (Rules.bPGOOptimize && Rules.bPGOProfile)
+            {
+                throw new BuildException("bPGOProfile and bPGOOptimize are mutually exclusive.");
+            }
 
-			if (Rules.bPGOOptimize && Rules.bPGOProfile)
-			{
-				throw new BuildException("bPGOProfile and bPGOOptimize are mutually exclusive.");
-			}
-
-			if (Rules.bPGOProfile)
-			{
-				GlobalCompileEnvironment.Definitions.Add("ENABLE_PGO_PROFILE=1");
-			}
-			else
-			{
-				GlobalCompileEnvironment.Definitions.Add("ENABLE_PGO_PROFILE=0");
-			}
+            if (Rules.bPGOProfile)
+            {
+                GlobalCompileEnvironment.Definitions.Add("ENABLE_PGO_PROFILE=1");
+            }
+            else
+            {
+                GlobalCompileEnvironment.Definitions.Add("ENABLE_PGO_PROFILE=0");
+            }
 
 			// Add the 'Engine/Source' path as a global include path for all modules
 			GlobalCompileEnvironment.IncludePaths.UserIncludePaths.Add(UnrealBuildTool.EngineSourceDirectory);
@@ -3761,39 +3773,39 @@ namespace UnrealBuildTool
 			}
 
 			if (Rules.bForceCompileDevelopmentAutomationTests)
-			{
-				GlobalCompileEnvironment.Definitions.Add("WITH_DEV_AUTOMATION_TESTS=1");
-			}
-			else
-			{
+            {
+                GlobalCompileEnvironment.Definitions.Add("WITH_DEV_AUTOMATION_TESTS=1");
+            }
+            else
+            {
                 switch(Configuration)
-				{
-					case UnrealTargetConfiguration.Test:
-					case UnrealTargetConfiguration.Shipping:
-						GlobalCompileEnvironment.Definitions.Add("WITH_DEV_AUTOMATION_TESTS=0");
-						break;
-					default:
-						GlobalCompileEnvironment.Definitions.Add("WITH_DEV_AUTOMATION_TESTS=1");
-						break;
-				}
-			}
+                {
+                    case UnrealTargetConfiguration.Test:
+                    case UnrealTargetConfiguration.Shipping:
+                        GlobalCompileEnvironment.Definitions.Add("WITH_DEV_AUTOMATION_TESTS=0");
+                        break;
+                    default:
+                        GlobalCompileEnvironment.Definitions.Add("WITH_DEV_AUTOMATION_TESTS=1");
+                        break;
+                }
+            }
 
-			if (Rules.bForceCompilePerformanceAutomationTests)
-			{
-				GlobalCompileEnvironment.Definitions.Add("WITH_PERF_AUTOMATION_TESTS=1");
-			}
-			else
-			{
-				switch (Configuration)
-				{
-					case UnrealTargetConfiguration.Shipping:
-						GlobalCompileEnvironment.Definitions.Add("WITH_PERF_AUTOMATION_TESTS=0");
-						break;
-					default:
-						GlobalCompileEnvironment.Definitions.Add("WITH_PERF_AUTOMATION_TESTS=1");
-						break;
-				}
-			}
+            if (Rules.bForceCompilePerformanceAutomationTests)
+            {
+                GlobalCompileEnvironment.Definitions.Add("WITH_PERF_AUTOMATION_TESTS=1");
+            }
+            else
+            {
+                switch (Configuration)
+                {
+                    case UnrealTargetConfiguration.Shipping:
+                        GlobalCompileEnvironment.Definitions.Add("WITH_PERF_AUTOMATION_TESTS=0");
+                        break;
+                    default:
+                        GlobalCompileEnvironment.Definitions.Add("WITH_PERF_AUTOMATION_TESTS=1");
+                        break;
+                }
+            }
 
 			// By default, shadow source files for this target in the root OutputDirectory
 			GlobalCompileEnvironment.LocalShadowDirectory = LinkIntermediateDirectory;
@@ -3844,14 +3856,14 @@ namespace UnrealBuildTool
 				GlobalCompileEnvironment.Definitions.Add("WITH_PLUGIN_SUPPORT=0");
 			}
 
-			if (Rules.bWithPerfCounters)
-			{
-				GlobalCompileEnvironment.Definitions.Add("WITH_PERFCOUNTERS=1");
-			}
-			else
-			{
-				GlobalCompileEnvironment.Definitions.Add("WITH_PERFCOUNTERS=0");
-			}
+            if (Rules.bWithPerfCounters)
+            {
+                GlobalCompileEnvironment.Definitions.Add("WITH_PERFCOUNTERS=1");
+            }
+            else
+            {
+                GlobalCompileEnvironment.Definitions.Add("WITH_PERFCOUNTERS=0");
+            }
 
 			if (Rules.bUseLoggingInShipping)
 			{
@@ -3871,14 +3883,14 @@ namespace UnrealBuildTool
 				GlobalCompileEnvironment.Definitions.Add("WITH_LOGGING_TO_MEMORY=0");
 			}
 
-			if (Rules.bUseCacheFreedOSAllocs)
-			{
-				GlobalCompileEnvironment.Definitions.Add("USE_CACHE_FREED_OS_ALLOCS=1");
-			}
-			else
-			{
-				GlobalCompileEnvironment.Definitions.Add("USE_CACHE_FREED_OS_ALLOCS=0");
-			}
+            if (Rules.bUseCacheFreedOSAllocs)
+            {
+                GlobalCompileEnvironment.Definitions.Add("USE_CACHE_FREED_OS_ALLOCS=1");
+            }
+            else
+            {
+                GlobalCompileEnvironment.Definitions.Add("USE_CACHE_FREED_OS_ALLOCS=0");
+            }
 
 			if (Rules.bUseChecksInShipping)
 			{
@@ -3976,10 +3988,10 @@ namespace UnrealBuildTool
 			}
 		}
 
-		/// <summary>
-		/// Create a rules object for the given module, and set any default values for this target
-		/// </summary>
-		private ModuleRules CreateModuleRulesAndSetDefaults(string ModuleName, string ReferenceChain, out FileReference ModuleFileName)
+        /// <summary>
+        /// Create a rules object for the given module, and set any default values for this target
+        /// </summary>
+        private ModuleRules CreateModuleRulesAndSetDefaults(string ModuleName, string ReferenceChain, out FileReference ModuleFileName)
 		{
 			// Create the rules from the assembly
 			ModuleRules RulesObject = RulesAssembly.CreateModuleRules(ModuleName, Rules, ReferenceChain, out ModuleFileName);
@@ -4210,8 +4222,8 @@ namespace UnrealBuildTool
 					{
 						if(DirectoryReference.Exists(BaseSourceDirectory))
 						{
-							RulesObject.PublicIncludePaths.Add(NormalizeIncludePath(BaseSourceDirectory));
-						}
+						RulesObject.PublicIncludePaths.Add(NormalizeIncludePath(BaseSourceDirectory));
+					}
 					}
 
 					// Resolve private include paths against the project source root
@@ -4251,7 +4263,7 @@ namespace UnrealBuildTool
 					if (bDiscoverFiles && !UnrealBuildTool.IsUnderAnInstalledDirectory(ModuleFileName.Directory))
 					{
 						List<FileReference> SourceFilePaths = new List<FileReference>();
-							SourceFilePaths = SourceFileSearch.FindModuleSourceFiles(ModuleRulesFile: ModuleFileName);
+						SourceFilePaths = SourceFileSearch.FindModuleSourceFiles(ModuleRulesFile: ModuleFileName);
 						FoundSourceFiles = GetCPlusPlusFilesToBuild(SourceFilePaths, ModuleDirectory, Platform);
 					}
 				}
@@ -4507,22 +4519,22 @@ namespace UnrealBuildTool
 			return FilteredFileItems;
 		}
 
-		static bool ShouldIncludeNativizedAssets(DirectoryReference GameProjectDirectory)
-		{
-			ConfigHierarchy Config = ConfigCache.ReadHierarchy(ConfigHierarchyType.Game, GameProjectDirectory, BuildHostPlatform.Current.Platform);
-			if (Config != null)
-			{
-				// Determine whether or not the user has enabled nativization of Blueprint assets at cook time (default is 'Disabled')
-				string BlueprintNativizationMethod;
-				if (!Config.TryGetValue("/Script/UnrealEd.ProjectPackagingSettings", "BlueprintNativizationMethod", out BlueprintNativizationMethod))
-				{
-					BlueprintNativizationMethod = "Disabled";
-				}
-
-				return BlueprintNativizationMethod != "Disabled";
-			}
+        static bool ShouldIncludeNativizedAssets(DirectoryReference GameProjectDirectory)
+        {
+            ConfigHierarchy Config = ConfigCache.ReadHierarchy(ConfigHierarchyType.Game, GameProjectDirectory, BuildHostPlatform.Current.Platform);
+            if (Config != null)
+            {
+                // Determine whether or not the user has enabled nativization of Blueprint assets at cook time (default is 'Disabled')
+                string BlueprintNativizationMethod;
+                if (!Config.TryGetValue("/Script/UnrealEd.ProjectPackagingSettings", "BlueprintNativizationMethod", out BlueprintNativizationMethod))
+                {
+                    BlueprintNativizationMethod = "Disabled";
+                }
+
+                return BlueprintNativizationMethod != "Disabled";
+            }
 
-			return false;
-		}
+            return false;
+        }
 	}
 }
