UBT: Move code to enumerate Slate runtime dependencies into the Slate module. Doesn't need to be done inside core UBT.

#rnx

diff --git a/Engine/Source/Programs/UnrealBuildTool/Configuration/UEBuildTarget.cs b/Engine/Source/Programs/UnrealBuildTool/Configuration/UEBuildTarget.cs
--- a/Engine/Source/Programs/UnrealBuildTool/Configuration/UEBuildTarget.cs
+++ b/Engine/Source/Programs/UnrealBuildTool/Configuration/UEBuildTarget.cs
@@ -1601,19 +1601,6 @@ namespace UnrealBuildTool
 				}
 			}
 
-			// Add slate runtime dependencies
-			if (Rules.bUsesSlate)
-			{
-				AddRuntimeDependenciesFromDir(DirectoryReference.Combine(UnrealBuildTool.EngineDirectory, "Content", "Slate"), StagedFileType.UFS);
-				AddRuntimeDependenciesFromDir(DirectoryReference.Combine(UnrealBuildTool.EngineDirectory, "Content", "SlateDebug"), StagedFileType.UFS);
-
-				if (ProjectFile != null)
-				{
-					AddRuntimeDependenciesFromDir(DirectoryReference.Combine(ProjectDirectory, "Content", "Slate"), StagedFileType.UFS);
-					AddRuntimeDependenciesFromDir(DirectoryReference.Combine(ProjectDirectory, "Content", "SlateDebug"), StagedFileType.UFS);
-				}
-			}
-
 			// Find all the modules which are part of this target
 			HashSet<UEBuildModule> UniqueLinkedModules = new HashSet<UEBuildModule>();
 			foreach (UEBuildBinary Binary in Binaries)
@@ -1759,22 +1746,6 @@ namespace UnrealBuildTool
 			Receipt.AdditionalProperties.Add(new ReceiptProperty("SDK", ToolChain.GetSDKVersion()));
 		}
 
-		/// <summary>
-		/// Add the contents of a directory as runtime dependencies
-		/// </summary>
-		/// <param name="BaseDir">The base directory to enumerate files from</param>
-		/// <param name="Type">How the file should be staged</param>
-		void AddRuntimeDependenciesFromDir(DirectoryReference BaseDir, StagedFileType Type)
-		{
-			if (DirectoryReference.Exists(BaseDir))
-			{
-				foreach (FileReference File in DirectoryReference.EnumerateFiles(BaseDir, "*", SearchOption.AllDirectories))
-				{
-					Receipt.RuntimeDependencies.Add(File, Type);
-				}
-			}
-		}
-
 		/// <summary>
 		/// Try to recycle the build id from existing version manifests in the engine directory rather than generating a new one, if no engine binaries are being modified.
 		/// This allows sharing engine binaries when switching between projects and switching between UE4 and a game-specific project. Note that different targets may require
diff --git a/Engine/Source/Runtime/Slate/Slate.Build.cs b/Engine/Source/Runtime/Slate/Slate.Build.cs
--- a/Engine/Source/Runtime/Slate/Slate.Build.cs
+++ b/Engine/Source/Runtime/Slate/Slate.Build.cs
@@ -53,5 +53,18 @@ public class Slate : ModuleRules
 		{
 			AddEngineThirdPartyPrivateStaticDependencies(Target, "SDL2");
 		}
+
+		// Add slate runtime dependencies
+		if (Target.bUsesSlate)
+		{
+			RuntimeDependencies.Add("$(EngineDir)/Content/Slate/...", StagedFileType.UFS);
+			RuntimeDependencies.Add("$(EngineDir)/Content/SlateDebug/...", StagedFileType.UFS);
+
+			if (Target.ProjectFile != null)
+			{
+				RuntimeDependencies.Add("$(ProjectDir)/Content/Slate/...", StagedFileType.UFS);
+				RuntimeDependencies.Add("$(ProjectDir)/Content/SlateDebug/...", StagedFileType.UFS);
+			}
+		}
 	}
 }
