UBT: Fix ability to pass macros via the command line. Syntax is now "-define:FOO=BAR" rather than "-define FOO=BAR", to eliminate the need for non-standard argument parsing.

diff --git a/Engine/Source/Programs/UnrealBuildTool/Configuration/TargetRules.cs b/Engine/Source/Programs/UnrealBuildTool/Configuration/TargetRules.cs
--- a/Engine/Source/Programs/UnrealBuildTool/Configuration/TargetRules.cs
+++ b/Engine/Source/Programs/UnrealBuildTool/Configuration/TargetRules.cs
@@ -1048,7 +1048,7 @@ namespace UnrealBuildTool
 		/// Macros to define globally across the whole target.
 		/// </summary>
 		[RequiresUniqueBuildEnvironment]
-		[CommandLine("-Define", ValueAfterSpace = true)]
+		[CommandLine("-Define:")]
 		public List<string> GlobalDefinitions = new List<string>();
 
 		/// <summary>
diff --git a/Engine/Source/Programs/UnrealBuildTool/System/CommandLine.cs b/Engine/Source/Programs/UnrealBuildTool/System/CommandLine.cs
--- a/Engine/Source/Programs/UnrealBuildTool/System/CommandLine.cs
+++ b/Engine/Source/Programs/UnrealBuildTool/System/CommandLine.cs
@@ -1,4 +1,4 @@
-﻿// Copyright 1998-2018 Epic Games, Inc. All Rights Reserved.
+// Copyright 1998-2018 Epic Games, Inc. All Rights Reserved.
 
 using System;
 using System.Collections.Generic;
@@ -76,7 +76,7 @@ namespace UnrealBuildTool
 						}
 						else
 						{
-							if(FieldInfo.FieldType != typeof(bool) && !Attribute.ValueAfterSpace && Attribute.Value == null && !Prefix.EndsWith("="))
+							if(FieldInfo.FieldType != typeof(bool) && Attribute.Value == null && !Prefix.EndsWith("=") && !Prefix.EndsWith(":"))
 							{
 								Prefix = Prefix + "=";
 							}
@@ -94,7 +94,7 @@ namespace UnrealBuildTool
 				if(Argument.Length > 0 && Argument[0] == '-')
 				{
 					// Get the length of the argument prefix
-					int EqualsIdx = Argument.IndexOf('=');
+					int EqualsIdx = Argument.IndexOfAny(new char[]{ '=', ':' });
 					string Prefix = (EqualsIdx == -1)? Argument : Argument.Substring(0, EqualsIdx + 1);
 
 					// Check if there's a matching argument registered
@@ -119,17 +119,6 @@ namespace UnrealBuildTool
 						{
 							AssignValue(Parameter, Argument.Substring(EqualsIdx + 1), TargetObject, AssignedFieldToParameter);
 						}
-						else if(Parameter.Attribute.ValueAfterSpace)
-						{
-							if(Idx >= Arguments.Count)
-							{
-								Log.WriteLine(LogEventType.Warning, "Missing parameter for {0}", Parameter.Prefix);
-							}
-							else
-							{
-								AssignValue(Parameter, Arguments[NextIdx++], TargetObject, AssignedFieldToParameter);
-							}
-						}
 						else if(Parameter.FieldInfo.FieldType == typeof(bool))
 						{
 							AssignValue(Parameter, "true", TargetObject, AssignedFieldToParameter);
diff --git a/Engine/Source/Programs/UnrealBuildTool/System/CommandLineAttribute.cs b/Engine/Source/Programs/UnrealBuildTool/System/CommandLineAttribute.cs
--- a/Engine/Source/Programs/UnrealBuildTool/System/CommandLineAttribute.cs
+++ b/Engine/Source/Programs/UnrealBuildTool/System/CommandLineAttribute.cs
@@ -1,4 +1,4 @@
-﻿// Copyright 1998-2018 Epic Games, Inc. All Rights Reserved.
+// Copyright 1998-2018 Epic Games, Inc. All Rights Reserved.
 
 using System;
 using System.Collections.Generic;
@@ -25,11 +25,6 @@ namespace UnrealBuildTool
 		/// </summary>
 		public string Value = null;
 
-		/// <summary>
-		/// Indicates that the value of the argument is separated by a space rather than an equals symbol.
-		/// </summary>
-		public bool ValueAfterSpace = false;
-
 		/// <summary>
 		/// For collection types, specifies the separator character between multiple arguments
 		/// </summary>
