UBT: Change DebugGame configurations to output a separate executable rather than requiring a -Debug argument at runtime. Previous behavior was a common source of errors.

Engine modules are still shared between Development and DebugGame, but the launch module sets a flag in Core on startup indicating the game configuration.

diff --git a/Engine/Build/InstalledEngineBuild.xml b/Engine/Build/InstalledEngineBuild.xml
--- a/Engine/Build/InstalledEngineBuild.xml
+++ b/Engine/Build/InstalledEngineBuild.xml
@@ -120,6 +120,7 @@
 		</Node>
 		<Node Name="Compile UE4Editor Win64" Requires="Compile UnrealHeaderTool Win64" Produces="#UE4Editor Win64;#UE4Editor Win64 Unstripped;#UE4Editor Win64 Stripped;#UE4Editor Win64 Unsigned;#UE4Editor Win64 Signed;#UE4Editor Win64 Includes">
 			<Compile Target="UE4Editor" Platform="Win64" Configuration="Development" Tag="#UE4Editor Win64" Arguments="-precompile"/>
+			<Compile Target="UE4Editor" Platform="Win64" Configuration="DebugGame" Tag="#UE4Editor Win64" Arguments="-precompile"/>
 
 			<Do If="$(EmbedSrcSrvInfo)">
 				<!-- Embed source info into the PDB files. Should be done from this machine to ensure that paths are correct. -->
@@ -160,6 +161,7 @@
 		</Node>
 		<Node Name="Compile UE4Editor Mac" Requires="Compile UnrealHeaderTool Mac" Produces="#UE4Editor Mac;#UE4Editor Mac Unstripped;#UE4Editor Mac Stripped;#UE4Editor Mac Unsigned;#UE4Editor Mac Signed;#UE4Editor Mac Includes">
 			<Compile Target="UE4Editor" Platform="Mac" Configuration="Development" Tag="#UE4Editor Mac" Arguments="-precompile"/>
+			<Compile Target="UE4Editor" Platform="Mac" Configuration="DebugGame" Tag="#UE4Editor Mac" Arguments="-precompile"/>
 
 			<Tag Files="#UE4Editor Mac" Filter="$(MacStripFilter)" With="#UE4Editor Mac Unstripped"/>
 
@@ -192,6 +194,7 @@
 		</Node>
 		<Node Name="Compile UE4Editor Linux" Requires="Compile UnrealHeaderTool $(HostPlatform)" Produces="#UE4Editor Linux;#UE4Editor Linux Unstripped;#UE4Editor Linux Stripped;#UE4Editor Linux Includes">
 			<Compile Target="UE4Editor" Platform="Linux" Configuration="Development" Tag="#UE4Editor Linux" Arguments="-precompile"/>
+			<Compile Target="UE4Editor" Platform="Linux" Configuration="DebugGame" Tag="#UE4Editor Linux" Arguments="-precompile"/>
 
 			<Tag Files="#UE4Editor Linux" Filter="$(LinuxStripFilter)" With="#UE4Editor Linux Unstripped"/>
 
diff --git a/Engine/Programs/UnrealGameSync/UnrealGameSync.ini b/Engine/Programs/UnrealGameSync/UnrealGameSync.ini
--- a/Engine/Programs/UnrealGameSync/UnrealGameSync.ini
+++ b/Engine/Programs/UnrealGameSync/UnrealGameSync.ini
@@ -1,5 +1,6 @@
 [Options]
 UseFastModularVersioning=1
+DebugGameHasSeparateExecutable=1
 +ExcludeChanges=^Rebuild lightmaps\.$
 +ExcludeChanges=^Rebuild lightmaps & texture streaming\.$
 +ExcludeChanges=^Added collection '.*'\s*
diff --git a/Engine/Source/Developer/LauncherServices/Private/Launcher/LauncherUATTask.h b/Engine/Source/Developer/LauncherServices/Private/Launcher/LauncherUATTask.h
--- a/Engine/Source/Developer/LauncherServices/Private/Launcher/LauncherUATTask.h
+++ b/Engine/Source/Developer/LauncherServices/Private/Launcher/LauncherUATTask.h
@@ -73,11 +73,6 @@ protected:
 		if(EditorExe.Len() > 0)
 		{
 			UATCommandLine += FString::Printf(TEXT(" -ue4exe=\"%s\""), *EditorExe);
-
-			if (FApp::IsRunningDebug())
-			{
-				UATCommandLine += TEXT(" -UseDebugParamForEditorExe");
-			}
 		}
 
 		// specialized command arguments for this particular task
diff --git a/Engine/Source/Developer/LauncherServices/Private/Profiles/LauncherProfile.h b/Engine/Source/Developer/LauncherServices/Private/Profiles/LauncherProfile.h
--- a/Engine/Source/Developer/LauncherServices/Private/Profiles/LauncherProfile.h
+++ b/Engine/Source/Developer/LauncherServices/Private/Profiles/LauncherProfile.h
@@ -1092,7 +1092,6 @@ public:
 		Writer.WriteValue("nocompile", !IsBuildingUAT());
 		Writer.WriteValue("nocompileeditor", FApp::IsEngineInstalled());
 		Writer.WriteValue("ue4exe", GetEditorExe());
-		Writer.WriteValue("usedebugparamforeditorexe", FApp::IsRunningDebug());
 		Writer.WriteValue("utf8output", true);
 
 		// client configurations
@@ -1400,7 +1399,6 @@ public:
 		"skipcookingeditorcontent", "true/false"
 		"numcookerstospawn", "8"
 		"compressed", "true/false"
-		"usedebugparamforeditorexe", "true/false"
 		"iterativecooking", "true/false"
 		"skipcookonthefly", "true/false"
 		"cookall", "true/false"
diff --git a/Engine/Source/Editor/MainFrame/Private/Frame/MainFrameActions.cpp b/Engine/Source/Editor/MainFrame/Private/Frame/MainFrameActions.cpp
--- a/Engine/Source/Editor/MainFrame/Private/Frame/MainFrameActions.cpp
+++ b/Engine/Source/Editor/MainFrame/Private/Frame/MainFrameActions.cpp
@@ -402,11 +402,6 @@ void FMainFrameActionCallbacks::CookContent(const FName InPlatformInfoName)
 		OptionalParams += TEXT(" -iterate");
 	}
 
-	if (FApp::IsRunningDebug())
-	{
-		OptionalParams += TEXT(" -UseDebugParamForEditorExe");
-	}
-
 	FString ProjectPath = FPaths::IsProjectFilePathSet() ? FPaths::ConvertRelativePathToFull(FPaths::GetProjectFilePath()) : FPaths::RootDir() / FApp::GetProjectName() / FApp::GetProjectName() + TEXT(".uproject");
 	FString CommandLine = FString::Printf(TEXT("BuildCookRun %s%s -nop4 -project=\"%s\" -cook -skipstage -ue4exe=%s %s -utf8output"),
 		GetUATCompilationFlags(),
@@ -722,11 +717,6 @@ void FMainFrameActionCallbacks::PackageProject( const FName InPlatformInfoName )
 		OptionalParams += FString::Printf(TEXT(" -NumCookersToSpawn=%d"), NumCookers); 
 	}
 
-	if (FApp::IsRunningDebug())
-	{
-		OptionalParams += TEXT(" -UseDebugParamForEditorExe");
-	}
-
 	FString Configuration = FindObject<UEnum>(ANY_PACKAGE, TEXT("EProjectPackagingBuildConfigurations"))->GetNameStringByValue(PackagingSettings->BuildConfiguration);
 	Configuration = Configuration.Replace(TEXT("PPBC_"), TEXT(""));
 	if (Configuration.Right(6) == TEXT("Client"))
diff --git a/Engine/Source/Editor/UnrealEd/Private/PlayLevel.cpp b/Engine/Source/Editor/UnrealEd/Private/PlayLevel.cpp
--- a/Engine/Source/Editor/UnrealEd/Private/PlayLevel.cpp
+++ b/Engine/Source/Editor/UnrealEd/Private/PlayLevel.cpp
@@ -1330,10 +1330,6 @@ void UEditorEngine::PlayStandaloneLocalPc(FString MapNameOverride, FIntPoint* Wi
 	}
 
 	FString AdditionalParameters(TEXT(" -messaging -SessionName=\"Play in Standalone Game\""));
-	if (FApp::IsRunningDebug())
-	{
-		AdditionalParameters += TEXT(" -debug");
-	}
 
 	const FString PreviewGameLanguage = FTextLocalizationManager::Get().GetConfiguredGameLocalizationPreviewLanguage();
 	if (!PreviewGameLanguage.IsEmpty())
diff --git a/Engine/Source/Programs/AutomationTool/AutomationUtils/ProjectParams.cs b/Engine/Source/Programs/AutomationTool/AutomationUtils/ProjectParams.cs
--- a/Engine/Source/Programs/AutomationTool/AutomationUtils/ProjectParams.cs
+++ b/Engine/Source/Programs/AutomationTool/AutomationUtils/ProjectParams.cs
@@ -331,7 +331,6 @@ namespace AutomationTool
 			this.NumClients = InParams.NumClients;
             this.Compressed = InParams.Compressed;
 			this.AdditionalPakOptions = InParams.AdditionalPakOptions;
-            this.UseDebugParamForEditorExe = InParams.UseDebugParamForEditorExe;
 			this.Archive = InParams.Archive;
 			this.ArchiveDirectoryParam = InParams.ArchiveDirectoryParam;
 			this.ArchiveMetaData = InParams.ArchiveMetaData;
@@ -392,7 +391,6 @@ namespace AutomationTool
 			bool? Clean = null,
             bool? Compressed = null,
 			string AdditionalPakOptions = null,
-            bool? UseDebugParamForEditorExe = null,
             bool? IterativeCooking = null,
 			string IterateSharedCookedBuild = null,
 			bool? IterateSharedBuildUsePrecompiledExe = null,
@@ -633,7 +631,6 @@ namespace AutomationTool
             }
             this.Compressed = GetParamValueIfNotSpecified(Command, Compressed, this.Compressed, "compressed");
 			this.AdditionalPakOptions = ParseParamValueIfNotSpecified(Command, AdditionalPakOptions, "AdditionalPakOptions");
-            this.UseDebugParamForEditorExe = GetParamValueIfNotSpecified(Command, UseDebugParamForEditorExe, this.UseDebugParamForEditorExe, "UseDebugParamForEditorExe");
 			this.IterativeCooking = GetParamValueIfNotSpecified(Command, IterativeCooking, this.IterativeCooking, new string[] { "iterativecooking", "iterate" });
 			this.IterateSharedCookedBuild = GetParamValueIfNotSpecified(Command, false, false, "iteratesharedcookedbuild") ? "usesyncedbuild" : null;
 			this.IterateSharedCookedBuild = ParseParamValueIfNotSpecified(Command, IterateSharedCookedBuild, "IterateSharedCookedBuild", String.Empty);
@@ -1534,11 +1531,6 @@ namespace AutomationTool
 		/// </summary>
 		public string AdditionalPakOptions;
 
-        /// <summary>
-        /// put -debug on the editorexe commandline
-        /// </summary>
-        public bool UseDebugParamForEditorExe;
-
         /// <summary>
         /// Cook: Do not include a version number in the cooked content
         /// </summary>
@@ -2628,7 +2620,6 @@ namespace AutomationTool
 				CommandUtils.LogLog("ClientTargetPlatform={0}", string.Join(",", ClientTargetPlatforms));
 				CommandUtils.LogLog("Compressed={0}", Compressed);
 				CommandUtils.LogLog("AdditionalPakOptions={0}", AdditionalPakOptions);
-				CommandUtils.LogLog("UseDebugParamForEditorExe={0}", UseDebugParamForEditorExe);
 				CommandUtils.LogLog("CookOnTheFly={0}", CookOnTheFly);
 				CommandUtils.LogLog("CookOnTheFlyStreaming={0}", CookOnTheFlyStreaming);
 				CommandUtils.LogLog("UnversionedCookedContent={0}", UnversionedCookedContent);
diff --git a/Engine/Source/Programs/AutomationTool/Scripts/CookCommand.Automation.cs b/Engine/Source/Programs/AutomationTool/Scripts/CookCommand.Automation.cs
--- a/Engine/Source/Programs/AutomationTool/Scripts/CookCommand.Automation.cs
+++ b/Engine/Source/Programs/AutomationTool/Scripts/CookCommand.Automation.cs
@@ -1,4 +1,4 @@
-﻿// Copyright 1998-2018 Epic Games, Inc. All Rights Reserved.
+// Copyright 1998-2018 Epic Games, Inc. All Rights Reserved.
 using System;
 using System.Collections.Generic;
 using System.IO;
@@ -60,10 +60,6 @@ public partial class Project : CommandUtils
 				{
 					COTFCommandLine += " -iterate -iteratehash";
 				}
-				if (Params.UseDebugParamForEditorExe)
-				{
-					COTFCommandLine += " -debug";
-				}
 
 				var ServerLogFile = CombinePaths(LogFolderOutsideOfSandbox, "Server.log");
 				Platform ClientPlatformInst = Params.ClientTargetPlatformInstances[0];
@@ -155,10 +151,6 @@ public partial class Project : CommandUtils
 				{
 					CommandletParams += " -FastCook";
 				}
-                if (Params.UseDebugParamForEditorExe)
-                {
-                    CommandletParams += " -debug";
-                }
                 if (Params.Manifests)
                 {
                     CommandletParams += " -manifests";
diff --git a/Engine/Source/Programs/UnrealBuildTool/Configuration/UEBuildModuleCPP.cs b/Engine/Source/Programs/UnrealBuildTool/Configuration/UEBuildModuleCPP.cs
--- a/Engine/Source/Programs/UnrealBuildTool/Configuration/UEBuildModuleCPP.cs
+++ b/Engine/Source/Programs/UnrealBuildTool/Configuration/UEBuildModuleCPP.cs
@@ -1180,19 +1180,6 @@ namespace UnrealBuildTool
 				Result.Definitions.Add("UE_IS_ENGINE_MODULE=0");
 			}
 
-			// Switch the optimization flag if we're building a game module. Also pass the definition for building in DebugGame along (see ModuleManager.h for notes).
-			if (!Rules.bTreatAsEngineModule)
-			{
-				if (Target.Configuration == UnrealTargetConfiguration.DebugGame)
-				{
-					Result.Definitions.Add("UE_BUILD_DEVELOPMENT_WITH_DEBUGGAME=1");
-				}
-				else
-				{
-					Result.Definitions.Add("UE_BUILD_DEVELOPMENT_WITH_DEBUGGAME=0");
-				}
-			}
-
 			// For game modules, set the define for the project name. This will be used by the IMPLEMENT_PRIMARY_GAME_MODULE macro.
 			if (!Rules.bTreatAsEngineModule)
 			{
@@ -1247,19 +1234,6 @@ namespace UnrealBuildTool
 				CompileEnvironment.Definitions.Add("UE_IS_ENGINE_MODULE=0");
 			}
 
-			// Switch the optimization flag if we're building a game module. Also pass the definition for building in DebugGame along (see ModuleManager.h for notes).
-			if (!Rules.bTreatAsEngineModule)
-			{
-				if (Target.Configuration == UnrealTargetConfiguration.DebugGame)
-				{
-					CompileEnvironment.Definitions.Add("UE_BUILD_DEVELOPMENT_WITH_DEBUGGAME=1");
-				}
-				else
-				{
-					CompileEnvironment.Definitions.Add("UE_BUILD_DEVELOPMENT_WITH_DEBUGGAME=0");
-				}
-			}
-
 			// Add the module's private definitions.
 			CompileEnvironment.Definitions.AddRange(PublicDefinitions);
 
diff --git a/Engine/Source/Programs/UnrealBuildTool/Configuration/UEBuildTarget.cs b/Engine/Source/Programs/UnrealBuildTool/Configuration/UEBuildTarget.cs
--- a/Engine/Source/Programs/UnrealBuildTool/Configuration/UEBuildTarget.cs
+++ b/Engine/Source/Programs/UnrealBuildTool/Configuration/UEBuildTarget.cs
@@ -3182,7 +3182,7 @@ namespace UnrealBuildTool
 
 			// Get the configuration that this module will be built in. Engine modules compiled in DebugGame will use Development.
 			UnrealTargetConfiguration ModuleConfiguration = Configuration;
-			if (Configuration == UnrealTargetConfiguration.DebugGame && !RulesAssembly.IsGameModule(ModuleName))
+			if (Configuration == UnrealTargetConfiguration.DebugGame && !RulesAssembly.IsGameModule(ModuleName) && !ModuleName.Equals(Rules.LaunchModuleName, StringComparison.InvariantCultureIgnoreCase))
 			{
 				ModuleConfiguration = UnrealTargetConfiguration.Development;
 			}
@@ -3301,10 +3301,6 @@ namespace UnrealBuildTool
 		{
 			// Get the configuration for the executable. If we're building DebugGame, and this executable only contains engine modules, use the same name as development.
 			UnrealTargetConfiguration ExeConfiguration = Configuration;
-			if (Configuration == UnrealTargetConfiguration.DebugGame && !bIncludesGameModules)
-			{
-				ExeConfiguration = UnrealTargetConfiguration.Development;
-			}
 
 			// Build the binary path
 			DirectoryReference BinaryDirectory = DirectoryReference.Combine(BaseDirectory, "Binaries", Platform.ToString());
@@ -3610,11 +3606,22 @@ namespace UnrealBuildTool
 			// Create the launch module
 			UEBuildModuleCPP LaunchModule = FindOrCreateCppModuleByName(Rules.LaunchModuleName, TargetRulesFile.GetFileName());
 
+			// Get the intermediate directory for the launch module directory. This can differ from the standard engine intermediate directory because it is always configuration-specific.
+			DirectoryReference IntermediateDirectory;
+			if(LaunchModule.RulesFile.IsUnderDirectory(UnrealBuildTool.EngineDirectory) && !ShouldCompileMonolithic())
+			{
+				IntermediateDirectory = DirectoryReference.Combine(UnrealBuildTool.EngineDirectory, PlatformIntermediateFolder, AppName, Configuration.ToString());
+			}
+			else
+			{
+				IntermediateDirectory = ProjectIntermediateDirectory;
+			}
+
 			// Create the binary
 			UEBuildBinary Binary = new UEBuildBinary(
 				Type: Rules.bShouldCompileAsDLL? UEBuildBinaryType.DynamicLinkLibrary : UEBuildBinaryType.Executable,
 				OutputFilePaths: OutputPaths,
-				IntermediateDirectory: (!LaunchModule.RulesFile.IsUnderDirectory(UnrealBuildTool.EngineDirectory) || ShouldCompileMonolithic()) ? ProjectIntermediateDirectory : EngineIntermediateDirectory,
+				IntermediateDirectory: IntermediateDirectory,
 				bAllowExports: Rules.bHasExports,
 				PrimaryModule: LaunchModule,
 				bUsePrecompiled: LaunchModule.Rules.bUsePrecompiled && OutputPaths[0].IsUnderDirectory(UnrealBuildTool.EngineDirectory)
diff --git a/Engine/Source/Programs/UnrealBuildTool/ProjectFiles/CodeLite/CodeLiteProject.cs b/Engine/Source/Programs/UnrealBuildTool/ProjectFiles/CodeLite/CodeLiteProject.cs
--- a/Engine/Source/Programs/UnrealBuildTool/ProjectFiles/CodeLite/CodeLiteProject.cs
+++ b/Engine/Source/Programs/UnrealBuildTool/ProjectFiles/CodeLite/CodeLiteProject.cs
@@ -363,10 +363,6 @@ namespace UnrealBuildTool
 							if (TargetName != "UE4Editor" && GameProjectFile != "")
 							{
 								string commandArguments = "\"" + GameProjectFile + "\"" + " -game";
-								if (CurConf.ToString ().Contains ("Debug")) 
-								{
-									commandArguments += " -debug";
-								}
 								XAttribute CommandArguments = new XAttribute("CommandArguments", commandArguments);
 								CodeLiteConfigurationGeneral.Add(CommandArguments);
 							}
diff --git a/Engine/Source/Programs/UnrealBuildTool/ProjectFiles/VisualStudio/VCProject.cs b/Engine/Source/Programs/UnrealBuildTool/ProjectFiles/VisualStudio/VCProject.cs
--- a/Engine/Source/Programs/UnrealBuildTool/ProjectFiles/VisualStudio/VCProject.cs
+++ b/Engine/Source/Programs/UnrealBuildTool/ProjectFiles/VisualStudio/VCProject.cs
@@ -1201,7 +1201,7 @@ namespace UnrealBuildTool
 
 					// Make the output file path
 					FileReference NMakePath = FileReference.Combine(OutputDirectory, BaseExeName);
-					if (Configuration != TargetRulesObject.UndecoratedConfiguration && (Configuration != UnrealTargetConfiguration.DebugGame || bShouldCompileMonolithic))
+					if (Configuration != TargetRulesObject.UndecoratedConfiguration)
 					{
 						NMakePath += "-" + UBTPlatformName + "-" + UBTConfigurationName;
 					}
@@ -1311,11 +1311,6 @@ namespace UnrealBuildTool
 								DebugOptions += ProjectName;
 							}
 
-							if (Configuration == UnrealTargetConfiguration.Debug || Configuration == UnrealTargetConfiguration.DebugGame)
-							{
-								DebugOptions += " -debug";
-							}
-
 							VCUserFileContent.Append(
 								"		<LocalDebuggerCommandArguments>" + DebugOptions + "</LocalDebuggerCommandArguments>" + ProjectFileGenerator.NewLine
 								);
diff --git a/Engine/Source/Programs/UnrealBuildTool/ProjectFiles/VisualStudioCode/VSCodeProjectFileGenerator.cs b/Engine/Source/Programs/UnrealBuildTool/ProjectFiles/VisualStudioCode/VSCodeProjectFileGenerator.cs
--- a/Engine/Source/Programs/UnrealBuildTool/ProjectFiles/VisualStudioCode/VSCodeProjectFileGenerator.cs
+++ b/Engine/Source/Programs/UnrealBuildTool/ProjectFiles/VisualStudioCode/VSCodeProjectFileGenerator.cs
@@ -970,11 +970,6 @@ namespace UnrealBuildTool
 											OutFile.AddUnnamedField(InProject.Name);
 										}
 									}
-
-									if (BuildProduct.Config == UnrealTargetConfiguration.Debug || BuildProduct.Config == UnrealTargetConfiguration.DebugGame)
-									{
-										OutFile.AddUnnamedField("-debug");
-									}
 								}
 
 							}
diff --git a/Engine/Source/Runtime/Core/Private/Mac/MacPlatformProcess.cpp b/Engine/Source/Runtime/Core/Private/Mac/MacPlatformProcess.cpp
--- a/Engine/Source/Runtime/Core/Private/Mac/MacPlatformProcess.cpp
+++ b/Engine/Source/Runtime/Core/Private/Mac/MacPlatformProcess.cpp
@@ -69,7 +69,7 @@ FString FMacPlatformProcess::GenerateApplicationPath( const FString& AppName, EB
 	
 	FString PlatformName = TEXT("Mac");
 	FString ExecutableName = AppName;
-	if (BuildConfiguration != EBuildConfigurations::Development && BuildConfiguration != EBuildConfigurations::DebugGame)
+	if (BuildConfiguration != EBuildConfigurations::Development)
 	{
 		ExecutableName += FString::Printf(TEXT("-%s-%s"), *PlatformName, EBuildConfigurations::ToString(BuildConfiguration));
 	}
diff --git a/Engine/Source/Runtime/Core/Private/Misc/App.cpp b/Engine/Source/Runtime/Core/Private/Misc/App.cpp
--- a/Engine/Source/Runtime/Core/Private/Misc/App.cpp
+++ b/Engine/Source/Runtime/Core/Private/Misc/App.cpp
@@ -13,6 +13,10 @@ DEFINE_LOG_CATEGORY_STATIC(LogApp, Log, All);
 /* FApp static initialization
  *****************************************************************************/
 
+#if UE_BUILD_DEVELOPMENT
+bool FApp::bIsDebugGame = false;
+#endif
+
 FGuid FApp::InstanceId = FGuid::NewGuid();
 FGuid FApp::SessionId = FGuid::NewGuid();
 FString FApp::SessionName = FString();
@@ -81,13 +85,7 @@ EBuildConfigurations::Type FApp::GetBuildConfiguration()
 	return EBuildConfigurations::Debug;
 
 #elif UE_BUILD_DEVELOPMENT
-	// Detect DebugGame using an extern variable in monolithic configurations, or a command line argument in modular configurations.
-	#if IS_MONOLITHIC
-		extern const bool GIsDebugGame;
-		return GIsDebugGame? EBuildConfigurations::DebugGame : EBuildConfigurations::Development;
-	#else
-		return IsRunningDebug() ? EBuildConfigurations::DebugGame : EBuildConfigurations::Development;
-	#endif
+	return bIsDebugGame ? EBuildConfigurations::DebugGame : EBuildConfigurations::Development;
 
 #elif UE_BUILD_SHIPPING
 	return EBuildConfigurations::Shipping;
@@ -100,14 +98,12 @@ EBuildConfigurations::Type FApp::GetBuildConfiguration()
 #endif
 }
 
-bool FApp::IsRunningDebug()
+#if UE_BUILD_DEVELOPMENT
+void FApp::SetDebugGame(bool bInIsDebugGame)
 {
-	static FString RunConfig;
-	static const bool bHasRunConfig = FParse::Value(FCommandLine::Get(), TEXT("RunConfig="), RunConfig);
-	static const bool bRunningDebug = FParse::Param(FCommandLine::Get(), TEXT("debug"))
-	                                  || (bHasRunConfig && RunConfig.StartsWith(TEXT("Debug")));
-	return bRunningDebug;
+	bIsDebugGame = bInIsDebugGame;
 }
+#endif
 
 FString FApp::GetBuildDate()
 {
diff --git a/Engine/Source/Runtime/Core/Private/Modules/BuildVersion.cpp b/Engine/Source/Runtime/Core/Private/Modules/BuildVersion.cpp
--- a/Engine/Source/Runtime/Core/Private/Modules/BuildVersion.cpp
+++ b/Engine/Source/Runtime/Core/Private/Modules/BuildVersion.cpp
@@ -31,6 +31,21 @@ FString FBuildVersion::GetFileNameForCurrentExecutable()
 	{
 		AppExecutableName = AppExecutableName.Left(AppExecutableName.Len() - 4);
 	}
+#endif
+#if UE_BUILD_DEVELOPMENT
+	// For DebugGame builds we want to read the development version file because it's in an engine folder
+	if (FApp::GetBuildConfiguration() == EBuildConfigurations::DebugGame)
+	{
+		for (int32 Idx = 0; Idx < 2; Idx++)
+		{
+			int EndIdx;
+			if (!AppExecutableName.FindLastChar('-', EndIdx))
+			{
+				break;
+			}
+			AppExecutableName = AppExecutableName.Left(EndIdx);
+		}
+	}
 #endif
 	FString Result = FPlatformProcess::GetModulesDirectory() / AppExecutableName + TEXT(".version");
 	return Result;
diff --git a/Engine/Source/Runtime/Core/Private/Unix/UnixPlatformProcess.cpp b/Engine/Source/Runtime/Core/Private/Unix/UnixPlatformProcess.cpp
--- a/Engine/Source/Runtime/Core/Private/Unix/UnixPlatformProcess.cpp
+++ b/Engine/Source/Runtime/Core/Private/Unix/UnixPlatformProcess.cpp
@@ -369,7 +369,7 @@ FString FUnixPlatformProcess::GenerateApplicationPath( const FString& AppName, E
 	FString PlatformName = FPlatformProcess::GetBinariesSubdirectory();
 	FString ExecutablePath = FString::Printf(TEXT("../../../Engine/Binaries/%s/%s"), *PlatformName, *AppName);
 	
-	if (BuildConfiguration != EBuildConfigurations::Development && BuildConfiguration != EBuildConfigurations::DebugGame)
+	if (BuildConfiguration != EBuildConfigurations::Development)
 	{
 		ExecutablePath += FString::Printf(TEXT("-%s-%s"), *PlatformName, EBuildConfigurations::ToString(BuildConfiguration));
 	}
diff --git a/Engine/Source/Runtime/Core/Private/Windows/WindowsPlatformProcess.cpp b/Engine/Source/Runtime/Core/Private/Windows/WindowsPlatformProcess.cpp
--- a/Engine/Source/Runtime/Core/Private/Windows/WindowsPlatformProcess.cpp
+++ b/Engine/Source/Runtime/Core/Private/Windows/WindowsPlatformProcess.cpp
@@ -107,7 +107,7 @@ FString FWindowsPlatformProcess::GenerateApplicationPath( const FString& AppName
 	FString PlatformName = GetBinariesSubdirectory();
 	FString ExecutablePath = FString::Printf(TEXT("..\\..\\..\\Engine\\Binaries\\%s\\%s"), *PlatformName, *AppName);
 
-	if (BuildConfiguration != EBuildConfigurations::Development && BuildConfiguration != EBuildConfigurations::DebugGame)
+	if (BuildConfiguration != EBuildConfigurations::Development)
 	{
 		ExecutablePath += FString::Printf(TEXT("-%s-%s"), *PlatformName, EBuildConfigurations::ToString(BuildConfiguration));
 	}
diff --git a/Engine/Source/Runtime/Core/Public/Misc/App.h b/Engine/Source/Runtime/Core/Public/Misc/App.h
--- a/Engine/Source/Runtime/Core/Public/Misc/App.h
+++ b/Engine/Source/Runtime/Core/Public/Misc/App.h
@@ -41,10 +41,14 @@ public:
 	 */
 	static EBuildConfigurations::Type GetBuildConfiguration();
 
+#if UE_BUILD_DEVELOPMENT
 	/**
-	 * Gets whether the application is running with Debug game libraries (set from commandline)
+	 * For development configurations, sets whether the application should load DebugGame game modules.
+	 *
+	 * @param Whether we're running in debug game or not.
 	 */
-	static bool IsRunningDebug();
+	static void SetDebugGame(bool bIsDebugGame);
+#endif
 
 	/*
 	* Gets the unique version string for this build. This string is not assumed to have any particular format other being a unique identifier for the build.
@@ -676,6 +680,11 @@ public:
 
 private:
 
+#if UE_BUILD_DEVELOPMENT
+	/** The current build configuration */
+	static bool bIsDebugGame;
+#endif
+
 	/** Holds the instance identifier. */
 	static FGuid InstanceId;
 
diff --git a/Engine/Source/Runtime/Core/Public/Modules/ModuleManager.h b/Engine/Source/Runtime/Core/Public/Modules/ModuleManager.h
--- a/Engine/Source/Runtime/Core/Public/Modules/ModuleManager.h
+++ b/Engine/Source/Runtime/Core/Public/Modules/ModuleManager.h
@@ -797,24 +797,12 @@ class FDefaultGameModuleImpl
 		} \
 	} GEncryptionKeyRegistration;
 
-/**
- * Macro for declaring the GIsDebugGame variable for monolithic development builds. NB: This define, and the UE_BUILD_DEVELOPMENT_WITH_DEBUGGAME defines like it, should NEVER be 
- * directly used or defined for engine code, because it prevents sharing the same build products with the development build (important for Launcher build sizes). In modular builds, 
- * DebugGame modules will be loaded by specifying the -debug parameter on the command-line.
- */
-#if IS_MONOLITHIC && UE_BUILD_DEVELOPMENT
-	#define IMPLEMENT_DEBUGGAME() extern const bool GIsDebugGame = (UE_BUILD_DEVELOPMENT_WITH_DEBUGGAME != 0);
-#else
-	#define IMPLEMENT_DEBUGGAME()
-#endif 
-
 #if IS_PROGRAM
 
 	#if IS_MONOLITHIC
 		#define IMPLEMENT_APPLICATION( ModuleName, GameName ) \
 			/* For monolithic builds, we must statically define the game's name string (See Core.h) */ \
 			TCHAR GInternalProjectName[64] = TEXT( GameName ); \
-			IMPLEMENT_DEBUGGAME() \
 			IMPLEMENT_FOREIGN_ENGINE_DIR() \
 			IMPLEMENT_SIGNING_KEY_REGISTRATION() \
 			IMPLEMENT_ENCRYPTION_KEY_REGISTRATION() \
@@ -851,7 +839,6 @@ class FDefaultGameModuleImpl
 			TCHAR GInternalProjectName[64] = TEXT( PREPROCESSOR_TO_STRING(UE_PROJECT_NAME) ); \
 			/* Implement the GIsGameAgnosticExe variable (See Core.h). */ \
 			bool GIsGameAgnosticExe = false; \
-			IMPLEMENT_DEBUGGAME() \
 			IMPLEMENT_FOREIGN_ENGINE_DIR() \
 			IMPLEMENT_SIGNING_KEY_REGISTRATION() \
 			IMPLEMENT_ENCRYPTION_KEY_REGISTRATION() \
@@ -863,7 +850,6 @@ class FDefaultGameModuleImpl
 		#define IMPLEMENT_PRIMARY_GAME_MODULE( ModuleImplClass, ModuleName, DEPRECATED_GameName ) \
 			/* For monolithic builds, we must statically define the game's name string (See Core.h) */ \
 			TCHAR GInternalProjectName[64] = TEXT( PREPROCESSOR_TO_STRING(UE_PROJECT_NAME) ); \
-			IMPLEMENT_DEBUGGAME() \
 			PER_MODULE_BOILERPLATE \
 			IMPLEMENT_FOREIGN_ENGINE_DIR() \
 			IMPLEMENT_SIGNING_KEY_REGISTRATION() \
diff --git a/Engine/Source/Runtime/Launch/Launch.Build.cs b/Engine/Source/Runtime/Launch/Launch.Build.cs
--- a/Engine/Source/Runtime/Launch/Launch.Build.cs
+++ b/Engine/Source/Runtime/Launch/Launch.Build.cs
@@ -36,6 +36,16 @@ public class Launch : ModuleRules
 				"UtilityShaders",
 		});
 
+		// Set a macro allowing us to switch between debuggame/development configuration
+		if(Target.Configuration == UnrealTargetConfiguration.DebugGame)
+		{
+			PrivateDefinitions.Add("UE_BUILD_DEVELOPMENT_WITH_DEBUGGAME=1");
+		}
+		else
+		{
+			PrivateDefinitions.Add("UE_BUILD_DEVELOPMENT_WITH_DEBUGGAME=0");
+		}
+
 		// Enable the LauncherCheck module to be used for platforms that support the Launcher.
 		// Projects should set Target.bUseLauncherChecks in their Target.cs to enable the functionality.
 		if (Target.bUseLauncherChecks &&
diff --git a/Engine/Source/Runtime/Launch/Private/LaunchEngineLoop.cpp b/Engine/Source/Runtime/Launch/Private/LaunchEngineLoop.cpp
--- a/Engine/Source/Runtime/Launch/Private/LaunchEngineLoop.cpp
+++ b/Engine/Source/Runtime/Launch/Private/LaunchEngineLoop.cpp
@@ -988,6 +988,11 @@ int32 FEngineLoop::PreInit(const TCHAR* CmdLine)
 {
 	FMemory::SetupTLSCachesOnCurrentThread();
 
+	// Set the flag for whether we've build DebugGame instead of Development. The engine does not know this (whereas the launch module does) because it is always built in development.
+#ifdef UE_BUILD_DEVELOPMENT_WITH_DEBUGGAME
+	FApp::SetDebugGame(UE_BUILD_DEVELOPMENT_WITH_DEBUGGAME != 0);
+#endif
+
 	// disable/enable LLM based on commandline
 	LLM(FLowLevelMemTracker::Get().ProcessCommandLine(CmdLine));
 	LLM_SCOPE(ELLMTag::EnginePreInitMemory);
diff --git a/Engine/Source/Runtime/UE4Game/Private/UE4Game.cpp b/Engine/Source/Runtime/UE4Game/Private/UE4Game.cpp
--- a/Engine/Source/Runtime/UE4Game/Private/UE4Game.cpp
+++ b/Engine/Source/Runtime/UE4Game/Private/UE4Game.cpp
@@ -12,6 +12,5 @@ IMPLEMENT_MODULE(FDefaultModuleImpl, UE4Game);
 PER_MODULE_BOILERPLATE
 bool GIsGameAgnosticExe = true;
 TCHAR GInternalProjectName[64] = TEXT("");
-IMPLEMENT_DEBUGGAME()
 IMPLEMENT_FOREIGN_ENGINE_DIR()
 #endif
