UBT: Include contents of modified intermediate files in the log, to make it easier to debug hidden dependencies.

diff --git a/Engine/Source/Programs/UnrealBuildTool/System/FileItem.cs b/Engine/Source/Programs/UnrealBuildTool/System/FileItem.cs
--- a/Engine/Source/Programs/UnrealBuildTool/System/FileItem.cs
+++ b/Engine/Source/Programs/UnrealBuildTool/System/FileItem.cs
@@ -319,17 +319,7 @@ namespace UnrealBuildTool
 				string CurrentContents = Utils.ReadAllText(AbsolutePath.FullName);
 				if(!String.Equals(CurrentContents, Contents, StringComparison.InvariantCultureIgnoreCase))
 				{
-					try
-					{
-						FileReference PrevAbsolutePath = new FileReference(AbsolutePath.FullName + ".prev");
-						FileReference.Delete(PrevAbsolutePath);
-						FileReference.Move(AbsolutePath, PrevAbsolutePath);
-						Log.TraceLog("Updating {0} - contents have changed (previous version renamed to {1}).", AbsolutePath.FullName, PrevAbsolutePath);
-					}
-					catch
-					{
-						Log.TraceLog("Updating {0} - contents have changed (unable to rename). Previous:\n  {1}\nNew:\n  {2}", AbsolutePath.FullName, CurrentContents.Replace("\n", "\n  "), Contents.Replace("\n", "\n  "));
-					}
+					Log.TraceLog("Updating {0} - contents have changed. Previous:\n  {1}\nNew:\n  {2}", AbsolutePath.FullName, CurrentContents.Replace("\n", "\n  "), Contents.Replace("\n", "\n  "));
 					File.WriteAllText(AbsolutePath.FullName, Contents, GetEncodingForString(Contents));
 				}
 			}
