UBT: Fix assumption that bPrecompile is equivalent to bBuildAllModules. This is no longer the case; they are now controlled by separate options. Should fix CIS errors building the editor.

diff --git a/Engine/Source/Programs/UnrealBuildTool/Configuration/ModuleRules.cs b/Engine/Source/Programs/UnrealBuildTool/Configuration/ModuleRules.cs
--- a/Engine/Source/Programs/UnrealBuildTool/Configuration/ModuleRules.cs
+++ b/Engine/Source/Programs/UnrealBuildTool/Configuration/ModuleRules.cs
@@ -652,7 +652,7 @@ namespace UnrealBuildTool
 		/// </summary>
 		/// <param name="RulesFile">Path to the module rules file</param>
 		/// <returns>True if the module can be precompiled, false otherwise</returns>
-		internal bool CanPrecompile(FileReference RulesFile)
+		internal bool IsValidForTarget(FileReference RulesFile)
 		{
 			if(Type == ModuleRules.ModuleType.CPlusPlus)
 			{
diff --git a/Engine/Source/Programs/UnrealBuildTool/Configuration/UEBuildTarget.cs b/Engine/Source/Programs/UnrealBuildTool/Configuration/UEBuildTarget.cs
--- a/Engine/Source/Programs/UnrealBuildTool/Configuration/UEBuildTarget.cs
+++ b/Engine/Source/Programs/UnrealBuildTool/Configuration/UEBuildTarget.cs
@@ -2831,23 +2831,23 @@ namespace UnrealBuildTool
 		protected void AddAllValidModulesToTarget()
 		{
 			// Find all the modules that are part of the target
-			HashSet<string> PrecompiledModuleNames = new HashSet<string>();
+			HashSet<string> ValidModuleNames = new HashSet<string>();
 			foreach (UEBuildModuleCPP Module in Modules.Values.OfType<UEBuildModuleCPP>())
 			{
-				if(Module.Binary != null && Module.Rules.bPrecompile)
+				if(Module.Binary != null)
 				{
-					PrecompiledModuleNames.Add(Module.Name);
+					ValidModuleNames.Add(Module.Name);
 				}
 			}
 
-			// If we're precompiling a base engine target, create binaries for all the engine modules that are compatible with it.
+			// If we're compiling a base engine target, create binaries for all the engine modules that are compatible with it.
 			if (ProjectFile == null && TargetType != TargetType.Program)
 			{
 				// Find all the known module names in this assembly
 				List<string> ModuleNames = new List<string>();
 				RulesAssembly.GetAllModuleNames(ModuleNames);
 
-				// Find all the platform folders to exclude from the list of precompiled modules
+				// Find all the platform folders to exclude from the list of valid modules
 				List<FileSystemName> ExcludeFolders = new List<FileSystemName>();
 				foreach (UnrealTargetPlatform TargetPlatform in Enum.GetValues(typeof(UnrealTargetPlatform)))
 				{
@@ -2918,7 +2918,7 @@ namespace UnrealBuildTool
 					}
 				}
 
-				// Add all the plugin modules that need to be precompiled
+				// Add all the plugin modules that need to be compiled
 				List<PluginInfo> Plugins = RulesAssembly.EnumeratePlugins().ToList();
 				foreach(PluginInfo Plugin in Plugins)
 				{
@@ -2941,7 +2941,7 @@ namespace UnrealBuildTool
 					}
 				}
 
-				// Create rules for each remaining module, and check that it's set to be precompiled
+				// Create rules for each remaining module, and check that it's set to be compiled
 				foreach(string FilteredModuleName in FilteredModuleNames)
 				{
 					FileReference ModuleFileName = null;
@@ -2958,37 +2958,37 @@ namespace UnrealBuildTool
 					}
 
 					// Figure out if it can be precompiled
-					if (ModuleRules != null && ModuleRules.bPrecompile)
+					if (ModuleRules != null && ModuleRules.IsValidForTarget(ModuleFileName))
 					{
-						PrecompiledModuleNames.Add(FilteredModuleName);
+						ValidModuleNames.Add(FilteredModuleName);
 					}
 				}
 			}
 
 			// Now create all the precompiled modules, making sure they don't reference anything that's not in the precompiled set
-			HashSet<UEBuildModuleCPP> PrecompiledModules = new HashSet<UEBuildModuleCPP>();
-			foreach(string ModuleName in PrecompiledModuleNames)
+			HashSet<UEBuildModuleCPP> ValidModules = new HashSet<UEBuildModuleCPP>();
+			foreach(string ModuleName in ValidModuleNames)
 			{
-				const string PrecompileReferenceChain = "precompile option";
+				const string PrecompileReferenceChain = "allmodules option";
 				UEBuildModuleCPP Module = (UEBuildModuleCPP)FindOrCreateModuleByName(ModuleName, PrecompileReferenceChain);
 				Module.RecursivelyCreateModules(FindOrCreateModuleByName, PrecompileReferenceChain);
-				PrecompiledModules.Add(Module);
+				ValidModules.Add(Module);
 			}
 
 			// Make sure precompiled modules don't reference any non-precompiled modules
-			foreach(UEBuildModuleCPP PrecompiledModule in PrecompiledModules)
+			foreach(UEBuildModuleCPP ValidModule in ValidModules)
 			{
-				foreach(UEBuildModuleCPP ReferencedModule in PrecompiledModule.GetDependencies(false, true).OfType<UEBuildModuleCPP>())
+				foreach(UEBuildModuleCPP ReferencedModule in ValidModule.GetDependencies(false, true).OfType<UEBuildModuleCPP>())
 				{
-					if(!PrecompiledModules.Contains(ReferencedModule))
+					if(!ValidModules.Contains(ReferencedModule))
 					{
-						Log.TraceWarning("Precompiled module '{0}' will not be usable without non-precompiled module '{1}'.", PrecompiledModule.Name, ReferencedModule.Name);
+						Log.TraceWarning("Module '{0}' is not usable without module '{1}', which is not valid for this target.", ValidModule.Name, ReferencedModule.Name);
 					}
 				}
 			}
 
 			// Make sure every module is built
-			foreach(UEBuildModuleCPP Module in PrecompiledModules)
+			foreach(UEBuildModuleCPP Module in ValidModules)
 			{
 				if(Module.Binary == null)
 				{
diff --git a/Engine/Source/Programs/UnrealBuildTool/System/RulesAssembly.cs b/Engine/Source/Programs/UnrealBuildTool/System/RulesAssembly.cs
--- a/Engine/Source/Programs/UnrealBuildTool/System/RulesAssembly.cs
+++ b/Engine/Source/Programs/UnrealBuildTool/System/RulesAssembly.cs
@@ -365,11 +365,6 @@ namespace UnrealBuildTool
 				}
 				Constructor.Invoke(RulesObject, new object[] { Target });
 
-				// Update the precompiled flags
-				if(RulesObject.bPrecompile && !RulesObject.CanPrecompile(ModuleFileName))
-				{
-					RulesObject.bPrecompile = false;
-				}
 				return RulesObject;
 			}
 			catch (Exception Ex)
