UBT: Remove indirect calls to determine extensions for object files and precompiled headers. The toolchain knows the correct convention for the platform.

#rnx

diff --git a/Engine/Source/Programs/UnrealBuildTool/Configuration/UEBuildBinary.cs b/Engine/Source/Programs/UnrealBuildTool/Configuration/UEBuildBinary.cs
--- a/Engine/Source/Programs/UnrealBuildTool/Configuration/UEBuildBinary.cs
+++ b/Engine/Source/Programs/UnrealBuildTool/Configuration/UEBuildBinary.cs
@@ -31,16 +31,6 @@ namespace UnrealBuildTool
 		/// A static library (.lib or .a)
 		/// </summary>
 		StaticLibrary,
-
-		/// <summary>
-		/// An object file (.obj or .o)
-		/// </summary>
-		Object,
-
-		/// <summary>
-		/// A precompiled header (.pch or .gch)
-		/// </summary>
-		PrecompiledHeader
 	}
 
 	/// <summary>
diff --git a/Engine/Source/Programs/UnrealBuildTool/Platform/Android/AndroidToolChain.cs b/Engine/Source/Programs/UnrealBuildTool/Platform/Android/AndroidToolChain.cs
--- a/Engine/Source/Programs/UnrealBuildTool/Platform/Android/AndroidToolChain.cs
+++ b/Engine/Source/Programs/UnrealBuildTool/Platform/Android/AndroidToolChain.cs
@@ -1251,7 +1251,7 @@ namespace UnrealBuildTool
 
 			string BasePCHName = "";
 			UEBuildPlatform BuildPlatform = UEBuildPlatform.GetBuildPlatform(UEBuildPlatform.CPPTargetPlatformToUnrealTargetPlatform(CompileEnvironment.Platform));
-			string PCHExtension = BuildPlatform.GetBinaryExtension(UEBuildBinaryType.PrecompiledHeader);
+			string PCHExtension = ".gch";
 			if (CompileEnvironment.PrecompiledHeaderAction == PrecompiledHeaderAction.Include)
 			{
 				BasePCHName = RemoveArchName(CompileEnvironment.PrecompiledHeaderFile.AbsolutePath).Replace(PCHExtension, "");
@@ -1386,12 +1386,15 @@ namespace UnrealBuildTool
 								CompileAction.PrerequisiteItems.Add(ArchPrecompiledHeaderFile);
 							}
 
-							string ObjectFileExtension = BuildPlatform.GetBinaryExtension(UEBuildBinaryType.Object);
-
+							string ObjectFileExtension;
 							if(CompileEnvironment.AdditionalArguments != null && CompileEnvironment.AdditionalArguments.Contains("-emit-llvm"))
 							{
 								ObjectFileExtension = ".bc";
 							}
+							else
+							{
+								ObjectFileExtension = ".o";
+							}
 
 							// Add the object file to the produced item list.
 							FileItem ObjectFile = FileItem.GetItemByFileReference(
diff --git a/Engine/Source/Programs/UnrealBuildTool/Platform/Android/UEBuildAndroid.cs b/Engine/Source/Programs/UnrealBuildTool/Platform/Android/UEBuildAndroid.cs
--- a/Engine/Source/Programs/UnrealBuildTool/Platform/Android/UEBuildAndroid.cs
+++ b/Engine/Source/Programs/UnrealBuildTool/Platform/Android/UEBuildAndroid.cs
@@ -159,10 +159,6 @@ namespace UnrealBuildTool
 					return ".so";
 				case UEBuildBinaryType.StaticLibrary:
 					return ".a";
-				case UEBuildBinaryType.Object:
-					return ".o";
-				case UEBuildBinaryType.PrecompiledHeader:
-					return ".gch";
 			}
 			return base.GetBinaryExtension(InBinaryType);
 		}
diff --git a/Engine/Source/Programs/UnrealBuildTool/Platform/HTML5/HTML5ToolChain.cs b/Engine/Source/Programs/UnrealBuildTool/Platform/HTML5/HTML5ToolChain.cs
--- a/Engine/Source/Programs/UnrealBuildTool/Platform/HTML5/HTML5ToolChain.cs
+++ b/Engine/Source/Programs/UnrealBuildTool/Platform/HTML5/HTML5ToolChain.cs
@@ -390,14 +390,8 @@ namespace UnrealBuildTool
 
 				// Add the source file path to the command-line.
 				string FileArguments = string.Format(" \"{0}\"", SourceFile.AbsolutePath);
-				string ObjectFileExtension = UEBuildPlatform.GetBuildPlatform(UnrealTargetPlatform.HTML5).GetBinaryExtension(UEBuildBinaryType.Object);
 				// Add the object file to the produced item list.
-				FileItem ObjectFile = FileItem.GetItemByFileReference(
-					FileReference.Combine(
-						OutputDir,
-						Path.GetFileName(SourceFile.AbsolutePath) + ObjectFileExtension
-						)
-					);
+				FileItem ObjectFile = FileItem.GetItemByFileReference(FileReference.Combine(OutputDir, Path.GetFileName(SourceFile.AbsolutePath) + ".bc"));
 				CompileAction.ProducedItems.Add(ObjectFile);
 				FileArguments += string.Format(" -o \"{0}\"", ObjectFile.AbsolutePath);
 
diff --git a/Engine/Source/Programs/UnrealBuildTool/Platform/HTML5/UEBuildHTML5.cs b/Engine/Source/Programs/UnrealBuildTool/Platform/HTML5/UEBuildHTML5.cs
--- a/Engine/Source/Programs/UnrealBuildTool/Platform/HTML5/UEBuildHTML5.cs
+++ b/Engine/Source/Programs/UnrealBuildTool/Platform/HTML5/UEBuildHTML5.cs
@@ -95,10 +95,6 @@ namespace UnrealBuildTool
 					return ".js";
 				case UEBuildBinaryType.StaticLibrary:
 					return ".bc";
-				case UEBuildBinaryType.Object:
-					return ".bc";
-				case UEBuildBinaryType.PrecompiledHeader:
-					return ".gch";
 			}
 
 			return base.GetBinaryExtension(InBinaryType);
diff --git a/Engine/Source/Programs/UnrealBuildTool/Platform/IOS/IOSToolChain.cs b/Engine/Source/Programs/UnrealBuildTool/Platform/IOS/IOSToolChain.cs
--- a/Engine/Source/Programs/UnrealBuildTool/Platform/IOS/IOSToolChain.cs
+++ b/Engine/Source/Programs/UnrealBuildTool/Platform/IOS/IOSToolChain.cs
@@ -644,10 +644,7 @@ namespace UnrealBuildTool
 
 			if (CompileEnvironment.PrecompiledHeaderAction == PrecompiledHeaderAction.Include)
 			{
-				// Add the precompiled header file's path to the include path so GCC can find it.
-				// This needs to be before the other include paths to ensure GCC uses it instead of the source header file.
-				string PrecompiledFileExtension = UEBuildPlatform.GetBuildPlatform(UnrealTargetPlatform.IOS).GetBinaryExtension(UEBuildBinaryType.PrecompiledHeader);
-				PCHArguments += string.Format(" -include \"{0}\"", ConvertPath(CompileEnvironment.PrecompiledHeaderFile.AbsolutePath.Replace(PrecompiledFileExtension, "")));
+				PCHArguments += string.Format(" -include \"{0}\"", ConvertPath(CompileEnvironment.PrecompiledHeaderIncludeFilename.FullName));
 			}
 
 			foreach(FileItem ForceIncludeFile in CompileEnvironment.ForceIncludeFiles)
@@ -741,14 +738,8 @@ namespace UnrealBuildTool
 				string OutputFilePath = null;
 				if (CompileEnvironment.PrecompiledHeaderAction == PrecompiledHeaderAction.Create)
 				{
-					string PrecompiledFileExtension = UEBuildPlatform.GetBuildPlatform(UnrealTargetPlatform.IOS).GetBinaryExtension(UEBuildBinaryType.PrecompiledHeader);
 					// Add the precompiled header file to the produced item list.
-					FileItem PrecompiledHeaderFile = FileItem.GetItemByFileReference(
-						FileReference.Combine(
-							OutputDir,
-							Path.GetFileName(SourceFile.AbsolutePath) + PrecompiledFileExtension
-							)
-						);
+					FileItem PrecompiledHeaderFile = FileItem.GetItemByFileReference(FileReference.Combine(OutputDir, Path.GetFileName(SourceFile.AbsolutePath) + ".gch"));
 
 					FileItem RemotePrecompiledHeaderFile = LocalToRemoteFileItem(PrecompiledHeaderFile, false);
 					CompileAction.ProducedItems.Add(RemotePrecompiledHeaderFile);
@@ -764,14 +755,8 @@ namespace UnrealBuildTool
 						CompileAction.bIsUsingPCH = true;
 						CompileAction.PrerequisiteItems.Add(CompileEnvironment.PrecompiledHeaderFile);
 					}
-					string ObjectFileExtension = UEBuildPlatform.GetBuildPlatform(UnrealTargetPlatform.IOS).GetBinaryExtension(UEBuildBinaryType.Object);
 					// Add the object file to the produced item list.
-					FileItem ObjectFile = FileItem.GetItemByFileReference(
-						FileReference.Combine(
-							OutputDir,
-							Path.GetFileName(SourceFile.AbsolutePath) + ObjectFileExtension
-							)
-						);
+					FileItem ObjectFile = FileItem.GetItemByFileReference(FileReference.Combine(OutputDir, Path.GetFileName(SourceFile.AbsolutePath) + ".o"));
 
 					FileItem RemoteObjectFile = LocalToRemoteFileItem(ObjectFile, false);
 					CompileAction.ProducedItems.Add(RemoteObjectFile);
diff --git a/Engine/Source/Programs/UnrealBuildTool/Platform/IOS/UEBuildIOS.cs b/Engine/Source/Programs/UnrealBuildTool/Platform/IOS/UEBuildIOS.cs
--- a/Engine/Source/Programs/UnrealBuildTool/Platform/IOS/UEBuildIOS.cs
+++ b/Engine/Source/Programs/UnrealBuildTool/Platform/IOS/UEBuildIOS.cs
@@ -606,10 +606,6 @@ namespace UnrealBuildTool
 					return "";
 				case UEBuildBinaryType.StaticLibrary:
 					return ".a";
-				case UEBuildBinaryType.Object:
-					return ".o";
-				case UEBuildBinaryType.PrecompiledHeader:
-					return ".gch";
 			}
 			return base.GetBinaryExtension(InBinaryType);
 		}
diff --git a/Engine/Source/Programs/UnrealBuildTool/Platform/Linux/LinuxToolChain.cs b/Engine/Source/Programs/UnrealBuildTool/Platform/Linux/LinuxToolChain.cs
--- a/Engine/Source/Programs/UnrealBuildTool/Platform/Linux/LinuxToolChain.cs
+++ b/Engine/Source/Programs/UnrealBuildTool/Platform/Linux/LinuxToolChain.cs
@@ -1013,10 +1013,7 @@ namespace UnrealBuildTool
 
 			if (CompileEnvironment.PrecompiledHeaderAction == PrecompiledHeaderAction.Include)
 			{
-				// Add the precompiled header file's path to the include path so Clang can find it.
-				// This needs to be before the other include paths to ensure Clang uses it instead of the source header file.
-				string PrecompiledFileExtension = BuildPlatform.GetBinaryExtension(UEBuildBinaryType.PrecompiledHeader);
-				PCHArguments += string.Format(" -include \"{0}\"", CompileEnvironment.PrecompiledHeaderFile.AbsolutePath.Replace(PrecompiledFileExtension, "").Replace('\\', '/'));
+				PCHArguments += string.Format(" -include \"{0}\"", CompileEnvironment.PrecompiledHeaderIncludeFilename.FullName.Replace('\\', '/'));
 			}
 
 			foreach(FileItem ForceIncludeFile in CompileEnvironment.ForceIncludeFiles)
@@ -1084,14 +1081,8 @@ namespace UnrealBuildTool
 
 				if (CompileEnvironment.PrecompiledHeaderAction == PrecompiledHeaderAction.Create)
 				{
-					string PrecompiledFileExtension = BuildPlatform.GetBinaryExtension(UEBuildBinaryType.PrecompiledHeader);
 					// Add the precompiled header file to the produced item list.
-					FileItem PrecompiledHeaderFile = FileItem.GetItemByFileReference(
-						FileReference.Combine(
-							OutputDir,
-							Path.GetFileName(SourceFile.AbsolutePath) + PrecompiledFileExtension
-							)
-						);
+					FileItem PrecompiledHeaderFile = FileItem.GetItemByFileReference(FileReference.Combine(OutputDir, Path.GetFileName(SourceFile.AbsolutePath) + ".gch"));
 
 					CompileAction.ProducedItems.Add(PrecompiledHeaderFile);
 					Result.PrecompiledHeaderFile = PrecompiledHeaderFile;
@@ -1107,14 +1098,8 @@ namespace UnrealBuildTool
 						CompileAction.PrerequisiteItems.Add(CompileEnvironment.PrecompiledHeaderFile);
 					}
 
-					string ObjectFileExtension = BuildPlatform.GetBinaryExtension(UEBuildBinaryType.Object);
 					// Add the object file to the produced item list.
-					FileItem ObjectFile = FileItem.GetItemByFileReference(
-						FileReference.Combine(
-							OutputDir,
-							Path.GetFileName(SourceFile.AbsolutePath) + ObjectFileExtension
-							)
-						);
+					FileItem ObjectFile = FileItem.GetItemByFileReference(FileReference.Combine(OutputDir, Path.GetFileName(SourceFile.AbsolutePath) + ".o"));
 					CompileAction.ProducedItems.Add(ObjectFile);
 					Result.ObjectFiles.Add(ObjectFile);
 
diff --git a/Engine/Source/Programs/UnrealBuildTool/Platform/Linux/UEBuildLinux.cs b/Engine/Source/Programs/UnrealBuildTool/Platform/Linux/UEBuildLinux.cs
--- a/Engine/Source/Programs/UnrealBuildTool/Platform/Linux/UEBuildLinux.cs
+++ b/Engine/Source/Programs/UnrealBuildTool/Platform/Linux/UEBuildLinux.cs
@@ -230,10 +230,6 @@ namespace UnrealBuildTool
 					return "";
 				case UEBuildBinaryType.StaticLibrary:
 					return ".a";
-				case UEBuildBinaryType.Object:
-					return ".o";
-				case UEBuildBinaryType.PrecompiledHeader:
-					return ".gch";
 			}
 			return base.GetBinaryExtension(InBinaryType);
 		}
diff --git a/Engine/Source/Programs/UnrealBuildTool/Platform/Mac/MacToolChain.cs b/Engine/Source/Programs/UnrealBuildTool/Platform/Mac/MacToolChain.cs
--- a/Engine/Source/Programs/UnrealBuildTool/Platform/Mac/MacToolChain.cs
+++ b/Engine/Source/Programs/UnrealBuildTool/Platform/Mac/MacToolChain.cs
@@ -429,9 +429,8 @@ namespace UnrealBuildTool
 			{
 				// Add the precompiled header file's path to the include path so GCC can find it.
 				// This needs to be before the other include paths to ensure GCC uses it instead of the source header file.
-				string PrecompiledFileExtension = UEBuildPlatform.GetBuildPlatform(UnrealTargetPlatform.Mac).GetBinaryExtension(UEBuildBinaryType.PrecompiledHeader);
 				PCHArguments.Append(" -include \"");
-				PCHArguments.Append(CompileEnvironment.PrecompiledHeaderFile.AbsolutePath.Replace(PrecompiledFileExtension, ""));
+				PCHArguments.Append(CompileEnvironment.PrecompiledHeaderIncludeFilename);
 				PCHArguments.Append("\"");
 			}
 
@@ -533,14 +532,8 @@ namespace UnrealBuildTool
 				string OutputFilePath = null;
 				if (CompileEnvironment.PrecompiledHeaderAction == PrecompiledHeaderAction.Create)
 				{
-					string PrecompiledHeaderExtension = UEBuildPlatform.GetBuildPlatform(UnrealTargetPlatform.Mac).GetBinaryExtension(UEBuildBinaryType.PrecompiledHeader);
 					// Add the precompiled header file to the produced item list.
-					FileItem PrecompiledHeaderFile = FileItem.GetItemByFileReference(
-						FileReference.Combine(
-							OutputDir,
-							Path.GetFileName(SourceFile.AbsolutePath) + PrecompiledHeaderExtension
-							)
-						);
+					FileItem PrecompiledHeaderFile = FileItem.GetItemByFileReference(FileReference.Combine(OutputDir, Path.GetFileName(SourceFile.AbsolutePath) + ".gch"));
 
 					FileItem RemotePrecompiledHeaderFile = LocalToRemoteFileItem(PrecompiledHeaderFile, false);
 					CompileAction.ProducedItems.Add(RemotePrecompiledHeaderFile);
@@ -556,14 +549,8 @@ namespace UnrealBuildTool
 						CompileAction.bIsUsingPCH = true;
 						CompileAction.PrerequisiteItems.Add(CompileEnvironment.PrecompiledHeaderFile);
 					}
-					string ObjectFileExtension = UEBuildPlatform.GetBuildPlatform(UnrealTargetPlatform.Mac).GetBinaryExtension(UEBuildBinaryType.Object);
 					// Add the object file to the produced item list.
-					FileItem ObjectFile = FileItem.GetItemByFileReference(
-						FileReference.Combine(
-							OutputDir,
-							Path.GetFileName(SourceFile.AbsolutePath) + ObjectFileExtension
-							)
-						);
+					FileItem ObjectFile = FileItem.GetItemByFileReference(FileReference.Combine(OutputDir, Path.GetFileName(SourceFile.AbsolutePath) + ".o"));
 
 					FileItem RemoteObjectFile = LocalToRemoteFileItem(ObjectFile, false);
 					CompileAction.ProducedItems.Add(RemoteObjectFile);
diff --git a/Engine/Source/Programs/UnrealBuildTool/Platform/Mac/UEBuildMac.cs b/Engine/Source/Programs/UnrealBuildTool/Platform/Mac/UEBuildMac.cs
--- a/Engine/Source/Programs/UnrealBuildTool/Platform/Mac/UEBuildMac.cs
+++ b/Engine/Source/Programs/UnrealBuildTool/Platform/Mac/UEBuildMac.cs
@@ -206,10 +206,6 @@ namespace UnrealBuildTool
 					return "";
 				case UEBuildBinaryType.StaticLibrary:
 					return ".a";
-				case UEBuildBinaryType.Object:
-					return ".o";
-				case UEBuildBinaryType.PrecompiledHeader:
-					return ".gch";
 			}
 			return base.GetBinaryExtension(InBinaryType);
 		}
@@ -228,8 +224,6 @@ namespace UnrealBuildTool
 				case UEBuildBinaryType.Executable:
 					return Target.bUsePDBFiles ? new string[] {".dSYM"} : new string[] {};
 				case UEBuildBinaryType.StaticLibrary:
-				case UEBuildBinaryType.Object:
-				case UEBuildBinaryType.PrecompiledHeader:
 				default:
 					return new string [] {};
 			}
diff --git a/Engine/Source/Programs/UnrealBuildTool/Platform/Windows/UEBuildWindows.cs b/Engine/Source/Programs/UnrealBuildTool/Platform/Windows/UEBuildWindows.cs
--- a/Engine/Source/Programs/UnrealBuildTool/Platform/Windows/UEBuildWindows.cs
+++ b/Engine/Source/Programs/UnrealBuildTool/Platform/Windows/UEBuildWindows.cs
@@ -1256,10 +1256,6 @@ namespace UnrealBuildTool
 					return ".exe";
 				case UEBuildBinaryType.StaticLibrary:
 					return ".lib";
-				case UEBuildBinaryType.Object:
-					return ".obj";
-				case UEBuildBinaryType.PrecompiledHeader:
-					return ".pch";
 			}
 			return base.GetBinaryExtension(InBinaryType);
 		}
diff --git a/Engine/Source/Programs/UnrealBuildTool/Platform/Windows/VCToolChain.cs b/Engine/Source/Programs/UnrealBuildTool/Platform/Windows/VCToolChain.cs
--- a/Engine/Source/Programs/UnrealBuildTool/Platform/Windows/VCToolChain.cs
+++ b/Engine/Source/Programs/UnrealBuildTool/Platform/Windows/VCToolChain.cs
@@ -912,12 +912,11 @@ namespace UnrealBuildTool
 					// path -- it might be a private PCH header and we need to make sure that its found!
 					AddIncludePath(FileArguments, SourceFile.Location.Directory);
 
-					string PrecompiledFileExtension = UEBuildPlatform.GetBuildPlatform(UnrealTargetPlatform.Win64).GetBinaryExtension(UEBuildBinaryType.PrecompiledHeader);
 					// Add the precompiled header file to the produced items list.
 					FileItem PrecompiledHeaderFile = FileItem.GetItemByFileReference(
 						FileReference.Combine(
 							OutputDir,
-							Path.GetFileName(SourceFile.AbsolutePath) + PrecompiledFileExtension
+							Path.GetFileName(SourceFile.AbsolutePath) + ".pch"
 							)
 						);
 					CompileAction.ProducedItems.Add(PrecompiledHeaderFile);
@@ -975,9 +974,8 @@ namespace UnrealBuildTool
 
 				if (bEmitsObjectFile)
 				{
-					string ObjectFileExtension = UEBuildPlatform.GetBuildPlatform(UnrealTargetPlatform.Win64).GetBinaryExtension(UEBuildBinaryType.Object);
 					// Add the object file to the produced item list.
-					string ObjectLeafFilename = Path.GetFileName(SourceFile.AbsolutePath) + ObjectFileExtension;
+					string ObjectLeafFilename = Path.GetFileName(SourceFile.AbsolutePath) + ".obj";
 					FileItem ObjectFile = FileItem.GetItemByFileReference(FileReference.Combine(OutputDir, ObjectLeafFilename));
 					if (Target.WindowsPlatform.ObjSrcMapFile != null)
 					{
