UBT: Remove action graph visualization code; no longer used.

#rnx

diff --git a/Engine/Source/Programs/UnrealBuildTool/System/ActionGraph.cs b/Engine/Source/Programs/UnrealBuildTool/System/ActionGraph.cs
--- a/Engine/Source/Programs/UnrealBuildTool/System/ActionGraph.cs
+++ b/Engine/Source/Programs/UnrealBuildTool/System/ActionGraph.cs
@@ -1265,348 +1265,5 @@ namespace UnrealBuildTool
 		{
 			return IsCPPImplementationFile(FileItem) || IsCPPIncludeFile(FileItem) || IsCPPResourceFile(FileItem);
 		}
-
-
-
-		/// <summary>
-		/// Types of action graph visualizations that we can emit
-		/// </summary>
-		public enum ActionGraphVisualizationType
-		{
-			OnlyActions,
-			ActionsWithFiles,
-			ActionsWithFilesAndHeaders,
-			OnlyFilesAndHeaders,
-			OnlyCPlusPlusFilesAndHeaders
-		}
-
-
-
-		/// <summary>
-		/// Saves the action graph (and include dependency network) to a graph gile
-		/// </summary>
-		/// <param name="Headers">Cached header data for this target</param>
-		/// <param name="Filename">File name to emit</param>
-		/// <param name="Description">Description to be stored in graph metadata</param>
-		/// <param name="VisualizationType">Type of graph to create</param>
-		/// <param name="Actions">All actions</param>
-		/// <param name="IncludeCompileActions">True if we should include compile actions.  If disabled, only the static link actions will be shown, which is useful to see module relationships</param>
-		public static void SaveActionGraphVisualization(CPPHeaders Headers, string Filename, string Description, ActionGraphVisualizationType VisualizationType, List<Action> Actions, bool IncludeCompileActions = true)
-		{
-			// True if we should include individual files in the graph network, or false to include only the build actions
-			bool IncludeFiles = VisualizationType != ActionGraphVisualizationType.OnlyActions;
-			bool OnlyIncludeCPlusPlusFiles = VisualizationType == ActionGraphVisualizationType.OnlyCPlusPlusFilesAndHeaders;
-
-			// True if want to show actions in the graph, otherwise we're only showing files
-			bool IncludeActions = VisualizationType != ActionGraphVisualizationType.OnlyFilesAndHeaders && VisualizationType != ActionGraphVisualizationType.OnlyCPlusPlusFilesAndHeaders;
-
-			// True if C++ header dependencies should be expanded into the graph, or false to only have .cpp files
-			bool ExpandCPPHeaderDependencies = IncludeFiles && (VisualizationType == ActionGraphVisualizationType.ActionsWithFilesAndHeaders || VisualizationType == ActionGraphVisualizationType.OnlyFilesAndHeaders || VisualizationType == ActionGraphVisualizationType.OnlyCPlusPlusFilesAndHeaders);
-
-			DateTime TimerStartTime = DateTime.UtcNow;
-
-			List<GraphNode> GraphNodes = new List<GraphNode>();
-
-			Dictionary<FileItem, GraphNode> FileToGraphNodeMap = new Dictionary<FileItem, GraphNode>();
-
-			// Filter our list of actions
-			List<Action> FilteredActions = new List<Action>();
-			{
-				for (int ActionIndex = 0; ActionIndex < Actions.Count; ++ActionIndex)
-				{
-					Action Action = Actions[ActionIndex];
-
-					if (!IncludeActions || IncludeCompileActions || (Action.ActionType != ActionType.Compile))
-					{
-						FilteredActions.Add(Action);
-					}
-				}
-			}
-
-
-			HashSet<FileItem> FilesToCreateNodesFor = new HashSet<FileItem>();
-			for (int ActionIndex = 0; ActionIndex < FilteredActions.Count; ++ActionIndex)
-			{
-				Action Action = FilteredActions[ActionIndex];
-
-				if (IncludeActions)
-				{
-					GraphNode GraphNode = new GraphNode()
-					{
-						Id = GraphNodes.Count,
-
-						// Don't bother including "Link" text if we're excluding compile actions
-						Label = IncludeCompileActions ? (Action.ActionType.ToString() + " " + Action.StatusDescription) : Action.StatusDescription
-					};
-
-					switch (Action.ActionType)
-					{
-						case ActionType.BuildProject:
-							GraphNode.Color = new GraphColor() { R = 0.3f, G = 1.0f, B = 1.0f, A = 1.0f };
-							GraphNode.Size = 1.1f;
-							break;
-
-						case ActionType.Compile:
-							GraphNode.Color = new GraphColor() { R = 0.3f, G = 1.0f, B = 0.3f, A = 1.0f };
-							break;
-
-						case ActionType.Link:
-							GraphNode.Color = new GraphColor() { R = 0.3f, G = 0.3f, B = 1.0f, A = 1.0f };
-							GraphNode.Size = 1.2f;
-							break;
-					}
-
-					GraphNodes.Add(GraphNode);
-				}
-
-				if (IncludeFiles)
-				{
-					foreach (FileItem ProducedFileItem in Action.ProducedItems)
-					{
-						if (!OnlyIncludeCPlusPlusFiles || IsCPPFile(ProducedFileItem))
-						{
-							FilesToCreateNodesFor.Add(ProducedFileItem);
-						}
-					}
-
-					foreach (FileItem PrerequisiteFileItem in Action.PrerequisiteItems)
-					{
-						if (!OnlyIncludeCPlusPlusFiles || IsCPPFile(PrerequisiteFileItem))
-						{
-							FilesToCreateNodesFor.Add(PrerequisiteFileItem);
-						}
-					}
-				}
-			}
-
-
-			Dictionary<FileItem, List<FileItem>> OverriddenPrerequisites = new Dictionary<FileItem, List<FileItem>>();
-
-			// Determine the average size of all of the C++ source files
-			Int64 AverageCPPFileSize;
-			{
-				Int64 TotalFileSize = 0;
-				int CPPFileCount = 0;
-				foreach (FileItem FileItem in FilesToCreateNodesFor)
-				{
-					if (IsCPPFile(FileItem))
-					{
-						++CPPFileCount;
-						TotalFileSize += new FileInfo(FileItem.AbsolutePath).Length;
-					}
-				}
-
-				if (CPPFileCount > 0)
-				{
-					AverageCPPFileSize = TotalFileSize / CPPFileCount;
-				}
-				else
-				{
-					AverageCPPFileSize = 1;
-				}
-			}
-
-			foreach (FileItem FileItem in FilesToCreateNodesFor)
-			{
-				GraphNode FileGraphNode = new GraphNode()
-				{
-					Id = GraphNodes.Count,
-					Label = Path.GetFileName(FileItem.AbsolutePath)
-				};
-
-				if (FileItem.AbsolutePath.EndsWith(".h", StringComparison.InvariantCultureIgnoreCase) ||
-					FileItem.AbsolutePath.EndsWith(".inl", StringComparison.InvariantCultureIgnoreCase))
-				{
-					// Header file
-					FileGraphNode.Color = new GraphColor() { R = 0.9f, G = 0.2f, B = 0.9f, A = 1.0f };
-				}
-				else if (FileItem.AbsolutePath.EndsWith(".cpp", StringComparison.InvariantCultureIgnoreCase) ||
-						 FileItem.AbsolutePath.EndsWith(".c", StringComparison.InvariantCultureIgnoreCase) ||
-						 FileItem.AbsolutePath.EndsWith(".mm", StringComparison.InvariantCultureIgnoreCase))
-				{
-					// C++ file
-					FileGraphNode.Color = new GraphColor() { R = 1.0f, G = 1.0f, B = 0.3f, A = 1.0f };
-				}
-				else
-				{
-					// Other file
-					FileGraphNode.Color = new GraphColor() { R = 0.4f, G = 0.4f, B = 0.1f, A = 1.0f };
-				}
-
-				// Set the size of the file node based on the size of the file on disk
-				bool bIsCPPFile = IsCPPFile(FileItem);
-				if (bIsCPPFile)
-				{
-					float MinNodeSize = 0.25f;
-					float MaxNodeSize = 2.0f;
-					long FileSize = new FileInfo(FileItem.AbsolutePath).Length;
-					float FileSizeScale = (float)((double)FileSize / (double)AverageCPPFileSize);
-
-					float SourceFileSizeScaleFactor = 0.1f;		// How much to make nodes for files bigger or larger based on their difference from the average file's size
-					FileGraphNode.Size = Math.Min(Math.Max(1.0f + SourceFileSizeScaleFactor * FileSizeScale, MinNodeSize), MaxNodeSize);
-				}
-
-				//@todo: Testing out attribute support.  Replace with an attribute that is actually useful!
-				//if( FileItem.PrecompiledHeaderIncludeFilename != null )
-				//{ 
-				//FileGraphNode.Attributes[ "PCHFile" ] = Path.GetFileNameWithoutExtension( FileItem.PrecompiledHeaderIncludeFilename );
-				//}
-
-				FileToGraphNodeMap[FileItem] = FileGraphNode;
-				GraphNodes.Add(FileGraphNode);
-
-				if (ExpandCPPHeaderDependencies && bIsCPPFile)
-				{
-					List<DependencyInclude> DirectlyIncludedFilenames = Headers.GetDirectIncludeDependencies(FileItem, bOnlyCachedDependencies: false);
-
-					// Resolve the included file name to an actual file.
-					List<FileItem> DirectlyIncludedFiles =
-						DirectlyIncludedFilenames
-						.Where(DirectlyIncludedFilename => (DirectlyIncludedFilename.IncludeResolvedNameIfSuccessful != null))
-						.Select(DirectlyIncludedFilename => DirectlyIncludedFilename.IncludeResolvedNameIfSuccessful)
-						// Skip same include over and over (.inl files)
-						.Distinct()
-						.Select(FileItem.GetItemByFileReference)
-						.ToList();
-
-					OverriddenPrerequisites[FileItem] = DirectlyIncludedFiles;
-				}
-			}
-
-
-			// Connect everything together
-			List<GraphEdge> GraphEdges = new List<GraphEdge>();
-
-			if (IncludeActions)
-			{
-				for (int ActionIndex = 0; ActionIndex < FilteredActions.Count; ++ActionIndex)
-				{
-					Action Action = FilteredActions[ActionIndex];
-					GraphNode ActionGraphNode = GraphNodes[ActionIndex];
-
-					List<FileItem> ActualPrerequisiteItems = Action.PrerequisiteItems;
-					if (IncludeFiles && ExpandCPPHeaderDependencies && Action.ActionType == ActionType.Compile)
-					{
-						// The first prerequisite is always the .cpp file to compile
-						FileItem CPPFile = Action.PrerequisiteItems[0];
-						if (!IsCPPFile(CPPFile))
-						{
-							throw new BuildException("Was expecting a C++ file as the first prerequisite for a Compile action");
-						}
-
-						ActualPrerequisiteItems = new List<FileItem>();
-						ActualPrerequisiteItems.Add(CPPFile);
-					}
-
-
-					foreach (FileItem PrerequisiteFileItem in ActualPrerequisiteItems)
-					{
-						if (IncludeFiles)
-						{
-							GraphNode PrerequisiteFileGraphNode;
-							if (FileToGraphNodeMap.TryGetValue(PrerequisiteFileItem, out PrerequisiteFileGraphNode))
-							{
-								// Connect a file our action is dependent on, to our action itself
-								GraphEdge GraphEdge = new GraphEdge()
-								{
-									Id = GraphEdges.Count,
-									Source = PrerequisiteFileGraphNode,
-									Target = ActionGraphNode,
-								};
-
-								GraphEdges.Add(GraphEdge);
-							}
-							else
-							{
-								// Not a file we were tracking
-								// Log.TraceInformation( "Unknown file: " + PrerequisiteFileItem.AbsolutePath );
-							}
-						}
-						else if (PrerequisiteFileItem.ProducingAction != null)
-						{
-							// Not showing files, so connect the actions together
-							int ProducingActionIndex = FilteredActions.IndexOf(PrerequisiteFileItem.ProducingAction);
-							if (ProducingActionIndex != -1)
-							{
-								GraphNode SourceGraphNode = GraphNodes[ProducingActionIndex];
-
-								GraphEdge GraphEdge = new GraphEdge()
-								{
-									Id = GraphEdges.Count,
-									Source = SourceGraphNode,
-									Target = ActionGraphNode,
-								};
-
-								GraphEdges.Add(GraphEdge);
-							}
-							else
-							{
-								// Our producer action was filtered out
-							}
-						}
-					}
-
-					foreach (FileItem ProducedFileItem in Action.ProducedItems)
-					{
-						if (IncludeFiles)
-						{
-							if (!OnlyIncludeCPlusPlusFiles || IsCPPFile(ProducedFileItem))
-							{
-								GraphNode ProducedFileGraphNode = FileToGraphNodeMap[ProducedFileItem];
-
-								GraphEdge GraphEdge = new GraphEdge()
-								{
-									Id = GraphEdges.Count,
-									Source = ActionGraphNode,
-									Target = ProducedFileGraphNode,
-								};
-
-								GraphEdges.Add(GraphEdge);
-							}
-						}
-					}
-				}
-			}
-
-			if (IncludeFiles && ExpandCPPHeaderDependencies)
-			{
-				// Fill in overridden prerequisites
-				foreach (KeyValuePair<FileItem, List<FileItem>> FileAndPrerequisites in OverriddenPrerequisites)
-				{
-					FileItem FileItem = FileAndPrerequisites.Key;
-					List<FileItem> FilePrerequisites = FileAndPrerequisites.Value;
-
-					GraphNode FileGraphNode = FileToGraphNodeMap[FileItem];
-					foreach (FileItem PrerequisiteFileItem in FilePrerequisites)
-					{
-						GraphNode PrerequisiteFileGraphNode;
-						if (FileToGraphNodeMap.TryGetValue(PrerequisiteFileItem, out PrerequisiteFileGraphNode))
-						{
-							GraphEdge GraphEdge = new GraphEdge()
-							{
-								Id = GraphEdges.Count,
-								Source = PrerequisiteFileGraphNode,
-								Target = FileGraphNode,
-							};
-
-							GraphEdges.Add(GraphEdge);
-						}
-						else
-						{
-							// Some other header that we don't track directly
-							//Log.TraceInformation( "File not known: " + PrerequisiteFileItem.AbsolutePath );
-						}
-					}
-				}
-			}
-
-			GraphVisualization.WriteGraphFile(Filename, Description, GraphNodes, GraphEdges);
-
-			if (UnrealBuildTool.bPrintPerformanceInfo)
-			{
-				TimeSpan TimerDuration = DateTime.UtcNow - TimerStartTime;
-				Log.TraceInformation("Generating and saving ActionGraph took " + TimerDuration.TotalSeconds + "s");
-			}
-		}
-	};
+	}
 }
diff --git a/Engine/Source/Programs/UnrealBuildTool/UnrealBuildTool.cs b/Engine/Source/Programs/UnrealBuildTool/UnrealBuildTool.cs
--- a/Engine/Source/Programs/UnrealBuildTool/UnrealBuildTool.cs
+++ b/Engine/Source/Programs/UnrealBuildTool/UnrealBuildTool.cs
@@ -1494,26 +1494,6 @@ namespace UnrealBuildTool
 
 							// Update mapping of the target name to the list of UObject modules in this target
 							TargetNameToUObjectModules[Target.GetTargetName()] = TargetUObjectModules;
-
-							if ((BuildConfiguration.bXGEExport && BuildConfiguration.bGenerateManifest) || (!ProjectFileGenerator.bGenerateProjectFiles && !BuildConfiguration.bGenerateManifest && !BuildConfiguration.bCleanProject))
-							{
-								// Generate an action graph if we were asked to do that.  The graph generation needs access to the include dependency cache, so
-								// we generate it before saving and cleaning that up.
-								if (GeneratingActionGraph)
-								{
-									// The graph generation feature currently only works with a single target at a time.  This is because of how we need access
-									// to include dependencies for the target, but those aren't kept around as we process multiple targets
-									if (TargetSettings.Count != 1)
-									{
-										throw new BuildException("ERROR: The '-graph' option only works with a single target at a time");
-									}
-									////ActionGraph.FinalizeActionGraph();
-									List<Action> ActionsToExecute = ActionGraph.AllActions;
-
-									ActionGraph.ActionGraphVisualizationType VisualizationType = ActionGraph.ActionGraphVisualizationType.OnlyCPlusPlusFilesAndHeaders;
-									ActionGraph.SaveActionGraphVisualization(TargetToHeaders[Target], FileReference.Combine(UnrealBuildTool.EngineDirectory, "Intermediate", "Build", Target.GetTargetName() + ".gexf").FullName, Target.GetTargetName(), VisualizationType, ActionsToExecute);
-								}
-							}
 						}
 					}
 
diff --git a/Engine/Source/Programs/UnrealBuildTool/UnrealBuildTool.csproj b/Engine/Source/Programs/UnrealBuildTool/UnrealBuildTool.csproj
--- a/Engine/Source/Programs/UnrealBuildTool/UnrealBuildTool.csproj
+++ b/Engine/Source/Programs/UnrealBuildTool/UnrealBuildTool.csproj
@@ -350,9 +350,6 @@
     <Compile Include="ToolChain\RemoteToolChain.cs" />
     <Compile Include="ToolChain\UEToolChain.cs" />
     <Compile Include="System\XGE.cs" />
-    <Compile Include="Utilities\GraphVisualization.cs">
-      <SubType>Code</SubType>
-    </Compile>
     <Compile Include="Platform\Windows\UEBuildWindows.cs" />
     <Compile Include="Platform\Windows\VCEnvironment.cs" />
     <Compile Include="Platform\Windows\VCToolChain.cs" />
@@ -409,7 +406,5 @@
       <Name>DotNETUtilities</Name>
     </ProjectReference>
   </ItemGroup>
-  <ItemGroup>
-    <Folder Include="ProjectFiles\VisualStudioMac\" />
-  </ItemGroup>
+  <ItemGroup />
 </Project>
\ No newline at end of file
diff --git a/Engine/Source/Programs/UnrealBuildTool/Utilities/GraphVisualization.cs b/Engine/Source/Programs/UnrealBuildTool/Utilities/GraphVisualization.cs
deleted file mode 100644
