UBT: Remove telemetry code. This has never proved useful for analyzing performance due to the number of incidental factors that affect build times (eg. number of files being compiled).

#rnx

diff --git a/Engine/Source/Programs/UnrealBuildTool/System/ActionGraph.cs b/Engine/Source/Programs/UnrealBuildTool/System/ActionGraph.cs
--- a/Engine/Source/Programs/UnrealBuildTool/System/ActionGraph.cs
+++ b/Engine/Source/Programs/UnrealBuildTool/System/ActionGraph.cs
@@ -481,21 +481,6 @@ namespace UnrealBuildTool
 				ExecutorName = Executor.Name;
 				Result = Executor.ExecuteActions(ActionsToExecute, BuildConfiguration.bLogDetailedActionStats);
 
-				if(Telemetry.IsAvailable() && !BuildConfiguration.bXGEExport)
-				{
-					Telemetry.SendEvent("BuildTime",
-						"ExecutorType", ExecutorName,
-						"OS", Environment.OSVersion.ToString(),
-						"MachineName", Environment.MachineName,
-						"NumLogicalCores", Environment.ProcessorCount.ToString(),
-						"Targets", TargetInfoForTelemetry,
-						"NumActions", ActionsToExecute.Count.ToString(),
-						"NumCompileActions", ActionsToExecute.Count(x => x.ActionType == ActionType.Compile).ToString(),
-						"NumPchCompileActions", ActionsToExecute.Count(x => x.ActionType == ActionType.Compile && x.ProducedItems.Any(y => y.Location.HasExtension(".pch") || y.Location.HasExtension(".gch"))).ToString(),
-						"NumLinkActions", ActionsToExecute.Count(x => x.ActionType == ActionType.Link).ToString(),
-						"ElapsedTime", (DateTime.UtcNow - StartTime).TotalSeconds.ToString());
-				}
-
 				if (!BuildConfiguration.bXGEExport)
 				{
 					// Verify the link outputs were created (seems to happen with Win64 compiles)
diff --git a/Engine/Source/Programs/UnrealBuildTool/System/Telemetry.cs b/Engine/Source/Programs/UnrealBuildTool/System/Telemetry.cs
deleted file mode 100644
diff --git a/Engine/Source/Programs/UnrealBuildTool/System/XGE.cs b/Engine/Source/Programs/UnrealBuildTool/System/XGE.cs
--- a/Engine/Source/Programs/UnrealBuildTool/System/XGE.cs
+++ b/Engine/Source/Programs/UnrealBuildTool/System/XGE.cs
@@ -132,39 +132,6 @@ namespace UnrealBuildTool
 				string XGETaskFilePath = FileReference.Combine(UnrealBuildTool.EngineDirectory, "Intermediate", "Build", "XGETasks.xml").FullName;
 				WriteTaskFile(Actions, XGETaskFilePath, ProgressWriter.bWriteMarkup, false);
 
-				// Try to execute the XGE tasks, and if XGE is available, skip the local execution fallback.
-				if (Telemetry.IsAvailable())
-				{
-					try
-					{
-						// @todo: find a way to report that for other host platforms
-						if (BuildHostPlatform.Current.Platform == UnrealTargetPlatform.Win64)
-						{
-							const string BuilderKey = "HKEY_LOCAL_MACHINE\\SOFTWARE\\Wow6432Node\\Xoreax\\IncrediBuild\\Builder";
-
-							string CPUUtilization = Registry.GetValue(BuilderKey, "ForceCPUCount", "").ToString();
-							string AvoidTaskExecutionOnLocalMachine = Registry.GetValue(BuilderKey, "AvoidLocalExec", "").ToString();
-							string RestartRemoteProcessesOnLocalMachine = Registry.GetValue(BuilderKey, "AllowDoubleTargets", "").ToString();
-							string LimitMaxNumberOfCores = Registry.GetValue(BuilderKey, "MaxHelpers", "").ToString();
-							string WriteOutputToDiskInBackground = Registry.GetValue(BuilderKey, "LazyOutputWriter_Beta", "").ToString();
-							string MaxConcurrentPDBs = Registry.GetValue(BuilderKey, "MaxConcurrentPDBs", "").ToString();
-							string EnabledAsHelper = Registry.GetValue(BuilderKey, "LastEnabled", "").ToString();
-
-							Telemetry.SendEvent("XGESettings.2",
-								"CPUUtilization", CPUUtilization,
-								"AvoidTaskExecutionOnLocalMachine", AvoidTaskExecutionOnLocalMachine,
-								"RestartRemoteProcessesOnLocalMachine", RestartRemoteProcessesOnLocalMachine,
-								"LimitMaxNumberOfCores", LimitMaxNumberOfCores,
-								"WriteOutputToDiskInBackground", WriteOutputToDiskInBackground,
-								"MaxConcurrentPDBs", MaxConcurrentPDBs,
-								"EnabledAsHelper", EnabledAsHelper);
-						}
-					}
-					catch
-					{
-					}
-				}
-
 				XGEResult = ExecuteTaskFileWithProgressMarkup(XGETaskFilePath, Actions.Count);
 			}
 			return XGEResult;
diff --git a/Engine/Source/Programs/UnrealBuildTool/UnrealBuildTool.cs b/Engine/Source/Programs/UnrealBuildTool/UnrealBuildTool.cs
--- a/Engine/Source/Programs/UnrealBuildTool/UnrealBuildTool.cs
+++ b/Engine/Source/Programs/UnrealBuildTool/UnrealBuildTool.cs
@@ -632,8 +632,6 @@ namespace UnrealBuildTool
 					Log.TraceVerbose("UnrealBuildTool (DEBUG OUTPUT MODE)");
 					Log.TraceVerbose("Command-line: {0}", String.Join(" ", Arguments));
 
-					Telemetry.Initialize();
-
 					string GameName = null;
 					bool bSpecificModulesOnly = false;
 
@@ -956,8 +954,6 @@ namespace UnrealBuildTool
 				}
 			}
 
-			Telemetry.Shutdown();
-
 			// Print some performance info
 			Log.TraceVerbose("Execution time: {0}", (DateTime.UtcNow - StartTime).TotalSeconds);
 
diff --git a/Engine/Source/Programs/UnrealBuildTool/UnrealBuildTool.csproj b/Engine/Source/Programs/UnrealBuildTool/UnrealBuildTool.csproj
--- a/Engine/Source/Programs/UnrealBuildTool/UnrealBuildTool.csproj
+++ b/Engine/Source/Programs/UnrealBuildTool/UnrealBuildTool.csproj
@@ -302,11 +302,7 @@
     </When>
   </Choose>
   <Choose>
-    <When Condition="Exists('System\NotForLicensees\TelemetryProviderET.cs')">
-      <ItemGroup>
-        <Compile Include="System\NotForLicensees\TelemetryProviderET.cs" />
-      </ItemGroup>
-    </When>
+    <When Condition="Exists('System\NotForLicensees\TelemetryProviderET.cs')" />
   </Choose>
   <ItemGroup>
     <Compile Include="..\DotNETCommon\MetaData.cs">
@@ -342,7 +338,6 @@
     <Compile Include="System\RPCUtilHelper.cs" />
     <Compile Include="System\RulesCompiler.cs" />
     <Compile Include="System\SourceFileSearch.cs" />
-    <Compile Include="System\Telemetry.cs" />
     <Compile Include="System\UEPlatformProjectGenerator.cs" />
     <Compile Include="System\Unity.cs" />
     <Compile Include="System\UProjectInfo.cs" />
