UBT: Remove unqiue id from all actions in the action graph. This is only used for printing debug info in the case of a (rare) cycle in the action graph, so just look it up when needed.

#rnx

diff --git a/Engine/Source/Programs/UnrealBuildTool/System/ActionGraph.cs b/Engine/Source/Programs/UnrealBuildTool/System/ActionGraph.cs
--- a/Engine/Source/Programs/UnrealBuildTool/System/ActionGraph.cs
+++ b/Engine/Source/Programs/UnrealBuildTool/System/ActionGraph.cs
@@ -129,11 +129,6 @@ namespace UnrealBuildTool
 		/// Preparation only (not serialized)
 		///
 
-		/// <summary>
-		/// Unique action identifier.  Used for displaying helpful info about detected cycles in the graph.
-		/// </summary>
-		public readonly int UniqueId;
-
 		/// <summary>
 		/// Total number of actions depending on this one.
 		/// </summary>
@@ -160,10 +155,9 @@ namespace UnrealBuildTool
 		public DateTimeOffset EndTime = DateTimeOffset.MinValue;
 
 
-		public Action(ActionType InActionType, int InUniqueId)
+		public Action(ActionType InActionType)
 		{
 			ActionType = InActionType;
-			UniqueId = InUniqueId;
 		}
 
 		public Action(SerializationInfo SerializationInfo, StreamingContext StreamingContext)
@@ -276,11 +270,6 @@ namespace UnrealBuildTool
 
 	class ActionGraph
 	{
-		/// <summary>
-		/// Unique Id given to all actions added to this graph
-		/// </summary>
-		public int NextUniqueId;
-
 		/// <summary>
 		/// List of all the actions
 		/// </summary>
@@ -289,13 +278,11 @@ namespace UnrealBuildTool
 		public ActionGraph()
 		{
 			XmlConfig.ApplyTo(this);
-			NextUniqueId = 0;
 		}
 
 		public Action Add(ActionType Type)
 		{
-			int UniqueId = ++NextUniqueId;
-			Action NewAction = new Action(Type, UniqueId);
+			Action NewAction = new Action(Type);
 			AllActions.Add(NewAction);
 			return NewAction;
 		}
@@ -788,13 +775,20 @@ namespace UnrealBuildTool
 			// If there are any cyclical actions, throw an exception.
 			if (ActionIsNonCyclical.Count < AllActions.Count)
 			{
+				// Find the index of each action
+				Dictionary<Action, int> ActionToIndex = new Dictionary<Action, int>();
+				for(int Idx = 0; Idx < AllActions.Count; Idx++)
+				{
+					ActionToIndex[AllActions[Idx]] = Idx;
+				}
+
 				// Describe the cyclical actions.
 				string CycleDescription = "";
 				foreach (Action Action in AllActions)
 				{
 					if (!ActionIsNonCyclical.ContainsKey(Action))
 					{
-						CycleDescription += string.Format("Action #{0}: {1}\n", Action.UniqueId, Action.CommandPath);
+						CycleDescription += string.Format("Action #{0}: {1}\n", ActionToIndex[Action], Action.CommandPath);
 						CycleDescription += string.Format("\twith arguments: {0}\n", Action.CommandArguments);
 						foreach (FileItem PrerequisiteItem in Action.PrerequisiteItems)
 						{
@@ -811,11 +805,11 @@ namespace UnrealBuildTool
 							{
 								if (CyclicPrerequisiteAction.ProducedItems.Count == 1)
 								{
-									CycleDescription += string.Format("\t\t{0} (produces: {1})\n", CyclicPrerequisiteAction.UniqueId, CyclicPrerequisiteAction.ProducedItems[0].AbsolutePath);
+									CycleDescription += string.Format("\t\t{0} (produces: {1})\n", ActionToIndex[CyclicPrerequisiteAction], CyclicPrerequisiteAction.ProducedItems[0].AbsolutePath);
 								}
 								else
 								{
-									CycleDescription += string.Format("\t\t{0}\n", CyclicPrerequisiteAction.UniqueId);
+									CycleDescription += string.Format("\t\t{0}\n", ActionToIndex[CyclicPrerequisiteAction]);
 									foreach (FileItem CyclicProducedItem in CyclicPrerequisiteAction.ProducedItems)
 									{
 										CycleDescription += string.Format("\t\t\tproduces:   {0}\n", CyclicProducedItem.AbsolutePath);
