UBT: Remove the ToHashSet() extension method.

#rnx

diff --git a/Engine/Source/Programs/UnrealBuildTool/System/ActionGraph.cs b/Engine/Source/Programs/UnrealBuildTool/System/ActionGraph.cs
--- a/Engine/Source/Programs/UnrealBuildTool/System/ActionGraph.cs
+++ b/Engine/Source/Programs/UnrealBuildTool/System/ActionGraph.cs
@@ -363,21 +363,21 @@ namespace UnrealBuildTool
 			CreateDirectoriesForProducedItems(OutdatedActionDictionary);
 
 			// Build a list of actions that are both needed for this target and outdated.
-			HashSet<Action> ActionsToExecute = AllActions.Where(Action => Action.CommandPath != null && IsActionOutdatedMap.ContainsKey(Action) && OutdatedActionDictionary[Action]).ToHashSet();
+			HashSet<Action> ActionsToExecute = new HashSet<Action>(AllActions.Where(Action => Action.CommandPath != null && IsActionOutdatedMap.ContainsKey(Action) && OutdatedActionDictionary[Action]));
 
 			// Remove link actions if asked to
 			if (BuildConfiguration.bSkipLinkingWhenNothingToCompile)
 			{
 				// Get all items produced by a compile action
-				HashSet<FileItem> ProducedItems = ActionsToExecute.Where(Action => Action.ActionType == ActionType.Compile).SelectMany(x => x.ProducedItems).ToHashSet();
+				HashSet<FileItem> ProducedItems = new HashSet<FileItem>(ActionsToExecute.Where(Action => Action.ActionType == ActionType.Compile).SelectMany(x => x.ProducedItems));
 
 				// Get all link actions which have no out-of-date prerequisites
-				HashSet<Action> UnlinkedActions = ActionsToExecute.Where(Action => Action.ActionType == ActionType.Link && !ProducedItems.Overlaps(Action.PrerequisiteItems)).ToHashSet();
+				HashSet<Action> UnlinkedActions = new HashSet<Action>(ActionsToExecute.Where(Action => Action.ActionType == ActionType.Link && !ProducedItems.Overlaps(Action.PrerequisiteItems)));
 
 				// Don't regard an action as unlinked if there is an associated 'failed.hotreload' file.
 				UnlinkedActions.RemoveWhere(Action => Action.ProducedItems.Any(Item => File.Exists(Path.Combine(Path.GetDirectoryName(Item.AbsolutePath), "failed.hotreload"))));
 
-				HashSet<Action> UnlinkedActionsWithFailedHotreload = ActionsToExecute.Where(Action => Action.ActionType == ActionType.Link && !ProducedItems.Overlaps(Action.PrerequisiteItems)).ToHashSet();
+				HashSet<Action> UnlinkedActionsWithFailedHotreload = new HashSet<Action>(ActionsToExecute.Where(Action => Action.ActionType == ActionType.Link && !ProducedItems.Overlaps(Action.PrerequisiteItems)));
 
 				// Remove unlinked items
 				ActionsToExecute.ExceptWith(UnlinkedActions);
@@ -386,13 +386,13 @@ namespace UnrealBuildTool
 				for (;;)
 				{
 					// Get all prerequisite items of a link action
-					HashSet<Action> PrerequisiteLinkActions = ActionsToExecute.Where(Action => Action.ActionType == ActionType.Link).SelectMany(x => x.PrerequisiteItems).Select(Item => Item.ProducingAction).ToHashSet();
+					HashSet<Action> PrerequisiteLinkActions = new HashSet<Action>(ActionsToExecute.Where(Action => Action.ActionType == ActionType.Link).SelectMany(x => x.PrerequisiteItems).Select(Item => Item.ProducingAction));
 
 					// Find all unlinked actions that need readding
-					HashSet<Action> UnlinkedActionsToReadd = UnlinkedActions.Where(Action => PrerequisiteLinkActions.Contains(Action)).ToHashSet();
+					HashSet<Action> UnlinkedActionsToReadd = new HashSet<Action>(UnlinkedActions.Where(Action => PrerequisiteLinkActions.Contains(Action)));
 
 					// Also re-add any DLL whose import library is being rebuilt. These may be separate actions, and the import library will reference the new DLL even if it isn't being compiled itself, so it must exist.
-					HashSet<FileReference> ProducedItemsToReAdd = ActionsToExecute.SelectMany(x => x.ProducedItems).Select(x => x.Location).Where(x => x.HasExtension(".lib")).Select(x => x.ChangeExtension(".suppressed.lib")).ToHashSet();
+					HashSet<FileReference> ProducedItemsToReAdd = new HashSet<FileReference>(ActionsToExecute.SelectMany(x => x.ProducedItems).Select(x => x.Location).Where(x => x.HasExtension(".lib")).Select(x => x.ChangeExtension(".suppressed.lib")));
 					UnlinkedActionsToReadd.UnionWith(UnlinkedActions.Where(x => x.ProducedItems.Any(y => ProducedItemsToReAdd.Contains(y.Location))));
 
 					// Bail if we didn't find anything
@@ -412,7 +412,7 @@ namespace UnrealBuildTool
 				}
 
 				// Remove actions that are wholly dependent on unlinked actions
-				ActionsToExecute = ActionsToExecute.Where(Action => Action.PrerequisiteItems.Count == 0 || !Action.PrerequisiteItems.Select(Item => Item.ProducingAction).ToHashSet().IsSubsetOf(UnlinkedActions)).ToHashSet();
+				ActionsToExecute = new HashSet<Action>(ActionsToExecute.Where(Action => Action.PrerequisiteItems.Count == 0 || !new HashSet<Action>(Action.PrerequisiteItems.Select(Item => Item.ProducingAction)).IsSubsetOf(UnlinkedActions)));
 			}
 
 			if (UnrealBuildTool.bPrintPerformanceInfo)
diff --git a/Engine/Source/Programs/UnrealBuildTool/UnrealBuildTool.csproj b/Engine/Source/Programs/UnrealBuildTool/UnrealBuildTool.csproj
--- a/Engine/Source/Programs/UnrealBuildTool/UnrealBuildTool.csproj
+++ b/Engine/Source/Programs/UnrealBuildTool/UnrealBuildTool.csproj
@@ -247,7 +247,6 @@
     <Compile Include="ToolChain\RemoteExports.cs" />
     <Compile Include="UnrealBuildTool.cs" />
     <Compile Include="Utilities\Crypto.cs" />
-    <Compile Include="Utilities\HashSetExtensions.cs" />
     <Compile Include="Platform\Windows\PVSToolChain.cs" />
     <Compile Include="Platform\Windows\WindowsDeploy.cs" />
     <Compile Include="Platform\Windows\WinResources.cs" />
diff --git a/Engine/Source/Programs/UnrealBuildTool/Utilities/HashSetExtensions.cs b/Engine/Source/Programs/UnrealBuildTool/Utilities/HashSetExtensions.cs
deleted file mode 100644
