UBT: Rename Crypto.cs to EncryptionAndSigning.cs to match the name of the class inside it, and move it under the System folder.

#rnx

diff --git a/Engine/Source/Programs/UnrealBuildTool/System/EncryptionAndSigning.cs b/Engine/Source/Programs/UnrealBuildTool/System/EncryptionAndSigning.cs
new file mode 100644
--- /dev/null
+++ b/Engine/Source/Programs/UnrealBuildTool/System/EncryptionAndSigning.cs
@@ -0,0 +1,297 @@
+// Copyright 1998-2018 Epic Games, Inc. All Rights Reserved.
+
+using System;
+using System.Collections.Generic;
+using System.IO;
+using System.Linq;
+using System.Text;
+using Tools.DotNETCommon;
+
+namespace UnrealBuildTool
+{
+	/// <summary>
+	/// Helper functions for dealing with encryption and pak signing
+	/// </summary>
+	public static class EncryptionAndSigning
+	{
+		/// <summary>
+		/// Wrapper class for a single RSA key
+		/// </summary>
+		public class SigningKey
+		{
+			/// <summary>
+			/// Exponent
+			/// </summary>
+			public byte[] Exponent;
+
+			/// <summary>
+			/// Modulus
+			/// </summary>
+			public byte[] Modulus;
+		}
+
+		/// <summary>
+		/// Wrapper class for an RSA public/private key pair
+		/// </summary>
+		public class SigningKeyPair
+		{
+			/// <summary>
+			/// Public key
+			/// </summary>
+			public SigningKey PublicKey = new SigningKey();
+
+			/// <summary>
+			/// Private key
+			/// </summary>
+			public SigningKey PrivateKey = new SigningKey();
+		}
+
+		/// <summary>
+		/// Wrapper class for a 128 bit AES encryption key
+		/// </summary>
+		public class EncryptionKey
+		{
+			/// <summary>
+			/// 128 bit AES key
+			/// </summary>
+			public byte[] Key;
+		}
+
+		/// <summary>
+		/// Wrapper class for all crypto settings
+		/// </summary>
+		public class CryptoSettings
+		{
+			/// <summary>
+			/// AES encyption key
+			/// </summary>
+			public EncryptionKey EncryptionKey = null;
+
+			/// <summary>
+			/// RSA public/private key
+			/// </summary>
+			public SigningKeyPair SigningKey = null;
+
+			/// <summary>
+			/// Enable pak signature checking
+			/// </summary>
+			public bool bEnablePakSigning = false;
+
+			/// <summary>
+			/// Encrypt the index of the pak file. Stops the pak file being easily accessible by unrealpak
+			/// </summary>
+			public bool bEnablePakIndexEncryption = false;
+
+			/// <summary>
+			/// Encrypt all ini files in the pak. Good for game data obsfucation
+			/// </summary>
+			public bool bEnablePakIniEncryption = false;
+
+			/// <summary>
+			/// Encrypt the uasset files in the pak file. After cooking, uasset files only contain package metadata / nametables / export and import tables. Provides good string data obsfucation without
+			/// the downsides of full package encryption, with the security drawbacks of still having some data stored unencrypted 
+			/// </summary>
+			public bool bEnablePakUAssetEncryption = false;
+
+			/// <summary>
+			/// Encrypt all assets data files (including exp and ubulk) in the pak file. Likely to be slow, and to cause high data entropy (bad for delta patching)
+			/// </summary>
+			public bool bEnablePakFullAssetEncryption = false;
+
+			/// <summary>
+			/// Some platforms have their own data crypto systems, so allow the config settings to totally disable our own crypto
+			/// </summary>
+			public bool bDataCryptoRequired = false;
+
+			/// <summary>
+			/// 
+			/// </summary>
+			public bool IsAnyEncryptionEnabled()
+			{
+				return bEnablePakFullAssetEncryption || bEnablePakUAssetEncryption || bEnablePakIndexEncryption || bEnablePakIniEncryption;
+			}
+
+			/// <summary>
+			/// 
+			/// </summary>
+			public void Save(FileReference InFile)
+			{
+				FileReference.WriteAllText(InFile, fastJSON.JSON.Instance.ToJSON(this, new fastJSON.JSONParameters {}));
+			}
+		}
+
+		/// <summary>
+		/// Helper class for formatting incoming hex signing key strings
+		/// </summary>
+		private static string ProcessSigningKeyInputStrings(string InString)
+		{
+			if (InString.StartsWith("0x"))
+			{
+				InString = InString.Substring(2);
+			}
+			return InString.TrimStart('0');
+		}
+
+		/// <summary>
+		/// Parse crypto settings from INI file
+		/// </summary>
+		public static CryptoSettings ParseCryptoSettings(DirectoryReference InProjectDirectory, UnrealTargetPlatform InTargetPlatform)
+		{
+			CryptoSettings Settings = new CryptoSettings();
+
+			ConfigHierarchy Ini = ConfigCache.ReadHierarchy(ConfigHierarchyType.Engine, InProjectDirectory, InTargetPlatform);
+			Ini.GetBool("PlatformCrypto", "PlatformRequiresDataCrypto", out Settings.bDataCryptoRequired);
+
+			// For now, we'll just not parse any keys if data crypto is disabled for this platform. In the future, we might want to use
+			// these keys for non-data purposes (other general purpose encryption maybe?)
+			if (!Settings.bDataCryptoRequired)
+			{
+				return Settings;
+			}
+
+			{
+				// Start by parsing the legacy encryption.ini settings
+				Ini = ConfigCache.ReadHierarchy(ConfigHierarchyType.Encryption, InProjectDirectory, InTargetPlatform);
+				Ini.GetBool("Core.Encryption", "SignPak", out Settings.bEnablePakSigning);
+
+				string[] SigningKeyStrings = new string[3];
+				Ini.GetString("Core.Encryption", "rsa.privateexp", out SigningKeyStrings[0]);
+				Ini.GetString("Core.Encryption", "rsa.modulus", out SigningKeyStrings[1]);
+				Ini.GetString("Core.Encryption", "rsa.publicexp", out SigningKeyStrings[2]);
+
+				if (String.IsNullOrEmpty(SigningKeyStrings[0]) || String.IsNullOrEmpty(SigningKeyStrings[1]) || String.IsNullOrEmpty(SigningKeyStrings[2]))
+				{
+					SigningKeyStrings = null;
+				}
+				else
+				{
+					Settings.SigningKey = new SigningKeyPair();
+					Settings.SigningKey.PrivateKey.Exponent = ParseHexStringToByteArray(ProcessSigningKeyInputStrings(SigningKeyStrings[0]), 64);
+					Settings.SigningKey.PrivateKey.Modulus = ParseHexStringToByteArray(ProcessSigningKeyInputStrings(SigningKeyStrings[1]), 64);
+					Settings.SigningKey.PublicKey.Exponent = ParseHexStringToByteArray(ProcessSigningKeyInputStrings(SigningKeyStrings[2]), 64);
+					Settings.SigningKey.PublicKey.Modulus = Settings.SigningKey.PrivateKey.Modulus;
+
+					if ((Settings.SigningKey.PrivateKey.Exponent.Length > 64) ||
+						(Settings.SigningKey.PrivateKey.Modulus.Length > 64) ||
+						(Settings.SigningKey.PublicKey.Exponent.Length > 64) ||
+						(Settings.SigningKey.PublicKey.Modulus.Length > 64))
+					{
+						throw new Exception(string.Format("[{0}] Signing keys parsed from encryption.ini are too long. They must be a maximum of 64 bytes long!", InProjectDirectory));
+					}
+				}
+
+				Ini.GetBool("Core.Encryption", "EncryptPak", out Settings.bEnablePakIndexEncryption);
+				Settings.bEnablePakFullAssetEncryption = false;
+				Settings.bEnablePakUAssetEncryption = false;
+				Settings.bEnablePakIniEncryption = Settings.bEnablePakIndexEncryption;
+
+				string EncryptionKeyString;
+				Ini.GetString("Core.Encryption", "aes.key", out EncryptionKeyString);
+				Settings.EncryptionKey = new EncryptionKey();
+
+				if (EncryptionKeyString.Length > 0)
+				{
+					if (EncryptionKeyString.Length < 32)
+					{
+						Log.WriteLine(LogEventType.Warning, "AES key parsed from encryption.ini is too short. It must be 32 bytes, so will be padded with 0s, giving sub-optimal security!");
+					}
+					else if (EncryptionKeyString.Length > 32)
+					{
+						Log.WriteLine(LogEventType.Warning, "AES key parsed from encryption.ini is too long. It must be 32 bytes, so will be truncated!");
+					}
+
+					Settings.EncryptionKey.Key = ParseAnsiStringToByteArray(EncryptionKeyString, 32);
+				}
+			}
+
+			Ini = ConfigCache.ReadHierarchy(ConfigHierarchyType.Crypto, InProjectDirectory, InTargetPlatform);
+			string SectionName = "/Script/CryptoKeys.CryptoKeysSettings";
+			ConfigHierarchySection CryptoSection = Ini.FindSection(SectionName);
+
+			// If we have new format crypto keys, read them in over the top of the legacy settings
+			if (CryptoSection != null && CryptoSection.KeyNames.Count() > 0)
+			{
+				Ini.GetBool(SectionName, "bEnablePakSigning", out Settings.bEnablePakSigning);
+				Ini.GetBool(SectionName, "bEncryptPakIniFiles", out Settings.bEnablePakIniEncryption);
+				Ini.GetBool(SectionName, "bEncryptPakIndex", out Settings.bEnablePakIndexEncryption);
+				Ini.GetBool(SectionName, "bEncryptUAssetFiles", out Settings.bEnablePakUAssetEncryption);
+				Ini.GetBool(SectionName, "bEncryptAllAssetFiles", out Settings.bEnablePakFullAssetEncryption);
+
+				// Parse encryption key
+				string EncryptionKeyString;
+				Ini.GetString(SectionName, "EncryptionKey", out EncryptionKeyString);
+				if (!string.IsNullOrEmpty(EncryptionKeyString))
+				{
+					Settings.EncryptionKey = new EncryptionKey();
+					Settings.EncryptionKey.Key = System.Convert.FromBase64String(EncryptionKeyString);
+				}
+
+				// Parse signing key
+				string PrivateExponent, PublicExponent, Modulus;
+				Ini.GetString(SectionName, "SigningPrivateExponent", out PrivateExponent);
+				Ini.GetString(SectionName, "SigningModulus", out Modulus);
+				Ini.GetString(SectionName, "SigningPublicExponent", out PublicExponent);
+
+				if (!String.IsNullOrEmpty(PrivateExponent) && !String.IsNullOrEmpty(PublicExponent) && !String.IsNullOrEmpty(Modulus))
+				{
+					Settings.SigningKey = new SigningKeyPair();
+					Settings.SigningKey.PublicKey.Exponent = System.Convert.FromBase64String(PublicExponent);
+					Settings.SigningKey.PublicKey.Modulus = System.Convert.FromBase64String(Modulus);
+					Settings.SigningKey.PrivateKey.Exponent = System.Convert.FromBase64String(PrivateExponent);
+					Settings.SigningKey.PrivateKey.Modulus = Settings.SigningKey.PublicKey.Modulus;
+				}
+			}
+
+			return Settings;
+		}
+
+		/// <summary>
+		/// Take a hex string and parse into an array of bytes
+		/// </summary>
+		private static byte[] ParseHexStringToByteArray(string InString, int InMinimumLength)
+		{
+			if (InString.StartsWith("0x"))
+			{
+				InString = InString.Substring(2);
+			}
+
+			List<byte> Bytes = new List<byte>();
+			while (InString.Length > 0)
+			{
+				int CharsToParse = Math.Min(2, InString.Length);
+				string Value = InString.Substring(InString.Length - CharsToParse);
+				InString = InString.Substring(0, InString.Length - CharsToParse);
+				Bytes.Add(byte.Parse(Value, System.Globalization.NumberStyles.AllowHexSpecifier));
+			}
+
+			while (Bytes.Count < InMinimumLength)
+			{
+				Bytes.Add(0);
+			}
+
+			return Bytes.ToArray();
+		}
+
+		private static byte[] ParseAnsiStringToByteArray(string InString, Int32 InRequiredLength)
+		{
+			List<byte> Bytes = new List<byte>();
+
+			if (InString.Length > InRequiredLength)
+			{
+				InString = InString.Substring(0, InRequiredLength);
+			}
+
+			foreach (char C in InString)
+			{
+				Bytes.Add((byte)C);
+			}
+
+			while (Bytes.Count < InRequiredLength)
+			{
+				Bytes.Add(0);
+			}
+
+			return Bytes.ToArray();
+		}
+	}
+}
diff --git a/Engine/Source/Programs/UnrealBuildTool/UnrealBuildTool.csproj b/Engine/Source/Programs/UnrealBuildTool/UnrealBuildTool.csproj
--- a/Engine/Source/Programs/UnrealBuildTool/UnrealBuildTool.csproj
+++ b/Engine/Source/Programs/UnrealBuildTool/UnrealBuildTool.csproj
@@ -169,6 +169,7 @@
     <Compile Include="System\BuildVersion.cs" />
     <Compile Include="System\CppCompileEnvironment.cs" />
     <Compile Include="System\DataDrivenPlatformInfo.cs" />
+    <Compile Include="System\EncryptionAndSigning.cs" />
     <Compile Include="System\ParallelExecutor.cs" />
     <Compile Include="System\ModuleManifest.cs" />
     <Compile Include="System\PluginReferenceDescriptor.cs" />
@@ -246,7 +247,6 @@
   <ItemGroup>
     <Compile Include="ToolChain\RemoteExports.cs" />
     <Compile Include="UnrealBuildTool.cs" />
-    <Compile Include="Utilities\Crypto.cs" />
     <Compile Include="Platform\Windows\PVSToolChain.cs" />
     <Compile Include="Platform\Windows\WindowsDeploy.cs" />
     <Compile Include="Platform\Windows\WinResources.cs" />
diff --git a/Engine/Source/Programs/UnrealBuildTool/Utilities/Crypto.cs b/Engine/Source/Programs/UnrealBuildTool/Utilities/Crypto.cs
deleted file mode 100644
