UBT: Move all the action executors into their own folder in the project.

#rnx

diff --git a/Engine/Source/Programs/UnrealBuildTool/Executors/Distcc.cs b/Engine/Source/Programs/UnrealBuildTool/Executors/Distcc.cs
new file mode 100644
--- /dev/null
+++ b/Engine/Source/Programs/UnrealBuildTool/Executors/Distcc.cs
@@ -0,0 +1,386 @@
+// Copyright 1998-2018 Epic Games, Inc. All Rights Reserved.
+
+using System;
+using System.Collections.Generic;
+using System.IO;
+using System.Diagnostics;
+using System.Xml;
+using System.Text.RegularExpressions;
+using System.Linq;
+using System.Reflection;
+using Tools.DotNETCommon;
+
+namespace UnrealBuildTool
+{
+	class Distcc : ActionExecutor
+	{
+		/// <summary>
+		/// When enabled allows DMUCS/Distcc to fallback to local compilation when remote compiling fails. Defaults to true as separation of pre-process and compile stages can introduce non-fatal errors.
+		/// </summary>
+		[XmlConfigFile(Category = "BuildConfiguration")]
+		public bool bAllowDistccLocalFallback = true;
+
+		/// <summary>
+		/// When true enable verbose distcc output to aid debugging.
+		/// </summary>
+		[XmlConfigFile(Category = "BuildConfiguration")]
+		public bool bVerboseDistccOutput = false;
+
+		/// <summary>
+		/// Path to the Distcc and DMUCS executables.
+		/// </summary>
+		[XmlConfigFile(Category = "BuildConfiguration")]
+		public string DistccExecutablesPath = "/usr/local/bin";
+
+		/// <summary>
+		/// DMUCS coordinator hostname or IP address.
+		/// </summary>
+		[XmlConfigFile(Category = "BuildConfiguration")]
+		public string DMUCSCoordinator = "localhost";
+
+		/// <summary>
+		/// The DMUCS distinguishing property value to request for compile hosts.
+		/// </summary>
+		[XmlConfigFile(Category = "BuildConfiguration")]
+		public string DMUCSDistProp = "";
+
+		public Distcc()
+		{
+			XmlConfig.ApplyTo(this);
+
+			// The default for normal Mac users should be to use DistCode which installs as an Xcode plugin and provides dynamic host management
+			if (BuildHostPlatform.Current.Platform == UnrealTargetPlatform.Mac)
+			{
+				string UserDir = Environment.GetEnvironmentVariable("HOME");
+				string MacDistccExecutablesPath = UserDir + "/Library/Application Support/Developer/Shared/Xcode/Plug-ins/Distcc 3.2.xcplugin/Contents/usr/bin";
+
+				// But we should use the standard Posix directory when not installed - a build farm would use a static .dmucs hosts file not DistCode.
+				if (System.IO.Directory.Exists(MacDistccExecutablesPath))
+				{
+					DistccExecutablesPath = MacDistccExecutablesPath;
+				}
+
+				if (System.IO.File.Exists(UserDir + "/Library/Preferences/com.marksatt.DistCode.plist"))
+				{
+					using (System.Diagnostics.Process DefaultsProcess = new System.Diagnostics.Process())
+					{
+						try
+						{
+							DefaultsProcess.StartInfo.FileName = "/usr/bin/defaults";
+							DefaultsProcess.StartInfo.CreateNoWindow = true;
+							DefaultsProcess.StartInfo.UseShellExecute = false;
+							DefaultsProcess.StartInfo.RedirectStandardOutput = true;
+							DefaultsProcess.StartInfo.RedirectStandardError = true;
+							DefaultsProcess.StartInfo.Arguments = "read com.marksatt.DistCode DistProp";
+							DefaultsProcess.Start();
+							string Output = DefaultsProcess.StandardOutput.ReadToEnd();
+							DefaultsProcess.WaitForExit();
+							if (DefaultsProcess.ExitCode == 0)
+							{
+								DMUCSDistProp = Output;
+							}
+						}
+						catch (Exception)
+						{
+						}
+					}
+                    using (System.Diagnostics.Process CoordModeProcess = new System.Diagnostics.Process())
+                    {
+                        using (System.Diagnostics.Process DefaultsProcess = new System.Diagnostics.Process())
+                        {
+                            try
+                            {
+                                CoordModeProcess.StartInfo.FileName = "/usr/bin/defaults";
+                                CoordModeProcess.StartInfo.CreateNoWindow = true;
+                                CoordModeProcess.StartInfo.UseShellExecute = false;
+                                CoordModeProcess.StartInfo.RedirectStandardOutput = true;
+								CoordModeProcess.StartInfo.RedirectStandardError = true;
+                                CoordModeProcess.StartInfo.Arguments = "read com.marksatt.DistCode CoordinatorMode";
+                                CoordModeProcess.Start();
+                                string CoordModeProcessOutput = CoordModeProcess.StandardOutput.ReadToEnd();
+                                CoordModeProcess.WaitForExit();
+                                if (CoordModeProcess.ExitCode == 0 && CoordModeProcessOutput.StartsWith("1"))
+                                {
+                                    DefaultsProcess.StartInfo.FileName = "/usr/bin/defaults";
+                                    DefaultsProcess.StartInfo.CreateNoWindow = true;
+                                    DefaultsProcess.StartInfo.UseShellExecute = false;
+                                    DefaultsProcess.StartInfo.RedirectStandardOutput = true;
+									DefaultsProcess.StartInfo.RedirectStandardError = true;
+                                    DefaultsProcess.StartInfo.Arguments = "read com.marksatt.DistCode CoordinatorIP";
+                                    DefaultsProcess.Start();
+                                    string Output = DefaultsProcess.StandardOutput.ReadToEnd();
+                                    DefaultsProcess.WaitForExit();
+                                    if (DefaultsProcess.ExitCode == 0)
+                                    {
+                                        DMUCSCoordinator = Output;
+                                    }
+                                }
+                            }
+                            catch (Exception)
+                            {
+                            }
+                        }
+                    }
+				}
+			}
+		}
+
+		public override string Name
+		{
+			get { return "Distcc"; }
+		}
+
+		public override bool ExecuteActions(List<Action> Actions, bool bLogDetailedActionStats)
+		{
+			bool bDistccResult = true;
+			if (Actions.Count > 0)
+			{
+				// Time to sleep after each iteration of the loop in order to not busy wait.
+				const float LoopSleepTime = 0.1f;
+
+				int MaxActionsToExecuteInParallel = 0;
+
+				string UserDir = Environment.GetEnvironmentVariable("HOME");
+				string HostsInfo = UserDir + "/.dmucs/hosts-info";
+				string NumUBTBuildTasks = Environment.GetEnvironmentVariable("NumUBTBuildTasks");
+				Int32 MaxUBTBuildTasks = MaxActionsToExecuteInParallel;
+				if (Int32.TryParse(NumUBTBuildTasks, out MaxUBTBuildTasks))
+				{
+					MaxActionsToExecuteInParallel = MaxUBTBuildTasks;
+				}
+				else if (BuildHostPlatform.Current.Platform == UnrealTargetPlatform.Mac)
+				{
+					using (System.Diagnostics.Process DefaultsProcess = new System.Diagnostics.Process())
+					{
+						try
+						{
+							DefaultsProcess.StartInfo.FileName = "/usr/bin/defaults";
+							DefaultsProcess.StartInfo.CreateNoWindow = true;
+							DefaultsProcess.StartInfo.UseShellExecute = false;
+							DefaultsProcess.StartInfo.RedirectStandardOutput = true;
+							DefaultsProcess.StartInfo.Arguments = "read com.apple.dt.Xcode IDEBuildOperationMaxNumberOfConcurrentCompileTasks";
+							DefaultsProcess.Start();
+							string Output = DefaultsProcess.StandardOutput.ReadToEnd();
+							DefaultsProcess.WaitForExit();
+							if (DefaultsProcess.ExitCode == 0 && Int32.TryParse(Output, out MaxUBTBuildTasks))
+							{
+								MaxActionsToExecuteInParallel = MaxUBTBuildTasks;
+							}
+						}
+						catch (Exception)
+						{
+						}
+					}
+				}
+				else if (System.IO.File.Exists(HostsInfo))
+				{
+					System.IO.StreamReader File = new System.IO.StreamReader(HostsInfo);
+					string Line = null;
+					while ((Line = File.ReadLine()) != null)
+					{
+						string[] HostInfo = Line.Split(' ');
+						if (HostInfo.Count() == 3)
+						{
+							int NumCPUs = 0;
+							if (System.Int32.TryParse(HostInfo[1], out NumCPUs))
+							{
+								MaxActionsToExecuteInParallel += NumCPUs;
+							}
+						}
+					}
+					File.Close();
+				}
+				else
+				{
+					MaxActionsToExecuteInParallel = System.Environment.ProcessorCount;
+				}
+
+				if (bAllowDistccLocalFallback == false)
+				{
+					Environment.SetEnvironmentVariable("DISTCC_FALLBACK", "0");
+				}
+
+				if (bVerboseDistccOutput == true)
+				{
+					Environment.SetEnvironmentVariable("DISTCC_VERBOSE", "1");
+				}
+
+				string DistccExecutable = DistccExecutablesPath + "/distcc";
+				string GetHostExecutable = DistccExecutablesPath + "/gethost";
+
+				Log.TraceInformation("Performing {0} actions ({1} in parallel)", Actions.Count, MaxActionsToExecuteInParallel, DistccExecutable, GetHostExecutable);
+
+				Dictionary<Action, ActionThread> ActionThreadDictionary = new Dictionary<Action, ActionThread>();
+				int JobNumber = 1;
+				using (ProgressWriter ProgressWriter = new ProgressWriter("Compiling source code...", false))
+				{
+					int ProgressValue = 0;
+					while (true)
+					{
+						// Count the number of pending and still executing actions.
+						int NumUnexecutedActions = 0;
+						int NumExecutingActions = 0;
+						foreach (Action Action in Actions)
+						{
+							ActionThread ActionThread = null;
+							bool bFoundActionProcess = ActionThreadDictionary.TryGetValue(Action, out ActionThread);
+							if (bFoundActionProcess == false)
+							{
+								NumUnexecutedActions++;
+							}
+							else if (ActionThread != null)
+							{
+								if (ActionThread.bComplete == false)
+								{
+									NumUnexecutedActions++;
+									NumExecutingActions++;
+								}
+							}
+						}
+
+						// Update the current progress
+						int NewProgressValue = Actions.Count + 1 - NumUnexecutedActions;
+						if (ProgressValue != NewProgressValue)
+						{
+							ProgressWriter.Write(ProgressValue, Actions.Count + 1);
+							ProgressValue = NewProgressValue;
+						}
+
+						// If there aren't any pending actions left, we're done executing.
+						if (NumUnexecutedActions == 0)
+						{
+							break;
+						}
+
+						// If there are fewer actions executing than the maximum, look for pending actions that don't have any outdated
+						// prerequisites.
+						foreach (Action Action in Actions)
+						{
+							ActionThread ActionProcess = null;
+							bool bFoundActionProcess = ActionThreadDictionary.TryGetValue(Action, out ActionProcess);
+							if (bFoundActionProcess == false)
+							{
+								if (NumExecutingActions < Math.Max(1, MaxActionsToExecuteInParallel))
+								{
+									// Determine whether there are any prerequisites of the action that are outdated.
+									bool bHasOutdatedPrerequisites = false;
+									bool bHasFailedPrerequisites = false;
+									foreach (FileItem PrerequisiteItem in Action.PrerequisiteItems)
+									{
+										if (PrerequisiteItem.ProducingAction != null && Actions.Contains(PrerequisiteItem.ProducingAction))
+										{
+											ActionThread PrerequisiteProcess = null;
+											bool bFoundPrerequisiteProcess = ActionThreadDictionary.TryGetValue(PrerequisiteItem.ProducingAction, out PrerequisiteProcess);
+											if (bFoundPrerequisiteProcess == true)
+											{
+												if (PrerequisiteProcess == null)
+												{
+													bHasFailedPrerequisites = true;
+												}
+												else if (PrerequisiteProcess.bComplete == false)
+												{
+													bHasOutdatedPrerequisites = true;
+												}
+												else if (PrerequisiteProcess.ExitCode != 0)
+												{
+													bHasFailedPrerequisites = true;
+												}
+											}
+											else
+											{
+												bHasOutdatedPrerequisites = true;
+											}
+										}
+									}
+
+									// If there are any failed prerequisites of this action, don't execute it.
+									if (bHasFailedPrerequisites)
+									{
+										// Add a null entry in the dictionary for this action.
+										ActionThreadDictionary.Add(Action, null);
+									}
+									// If there aren't any outdated prerequisites of this action, execute it.
+									else if (!bHasOutdatedPrerequisites)
+									{
+										if (Action.ActionType == ActionType.Compile || Action.ActionType == ActionType.Link)
+										{
+											string TypeCommand = String.IsNullOrEmpty(DMUCSDistProp) ? "" : (" --type " + DMUCSDistProp);
+											string NewCommandArguments = "--server " + DMUCSCoordinator + TypeCommand + " --wait -1 \"" + DistccExecutable + "\" " + Action.CommandPath + " " + Action.CommandArguments;
+											
+											if (Action.ActionType == ActionType.Compile)
+											{
+												// Distcc separates preprocessing and compilation which means we must silence these warnings
+												NewCommandArguments += " -Wno-constant-logical-operand";
+												NewCommandArguments += " -Wno-parentheses-equality";
+											}
+											
+											Action.CommandPath = GetHostExecutable;
+											Action.CommandArguments = NewCommandArguments;
+										}
+
+										ActionThread ActionThread = new ActionThread(Action, JobNumber, Actions.Count);
+										JobNumber++;
+										ActionThread.Run();
+
+										ActionThreadDictionary.Add(Action, ActionThread);
+
+										NumExecutingActions++;
+									}
+								}
+							}
+						}
+
+						System.Threading.Thread.Sleep(TimeSpan.FromSeconds(LoopSleepTime));
+					}
+				}
+
+				Log.WriteLineIf(bLogDetailedActionStats, LogEventType.Console, "-------- Begin Detailed Action Stats ----------------------------------------------------------");
+				Log.WriteLineIf(bLogDetailedActionStats, LogEventType.Console, "^Action Type^Duration (seconds)^Tool^Task^Using PCH");
+
+				double TotalThreadSeconds = 0;
+
+				// Check whether any of the tasks failed and log action stats if wanted.
+				foreach (KeyValuePair<Action, ActionThread> ActionProcess in ActionThreadDictionary)
+				{
+					Action Action = ActionProcess.Key;
+					ActionThread ActionThread = ActionProcess.Value;
+
+					// Check for pending actions, preemptive failure
+					if (ActionThread == null)
+					{
+						bDistccResult = false;
+						continue;
+					}
+					// Check for executed action but general failure
+					if (ActionThread.ExitCode != 0)
+					{
+						bDistccResult = false;
+					}
+					// Log CPU time, tool and task.
+					double ThreadSeconds = Action.Duration.TotalSeconds;
+
+					Log.WriteLineIf(bLogDetailedActionStats,
+						LogEventType.Console,
+						"^{0}^{1:0.00}^{2}^{3}^{4}",
+						Action.ActionType.ToString(),
+						ThreadSeconds,
+						Path.GetFileName(Action.CommandPath),
+						Action.StatusDescription,
+						Action.bIsUsingPCH);
+
+					// Keep track of total thread seconds spent on tasks.
+					TotalThreadSeconds += ThreadSeconds;
+				}
+
+				Log.WriteLineIf(bLogDetailedActionStats || UnrealBuildTool.bPrintDebugInfo,
+					LogEventType.Console,
+					"-------- End Detailed Actions Stats -----------------------------------------------------------");
+
+				// Log total CPU seconds and numbers of processors involved in tasks.
+				Log.WriteLineIf(bLogDetailedActionStats || UnrealBuildTool.bPrintDebugInfo,
+					LogEventType.Console, "Cumulative thread seconds ({0} processors): {1:0.00}", System.Environment.ProcessorCount, TotalThreadSeconds);
+			}
+			return bDistccResult;
+		}
+	}
+}
diff --git a/Engine/Source/Programs/UnrealBuildTool/Executors/LocalExecutor.cs b/Engine/Source/Programs/UnrealBuildTool/Executors/LocalExecutor.cs
new file mode 100644
--- /dev/null
+++ b/Engine/Source/Programs/UnrealBuildTool/Executors/LocalExecutor.cs
@@ -0,0 +1,580 @@
+// Copyright 1998-2018 Epic Games, Inc. All Rights Reserved.
+
+using System;
+using System.Collections.Generic;
+using System.Text;
+using System.Text.RegularExpressions;
+using System.Diagnostics;
+using System.IO;
+using System.Reflection;
+using System.Threading;
+using System.Linq;
+using Tools.DotNETCommon;
+
+namespace UnrealBuildTool
+{
+	class ActionThread
+	{
+		/// <summary>
+		/// Cache the exit code from the command so that the executor can report errors
+		/// </summary>
+		public int ExitCode = 0;
+
+		/// <summary>
+		/// Set to true only when the local or RPC action is complete
+		/// </summary>
+		public bool bComplete = false;
+
+		/// <summary>
+		/// Cache the action that this thread is managing
+		/// </summary>
+		Action Action;
+
+		/// <summary>
+		/// For reporting status to the user
+		/// </summary>
+		int JobNumber;
+		int TotalJobs;
+
+		/// <summary>
+		/// Regex that matches environment variables in $(Variable) format.
+		/// </summary>
+		static Regex EnvironmentVariableRegex = new Regex("\\$\\(([\\d\\w]+)\\)");
+
+		/// <summary>
+		/// Replaces the environment variables references in a string with their values.
+		/// </summary>
+		public static string ExpandEnvironmentVariables(string Text)
+		{
+			foreach (Match EnvironmentVariableMatch in EnvironmentVariableRegex.Matches(Text))
+			{
+				string VariableValue = Environment.GetEnvironmentVariable(EnvironmentVariableMatch.Groups[1].Value);
+				Text = Text.Replace(EnvironmentVariableMatch.Value, VariableValue);
+			}
+			return Text;
+		}
+
+		/// <summary>
+		/// Constructor, takes the action to process
+		/// </summary>
+		public ActionThread(Action InAction, int InJobNumber, int InTotalJobs)
+		{
+			Action = InAction;
+			JobNumber = InJobNumber;
+			TotalJobs = InTotalJobs;
+		}
+
+		/// <summary>
+		/// Used when debuging Actions outputs all action return values to debug out
+		/// </summary>
+		/// <param name="sender"> Sending object</param>
+		/// <param name="e">  Event arguments (In this case, the line of string output)</param>
+		protected void ActionDebugOutput(object sender, DataReceivedEventArgs e)
+		{
+			string Output = e.Data;
+			if (Output == null)
+			{
+				return;
+			}
+
+			Log.TraceInformation(Output);
+		}
+
+
+		/// <summary>
+		/// The actual function to run in a thread. This is potentially long and blocking
+		/// </summary>
+		private void ThreadFunc()
+		{
+			// thread start time
+			Action.StartTime = DateTimeOffset.Now;
+
+			if (Action.ActionHandler != null)
+			{
+				// call the function and get the ExitCode and an output string
+				string Output;
+				Action.ActionHandler(Action, out ExitCode, out Output);
+
+				// Output status description (file name) when no output is returned
+				if (string.IsNullOrEmpty(Output))
+				{
+					if (Action.bShouldOutputStatusDescription)
+					{
+						Output = string.Format("[{0}/{1}] {2} {3}", JobNumber, TotalJobs, Action.CommandDescription, Action.StatusDescription);
+					}
+					else
+					{
+						Output = Action.StatusDescription;
+					}
+				}
+
+				Log.TraceInformation(Output);
+			}
+			else
+			{
+				// batch files (.bat, .cmd) need to be run via or cmd /c or shellexecute, 
+				// the latter which we can't use because we want to redirect input/output
+				bool bLaunchViaCmdExe = !Utils.IsRunningOnMono && !Path.GetExtension(Action.CommandPath).ToLower().EndsWith("exe");
+
+				// Create the action's process.
+				ProcessStartInfo ActionStartInfo = new ProcessStartInfo();
+				ActionStartInfo.WorkingDirectory = ExpandEnvironmentVariables(Action.WorkingDirectory);
+
+
+				string ExpandedCommandPath = ExpandEnvironmentVariables(Action.CommandPath);
+				if (bLaunchViaCmdExe)
+				{
+					ActionStartInfo.FileName = "cmd.exe";
+					ActionStartInfo.Arguments = string.Format
+					(
+						"/c \"{0} {1}\"",
+						ExpandEnvironmentVariables(ExpandedCommandPath),
+						ExpandEnvironmentVariables(Action.CommandArguments)
+					);
+				}
+				else
+				{
+					ActionStartInfo.FileName = ExpandedCommandPath;
+					ActionStartInfo.Arguments = ExpandEnvironmentVariables(Action.CommandArguments);
+				}
+
+				ActionStartInfo.UseShellExecute = false;
+				ActionStartInfo.RedirectStandardInput = false;
+				ActionStartInfo.RedirectStandardOutput = false;
+				ActionStartInfo.RedirectStandardError = false;
+
+				// Log command-line used to execute task if debug info printing is enabled.
+				if (UnrealBuildTool.bPrintDebugInfo)
+				{
+					Log.TraceVerbose("Executing: {0} {1}", ExpandedCommandPath, ActionStartInfo.Arguments);
+				}
+				// Log summary if wanted.
+				else if (Action.bShouldOutputStatusDescription)
+				{
+					string CommandDescription = Action.CommandDescription != null ? Action.CommandDescription : Path.GetFileName(ExpandedCommandPath);
+					if (string.IsNullOrEmpty(CommandDescription))
+					{
+						Log.TraceInformation(Action.StatusDescription);
+					}
+					else
+					{
+						Log.TraceInformation("[{0}/{1}] {2} {3}", JobNumber, TotalJobs, CommandDescription, Action.StatusDescription);
+					}
+				}
+
+				// Try to launch the action's process, and produce a friendly error message if it fails.
+				Process ActionProcess = null;
+				try
+				{
+					try
+					{
+						ActionProcess = new Process();
+						ActionProcess.StartInfo = ActionStartInfo;
+						bool bShouldRedirectOuput = Action.bPrintDebugInfo;
+						if (bShouldRedirectOuput)
+						{
+							ActionStartInfo.RedirectStandardOutput = true;
+							ActionStartInfo.RedirectStandardError = true;
+							ActionProcess.EnableRaisingEvents = true;
+
+							if (Action.bPrintDebugInfo)
+							{
+								ActionProcess.OutputDataReceived += new DataReceivedEventHandler(ActionDebugOutput);
+								ActionProcess.ErrorDataReceived += new DataReceivedEventHandler(ActionDebugOutput);
+							}
+						}
+						ActionProcess.Start();
+						if (bShouldRedirectOuput)
+						{
+							ActionProcess.BeginOutputReadLine();
+							ActionProcess.BeginErrorReadLine();
+						}
+					}
+					catch (Exception ex)
+					{
+						Log.TraceError("Failed to start local process for action: {0} {1}", Action.CommandPath, Action.CommandArguments);
+						Log.WriteException(ex, null);
+						ExitCode = 1;
+						bComplete = true;
+						return;
+					}
+
+					// wait for process to start
+					// NOTE: this may or may not be necessary; seems to depend on whether the system UBT is running on start the process in a timely manner.
+					int checkIterations = 0;
+					bool haveConfiguredProcess = false;
+					do
+					{
+						if (ActionProcess.HasExited)
+						{
+							if (haveConfiguredProcess == false)
+								Debug.WriteLine("Process for action exited before able to configure!");
+							break;
+						}
+
+						if (!haveConfiguredProcess)
+						{
+							try
+							{
+								ActionProcess.PriorityClass = ProcessPriorityClass.BelowNormal;
+								haveConfiguredProcess = true;
+							}
+							catch (Exception)
+							{
+							}
+							break;
+						}
+
+						Thread.Sleep(10);
+
+						checkIterations++;
+					} while (checkIterations < 100);
+					if (checkIterations == 100)
+					{
+						throw new BuildException("Failed to configure local process for action: {0} {1}", Action.CommandPath, Action.CommandArguments);
+					}
+
+					// block until it's complete
+					// @todo iosmerge: UBT had started looking at:	if (Utils.IsValidProcess(Process))
+					//    do we need to check that in the thread model?
+					ActionProcess.WaitForExit();
+
+					// capture exit code
+					ExitCode = ActionProcess.ExitCode;
+				}
+				finally
+				{
+					// As the process has finished now, free its resources. On non-Windows platforms, processes depend 
+					// on POSIX/BSD threading and these are limited per application. Disposing the Process releases 
+					// these thread resources.
+					if (ActionProcess != null)
+						ActionProcess.Close();
+				}
+			}
+
+			// track how long it took
+			Action.EndTime = DateTimeOffset.Now;
+
+			if (!Utils.IsRunningOnMono)
+			{
+				// let RPCUtilHelper clean up anything thread related
+				RPCUtilHelper.OnThreadComplete();
+			}
+
+			// we are done!!
+			bComplete = true;
+		}
+
+		/// <summary>
+		/// Starts a thread and runs the action in that thread
+		/// </summary>
+		public void Run()
+		{
+			Thread T = new Thread(ThreadFunc);
+			T.Start();
+		}
+	};
+
+	class LocalExecutor : ActionExecutor
+	{
+		/// <summary>
+		/// Processor count multiplier for local execution. Can be below 1 to reserve CPU for other tasks.
+		/// When using the local executor (not XGE), run a single action on each CPU core.  Note that you can set this to a larger value
+		/// to get slightly faster build times in many cases, but your computer's responsiveness during compiling may be much worse.
+		/// </summary>
+		[XmlConfigFile(Category = "BuildConfiguration")]
+		double ProcessorCountMultiplier = 1.0;
+
+		/// <summary>
+		/// Maximum processor count for local execution. 
+		/// </summary>
+		[XmlConfigFile(Category = "BuildConfiguration")]
+		int MaxProcessorCount = int.MaxValue;
+
+		public LocalExecutor()
+		{
+			XmlConfig.ApplyTo(this);
+		}
+
+		public override string Name
+		{
+			get { return "Local"; }
+		}
+
+		/// <summary>
+		/// Determines the maximum number of actions to execute in parallel, taking into account the resources available on this machine.
+		/// </summary>
+		/// <returns>Max number of actions to execute in parallel</returns>
+		public virtual int GetMaxActionsToExecuteInParallel()
+		{
+			// Get the number of logical processors
+			int NumLogicalCores = Utils.GetLogicalProcessorCount();
+
+			// Use WMI to figure out physical cores, excluding hyper threading.
+			int NumPhysicalCores = Utils.GetPhysicalProcessorCount();
+			if (NumPhysicalCores == -1)
+			{
+				NumPhysicalCores = NumLogicalCores;
+			}
+
+			// The number of actions to execute in parallel is trying to keep the CPU busy enough in presence of I/O stalls.
+			int MaxActionsToExecuteInParallel = 0;
+			if (NumPhysicalCores < NumLogicalCores && ProcessorCountMultiplier != 1.0)
+			{
+				// The CPU has more logical cores than physical ones, aka uses hyper-threading. 
+				// Use multiplier if provided
+				MaxActionsToExecuteInParallel = (int)(NumPhysicalCores * ProcessorCountMultiplier);
+			}
+			else if (NumPhysicalCores < NumLogicalCores && NumPhysicalCores > 4)
+			{
+				// The CPU has more logical cores than physical ones, aka uses hyper-threading. 
+				// Use average of logical and physical if we have "lots of cores"
+				MaxActionsToExecuteInParallel = Math.Max((int)(NumPhysicalCores + NumLogicalCores) / 2, NumLogicalCores - 4);
+			}
+			// No hyper-threading. Only kicking off a task per CPU to keep machine responsive.
+			else
+			{
+				MaxActionsToExecuteInParallel = NumPhysicalCores;
+			}
+
+#if !NET_CORE
+			if (Utils.IsRunningOnMono)
+			{
+				long PhysicalRAMAvailableMB = (new PerformanceCounter("Mono Memory", "Total Physical Memory").RawValue) / (1024 * 1024);
+				// heuristic: give each action at least 1.5GB of RAM (some clang instances will need more) if the total RAM is low, or 1GB on 16+GB machines
+				long MinMemoryPerActionMB = (PhysicalRAMAvailableMB < 16384) ? 3 * 1024 / 2 : 1024;
+				int MaxActionsAffordedByMemory = (int)(Math.Max(1, (PhysicalRAMAvailableMB) / MinMemoryPerActionMB));
+
+				MaxActionsToExecuteInParallel = Math.Min(MaxActionsToExecuteInParallel, MaxActionsAffordedByMemory);
+			}
+#endif
+
+			MaxActionsToExecuteInParallel = Math.Max(1, Math.Min(MaxActionsToExecuteInParallel, MaxProcessorCount));
+			return MaxActionsToExecuteInParallel;
+		}
+
+		/// <summary>
+		/// Executes the specified actions locally.
+		/// </summary>
+		/// <returns>True if all the tasks successfully executed, or false if any of them failed.</returns>
+		public override bool ExecuteActions(List<Action> Actions, bool bLogDetailedActionStats)
+		{
+			// Time to sleep after each iteration of the loop in order to not busy wait.
+			const float LoopSleepTime = 0.1f;
+
+			// The number of actions to execute in parallel is trying to keep the CPU busy enough in presence of I/O stalls.
+			int MaxActionsToExecuteInParallel = GetMaxActionsToExecuteInParallel();
+			Log.TraceInformation("Performing {0} actions ({1} in parallel)", Actions.Count, MaxActionsToExecuteInParallel);
+
+			Dictionary<Action, ActionThread> ActionThreadDictionary = new Dictionary<Action, ActionThread>();
+			int JobNumber = 1;
+			using (ProgressWriter ProgressWriter = new ProgressWriter("Compiling C++ source code...", false))
+			{
+				int ProgressValue = 0;
+				while (true)
+				{
+					// Count the number of pending and still executing actions.
+					int NumUnexecutedActions = 0;
+					int NumExecutingActions = 0;
+					foreach (Action Action in Actions)
+					{
+						ActionThread ActionThread = null;
+						bool bFoundActionProcess = ActionThreadDictionary.TryGetValue(Action, out ActionThread);
+						if (bFoundActionProcess == false)
+						{
+							NumUnexecutedActions++;
+						}
+						else if (ActionThread != null)
+						{
+							if (ActionThread.bComplete == false)
+							{
+								NumUnexecutedActions++;
+								NumExecutingActions++;
+							}
+						}
+					}
+
+					// Update the current progress
+					int NewProgressValue = Actions.Count + 1 - NumUnexecutedActions;
+					if (ProgressValue != NewProgressValue)
+					{
+						ProgressWriter.Write(ProgressValue, Actions.Count + 1);
+						ProgressValue = NewProgressValue;
+					}
+
+					// If there aren't any pending actions left, we're done executing.
+					if (NumUnexecutedActions == 0)
+					{
+						break;
+					}
+
+					// If there are fewer actions executing than the maximum, look for pending actions that don't have any outdated
+					// prerequisites.
+					foreach (Action Action in Actions)
+					{
+						ActionThread ActionProcess = null;
+						bool bFoundActionProcess = ActionThreadDictionary.TryGetValue(Action, out ActionProcess);
+						if (bFoundActionProcess == false)
+						{
+							if (NumExecutingActions < Math.Max(1, MaxActionsToExecuteInParallel))
+							{
+								// Determine whether there are any prerequisites of the action that are outdated.
+								bool bHasOutdatedPrerequisites = false;
+								bool bHasFailedPrerequisites = false;
+								foreach (FileItem PrerequisiteItem in Action.PrerequisiteItems)
+								{
+									if (PrerequisiteItem.ProducingAction != null && Actions.Contains(PrerequisiteItem.ProducingAction))
+									{
+										ActionThread PrerequisiteProcess = null;
+										bool bFoundPrerequisiteProcess = ActionThreadDictionary.TryGetValue(PrerequisiteItem.ProducingAction, out PrerequisiteProcess);
+										if (bFoundPrerequisiteProcess == true)
+										{
+											if (PrerequisiteProcess == null)
+											{
+												bHasFailedPrerequisites = true;
+											}
+											else if (PrerequisiteProcess.bComplete == false)
+											{
+												bHasOutdatedPrerequisites = true;
+											}
+											else if (PrerequisiteProcess.ExitCode != 0)
+											{
+												bHasFailedPrerequisites = true;
+											}
+										}
+										else
+										{
+											bHasOutdatedPrerequisites = true;
+										}
+									}
+								}
+
+								// If there are any failed prerequisites of this action, don't execute it.
+								if (bHasFailedPrerequisites)
+								{
+									// Add a null entry in the dictionary for this action.
+									ActionThreadDictionary.Add(Action, null);
+								}
+								// If there aren't any outdated prerequisites of this action, execute it.
+								else if (!bHasOutdatedPrerequisites)
+								{
+									ActionThread ActionThread = new ActionThread(Action, JobNumber, Actions.Count);
+									JobNumber++;
+
+									try
+									{
+										ActionThread.Run();
+									}
+									catch (Exception ex)
+									{
+										throw new BuildException(ex, "Failed to start thread for action: {0} {1}\r\n{2}", Action.CommandPath, Action.CommandArguments, ex.ToString());
+									}
+
+									ActionThreadDictionary.Add(Action, ActionThread);
+
+									NumExecutingActions++;
+								}
+							}
+						}
+					}
+
+					System.Threading.Thread.Sleep(TimeSpan.FromSeconds(LoopSleepTime));
+				}
+			}
+
+			Log.WriteLineIf(bLogDetailedActionStats, LogEventType.Console, "-------- Begin Detailed Action Stats ----------------------------------------------------------");
+			Log.WriteLineIf(bLogDetailedActionStats, LogEventType.Console, "^Action Type^Duration (seconds)^Tool^Task^Using PCH");
+
+			double TotalThreadSeconds = 0;
+
+			// Check whether any of the tasks failed and log action stats if wanted.
+			bool bSuccess = true;
+			double TotalBuildProjectTime = 0;
+			double TotalCompileTime = 0;
+			double TotalCreateAppBundleTime = 0;
+			double TotalGenerateDebugInfoTime = 0;
+			double TotalLinkTime = 0;
+			double TotalOtherActionsTime = 0;
+			foreach (KeyValuePair<Action, ActionThread> ActionProcess in ActionThreadDictionary)
+			{
+				Action Action = ActionProcess.Key;
+				ActionThread ActionThread = ActionProcess.Value;
+
+				// Check for pending actions, preemptive failure
+				if (ActionThread == null)
+				{
+					bSuccess = false;
+					continue;
+				}
+				// Check for executed action but general failure
+				if (ActionThread.ExitCode != 0)
+				{
+					bSuccess = false;
+				}
+				// Log CPU time, tool and task.
+				double ThreadSeconds = Action.Duration.TotalSeconds;
+
+				Log.WriteLineIf(bLogDetailedActionStats,
+					LogEventType.Console,
+					"^{0}^{1:0.00}^{2}^{3}^{4}",
+					Action.ActionType.ToString(),
+					ThreadSeconds,
+					Path.GetFileName(Action.CommandPath),
+					  Action.StatusDescription,
+					Action.bIsUsingPCH);
+
+				// Update statistics
+				switch (Action.ActionType)
+				{
+					case ActionType.BuildProject:
+						TotalBuildProjectTime += ThreadSeconds;
+						break;
+
+					case ActionType.Compile:
+						TotalCompileTime += ThreadSeconds;
+						break;
+
+					case ActionType.CreateAppBundle:
+						TotalCreateAppBundleTime += ThreadSeconds;
+						break;
+
+					case ActionType.GenerateDebugInfo:
+						TotalGenerateDebugInfoTime += ThreadSeconds;
+						break;
+
+					case ActionType.Link:
+						TotalLinkTime += ThreadSeconds;
+						break;
+
+					default:
+						TotalOtherActionsTime += ThreadSeconds;
+						break;
+				}
+
+				// Keep track of total thread seconds spent on tasks.
+				TotalThreadSeconds += ThreadSeconds;
+			}
+
+			Log.WriteLineIf(bLogDetailedActionStats, LogEventType.Console, "-------- End Detailed Actions Stats -----------------------------------------------------------");
+
+			// Log total CPU seconds and numbers of processors involved in tasks.
+			Log.WriteLineIf(bLogDetailedActionStats || UnrealBuildTool.bPrintDebugInfo,
+				LogEventType.Console, "Cumulative thread seconds ({0} processors): {1:0.00}", System.Environment.ProcessorCount, TotalThreadSeconds);
+
+			// Log detailed stats
+			Log.WriteLineIf(bLogDetailedActionStats || UnrealBuildTool.bPrintDebugInfo,
+				LogEventType.Console,
+				"Cumulative action seconds ({0} processors): {1:0.00} building projects, {2:0.00} compiling, {3:0.00} creating app bundles, {4:0.00} generating debug info, {5:0.00} linking, {6:0.00} other",
+				System.Environment.ProcessorCount,
+				TotalBuildProjectTime,
+				TotalCompileTime,
+				TotalCreateAppBundleTime,
+				TotalGenerateDebugInfoTime,
+				TotalLinkTime,
+				TotalOtherActionsTime
+			);
+
+			return bSuccess;
+		}
+	};
+}
diff --git a/Engine/Source/Programs/UnrealBuildTool/Executors/ParallelExecutor.cs b/Engine/Source/Programs/UnrealBuildTool/Executors/ParallelExecutor.cs
new file mode 100644
--- /dev/null
+++ b/Engine/Source/Programs/UnrealBuildTool/Executors/ParallelExecutor.cs
@@ -0,0 +1,283 @@
+// Copyright 1998-2018 Epic Games, Inc. All Rights Reserved.
+
+using System;
+using System.Collections.Generic;
+using System.Diagnostics;
+using System.Linq;
+using System.Text;
+using System.Threading;
+using System.Threading.Tasks;
+using Tools.DotNETCommon;
+
+namespace UnrealBuildTool
+{
+	/// <summary>
+	/// This executor is similar to LocalExecutor, but uses p/invoke on Windows to ensure that child processes are started at a lower priority and are terminated when the parent process terminates.
+	/// </summary>
+	class ParallelExecutor : ActionExecutor
+	{
+		[DebuggerDisplay("{Inner}")]
+		class BuildAction
+		{
+			public int SortIndex;
+			public Action Inner;
+
+			public HashSet<BuildAction> Dependencies = new HashSet<BuildAction>();
+			public int MissingDependencyCount;
+
+			public HashSet<BuildAction> Dependants = new HashSet<BuildAction>();
+			public int TotalDependantCount;
+
+			public List<string> LogLines = new List<string>();
+			public int ExitCode = -1;
+		}
+
+		/// <summary>
+		/// Processor count multiplier for local execution. Can be below 1 to reserve CPU for other tasks.
+		/// </summary>
+		[XmlConfigFile(Category = "BuildConfiguration")]
+		double ProcessorCountMultiplier = 1.0;
+
+		/// <summary>
+		/// Maximum processor count for local execution. 
+		/// </summary>
+		[XmlConfigFile(Category = "BuildConfiguration")]
+		int MaxProcessorCount = int.MaxValue;
+
+		/// <summary>
+		/// When enabled, will stop compiling targets after a compile error occurs.
+		/// </summary>
+		[XmlConfigFile(Category = "BuildConfiguration")]
+		bool bStopCompilationAfterErrors = false;
+
+		/// <summary>
+		/// Constructor
+		/// </summary>
+		public ParallelExecutor()
+		{
+			XmlConfig.ApplyTo(this);
+		}
+
+		/// <summary>
+		/// Returns the name of this executor
+		/// </summary>
+		public override string Name
+		{
+			get { return "Parallel"; }
+		}
+
+		/// <summary>
+		/// Checks whether the parallel executor can be used
+		/// </summary>
+		/// <returns>True if the parallel executor can be used</returns>
+		public static bool IsAvailable()
+		{
+			return Environment.OSVersion.Platform == PlatformID.Win32NT;
+		}
+
+		/// <summary>
+		/// Executes the specified actions locally.
+		/// </summary>
+		/// <returns>True if all the tasks successfully executed, or false if any of them failed.</returns>
+		public override bool ExecuteActions(List<Action> InputActions, bool bLogDetailedActionStats)
+		{
+			// Figure out how many processors to use
+			int MaxProcesses = Math.Min((int)(Environment.ProcessorCount * ProcessorCountMultiplier), MaxProcessorCount);
+			Log.TraceInformation("Building {0} {1} with {2} {3}...", InputActions.Count, (InputActions.Count == 1) ? "action" : "actions", MaxProcesses, (MaxProcesses == 1)? "process" : "processes");
+
+			// Create actions with all our internal metadata
+			List<BuildAction> Actions = new List<BuildAction>();
+			for(int Idx = 0; Idx < InputActions.Count; Idx++)
+			{
+				BuildAction Action = new BuildAction();
+				Action.SortIndex = Idx;
+				Action.Inner = InputActions[Idx];
+				Actions.Add(Action);
+			}
+
+			// Build a map of items to their producing actions
+			Dictionary<FileItem, BuildAction> FileToProducingAction = new Dictionary<FileItem, BuildAction>();
+			foreach(BuildAction Action in Actions)
+			{
+				foreach(FileItem ProducedItem in Action.Inner.ProducedItems)
+				{
+					FileToProducingAction[ProducedItem] = Action;
+				}
+			}
+
+			// Update all the actions with all their dependencies
+			foreach(BuildAction Action in Actions)
+			{
+				foreach(FileItem PrerequisiteItem in Action.Inner.PrerequisiteItems)
+				{
+					BuildAction Dependency;
+					if(FileToProducingAction.TryGetValue(PrerequisiteItem, out Dependency))
+					{
+						Action.Dependencies.Add(Dependency);
+						Dependency.Dependants.Add(Action);
+					}
+				}
+			}
+
+			// Figure out the recursive dependency count
+			HashSet<BuildAction> VisitedActions = new HashSet<BuildAction>();
+			foreach(BuildAction Action in Actions)
+			{
+				Action.MissingDependencyCount = Action.Dependencies.Count;
+				RecursiveIncDependents(Action, VisitedActions);
+			}
+
+			// Create the list of things to process
+			List<BuildAction> QueuedActions = new List<BuildAction>();
+			foreach(BuildAction Action in Actions)
+			{
+				if(Action.MissingDependencyCount == 0)
+				{
+					QueuedActions.Add(Action);
+				}
+			}
+
+			// Execute the actions
+			using (LogIndentScope Indent = new LogIndentScope("  "))
+			{
+				// Create a job object for all the child processes
+				bool bResult = true;
+				Dictionary<BuildAction, Thread> ExecutingActions = new Dictionary<BuildAction,Thread>();
+				List<BuildAction> CompletedActions = new List<BuildAction>();
+
+				using(ManagedProcessGroup ProcessGroup = new ManagedProcessGroup())
+				{
+					using(AutoResetEvent CompletedEvent = new AutoResetEvent(false))
+					{
+						int NumCompletedActions = 0;
+						using (ProgressWriter ProgressWriter = new ProgressWriter("Compiling C++ source code...", false))
+						{
+							while(QueuedActions.Count > 0 || ExecutingActions.Count > 0)
+							{
+								// Sort the actions by the number of things dependent on them
+								QueuedActions.Sort((A, B) => (A.TotalDependantCount == B.TotalDependantCount)? (B.SortIndex - A.SortIndex) : (B.TotalDependantCount - A.TotalDependantCount));
+
+								// Create threads up to the maximum number of actions
+								while(ExecutingActions.Count < MaxProcesses && QueuedActions.Count > 0)
+								{
+									BuildAction Action = QueuedActions[QueuedActions.Count - 1];
+									QueuedActions.RemoveAt(QueuedActions.Count - 1);
+
+									Thread ExecutingThread = new Thread(() => { ExecuteAction(ProcessGroup, Action, CompletedActions, CompletedEvent); });
+									ExecutingThread.Name = String.Format("Build:{0}", Action.Inner.StatusDescription);
+									ExecutingThread.Start();
+
+									ExecutingActions.Add(Action, ExecutingThread);
+								}
+
+								// Wait for something to finish
+								CompletedEvent.WaitOne();
+
+								// Wait for something to finish and flush it to the log
+								lock(CompletedActions)
+								{
+									foreach(BuildAction CompletedAction in CompletedActions)
+									{
+										// Join the thread
+										Thread CompletedThread = ExecutingActions[CompletedAction];
+										CompletedThread.Join();
+										ExecutingActions.Remove(CompletedAction);
+
+										// Write it to the log
+										if(CompletedAction.LogLines.Count > 0)
+										{
+											foreach(string LogLine in CompletedAction.LogLines)
+											{
+												Log.TraceInformation(LogLine);
+											}
+										}
+
+										// Update the progress
+										NumCompletedActions++;
+										ProgressWriter.Write(NumCompletedActions, InputActions.Count);
+
+										// Check the exit code
+										if(CompletedAction.ExitCode == 0)
+										{
+											// Mark all the dependents as done
+											foreach(BuildAction DependantAction in CompletedAction.Dependants)
+											{
+												if(--DependantAction.MissingDependencyCount == 0)
+												{
+													QueuedActions.Add(DependantAction);
+												}
+											}
+										}
+										else
+										{
+											// Update the exit code if it's not already set
+											if(bResult && CompletedAction.ExitCode != 0)
+											{
+												bResult = false;
+											}
+										}
+									}
+									CompletedActions.Clear();
+								}
+
+								// If we've already got a non-zero exit code, clear out the list of queued actions so nothing else will run
+								if(!bResult && bStopCompilationAfterErrors)
+								{
+									QueuedActions.Clear();
+								}
+							}
+						}
+					}
+				}
+
+				return bResult;
+			}
+		}
+
+		/// <summary>
+		/// Execute an individual action
+		/// </summary>
+		/// <param name="ProcessGroup">The process group</param>
+		/// <param name="Action">The action to execute</param>
+		/// <param name="CompletedActions">On completion, the list to add the completed action to</param>
+		/// <param name="CompletedEvent">Event to set once an event is complete</param>
+		static void ExecuteAction(ManagedProcessGroup ProcessGroup, BuildAction Action, List<BuildAction> CompletedActions, AutoResetEvent CompletedEvent)
+		{
+			if (Action.Inner.bShouldOutputStatusDescription && !String.IsNullOrEmpty(Action.Inner.StatusDescription))
+			{
+				Action.LogLines.Add(Action.Inner.StatusDescription);
+			}
+
+			using(ManagedProcess Process = new ManagedProcess(ProcessGroup, Action.Inner.CommandPath, Action.Inner.CommandArguments, Action.Inner.WorkingDirectory, null, null, ProcessPriorityClass.BelowNormal))
+			{
+				Action.LogLines.AddRange(Process.ReadAllLines());
+				Action.ExitCode = Process.ExitCode;
+			}
+
+			lock(CompletedActions)
+			{
+				CompletedActions.Add(Action);
+			}
+
+			CompletedEvent.Set();
+		}
+
+		/// <summary>
+		/// Increment the number of dependants of an action, recursively
+		/// </summary>
+		/// <param name="Action">The action to update</param>
+		/// <param name="VisitedActions">Set of visited actions</param>
+		private static void RecursiveIncDependents(BuildAction Action, HashSet<BuildAction> VisitedActions)
+		{
+			foreach(BuildAction Dependency in Action.Dependants)
+			{
+				if(!VisitedActions.Contains(Action))
+				{
+					VisitedActions.Add(Action);
+					Dependency.TotalDependantCount++;
+					RecursiveIncDependents(Dependency, VisitedActions);
+				}
+			}
+		}
+	}
+}
diff --git a/Engine/Source/Programs/UnrealBuildTool/Executors/SNDBS.cs b/Engine/Source/Programs/UnrealBuildTool/Executors/SNDBS.cs
new file mode 100644
--- /dev/null
+++ b/Engine/Source/Programs/UnrealBuildTool/Executors/SNDBS.cs
@@ -0,0 +1,429 @@
+﻿// Copyright 1998-2018 Epic Games, Inc. All Rights Reserved.
+
+using System;
+using System.Collections.Generic;
+using System.IO;
+using System.Diagnostics;
+using System.Xml;
+using System.Text.RegularExpressions;
+using System.Linq;
+using System.Reflection;
+using System.Threading;
+using System.ServiceProcess;
+using Tools.DotNETCommon;
+
+namespace UnrealBuildTool
+{
+	class SNDBS : ActionExecutor
+	{
+		/// <summary>
+		/// Processor count multiplier for local execution. Can be below 1 to reserve CPU for other tasks.
+		/// </summary>
+		[XmlConfigFile]
+		public double ProcessorCountMultiplier = 1.0;
+
+		/// <summary>
+		/// Maximum processor count for local execution. 
+		/// </summary>
+		[XmlConfigFile]
+		public int MaxProcessorCount = int.MaxValue;
+
+		/// <summary>
+		/// The number of actions to execute in parallel is trying to keep the CPU busy enough in presence of I/O stalls.
+		/// </summary>
+		int MaxActionsToExecuteInParallel;
+
+		/// <summary>
+		/// Unique id for new jobs
+		/// </summary>
+		int JobNumber;
+
+		public SNDBS()
+		{
+			XmlConfig.ApplyTo(this);
+		}
+
+		public override string Name
+		{
+			get { return "SNDBS"; }
+		}
+
+		/// <summary>
+		/// Used when debugging Actions outputs all action return values to debug out
+		/// </summary>
+		/// <param name="sender"> Sending object</param>
+		/// <param name="e">  Event arguments (In this case, the line of string output)</param>
+		static protected void ActionDebugOutput(object sender, DataReceivedEventArgs e)
+		{
+			string Output = e.Data;
+			if (Output == null)
+			{
+				return;
+			}
+
+			Log.TraceInformation(Output);
+		}
+
+		internal bool ExecuteLocalActions(List<Action> InLocalActions, Dictionary<Action, ActionThread> InActionThreadDictionary, int TotalNumJobs)
+		{
+			// Time to sleep after each iteration of the loop in order to not busy wait.
+			const float LoopSleepTime = 0.1f;
+
+			bool LocalActionsResult = true;
+
+			while (true)
+			{
+				// Count the number of pending and still executing actions.
+				int NumUnexecutedActions = 0;
+				int NumExecutingActions = 0;
+				foreach (Action Action in InLocalActions)
+				{
+					ActionThread ActionThread = null;
+					bool bFoundActionProcess = InActionThreadDictionary.TryGetValue(Action, out ActionThread);
+					if (bFoundActionProcess == false)
+					{
+						NumUnexecutedActions++;
+					}
+					else if (ActionThread != null)
+					{
+						if (ActionThread.bComplete == false)
+						{
+							NumUnexecutedActions++;
+							NumExecutingActions++;
+						}
+					}
+				}
+
+				// If there aren't any pending actions left, we're done executing.
+				if (NumUnexecutedActions == 0)
+				{
+					break;
+				}
+
+				// If there are fewer actions executing than the maximum, look for pending actions that don't have any outdated
+				// prerequisites.
+				foreach (Action Action in InLocalActions)
+				{
+					ActionThread ActionProcess = null;
+					bool bFoundActionProcess = InActionThreadDictionary.TryGetValue(Action, out ActionProcess);
+					if (bFoundActionProcess == false)
+					{
+						if (NumExecutingActions < Math.Max(1, MaxActionsToExecuteInParallel))
+						{
+							// Determine whether there are any prerequisites of the action that are outdated.
+							bool bHasOutdatedPrerequisites = false;
+							bool bHasFailedPrerequisites = false;
+							foreach (FileItem PrerequisiteItem in Action.PrerequisiteItems)
+							{
+								if (PrerequisiteItem.ProducingAction != null && InLocalActions.Contains(PrerequisiteItem.ProducingAction))
+								{
+									ActionThread PrerequisiteProcess = null;
+									bool bFoundPrerequisiteProcess = InActionThreadDictionary.TryGetValue(PrerequisiteItem.ProducingAction, out PrerequisiteProcess);
+									if (bFoundPrerequisiteProcess == true)
+									{
+										if (PrerequisiteProcess == null)
+										{
+											bHasFailedPrerequisites = true;
+										}
+										else if (PrerequisiteProcess.bComplete == false)
+										{
+											bHasOutdatedPrerequisites = true;
+										}
+										else if (PrerequisiteProcess.ExitCode != 0)
+										{
+											bHasFailedPrerequisites = true;
+										}
+									}
+									else
+									{
+										bHasOutdatedPrerequisites = true;
+									}
+								}
+							}
+
+							// If there are any failed prerequisites of this action, don't execute it.
+							if (bHasFailedPrerequisites)
+							{
+								// Add a null entry in the dictionary for this action.
+								InActionThreadDictionary.Add(Action, null);
+							}
+							// If there aren't any outdated prerequisites of this action, execute it.
+							else if (!bHasOutdatedPrerequisites)
+							{
+								ActionThread ActionThread = new ActionThread(Action, JobNumber, TotalNumJobs);
+								ActionThread.Run();
+
+								InActionThreadDictionary.Add(Action, ActionThread);
+
+								NumExecutingActions++;
+								JobNumber++;
+							}
+						}
+					}
+				}
+
+				System.Threading.Thread.Sleep(TimeSpan.FromSeconds(LoopSleepTime));
+			}
+
+			return LocalActionsResult;
+		}
+
+		internal bool ExecuteActions(List<Action> InActions, Dictionary<Action, ActionThread> InActionThreadDictionary)
+		{
+			// Build the script file that will be executed by SN-DBS
+			StreamWriter ScriptFile;
+			string ScriptFilename = Path.Combine(UnrealBuildTool.EngineDirectory.FullName, "Intermediate", "Build", "SNDBS.bat");
+
+			FileStream ScriptFileStream = new FileStream(ScriptFilename, FileMode.Create, FileAccess.ReadWrite, FileShare.Read);
+			ScriptFile = new StreamWriter(ScriptFileStream);
+			ScriptFile.AutoFlush = true;
+
+			int NumScriptedActions = 0;
+			List<Action> LocalActions = new List<Action>();
+			ActionThread DummyActionThread = new ActionThread(null, 1, 1);
+			foreach (Action Action in InActions)
+			{
+				ActionThread ActionProcess = null;
+				bool bFoundActionProcess = InActionThreadDictionary.TryGetValue(Action, out ActionProcess);
+				if (bFoundActionProcess == false)
+				{
+					// Determine whether there are any prerequisites of the action that are outdated.
+					bool bHasOutdatedPrerequisites = false;
+					bool bHasFailedPrerequisites = false;
+					foreach (FileItem PrerequisiteItem in Action.PrerequisiteItems)
+					{
+						if (PrerequisiteItem.ProducingAction != null && InActions.Contains(PrerequisiteItem.ProducingAction))
+						{
+							ActionThread PrerequisiteProcess = null;
+							bool bFoundPrerequisiteProcess = InActionThreadDictionary.TryGetValue(PrerequisiteItem.ProducingAction, out PrerequisiteProcess);
+							if (bFoundPrerequisiteProcess == true)
+							{
+								if (PrerequisiteProcess == null)
+								{
+									bHasFailedPrerequisites = true;
+								}
+								else if (PrerequisiteProcess.bComplete == false)
+								{
+									bHasOutdatedPrerequisites = true;
+								}
+								else if (PrerequisiteProcess.ExitCode != 0)
+								{
+									bHasFailedPrerequisites = true;
+								}
+							}
+							else
+							{
+								bHasOutdatedPrerequisites = true;
+							}
+						}
+					}
+
+					// If there are any failed prerequisites of this action, don't execute it.
+					if (bHasFailedPrerequisites)
+					{
+						// Add a null entry in the dictionary for this action.
+						InActionThreadDictionary.Add(Action, null);
+					}
+					// If there aren't any outdated prerequisites of this action, execute it.
+					else if (!bHasOutdatedPrerequisites)
+					{
+						if (Action.bCanExecuteRemotely == false || Action.bCanExecuteRemotelyWithSNDBS == false)
+						{
+							// Execute locally
+							LocalActions.Add(Action);
+						}
+						else
+						{
+							// Add to script for execution by SN-DBS
+							string NewCommandArguments = "\"" + Action.CommandPath + "\"" + " " + Action.CommandArguments;
+							ScriptFile.WriteLine(ActionThread.ExpandEnvironmentVariables(NewCommandArguments));
+							InActionThreadDictionary.Add(Action, DummyActionThread);
+							Action.StartTime = Action.EndTime = DateTimeOffset.Now;
+							Log.TraceInformation("[{0}/{1}] {2} {3}", JobNumber, InActions.Count, Action.CommandDescription, Action.StatusDescription);
+							JobNumber++;
+							NumScriptedActions++;
+						}
+					}
+				}
+			}
+
+			ScriptFile.Flush();
+			ScriptFile.Close();
+			ScriptFile.Dispose();
+			ScriptFile = null;
+
+			if (NumScriptedActions > 0)
+			{
+                // Create the process
+                string SCERoot = Environment.GetEnvironmentVariable("SCE_ROOT_DIR");
+                string SNDBSExecutable = Path.Combine(SCERoot, "Common/SN-DBS/bin/dbsbuild.exe");
+                ProcessStartInfo PSI = new ProcessStartInfo(SNDBSExecutable, String.Format("-q -p UE4 -s \"{0}\"", FileReference.Combine(UnrealBuildTool.EngineDirectory, "Intermediate", "Build", "sndbs.bat").FullName));
+				PSI.RedirectStandardOutput = true;
+				PSI.RedirectStandardError = true;
+				PSI.UseShellExecute = false;
+				PSI.CreateNoWindow = true;
+				PSI.WorkingDirectory = Path.GetFullPath("."); ;
+				Process NewProcess = new Process();
+				NewProcess.StartInfo = PSI;
+				NewProcess.OutputDataReceived += new DataReceivedEventHandler(ActionDebugOutput);
+				NewProcess.ErrorDataReceived += new DataReceivedEventHandler(ActionDebugOutput);
+				DateTimeOffset StartTime = DateTimeOffset.Now;
+				NewProcess.Start();
+				NewProcess.BeginOutputReadLine();
+				NewProcess.BeginErrorReadLine();
+
+				NewProcess.WaitForExit();
+
+				TimeSpan Duration;
+				DateTimeOffset EndTime = DateTimeOffset.Now;
+				if (EndTime == DateTimeOffset.MinValue)
+				{
+					Duration = DateTimeOffset.Now - StartTime;
+				}
+				else
+				{
+					Duration = EndTime - StartTime;
+				}
+
+				DummyActionThread.bComplete = true;
+				int ExitCode = NewProcess.ExitCode;
+				if (ExitCode != 0)
+				{
+					return false;
+				}
+			}
+
+			// Execute local tasks
+			if (LocalActions.Count > 0)
+			{
+				return ExecuteLocalActions(LocalActions, InActionThreadDictionary, InActions.Count);
+			}
+
+			return true;
+		}
+
+		public static bool IsAvailable()
+		{
+			string SCERoot = Environment.GetEnvironmentVariable("SCE_ROOT_DIR");
+			if(SCERoot == null)
+			{
+				return false;
+			}
+			if (!File.Exists(Path.Combine(SCERoot, "Common/SN-DBS/bin/dbsbuild.exe")))
+			{
+				return false;
+			}
+
+			ServiceController[] services = ServiceController.GetServices();
+			foreach (ServiceController service in services)
+			{
+				if (service.ServiceName.StartsWith("SNDBS") && service.Status == ServiceControllerStatus.Running)
+				{
+					return true;
+				}
+			}
+			return false;
+		}
+
+		public override bool ExecuteActions(List<Action> Actions, bool bLogDetailedActionStats)
+		{
+			bool SNDBSResult = true;
+			if (Actions.Count > 0)
+			{
+				// Use WMI to figure out physical cores, excluding hyper threading.
+				int NumCores = Utils.GetPhysicalProcessorCount();
+				
+				// If we failed to detect the number of cores, default to the logical processor count
+				if (NumCores == -1)
+				{
+					NumCores = System.Environment.ProcessorCount;
+				}
+				// The number of actions to execute in parallel is trying to keep the CPU busy enough in presence of I/O stalls.
+				MaxActionsToExecuteInParallel = 0;
+				// The CPU has more logical cores than physical ones, aka uses hyper-threading. 
+				if (NumCores < System.Environment.ProcessorCount)
+				{
+					MaxActionsToExecuteInParallel = (int)(NumCores * ProcessorCountMultiplier);
+				}
+				// No hyper-threading. Only kicking off a task per CPU to keep machine responsive.
+				else
+				{
+					MaxActionsToExecuteInParallel = NumCores;
+				}
+				MaxActionsToExecuteInParallel = Math.Min(MaxActionsToExecuteInParallel, MaxProcessorCount);
+
+				JobNumber = 1;
+				Dictionary<Action, ActionThread> ActionThreadDictionary = new Dictionary<Action, ActionThread>();
+
+				while (true)
+				{
+					bool bUnexecutedActions = false;
+					foreach (Action Action in Actions)
+					{
+						ActionThread ActionThread = null;
+						bool bFoundActionProcess = ActionThreadDictionary.TryGetValue(Action, out ActionThread);
+						if (bFoundActionProcess == false)
+						{
+							bUnexecutedActions = true;
+							if(!ExecuteActions(Actions, ActionThreadDictionary))
+							{
+								return false;
+							}
+							break;
+						}
+					}
+
+					if (bUnexecutedActions == false)
+					{
+						break;
+					}
+				}
+
+				Log.WriteLineIf(bLogDetailedActionStats, LogEventType.Console, "-------- Begin Detailed Action Stats ----------------------------------------------------------");
+				Log.WriteLineIf(bLogDetailedActionStats, LogEventType.Console, "^Action Type^Duration (seconds)^Tool^Task^Using PCH");
+
+				double TotalThreadSeconds = 0;
+
+				// Check whether any of the tasks failed and log action stats if wanted.
+				foreach (KeyValuePair<Action, ActionThread> ActionProcess in ActionThreadDictionary)
+				{
+					Action Action = ActionProcess.Key;
+					ActionThread ActionThread = ActionProcess.Value;
+
+					// Check for pending actions, preemptive failure
+					if (ActionThread == null)
+					{
+						SNDBSResult = false;
+						continue;
+					}
+					// Check for executed action but general failure
+					if (ActionThread.ExitCode != 0)
+					{
+						SNDBSResult = false;
+					}
+					// Log CPU time, tool and task.
+					double ThreadSeconds = Action.Duration.TotalSeconds;
+
+					Log.WriteLineIf(bLogDetailedActionStats,
+						LogEventType.Console,
+						"^{0}^{1:0.00}^{2}^{3}^{4}",
+						Action.ActionType.ToString(),
+						ThreadSeconds,
+						Path.GetFileName(Action.CommandPath),
+						Action.StatusDescription,
+						Action.bIsUsingPCH);
+
+					// Keep track of total thread seconds spent on tasks.
+					TotalThreadSeconds += ThreadSeconds;
+				}
+
+				Log.TraceInformation("-------- End Detailed Actions Stats -----------------------------------------------------------");
+
+				// Log total CPU seconds and numbers of processors involved in tasks.
+				Log.WriteLineIf(bLogDetailedActionStats || UnrealBuildTool.bPrintDebugInfo,
+					LogEventType.Console, "Cumulative thread seconds ({0} processors): {1:0.00}", System.Environment.ProcessorCount, TotalThreadSeconds);
+			}
+			return SNDBSResult;
+		}
+	}
+}
diff --git a/Engine/Source/Programs/UnrealBuildTool/Executors/XGE.cs b/Engine/Source/Programs/UnrealBuildTool/Executors/XGE.cs
new file mode 100644
--- /dev/null
+++ b/Engine/Source/Programs/UnrealBuildTool/Executors/XGE.cs
@@ -0,0 +1,514 @@
+// Copyright 1998-2018 Epic Games, Inc. All Rights Reserved.
+
+using System;
+using System.Collections;
+using System.Collections.Generic;
+using System.IO;
+using System.Diagnostics;
+using System.Xml;
+using System.Text.RegularExpressions;
+using System.Linq;
+using System.Reflection;
+using Microsoft.Win32;
+using System.Text;
+using Tools.DotNETCommon;
+
+namespace UnrealBuildTool
+{
+	class XGE : ActionExecutor
+	{
+		/// <summary>
+		/// Whether to use the no_watchdog_thread option to prevent VS2015 toolchain stalls.
+		/// </summary>
+		[XmlConfigFile(Category = "BuildConfiguration")]
+		bool bXGENoWatchdogThread = false;
+
+		/// <summary>
+		/// Whether to display the XGE build monitor.
+		/// </summary>
+		[XmlConfigFile(Category = "BuildConfiguration")]
+		bool bShowXGEMonitor = false;
+
+		/// <summary>
+		/// When enabled, XGE will stop compiling targets after a compile error occurs.  Recommended, as it saves computing resources for others.
+		/// </summary>
+		[XmlConfigFile(Category = "BuildConfiguration")]
+		bool bStopXGECompilationAfterErrors = false;
+
+		private const string ProgressMarkupPrefix = "@action";
+
+		public XGE()
+		{
+			XmlConfig.ApplyTo(this);
+		}
+
+		public override string Name
+		{
+			get { return "XGE"; }
+		}
+
+		public static bool IsAvailable()
+		{
+			// Get the name of the XgConsole executable.
+			string XgConsole = "xgConsole";
+			if (BuildHostPlatform.Current.Platform == UnrealTargetPlatform.Win64)
+			{
+				XgConsole = "xgConsole.exe";
+			}
+			else if (BuildHostPlatform.Current.Platform == UnrealTargetPlatform.Linux)
+			{
+				XgConsole = "ib_console";
+			}
+
+			// Search the path for it
+			string PathVariable = Environment.GetEnvironmentVariable("PATH");
+			foreach (string SearchPath in PathVariable.Split(Path.PathSeparator))
+			{
+				try
+				{
+					string PotentialPath = Path.Combine(SearchPath, XgConsole);
+					if(File.Exists(PotentialPath))
+					{
+						return true;
+					}
+				}
+				catch(ArgumentException)
+				{
+					// PATH variable may contain illegal characters; just ignore them.
+				}
+			}
+			return false;
+		}
+
+		// precompile the Regex needed to parse the XGE output (the ones we want are of the form "File (Duration at +time)"
+		private static Regex XGEDurationRegex = new Regex(@"(?<Filename>.*) *\((?<Duration>[0-9:\.]+) at [0-9\+:\.]+\)", RegexOptions.ExplicitCapture);
+
+		public static void ExportActions(List<Action> ActionsToExecute)
+		{
+			for(int FileNum = 0;;FileNum++)
+			{
+				string OutFile = Path.Combine(UnrealBuildTool.EngineDirectory.FullName, "Intermediate", "Build", String.Format("UBTExport.{0}.xge.xml", FileNum.ToString("D3")));
+				if(!File.Exists(OutFile))
+				{
+					ExportActions(ActionsToExecute, OutFile);
+					break;
+				}
+			}
+		}
+
+		public static void ExportActions(List<Action> ActionsToExecute, string OutFile)
+		{
+			WriteTaskFile(ActionsToExecute, OutFile, ProgressWriter.bWriteMarkup, bXGEExport: true);
+			Log.TraceInformation("XGEEXPORT: Exported '{0}'", OutFile);
+		}
+
+		public override bool ExecuteActions(List<Action> ActionsToExecute, bool bLogDetailedActionStats)
+		{
+			bool XGEResult = true;
+
+			// Batch up XGE execution by actions with the same output event handler.
+			List<Action> ActionBatch = new List<Action>();
+			ActionBatch.Add(ActionsToExecute[0]);
+			for (int ActionIndex = 1; ActionIndex < ActionsToExecute.Count && XGEResult; ++ActionIndex)
+			{
+				Action CurrentAction = ActionsToExecute[ActionIndex];
+				ActionBatch.Add(CurrentAction);
+			}
+			if (ActionBatch.Count > 0 && XGEResult)
+			{
+				XGEResult = ExecuteActionBatch(ActionBatch);
+				ActionBatch.Clear();
+			}
+
+			return XGEResult;
+		}
+
+		bool ExecuteActionBatch(List<Action> Actions)
+		{
+			bool XGEResult = true;
+			if (Actions.Count > 0)
+			{
+				// Write the actions to execute to a XGE task file.
+				string XGETaskFilePath = FileReference.Combine(UnrealBuildTool.EngineDirectory, "Intermediate", "Build", "XGETasks.xml").FullName;
+				WriteTaskFile(Actions, XGETaskFilePath, ProgressWriter.bWriteMarkup, false);
+
+				XGEResult = ExecuteTaskFileWithProgressMarkup(XGETaskFilePath, Actions.Count);
+			}
+			return XGEResult;
+		}
+
+		/// <summary>
+		/// Writes a XGE task file containing the specified actions to the specified file path.
+		/// </summary>
+		static void WriteTaskFile(List<Action> InActions, string TaskFilePath, bool bProgressMarkup, bool bXGEExport)
+		{
+			Dictionary<string, string> ExportEnv = new Dictionary<string, string>();
+
+			List<Action> Actions = InActions;
+			if (bXGEExport)
+			{
+				IDictionary CurrentEnvironment = Environment.GetEnvironmentVariables();
+				foreach (System.Collections.DictionaryEntry Pair in CurrentEnvironment)
+				{
+					if (!UnrealBuildTool.InitialEnvironment.Contains(Pair.Key) || (string)(UnrealBuildTool.InitialEnvironment[Pair.Key]) != (string)(Pair.Value))
+					{
+						ExportEnv.Add((string)(Pair.Key), (string)(Pair.Value));
+					}
+				}
+
+				int NumSortErrors = 0;
+				for (int ActionIndex = 0; ActionIndex < InActions.Count; ActionIndex++)
+				{
+					Action Action = InActions[ActionIndex];
+					foreach (FileItem Item in Action.PrerequisiteItems)
+					{
+						if (Item.ProducingAction != null && InActions.Contains(Item.ProducingAction))
+						{
+							int DepIndex = InActions.IndexOf(Item.ProducingAction);
+							if (DepIndex > ActionIndex)
+							{
+								NumSortErrors++;
+							}
+						}
+					}
+				}
+				if (NumSortErrors > 0)
+				{
+					Actions = new List<Action>();
+					HashSet<int> UsedActions = new HashSet<int>();
+					for (int ActionIndex = 0; ActionIndex < InActions.Count; ActionIndex++)
+					{
+						if (UsedActions.Contains(ActionIndex))
+						{
+							continue;
+						}
+						Action Action = InActions[ActionIndex];
+						foreach (FileItem Item in Action.PrerequisiteItems)
+						{
+							if (Item.ProducingAction != null && InActions.Contains(Item.ProducingAction))
+							{
+								int DepIndex = InActions.IndexOf(Item.ProducingAction);
+								if (UsedActions.Contains(DepIndex))
+								{
+									continue;
+								}
+								Actions.Add(Item.ProducingAction);
+								UsedActions.Add(DepIndex);
+							}
+						}
+						Actions.Add(Action);
+						UsedActions.Add(ActionIndex);
+					}
+					for (int ActionIndex = 0; ActionIndex < Actions.Count; ActionIndex++)
+					{
+						Action Action = Actions[ActionIndex];
+						foreach (FileItem Item in Action.PrerequisiteItems)
+						{
+							if (Item.ProducingAction != null && Actions.Contains(Item.ProducingAction))
+							{
+								int DepIndex = Actions.IndexOf(Item.ProducingAction);
+								if (DepIndex > ActionIndex)
+								{
+									throw new BuildException("Action is not topologically sorted.\n  {0} {1}\nDependency\n  {2} {3}", Action.CommandPath, Action.CommandArguments, Item.ProducingAction.CommandPath, Item.ProducingAction.CommandArguments);
+								}
+							}
+						}
+					}
+				}
+
+			}
+
+
+			XmlDocument XGETaskDocument = new XmlDocument();
+
+			// <BuildSet FormatVersion="1">...</BuildSet>
+			XmlElement BuildSetElement = XGETaskDocument.CreateElement("BuildSet");
+			XGETaskDocument.AppendChild(BuildSetElement);
+			BuildSetElement.SetAttribute("FormatVersion", "1");
+
+			// <Environments>...</Environments>
+			XmlElement EnvironmentsElement = XGETaskDocument.CreateElement("Environments");
+			BuildSetElement.AppendChild(EnvironmentsElement);
+
+			// <Environment Name="Default">...</CompileEnvironment>
+			XmlElement EnvironmentElement = XGETaskDocument.CreateElement("Environment");
+			EnvironmentsElement.AppendChild(EnvironmentElement);
+			EnvironmentElement.SetAttribute("Name", "Default");
+
+			// <Tools>...</Tools>
+			XmlElement ToolsElement = XGETaskDocument.CreateElement("Tools");
+			EnvironmentElement.AppendChild(ToolsElement);
+
+			if (ExportEnv.Count > 0)
+			{
+				// <Variables>...</Variables>
+				XmlElement VariablesElement = XGETaskDocument.CreateElement("Variables");
+				EnvironmentElement.AppendChild(VariablesElement);
+
+				foreach (KeyValuePair<string, string> Pair in ExportEnv)
+				{
+					// <Variable>...</Variable>
+					XmlElement VariableElement = XGETaskDocument.CreateElement("Variable");
+					VariablesElement.AppendChild(VariableElement);
+					VariableElement.SetAttribute("Name", Pair.Key);
+					VariableElement.SetAttribute("Value", Pair.Value);
+				}
+			}
+
+			for (int ActionIndex = 0; ActionIndex < Actions.Count; ActionIndex++)
+			{
+				Action Action = Actions[ActionIndex];
+
+				// <Tool ... />
+				XmlElement ToolElement = XGETaskDocument.CreateElement("Tool");
+				ToolsElement.AppendChild(ToolElement);
+				ToolElement.SetAttribute("Name", string.Format("Tool{0}", ActionIndex));
+				ToolElement.SetAttribute("AllowRemote", Action.bCanExecuteRemotely.ToString());
+
+				// The XGE documentation says that 'AllowIntercept' must be set to 'true' for all tools where 'AllowRemote' is enabled
+				ToolElement.SetAttribute("AllowIntercept", Action.bCanExecuteRemotely.ToString());
+
+				string OutputPrefix = "";
+				if (bProgressMarkup)
+				{
+					OutputPrefix += ProgressMarkupPrefix;
+				}
+				if (Action.bShouldOutputStatusDescription)
+				{
+					OutputPrefix += Action.StatusDescription;
+				}
+				if (OutputPrefix.Length > 0)
+				{
+					ToolElement.SetAttribute("OutputPrefix", OutputPrefix);
+				}
+
+				// When running on Windows, differentiate between .exe and batch files.
+				// Those (.bat, .cmd) need to be run via cmd /c or shellexecute, 
+				// the latter which we can't use because we want to redirect input/output
+
+				bool bLaunchViaCmdExe = (BuildHostPlatform.Current.Platform == UnrealTargetPlatform.Win64) && (!Path.GetExtension(Action.CommandPath).ToLower().EndsWith(".exe"));
+
+				string CommandPath = "";
+				string CommandArguments = "";
+
+				if (bLaunchViaCmdExe)
+				{
+					CommandPath = "cmd.exe";
+					CommandArguments = string.Format
+					(
+						"/c \"\"{0}\" {1}\"",
+						(Action.CommandPath),
+						(Action.CommandArguments)
+					);
+				}
+				else
+				{
+					CommandPath = Action.CommandPath;
+					CommandArguments = Action.CommandArguments;
+				}
+
+				ToolElement.SetAttribute("Params", CommandArguments);
+				ToolElement.SetAttribute("Path", CommandPath);
+				ToolElement.SetAttribute("SkipIfProjectFailed", "true");
+				if (Action.bIsGCCCompiler)
+				{
+					ToolElement.SetAttribute("AutoReserveMemory", "*.gch");
+				}
+				else
+				{
+					ToolElement.SetAttribute("AutoReserveMemory", "*.pch");
+				}
+				ToolElement.SetAttribute(
+					"OutputFileMasks",
+					string.Join(
+						",",
+						Action.ProducedItems.ConvertAll<string>(
+							delegate(FileItem ProducedItem) { return ProducedItem.Location.GetFileName(); }
+							).ToArray()
+						)
+					);
+			}
+
+			// <Project Name="Default" Env="Default">...</Project>
+			XmlElement ProjectElement = XGETaskDocument.CreateElement("Project");
+			BuildSetElement.AppendChild(ProjectElement);
+			ProjectElement.SetAttribute("Name", "Default");
+			ProjectElement.SetAttribute("Env", "Default");
+
+			for (int ActionIndex = 0; ActionIndex < Actions.Count; ActionIndex++)
+			{
+				Action Action = Actions[ActionIndex];
+
+				// <Task ... />
+				XmlElement TaskElement = XGETaskDocument.CreateElement("Task");
+				ProjectElement.AppendChild(TaskElement);
+				TaskElement.SetAttribute("SourceFile", "");
+				if (!Action.bShouldOutputStatusDescription)
+				{
+					// If we were configured to not output a status description, then we'll instead
+					// set 'caption' text for this task, so that the XGE coordinator has something
+					// to display within the progress bars.  For tasks that are outputting a
+					// description, XGE automatically displays that text in the progress bar, so we
+					// only need to do this for tasks that output their own progress.
+					TaskElement.SetAttribute("Caption", Action.StatusDescription);
+				}
+				TaskElement.SetAttribute("Name", string.Format("Action{0}", ActionIndex));
+				TaskElement.SetAttribute("Tool", string.Format("Tool{0}", ActionIndex));
+				TaskElement.SetAttribute("WorkingDir", Action.WorkingDirectory);
+				TaskElement.SetAttribute("SkipIfProjectFailed", "true");
+
+				// Create a semi-colon separated list of the other tasks this task depends on the results of.
+				List<string> DependencyNames = new List<string>();
+				foreach (FileItem Item in Action.PrerequisiteItems)
+				{
+					if (Item.ProducingAction != null && Actions.Contains(Item.ProducingAction))
+					{
+						DependencyNames.Add(string.Format("Action{0}", Actions.IndexOf(Item.ProducingAction)));
+					}
+				}
+
+				if (DependencyNames.Count > 0)
+				{
+					TaskElement.SetAttribute("DependsOn", string.Join(";", DependencyNames.ToArray()));
+				}
+			}
+
+			// Write the XGE task XML to a temporary file.
+			using (FileStream OutputFileStream = new FileStream(TaskFilePath, FileMode.Create, FileAccess.Write))
+			{
+				XGETaskDocument.Save(OutputFileStream);
+			}
+		}
+
+		/// <summary>
+		/// The possible result of executing tasks with XGE.
+		/// </summary>
+		enum ExecutionResult
+		{
+			Unavailable,
+			TasksFailed,
+			TasksSucceeded,
+		}
+
+		/// <summary>
+		/// Executes the tasks in the specified file.
+		/// </summary>
+		/// <param name="TaskFilePath">- The path to the file containing the tasks to execute in XGE XML format.</param>
+		/// <param name="OutputEventHandler"></param>
+		/// <param name="ActionCount"></param>
+		/// <returns>Indicates whether the tasks were successfully executed.</returns>
+		bool ExecuteTaskFile(string TaskFilePath, DataReceivedEventHandler OutputEventHandler, int ActionCount)
+		{
+			// A bug in the UCRT can cause XGE to hang on VS2015 builds. Figure out if this hang is likely to effect this build and workaround it if able.
+			// @todo: There is a KB coming that will fix this. Once that KB is available, test if it is present. Stalls will not be a problem if it is.
+			//
+			// Stalls are possible. However there is a workaround in XGE build 1659 and newer that can avoid the issue.
+			string XGEVersion = (BuildHostPlatform.Current.Platform == UnrealTargetPlatform.Win64) ? (string)Registry.GetValue(@"HKEY_LOCAL_MACHINE\SOFTWARE\Wow6432Node\Xoreax\IncrediBuild\Builder", "Version", null) : null;
+			if (XGEVersion != null)
+			{
+				int XGEBuildNumber;
+				if (Int32.TryParse(XGEVersion, out XGEBuildNumber))
+				{
+					// Per Xoreax support, subtract 1001000 from the registry value to get the build number of the installed XGE.
+					if (XGEBuildNumber - 1001000 >= 1659)
+					{
+						bXGENoWatchdogThread = true;
+					}
+					// @todo: Stalls are possible and we don't have a workaround. What should we do? Most people still won't encounter stalls, we don't really
+					// want to disable XGE on them if it would have worked.
+				}
+			}
+
+			bool bSilentCompileOutput = false;
+			string SilentOption = bSilentCompileOutput ? "/Silent" : "";
+
+			ProcessStartInfo XGEStartInfo = new ProcessStartInfo(
+				"xgConsole",
+				string.Format("\"{0}\" /Rebuild /NoWait {1} /NoLogo {2} /ShowAgent /ShowTime {3}",
+					TaskFilePath,
+					bStopXGECompilationAfterErrors ? "/StopOnErrors" : "",
+					SilentOption,
+					bXGENoWatchdogThread ? "/no_watchdog_thread" : "")
+				);
+			XGEStartInfo.UseShellExecute = false;
+
+			// Use the IDE-integrated Incredibuild monitor to display progress.
+			XGEStartInfo.Arguments += " /UseIdeMonitor";
+
+			// Optionally display the external XGE monitor.
+			if (bShowXGEMonitor)
+			{
+				XGEStartInfo.Arguments += " /OpenMonitor";
+			}
+
+			try
+			{
+				// Start the process, redirecting stdout/stderr if requested.
+				Process XGEProcess = new Process();
+				XGEProcess.StartInfo = XGEStartInfo;
+				bool bShouldRedirectOuput = OutputEventHandler != null;
+				if (bShouldRedirectOuput)
+				{
+					XGEStartInfo.RedirectStandardError = true;
+					XGEStartInfo.RedirectStandardOutput = true;
+					XGEProcess.EnableRaisingEvents = true;
+					XGEProcess.OutputDataReceived += OutputEventHandler;
+					XGEProcess.ErrorDataReceived += OutputEventHandler;
+				}
+				XGEProcess.Start();
+				if (bShouldRedirectOuput)
+				{
+					XGEProcess.BeginOutputReadLine();
+					XGEProcess.BeginErrorReadLine();
+				}
+
+				Log.TraceInformation("Distributing {0} action{1} to XGE",
+					ActionCount,
+					ActionCount == 1 ? "" : "s");
+
+				// Wait until the process is finished and return whether it all the tasks successfully executed.
+				XGEProcess.WaitForExit();
+				return XGEProcess.ExitCode == 0;
+			}
+			catch (Exception)
+			{
+				return false;
+			}
+		}
+
+		/// <summary>
+		/// Executes the tasks in the specified file, parsing progress markup as part of the output.
+		/// </summary>
+		bool ExecuteTaskFileWithProgressMarkup(string TaskFilePath, int NumActions)
+		{
+			using (ProgressWriter Writer = new ProgressWriter("Compiling C++ source files...", false))
+			{
+				int NumCompletedActions = 0;
+
+				// Create a wrapper delegate that will parse the output actions
+				DataReceivedEventHandler EventHandlerWrapper = (Sender, Args) =>
+				{
+					if(Args.Data != null)
+					{
+						string Text = Args.Data;
+						if (Text.StartsWith(ProgressMarkupPrefix))
+						{
+							Writer.Write(++NumCompletedActions, NumActions);
+
+							// Strip out anything that is just an XGE timer. Some programs don't output anything except the progress text.
+							Text = Args.Data.Substring(ProgressMarkupPrefix.Length);
+							if(Text.StartsWith(" (") && Text.EndsWith(")"))
+							{
+								return;
+							}
+						}
+						Log.TraceInformation(Text);
+					}
+				};
+
+				// Run through the standard XGE executor
+				return ExecuteTaskFile(TaskFilePath, EventHandlerWrapper, NumActions);
+			}
+		}
+	}
+}
diff --git a/Engine/Source/Programs/UnrealBuildTool/System/Distcc.cs b/Engine/Source/Programs/UnrealBuildTool/System/Distcc.cs
deleted file mode 100644
diff --git a/Engine/Source/Programs/UnrealBuildTool/System/LocalExecutor.cs b/Engine/Source/Programs/UnrealBuildTool/System/LocalExecutor.cs
deleted file mode 100644
diff --git a/Engine/Source/Programs/UnrealBuildTool/System/ParallelExecutor.cs b/Engine/Source/Programs/UnrealBuildTool/System/ParallelExecutor.cs
deleted file mode 100644
diff --git a/Engine/Source/Programs/UnrealBuildTool/System/PrecompileMode.cs b/Engine/Source/Programs/UnrealBuildTool/System/PrecompileMode.cs
deleted file mode 100644
diff --git a/Engine/Source/Programs/UnrealBuildTool/System/SNDBS.cs b/Engine/Source/Programs/UnrealBuildTool/System/SNDBS.cs
deleted file mode 100644
diff --git a/Engine/Source/Programs/UnrealBuildTool/System/XGE.cs b/Engine/Source/Programs/UnrealBuildTool/System/XGE.cs
deleted file mode 100644
diff --git a/Engine/Source/Programs/UnrealBuildTool/UnrealBuildTool.csproj b/Engine/Source/Programs/UnrealBuildTool/UnrealBuildTool.csproj
--- a/Engine/Source/Programs/UnrealBuildTool/UnrealBuildTool.csproj
+++ b/Engine/Source/Programs/UnrealBuildTool/UnrealBuildTool.csproj
@@ -148,6 +148,11 @@
     <Compile Include="Configuration\UEBuildPlatformFactory.cs" />
     <Compile Include="Configuration\UEBuildPlatformSDK.cs" />
     <Compile Include="Configuration\UEBuildPlugin.cs" />
+    <Compile Include="Executors\Distcc.cs" />
+    <Compile Include="Executors\LocalExecutor.cs" />
+    <Compile Include="Executors\ParallelExecutor.cs" />
+    <Compile Include="Executors\SNDBS.cs" />
+    <Compile Include="Executors\XGE.cs" />
     <Compile Include="Platform\Lumin\UEDeployLumin.cs" />
     <Compile Include="ProjectFiles\CLion\CLionGenerator.cs" />
     <Compile Include="ProjectFiles\CMake\CMakefileGenerator.cs" />
@@ -170,11 +175,9 @@
     <Compile Include="System\CppCompileEnvironment.cs" />
     <Compile Include="System\DataDrivenPlatformInfo.cs" />
     <Compile Include="System\EncryptionAndSigning.cs" />
-    <Compile Include="System\ParallelExecutor.cs" />
     <Compile Include="System\ModuleManifest.cs" />
     <Compile Include="System\PluginReferenceDescriptor.cs" />
     <Compile Include="System\PrecompiledManifest.cs" />
-    <Compile Include="System\PrecompileMode.cs" />
     <Compile Include="System\RulesDocumentation.cs" />
     <Compile Include="System\PrecompiledHeaderInstance.cs" />
     <Compile Include="System\PrecompiledHeaderTemplate.cs" />
@@ -319,11 +322,9 @@
     <Compile Include="Configuration\BuildConfiguration.cs" />
     <Compile Include="System\BuildHostPlatform.cs" />
     <Compile Include="System\FlatCPPIncludeDepencencyCache.cs" />
-    <Compile Include="System\Distcc.cs" />
     <Compile Include="System\ModuleDescriptor.cs" />
     <Compile Include="System\PluginDescriptor.cs" />
     <Compile Include="System\ProjectDescriptor.cs" />
-    <Compile Include="System\SNDBS.cs" />
     <Compile Include="System\DirectoryLookupCache.cs" />
     <Compile Include="System\CPPHeaders.cs" />
     <Compile Include="System\DependencyCache.cs" />
@@ -332,7 +333,6 @@
     <Compile Include="System\FileItem.cs" />
     <Compile Include="System\JunkDeleter.cs" />
     <Compile Include="System\LinkEnvironment.cs" />
-    <Compile Include="System\LocalExecutor.cs" />
     <Compile Include="System\Plugins.cs" />
     <Compile Include="System\RPCUtilHelper.cs" />
     <Compile Include="System\RulesCompiler.cs" />
@@ -343,7 +343,6 @@
     <Compile Include="System\Utils.cs" />
     <Compile Include="ToolChain\RemoteToolChain.cs" />
     <Compile Include="ToolChain\UEToolChain.cs" />
-    <Compile Include="System\XGE.cs" />
     <Compile Include="Platform\Windows\UEBuildWindows.cs" />
     <Compile Include="Platform\Windows\VCEnvironment.cs" />
     <Compile Include="Platform\Windows\VCToolChain.cs" />
