Add support for UnrealPak plugins.

* Project and plugin modules can now specify an array of supported programs in the "WhitelistPrograms" field of their module descriptors, to allow modules to be loaded by programs.
* Programs can now load any runtime modules, as long as they are whitelisted.
* Programs under the engine directory can now use a shared build environment, so that building with a project file does not cause output binaries to be output to the project directory.
* UnrealPak is now always built by default when packaging
* Convert UnrealPak to a modular configuration

diff --git a/Engine/Plugins/Messaging/TcpMessaging/TcpMessaging.uplugin b/Engine/Plugins/Messaging/TcpMessaging/TcpMessaging.uplugin
--- a/Engine/Plugins/Messaging/TcpMessaging/TcpMessaging.uplugin
+++ b/Engine/Plugins/Messaging/TcpMessaging/TcpMessaging.uplugin
@@ -18,12 +18,13 @@
 	[
 		{
 			"Name" : "TcpMessaging",
-			"Type" : "RuntimeAndProgram",
+			"Type" : "Runtime",
 			"LoadingPhase" : "PreDefault",
 			"BlacklistPlatforms" :
 			[
 				"HTML5"
-			]
-    }
+			],
+			"WhitelistPrograms" : [ "UnrealFrontend" ]
+	    }
 	]
 }
\ No newline at end of file
diff --git a/Engine/Plugins/Messaging/UdpMessaging/UdpMessaging.uplugin b/Engine/Plugins/Messaging/UdpMessaging/UdpMessaging.uplugin
--- a/Engine/Plugins/Messaging/UdpMessaging/UdpMessaging.uplugin
+++ b/Engine/Plugins/Messaging/UdpMessaging/UdpMessaging.uplugin
@@ -18,12 +18,13 @@
 	[
 		{
 			"Name" : "UdpMessaging",
-			"Type" : "RuntimeAndProgram",
+			"Type" : "Runtime",
 			"LoadingPhase" : "PreDefault",
 			"BlacklistPlatforms" :
 			[
 				"HTML5"
-			]
+			],
+			"WhitelistPrograms" : [ "UnrealFrontend" ]
 		}
 	]
 }
\ No newline at end of file
diff --git a/Engine/Source/Programs/AutomationTool/Scripts/BuildPluginCommand.Automation.cs b/Engine/Source/Programs/AutomationTool/Scripts/BuildPluginCommand.Automation.cs
--- a/Engine/Source/Programs/AutomationTool/Scripts/BuildPluginCommand.Automation.cs
+++ b/Engine/Source/Programs/AutomationTool/Scripts/BuildPluginCommand.Automation.cs
@@ -172,7 +172,7 @@ class BuildPlugin : BuildCommand
 				bool bBuildDeveloperTools = (TargetType == TargetType.Editor || TargetType == TargetType.Program);
 				bool bBuildEditor = (TargetType == TargetType.Editor);
 				bool bBuildRequiresCookedData = (TargetType != TargetType.Editor && TargetType != TargetType.Program);
-				if (Module.IsCompiledInConfiguration(Platform, Configuration, TargetType, bBuildDeveloperTools, bBuildEditor, bBuildRequiresCookedData))
+				if (Module.IsCompiledInConfiguration(Platform, Configuration, TargetName, TargetType, bBuildDeveloperTools, bBuildEditor, bBuildRequiresCookedData))
 				{
 					bCompilePlatform = true;
 				}
diff --git a/Engine/Source/Programs/AutomationTool/Scripts/BuildProjectCommand.Automation.cs b/Engine/Source/Programs/AutomationTool/Scripts/BuildProjectCommand.Automation.cs
--- a/Engine/Source/Programs/AutomationTool/Scripts/BuildProjectCommand.Automation.cs
+++ b/Engine/Source/Programs/AutomationTool/Scripts/BuildProjectCommand.Automation.cs
@@ -19,9 +19,10 @@ public enum ProjectBuildTargets
 	Bootstrap = 1 << 3,
 	CrashReporter = 1 << 4,
 	Programs = 1 << 5,
+	UnrealPak = 1 << 6,
 
 	// All targets
-	All = Editor | ClientCooked | ServerCooked | Bootstrap | CrashReporter | Programs,
+	All = Editor | ClientCooked | ServerCooked | Bootstrap | CrashReporter | Programs | UnrealPak,
 }
 
 /// <summary>
@@ -95,6 +96,15 @@ public partial class Project : CommandUtils
 			}
 		}
 
+		// Build any tools we need to stage
+		if ((TargetMask & ProjectBuildTargets.UnrealPak) == ProjectBuildTargets.UnrealPak && !Automation.IsEngineInstalled())
+		{
+			if (Params.Pak && Params.EditorTargets.Contains("UnrealPak") == false)
+			{
+				Agenda.AddTargets(new string[] { "UnrealPak" }, HostPlatform.Current.HostEditorPlatform, UnrealTargetConfiguration.Development, Params.CodeBasedUprojectPath);
+			}
+		}
+
 		// Additional compile arguments
 		string AdditionalArgs = "";
 
diff --git a/Engine/Source/Programs/AutomationTool/Scripts/CopyBuildToStagingDirectory.Automation.cs b/Engine/Source/Programs/AutomationTool/Scripts/CopyBuildToStagingDirectory.Automation.cs
--- a/Engine/Source/Programs/AutomationTool/Scripts/CopyBuildToStagingDirectory.Automation.cs
+++ b/Engine/Source/Programs/AutomationTool/Scripts/CopyBuildToStagingDirectory.Automation.cs
@@ -137,17 +137,29 @@ public partial class Project : CommandUtils
 		CmdLine.Append(PlatformOptions);
 
 		return CmdLine.ToString();
+	}
+
+	static FileReference GetUnrealPakLocation()
+	{
+		if(HostPlatform.Current.HostEditorPlatform == UnrealTargetPlatform.Win64)
+		{
+			return FileReference.Combine(CommandUtils.EngineDirectory, "Binaries", "Win64", "UnrealPak.exe");
 		}
+		else
+		{
+			return FileReference.Combine(CommandUtils.EngineDirectory, "Binaries", HostPlatform.Current.HostEditorPlatform.ToString(), "UnrealPak");
+		}
+	}
 
 	static public void RunUnrealPak(ProjectParams Params, Dictionary<string, string> UnrealPakResponseFile, FileReference OutputLocation, FileReference PakOrderFileLocation, string PlatformOptions, bool Compressed, EncryptionAndSigning.CryptoSettings CryptoSettings, FileReference CryptoKeysCacheFilename, String PatchSourceContentPath)
-		{
+	{
 		if (UnrealPakResponseFile.Count < 1)
 		{
 			return;
 		}
 
 		string Arguments = GetUnrealPakArguments(UnrealPakResponseFile, OutputLocation, PakOrderFileLocation, PlatformOptions, Compressed, CryptoSettings, CryptoKeysCacheFilename, PatchSourceContentPath);
-		RunCommandlet(Params.RawProjectPath, Params.UE4Exe, "UnrealPak", Arguments);
+		RunAndLog(CmdEnv, GetUnrealPakLocation().FullName, Arguments, Options: ERunOptions.Default | ERunOptions.UTF8Output);
 	}
 
 	static public void LogDeploymentContext(DeploymentContext SC)
@@ -1552,7 +1564,9 @@ public partial class Project : CommandUtils
 		{
 			string CommandsFile = LogUtils.GetUniqueLogName(CombinePaths(CmdEnv.EngineSavedFolder, "UnrealPak-Commands"));
 			File.WriteAllLines(CommandsFile, Commands);
-			RunCommandlet(Params.RawProjectPath, Params.UE4Exe, "UnrealPak", String.Format("-batch={0}", MakePathSafeToUseWithCommandLine(CommandsFile)));
+
+			string Arguments = String.Format("{0} -batch={1}", MakePathSafeToUseWithCommandLine(Params.RawProjectPath.FullName), MakePathSafeToUseWithCommandLine(CommandsFile));
+			RunAndLog(CmdEnv, GetUnrealPakLocation().FullName, Arguments, Options: ERunOptions.Default | ERunOptions.UTF8Output);
 		}
 
 		// Do any additional processing on the command output
@@ -1641,8 +1655,8 @@ public partial class Project : CommandUtils
 				CmdLine += " -customint=\"PakReadOrdering=0\"";
 				CmdLine += " -stdout";
 
-				string UnrealPakLogFileName = "UnrealPak_" + PakName;
-				RunAndLog(CmdEnv, BPTExe, CmdLine, UnrealPakLogFileName, Options: ERunOptions.Default | ERunOptions.UTF8Output);
+				string BuildPatchToolLogFileName = "BuildPatchTool_" + PakName;
+				RunAndLog(CmdEnv, BPTExe, CmdLine, BuildPatchToolLogFileName, Options: ERunOptions.Default | ERunOptions.UTF8Output);
 
 				InternalUtils.SafeCopyFile(SourceManifestPath, DestManifestPath);
 			}
diff --git a/Engine/Source/Programs/UnrealBuildTool/Configuration/UEBuildTarget.cs b/Engine/Source/Programs/UnrealBuildTool/Configuration/UEBuildTarget.cs
--- a/Engine/Source/Programs/UnrealBuildTool/Configuration/UEBuildTarget.cs
+++ b/Engine/Source/Programs/UnrealBuildTool/Configuration/UEBuildTarget.cs
@@ -436,7 +436,11 @@ namespace UnrealBuildTool
 			// Set the default value for whether to use the shared build environment
 			if(RulesObject.BuildEnvironment == TargetBuildEnvironment.Default)
 			{
-				if(RulesObject.Type != TargetType.Program && (UnrealBuildTool.IsEngineInstalled() || RulesObject.LinkType != TargetLinkType.Monolithic))
+				if(RulesObject.Type == TargetType.Program && Desc.ProjectFile != null && TargetFileName.IsUnderDirectory(Desc.ProjectFile.Directory))
+				{
+					RulesObject.BuildEnvironment = TargetBuildEnvironment.Unique;
+				}
+				else if(UnrealBuildTool.IsEngineInstalled() || RulesObject.LinkType != TargetLinkType.Monolithic)
 				{
 					RulesObject.BuildEnvironment = TargetBuildEnvironment.Shared;
 				}
@@ -1058,10 +1062,16 @@ namespace UnrealBuildTool
 
 			// Set the build environment
 			bUseSharedBuildEnvironment = (Rules.BuildEnvironment == TargetBuildEnvironment.Shared);
-
 			if (bUseSharedBuildEnvironment)
 			{
-				AppName = GetAppNameForTargetType(Rules.Type);
+				if(Rules.Type == TargetType.Program)
+				{
+					AppName = TargetName;
+				}
+				else
+				{
+					AppName = GetAppNameForTargetType(Rules.Type);
+				}
 			}
 
 			// Figure out what the project directory is. If we have a uproject file, use that. Otherwise use the engine directory.
@@ -1118,7 +1128,7 @@ namespace UnrealBuildTool
 
 			// Construct the output paths for this target's executable
 			DirectoryReference OutputDirectory;
-			if (bCompileMonolithic || TargetType == TargetType.Program || !bUseSharedBuildEnvironment)
+			if (bCompileMonolithic || !bUseSharedBuildEnvironment)
 			{
 				OutputDirectory = ProjectDirectory;
 			}
@@ -2928,7 +2938,7 @@ namespace UnrealBuildTool
 						{
 							foreach (ModuleDescriptor ModuleDescriptor in Plugin.Descriptor.Modules)
 							{
-								if (ModuleDescriptor.IsCompiledInConfiguration(Platform, Configuration, TargetType, bAllowDeveloperModules && Rules.bBuildDeveloperTools, Rules.bBuildEditor, Rules.bBuildRequiresCookedData))
+								if (ModuleDescriptor.IsCompiledInConfiguration(Platform, Configuration, TargetName, TargetType, bAllowDeveloperModules && Rules.bBuildDeveloperTools, Rules.bBuildEditor, Rules.bBuildRequiresCookedData))
 								{
 									FileReference ModuleFileName = RulesAssembly.GetModuleFileName(ModuleDescriptor.Name);
 									if(!ModuleFileName.ContainsAnyNames(ExcludeFoldersArray, Plugin.Directory))
@@ -3273,15 +3283,15 @@ namespace UnrealBuildTool
 				{
 					if(Plugin.EnabledByDefault && !ReferencedNames.Contains(Plugin.Name))
 					{
-							ReferencedNames.Add(Plugin.Name);
+						ReferencedNames.Add(Plugin.Name);
 
-							PluginReferenceDescriptor PluginReference = new PluginReferenceDescriptor(Plugin.Name, null, true);
-							PluginReference.bOptional = true;
+						PluginReferenceDescriptor PluginReference = new PluginReferenceDescriptor(Plugin.Name, null, true);
+						PluginReference.bOptional = true;
 
-							AddPlugin(PluginReference, "default plugins", ExcludeFolders, NameToInstance, NameToInfo);
-						}
+						AddPlugin(PluginReference, "default plugins", ExcludeFolders, NameToInstance, NameToInfo);
 					}
 				}
+			}
 
 			// If this is a program, synthesize references for plugins which are enabled via the config file
 			if(TargetType == TargetType.Program)
@@ -3406,7 +3416,7 @@ namespace UnrealBuildTool
 				{
 					foreach (ModuleDescriptor ModuleInfo in Info.Descriptor.Modules)
 					{
-						if (ModuleInfo.IsCompiledInConfiguration(Platform, Configuration, TargetType, Rules.bBuildDeveloperTools, Rules.bBuildEditor, Rules.bBuildRequiresCookedData))
+						if (ModuleInfo.IsCompiledInConfiguration(Platform, Configuration, TargetName, TargetType, Rules.bBuildDeveloperTools, Rules.bBuildEditor, Rules.bBuildRequiresCookedData))
 						{
 							UEBuildModuleCPP Module = FindOrCreateCppModuleByName(ModuleInfo.Name, PluginReferenceChain);
 							if(!Instance.Modules.Contains(Module))
@@ -3821,6 +3831,9 @@ namespace UnrealBuildTool
 			GlobalCompileEnvironment.Definitions.Add("UBT_COMPILED_PLATFORM=" + Platform.ToString());
 			GlobalCompileEnvironment.Definitions.Add("UBT_COMPILED_TARGET=" + TargetType.ToString());
 
+			// Set the global app name
+			GlobalCompileEnvironment.Definitions.Add(String.Format("UE_APP_NAME=\"{0}\"", AppName));
+
 			// Initialize the compile and link environments for the platform, configuration, and project.
 			BuildPlatform.SetUpEnvironment(Rules, GlobalCompileEnvironment, GlobalLinkEnvironment);
 			BuildPlatform.SetUpConfigurationEnvironment(Rules, GlobalCompileEnvironment, GlobalLinkEnvironment);
diff --git a/Engine/Source/Programs/UnrealBuildTool/System/ModuleDescriptor.cs b/Engine/Source/Programs/UnrealBuildTool/System/ModuleDescriptor.cs
--- a/Engine/Source/Programs/UnrealBuildTool/System/ModuleDescriptor.cs
+++ b/Engine/Source/Programs/UnrealBuildTool/System/ModuleDescriptor.cs
@@ -167,6 +167,11 @@ namespace UnrealBuildTool
 		/// </summary>
 		public UnrealTargetConfiguration[] BlacklistTargetConfigurations;
 
+		/// <summary>
+		/// List of allowed programs
+		/// </summary>
+		public string[] WhitelistPrograms;
+
 		/// <summary>
 		/// List of additional dependencies for building this module.
 		/// </summary>
@@ -234,6 +239,12 @@ namespace UnrealBuildTool
 				Module.BlacklistTargetConfigurations = BlacklistTargetConfigurations;
 			}
 
+			string[] WhitelistPrograms;
+			if (InObject.TryGetStringArrayField("WhitelistPrograms", out WhitelistPrograms))
+			{
+				Module.WhitelistPrograms = WhitelistPrograms;
+			}
+
 			string[] AdditionalDependencies;
 			if (InObject.TryGetStringArrayField("AdditionalDependencies", out AdditionalDependencies))
 			{
@@ -307,6 +318,10 @@ namespace UnrealBuildTool
 				}
 				Writer.WriteArrayEnd();
 			}
+			if(WhitelistPrograms != null && WhitelistPrograms.Length > 0)
+			{
+				Writer.WriteStringArrayField("WhitelistPrograms", WhitelistPrograms);
+			}
 			if (AdditionalDependencies != null && AdditionalDependencies.Length > 0)
 			{
 				Writer.WriteArrayStart("AdditionalDependencies");
@@ -343,11 +358,12 @@ namespace UnrealBuildTool
 		/// </summary>
 		/// <param name="Platform">The platform being compiled for</param>
 		/// <param name="TargetConfiguration">The target configuration being compiled for</param>
+		/// <param name="TargetName">Name of the target being built</param>
 		/// <param name="TargetType">The type of the target being compiled</param>
 		/// <param name="bBuildDeveloperTools">Whether the configuration includes developer tools (typically UEBuildConfiguration.bBuildDeveloperTools for UBT callers)</param>
 		/// <param name="bBuildEditor">Whether the configuration includes the editor (typically UEBuildConfiguration.bBuildEditor for UBT callers)</param>
 		/// <param name="bBuildRequiresCookedData">Whether the configuration requires cooked content (typically UEBuildConfiguration.bBuildRequiresCookedData for UBT callers)</param>
-		public bool IsCompiledInConfiguration(UnrealTargetPlatform Platform, UnrealTargetConfiguration TargetConfiguration, TargetType TargetType, bool bBuildDeveloperTools, bool bBuildEditor, bool bBuildRequiresCookedData)
+		public bool IsCompiledInConfiguration(UnrealTargetPlatform Platform, UnrealTargetConfiguration TargetConfiguration, string TargetName, TargetType TargetType, bool bBuildDeveloperTools, bool bBuildEditor, bool bBuildRequiresCookedData)
 		{
 			// Check the platform is whitelisted
 			if (WhitelistPlatforms != null && WhitelistPlatforms.Length > 0 && !WhitelistPlatforms.Contains(Platform))
@@ -385,12 +401,18 @@ namespace UnrealBuildTool
 				return false;
 			}
 
+			// Check the program name is whitelisted
+			if(TargetType == TargetType.Program && (WhitelistPrograms == null || !WhitelistPrograms.Contains(TargetName)))
+			{
+				return false;
+			}
+
 			// Check the module is compatible with this target.
 			switch (Type)
 			{
 				case ModuleHostType.Runtime:
 				case ModuleHostType.RuntimeNoCommandlet:
-                    return TargetType != TargetType.Program;
+                    return true;
                 case ModuleHostType.CookedOnly:
                     return bBuildRequiresCookedData;
                 case ModuleHostType.RuntimeAndProgram:
diff --git a/Engine/Source/Programs/UnrealPak/UnrealPak.Target.cs b/Engine/Source/Programs/UnrealPak/UnrealPak.Target.cs
--- a/Engine/Source/Programs/UnrealPak/UnrealPak.Target.cs
+++ b/Engine/Source/Programs/UnrealPak/UnrealPak.Target.cs
@@ -1,4 +1,4 @@
-﻿// Copyright 1998-2018 Epic Games, Inc. All Rights Reserved.
+// Copyright 1998-2018 Epic Games, Inc. All Rights Reserved.
 
 using UnrealBuildTool;
 using System.Collections.Generic;
@@ -8,11 +8,12 @@ public class UnrealPakTarget : TargetRules
 	public UnrealPakTarget(TargetInfo Target) : base(Target)
 	{
 		Type = TargetType.Program;
-		LinkType = TargetLinkType.Monolithic;
+		LinkType = TargetLinkType.Modular;
 		LaunchModuleName = "UnrealPak";
 
 		bCompileLeanAndMeanUE = true;
 		bUseMallocProfiler = false;
+		bCompileWithPluginSupport = true;
 
 		// No editor needed
 		bBuildEditor = false;
diff --git a/Engine/Source/Runtime/Projects/Private/ModuleDescriptor.cpp b/Engine/Source/Runtime/Projects/Private/ModuleDescriptor.cpp
--- a/Engine/Source/Runtime/Projects/Private/ModuleDescriptor.cpp
+++ b/Engine/Source/Runtime/Projects/Private/ModuleDescriptor.cpp
@@ -80,9 +80,6 @@ const TCHAR* EHostType::ToString( const EHostType::Type Value )
 		case RuntimeNoCommandlet:
 			return TEXT( "RuntimeNoCommandlet" );
 
-		case RuntimeAndProgram:
-			return TEXT("RuntimeAndProgram");
-
 		case CookedOnly:
 			return TEXT("CookedOnly");
 
@@ -220,6 +217,17 @@ bool FModuleDescriptor::Read(const FJsonObject& Object, FText& OutFailReason)
 		}
 	}
 
+	// Read the supported programs
+	TSharedPtr<FJsonValue> WhitelistProgramsValue = Object.TryGetField(TEXT("WhitelistPrograms"));
+	if (WhitelistProgramsValue.IsValid() && WhitelistProgramsValue->Type == EJson::Array)
+	{
+		const TArray< TSharedPtr< FJsonValue > >& ProgramsArray = WhitelistProgramsValue->AsArray();
+		for (int Idx = 0; Idx < ProgramsArray.Num(); Idx++)
+		{
+			WhitelistPrograms.Add(ProgramsArray[Idx]->AsString());
+		}
+	}
+
 	// Read the additional dependencies
 	TSharedPtr<FJsonValue> AdditionalDependenciesValue = Object.TryGetField(TEXT("AdditionalDependencies"));
 	if (AdditionalDependenciesValue.IsValid() && AdditionalDependenciesValue->Type == EJson::Array)
@@ -328,6 +336,15 @@ void FModuleDescriptor::Write(TJsonWriter<>& Writer) const
 		}
 		Writer.WriteArrayEnd();
 	}
+	if (WhitelistPrograms.Num() > 0)
+	{
+		Writer.WriteArrayStart(TEXT("WhitelistPrograms"));
+		for (int Idx = 0; Idx < WhitelistPrograms.Num(); Idx++)
+		{
+			Writer.WriteValue(WhitelistPrograms[Idx]);
+		}
+		Writer.WriteArrayEnd();
+	}
 	if (AdditionalDependencies.Num() > 0)
 	{
 		Writer.WriteArrayStart(TEXT("AdditionalDependencies"));
@@ -396,18 +413,19 @@ bool FModuleDescriptor::IsCompiledInCurrentConfiguration() const
 		return false;
 	}
 
+#if IS_PROGRAM
+	// Check the program is allowed
+	if (!WhitelistPrograms.Contains(UE_APP_NAME))
+	{
+		return false;
+	}
+#endif
+
 	// Check the module is compatible with this target. This should match ModuleDescriptor.IsCompiledInConfiguration in UBT
 	switch (Type)
 	{
 	case EHostType::Runtime:
 	case EHostType::RuntimeNoCommandlet:
-#if IS_PROGRAM
-		return false;
-#else
-		return true;
-#endif
-
-	case EHostType::RuntimeAndProgram:
 		return true;
 
 	case EHostType::CookedOnly:
@@ -462,20 +480,14 @@ bool FModuleDescriptor::IsLoadedInCurrentConfiguration() const
 	// Check that the runtime environment allows it to be loaded
 	switch (Type)
 	{
-	case EHostType::RuntimeAndProgram:
-		#if (WITH_ENGINE || WITH_PLUGIN_SUPPORT)
-			return true;
-		#endif
-		break;
-
 	case EHostType::Runtime:
-		#if (WITH_ENGINE || WITH_PLUGIN_SUPPORT) && !IS_PROGRAM
+		#if WITH_ENGINE || WITH_PLUGIN_SUPPORT
 			return true;
 		#endif
 		break;
 	
 	case EHostType::RuntimeNoCommandlet:
-		#if (WITH_ENGINE || WITH_PLUGIN_SUPPORT)  && !IS_PROGRAM
+		#if WITH_ENGINE || WITH_PLUGIN_SUPPORT
 			if(!IsRunningCommandlet()) return true;
 		#endif
 		break;
diff --git a/Engine/Source/Runtime/Projects/Public/ModuleDescriptor.h b/Engine/Source/Runtime/Projects/Public/ModuleDescriptor.h
--- a/Engine/Source/Runtime/Projects/Public/ModuleDescriptor.h
+++ b/Engine/Source/Runtime/Projects/Public/ModuleDescriptor.h
@@ -69,7 +69,6 @@ namespace EHostType
 	{
 		Runtime,
 		RuntimeNoCommandlet,
-		RuntimeAndProgram,
 		CookedOnly,
 		Developer,
 		Editor,
@@ -131,6 +130,9 @@ struct PROJECTS_API FModuleDescriptor
 	/** List of disallowed target configurations */
 	TArray<FString> BlacklistTargetConfigurations;
 
+	/** List of supported programs. No blacklist, because weprograms don't support plugins by default. */
+	TArray<FString> WhitelistPrograms;
+
 	/** List of additional dependencies for building this module. */
 	TArray<FString> AdditionalDependencies;
 
