UBT: Fix spam while generating project files if Clang isn't installed.

#rnx

diff --git a/Engine/Source/Programs/UnrealBuildTool/Configuration/TargetRules.cs b/Engine/Source/Programs/UnrealBuildTool/Configuration/TargetRules.cs
--- a/Engine/Source/Programs/UnrealBuildTool/Configuration/TargetRules.cs
+++ b/Engine/Source/Programs/UnrealBuildTool/Configuration/TargetRules.cs
@@ -1398,6 +1398,14 @@ namespace UnrealBuildTool
 			get { return BuildHostPlatform.Current.Platform; }
 		}
 
+		/// <summary>
+		/// Expose the bGenerateProjectFiles flag to targets, so we can modify behavior as appropriate for better intellisense
+		/// </summary>
+		public bool bGenerateProjectFiles
+		{
+			get { return ProjectFileGenerator.bGenerateProjectFiles; }
+		}
+
 		/// <summary>
 		/// Setup the global environment for building this target
 		/// IMPORTANT: Game targets will *not* have this function called if they use the shared build environment.
@@ -2174,9 +2182,14 @@ namespace UnrealBuildTool
 			get { return Inner.bShouldCompileAsDLL; }
 		}
 
-		#if !__MonoCS__
-		#pragma warning restore C1591
-		#endif
+		public bool bGenerateProjectFiles
+		{
+			get { return Inner.bGenerateProjectFiles; }
+		}
+
+#if !__MonoCS__
+#pragma warning restore C1591
+#endif
 		#endregion
 
 		/// <summary>
