UBT: Change the order in which Visual C++ toolchains are enumerated to prefer full releases over preview releases.

diff --git a/Engine/Source/Programs/UnrealBuildTool/Platform/Windows/UEBuildWindows.cs b/Engine/Source/Programs/UnrealBuildTool/Platform/Windows/UEBuildWindows.cs
--- a/Engine/Source/Programs/UnrealBuildTool/Platform/Windows/UEBuildWindows.cs
+++ b/Engine/Source/Programs/UnrealBuildTool/Platform/Windows/UEBuildWindows.cs
@@ -613,6 +613,7 @@ namespace UnrealBuildTool
 				    else if(Compiler == WindowsCompiler.VisualStudio2017)
 				    {
 					    // Enumerate all the installed Visual Studio instances using the interop SDK. There may be several installed side by side on a single machine (preview, community, enterprise, etc...).
+						List<DirectoryReference> PreReleaseInstallDirs = new List<DirectoryReference>();
 					    try
 					    {
 						    SetupConfiguration Setup = new SetupConfiguration();
@@ -632,13 +633,22 @@ namespace UnrealBuildTool
 							    ISetupInstance2 Instance = (ISetupInstance2)Instances[0];
 							    if((Instance.GetState() & InstanceState.Local) == InstanceState.Local)
 							    {
-								    InstallDirs.Add(new DirectoryReference(Instance.GetInstallationPath()));
+									ISetupInstanceCatalog Catalog = (ISetupInstanceCatalog)Instance as ISetupInstanceCatalog;
+									if (Catalog != null && Catalog.IsPrerelease())
+									{
+										PreReleaseInstallDirs.Add(new DirectoryReference(Instance.GetInstallationPath()));
+									}
+									else
+									{
+										InstallDirs.Add(new DirectoryReference(Instance.GetInstallationPath()));
+									}
 							    }
 						    }
 					    }
 					    catch
 					    {
 					    }
+						InstallDirs.AddRange(PreReleaseInstallDirs);
 				    }
 				    else
 				    {
@@ -736,33 +746,33 @@ namespace UnrealBuildTool
 							    foreach(DirectoryReference ToolChainDir in DirectoryReference.EnumerateDirectories(ToolChainBaseDir))
 							    {
 								    VersionNumber Version;
-								    if(VersionNumber.TryParse(ToolChainDir.GetDirectoryName(), out Version) && IsValidToolChainDir2017(ToolChainDir))
+								    if(VersionNumber.TryParse(ToolChainDir.GetDirectoryName(), out Version) && IsValidToolChainDir2017(ToolChainDir) && !ToolChainVersionToDir.ContainsKey(Version))
 								    {
 									    ToolChainVersionToDir[Version] = ToolChainDir;
 								    }
 							    }
 						    }
 					    }
-    
-					    // Enumerate all the AutoSDK toolchains
-					    DirectoryReference PlatformDir;
-					    if(UEBuildPlatformSDK.TryGetHostPlatformAutoSDKDir(out PlatformDir))
-					    {
-						    DirectoryReference ToolChainBaseDir = DirectoryReference.Combine(PlatformDir, "Win64", "VS2017");
-						    if(DirectoryReference.Exists(ToolChainBaseDir))
-						    {
-							    foreach(DirectoryReference ToolChainDir in DirectoryReference.EnumerateDirectories(ToolChainBaseDir))
-							    {
-								    VersionNumber Version;
-								    if(VersionNumber.TryParse(ToolChainDir.GetDirectoryName(), out Version) && IsValidToolChainDir2017(ToolChainDir))
-								    {
-									    ToolChainVersionToDir[Version] = ToolChainDir;
-								    }
-							    }
-						    }
-					    }
-				    }
-				    else
+
+						// Enumerate all the AutoSDK toolchains
+						DirectoryReference PlatformDir;
+						if(UEBuildPlatformSDK.TryGetHostPlatformAutoSDKDir(out PlatformDir))
+						{
+							DirectoryReference ToolChainBaseDir = DirectoryReference.Combine(PlatformDir, "Win64", "VS2017");
+							if(DirectoryReference.Exists(ToolChainBaseDir))
+							{
+								foreach(DirectoryReference ToolChainDir in DirectoryReference.EnumerateDirectories(ToolChainBaseDir))
+								{
+									VersionNumber Version;
+									if(VersionNumber.TryParse(ToolChainDir.GetDirectoryName(), out Version) && IsValidToolChainDir2017(ToolChainDir) && !ToolChainVersionToDir.ContainsKey(Version))
+									{
+										ToolChainVersionToDir[Version] = ToolChainDir;
+									}
+								}
+							}
+						}
+					}
+					else
 				    {
 					    throw new BuildException("Unsupported compiler version ({0})", Compiler);
 				    }
