UBT: Add a build setting which allows creating a dedicated PCH for every file that's excluded from the unity working set (disabled by default). Improves iteration times when working on individual cpp files, but slows down iterating on header changes (and can take a lot of disk space for large changes).

Dedicated PCH contains all includes scraped from the top of each cpp file, until a non-#include directive is encountered.

diff --git a/Engine/Source/Programs/UnrealBuildTool/Configuration/TargetRules.cs b/Engine/Source/Programs/UnrealBuildTool/Configuration/TargetRules.cs
--- a/Engine/Source/Programs/UnrealBuildTool/Configuration/TargetRules.cs
+++ b/Engine/Source/Programs/UnrealBuildTool/Configuration/TargetRules.cs
@@ -726,6 +726,12 @@ namespace UnrealBuildTool
 		[XmlConfigFile(Category = "BuildConfiguration")]
 		public bool bAdaptiveUnityDisablesPCH = true;
 
+		/// <summary>
+		/// Creates a dedicated PCH for each source file in the working set, allowing faster iteration on cpp-only changes.
+		/// </summary>
+		[XmlConfigFile(Category = "BuildConfiguration")]
+		public bool bAdaptiveUnityCreatesDedicatedPCH = false;
+
 		/// <summary>
 		/// The number of source files in a game module before unity build will be activated for that module.  This
 		/// allows small game modules to have faster iterative compile times for single files, at the expense of slower full
@@ -1846,6 +1852,11 @@ namespace UnrealBuildTool
 			get { return Inner.bAdaptiveUnityDisablesPCH; }
 		}
 
+		public bool bAdaptiveUnityCreatesDedicatedPCH
+		{
+			get { return Inner.bAdaptiveUnityCreatesDedicatedPCH; }
+		}
+
 		public int MinGameModuleSourceFilesForUnityBuild
 		{
 			get { return Inner.MinGameModuleSourceFilesForUnityBuild; }
diff --git a/Engine/Source/Programs/UnrealBuildTool/Configuration/UEBuildModuleCPP.cs b/Engine/Source/Programs/UnrealBuildTool/Configuration/UEBuildModuleCPP.cs
--- a/Engine/Source/Programs/UnrealBuildTool/Configuration/UEBuildModuleCPP.cs
+++ b/Engine/Source/Programs/UnrealBuildTool/Configuration/UEBuildModuleCPP.cs
@@ -773,7 +773,7 @@ namespace UnrealBuildTool
 
 			bool bAdaptiveUnityDisablesPCH = (Target.bAdaptiveUnityDisablesPCH && Rules.PCHUsage == ModuleRules.PCHUsageMode.UseExplicitOrSharedPCHs);
 
-			if ((Target.bAdaptiveUnityDisablesOptimizations || bAdaptiveUnityDisablesPCH) && !Target.bStressTestUnity)
+			if ((Target.bAdaptiveUnityDisablesOptimizations || bAdaptiveUnityDisablesPCH || Target.bAdaptiveUnityCreatesDedicatedPCH) && !Target.bStressTestUnity)
 			{
 				foreach (FileItem File in SourceFiles)
 				{
@@ -808,13 +808,21 @@ namespace UnrealBuildTool
 				{
 					AdaptiveUnityEnvironment.bOptimizeCode = false;
 				}
-				AdaptiveUnityEnvironment.PrecompiledHeaderAction = PrecompiledHeaderAction.None;
 
-				// Write all the definitions out to a separate file
-				CreateHeaderForDefinitions(AdaptiveUnityEnvironment, IntermediateDirectory, "Adaptive");
-
-				// Compile the files
-				CPPOutput AdaptiveOutput = ToolChain.CompileCPPFiles(AdaptiveUnityEnvironment, AdaptiveFiles, IntermediateDirectory, Name, ActionGraph);
+				// Create a per-file PCH
+				CPPOutput AdaptiveOutput;
+				if(Target.bAdaptiveUnityCreatesDedicatedPCH)
+				{
+					AdaptiveOutput = CompileAdaptiveNonUnityFilesWithDedicatedPCH(ToolChain, AdaptiveUnityEnvironment, AdaptiveFiles, IntermediateDirectory, Name, ActionGraph);
+				}
+				else if(bAdaptiveUnityDisablesPCH)
+				{
+					AdaptiveOutput = CompileAdaptiveNonUnityFilesWithoutPCH(ToolChain, AdaptiveUnityEnvironment, AdaptiveFiles, IntermediateDirectory, Name, ActionGraph);
+				}
+				else
+				{
+					AdaptiveOutput = CompileAdaptiveNonUnityFilesWithPCH(ToolChain, AdaptiveUnityEnvironment, AdaptiveFiles, IntermediateDirectory, Name, ActionGraph);
+				}
 
 				// Merge output
 				OutputFiles.ObjectFiles.AddRange(AdaptiveOutput.ObjectFiles);
@@ -824,6 +832,73 @@ namespace UnrealBuildTool
 			return OutputFiles;
 		}
 
+		static CPPOutput CompileAdaptiveNonUnityFilesWithPCH(UEToolChain ToolChain, CppCompileEnvironment CompileEnvironment, List<FileItem> Files, DirectoryReference IntermediateDirectory, string ModuleName, ActionGraph ActionGraph)
+		{
+			// Compile the files
+			return ToolChain.CompileCPPFiles(CompileEnvironment, Files, IntermediateDirectory, ModuleName, ActionGraph);
+		}
+
+		static CPPOutput CompileAdaptiveNonUnityFilesWithoutPCH(UEToolChain ToolChain, CppCompileEnvironment CompileEnvironment, List<FileItem> Files, DirectoryReference IntermediateDirectory, string ModuleName, ActionGraph ActionGraph)
+		{
+			// Disable precompiled headers
+			CompileEnvironment.PrecompiledHeaderAction = PrecompiledHeaderAction.None;
+
+			// Write all the definitions out to a separate file
+			CreateHeaderForDefinitions(CompileEnvironment, IntermediateDirectory, "Adaptive");
+
+			// Compile the files
+			return ToolChain.CompileCPPFiles(CompileEnvironment, Files, IntermediateDirectory, ModuleName, ActionGraph);
+		}
+
+		static CPPOutput CompileAdaptiveNonUnityFilesWithDedicatedPCH(UEToolChain ToolChain, CppCompileEnvironment CompileEnvironment, List<FileItem> Files, DirectoryReference IntermediateDirectory, string ModuleName, ActionGraph ActionGraph)
+		{
+			CPPOutput Output = new CPPOutput();
+			foreach(FileItem File in Files)
+			{
+				// Build the contents of the wrapper file
+				StringBuilder WrapperContents = new StringBuilder();
+				using (StringWriter Writer = new StringWriter(WrapperContents))
+				{
+					Writer.WriteLine("// Dedicated PCH for {0}", File.AbsolutePath);
+					Writer.WriteLine();
+					WriteDefinitions(CompileEnvironment.Definitions, Writer);
+					Writer.WriteLine();
+					using(StreamReader Reader = new StreamReader(File.Location.FullName))
+					{
+						CppIncludeParser.CopyIncludeDirectives(Reader, Writer);
+					}
+				}
+
+				// Write the PCH header
+				FileReference DedicatedPchLocation = FileReference.Combine(IntermediateDirectory, String.Format("PCH.Dedicated.{0}.h", File.Location.GetFileNameWithoutExtension()));
+				FileItem DedicatedPchFile = FileItem.CreateIntermediateTextFile(DedicatedPchLocation, WrapperContents.ToString());
+				DedicatedPchFile.CachedIncludePaths = File.CachedIncludePaths;
+
+				// Create a new C++ environment to compile the PCH
+				CppCompileEnvironment PchEnvironment = new CppCompileEnvironment(CompileEnvironment);
+				PchEnvironment.Definitions.Clear();
+				PchEnvironment.IncludePaths.UserIncludePaths.Add(File.Location.Directory); // Need to be able to include headers in the same directory as the source file
+				PchEnvironment.PrecompiledHeaderAction = PrecompiledHeaderAction.Create;
+				PchEnvironment.PrecompiledHeaderIncludeFilename = DedicatedPchFile.Location;
+
+				// Create the action to compile the PCH file.
+				CPPOutput PchOutput = ToolChain.CompileCPPFiles(PchEnvironment, new List<FileItem>() { DedicatedPchFile }, IntermediateDirectory, ModuleName, ActionGraph);
+				Output.ObjectFiles.AddRange(PchOutput.ObjectFiles);
+
+				// Create a new C++ environment to compile the original file
+				CppCompileEnvironment FileEnvironment = new CppCompileEnvironment(CompileEnvironment);
+				FileEnvironment.Definitions.Clear();
+				FileEnvironment.PrecompiledHeaderAction = PrecompiledHeaderAction.Include;
+				FileEnvironment.PrecompiledHeaderIncludeFilename = DedicatedPchFile.Location;
+				FileEnvironment.PrecompiledHeaderFile = PchOutput.PrecompiledHeaderFile;
+
+				// Create the action to compile the PCH file.
+				CPPOutput FileOutput = ToolChain.CompileCPPFiles(FileEnvironment, new List<FileItem>() { File }, IntermediateDirectory, ModuleName, ActionGraph);
+				Output.ObjectFiles.AddRange(FileOutput.ObjectFiles);
+			}
+			return Output;
+		}
+
 		/// <summary>
 		/// Creates a header file containing all the preprocessor definitions for a compile environment, and force-include it. We allow a more flexible syntax for preprocessor definitions than
 		/// is typically allowed on the command line (allowing function macros or double-quote characters, for example). Ensuring all definitions are specified in a header files ensures consistent
diff --git a/Engine/Source/Programs/UnrealBuildTool/System/CppIncludeParser.cs b/Engine/Source/Programs/UnrealBuildTool/System/CppIncludeParser.cs
new file mode 100644
--- /dev/null
+++ b/Engine/Source/Programs/UnrealBuildTool/System/CppIncludeParser.cs
@@ -0,0 +1,178 @@
+using System;
+using System.Collections.Generic;
+using System.IO;
+using System.Linq;
+using System.Text;
+using System.Threading.Tasks;
+
+namespace UnrealBuildTool
+{
+	/// <summary>
+	/// Parses include directives from cpp source files, to make dedicated PCHs
+	/// </summary>
+	static class CppIncludeParser
+	{
+		/// <summary>
+		/// Copies include directives from one file to another, until an unsafe directive (or non-#include line) is found.
+		/// </summary>
+		/// <param name="Reader">Stream to read from</param>
+		/// <param name="Writer">Stream to write directives to</param>
+		public static void CopyIncludeDirectives(StreamReader Reader, StringWriter Writer)
+		{
+			StringBuilder Token = new StringBuilder();
+			while(TryReadToken(Reader, Token))
+			{
+				if(Token.Length > 1 || Token[0] != '\n')
+				{
+					if(Token.Length != 1 || Token[0] != '#')
+					{
+						break;
+					}
+					if(!TryReadToken(Reader, Token))
+					{
+						break;
+					}
+
+					string Directive = Token.ToString();
+					if(Directive == "pragma")
+					{
+						if(!TryReadToken(Reader, Token) || Token.ToString() != "once")
+						{
+							break;
+						}
+						if(!TryReadToken(Reader, Token) || Token.ToString() != "\n")
+						{
+							break;
+						}
+					}
+					else if(Directive == "include")
+					{
+						if(!TryReadToken(Reader, Token) || Token[0] != '\"')
+						{
+							break;
+						}
+
+						string IncludeFile = Token.ToString();
+
+						if(!IncludeFile.EndsWith(".h\"") && !IncludeFile.EndsWith(".h>"))
+						{
+							break;
+						}
+						if(IncludeFile.Equals("\"RequiredProgramMainCPPInclude.h\"", StringComparison.OrdinalIgnoreCase))
+						{
+							break;
+						}
+						if(!TryReadToken(Reader, Token) || Token.ToString() != "\n")
+						{
+							break;
+						}
+
+						Writer.WriteLine("#include {0}", IncludeFile);
+					}
+					else
+					{
+						break;
+					}
+				}
+			}
+		}
+
+		/// <summary>
+		/// Reads an individual token from the input stream
+		/// </summary>
+		/// <param name="Reader">Stream to read from</param>
+		/// <param name="Token">Buffer to store token read from the stream</param>
+		/// <returns>True if a token was read, false otherwise</returns>
+		static bool TryReadToken(StreamReader Reader, StringBuilder Token)
+		{
+			Token.Clear();
+
+			int NextChar;
+			for(;;)
+			{
+				NextChar = Reader.Read();
+				if(NextChar == -1)
+				{
+					return false;
+				}
+				if(NextChar != ' ' && NextChar != '\t' && NextChar != '\r')
+				{
+					if(NextChar != '/')
+					{
+						break;
+					}
+					else if(Reader.Peek() == '/')
+					{
+						Reader.Read();
+						for(;;)
+						{
+							NextChar = Reader.Read();
+							if(NextChar == -1)
+							{
+								return false;
+							}
+							if(NextChar == '\n')
+							{
+								break;
+							}
+						}
+					}
+					else if(Reader.Peek() == '*')
+					{
+						Reader.Read();
+						for(;;)
+						{
+							NextChar = Reader.Read();
+							if(NextChar == -1)
+							{
+								return false;
+							}
+							if(NextChar == '*' && Reader.Peek() == '/')
+							{
+								break;
+							}
+						}
+						Reader.Read();
+					}
+					else
+					{
+						break;
+					}
+				}
+			}
+
+			Token.Append((char)NextChar);
+
+			if(Char.IsLetterOrDigit((char)NextChar))
+			{
+				for(;;)
+				{
+					NextChar = Reader.Read();
+					if(NextChar == -1 || !Char.IsLetterOrDigit((char)NextChar))
+					{
+						break;
+					}
+					Token.Append((char)NextChar);
+				}
+			}
+			else if(NextChar == '\"' || NextChar == '<')
+			{
+				for(;;)
+				{
+					NextChar = Reader.Read();
+					if(NextChar == -1)
+					{
+						break;
+					}
+					Token.Append((char)NextChar);
+					if(NextChar == '\"' || NextChar == '>')
+					{
+						break;
+					}
+				}
+			}
+
+			return true;
+		}
+	}
+}
diff --git a/Engine/Source/Programs/UnrealBuildTool/UnrealBuildTool.csproj b/Engine/Source/Programs/UnrealBuildTool/UnrealBuildTool.csproj
--- a/Engine/Source/Programs/UnrealBuildTool/UnrealBuildTool.csproj
+++ b/Engine/Source/Programs/UnrealBuildTool/UnrealBuildTool.csproj
@@ -173,6 +173,7 @@
     <Compile Include="ProjectFiles\Xcode\XcodeProjectFileGenerator.cs" />
     <Compile Include="System\BuildVersion.cs" />
     <Compile Include="System\CppCompileEnvironment.cs" />
+    <Compile Include="System\CppIncludeParser.cs" />
     <Compile Include="System\DataDrivenPlatformInfo.cs" />
     <Compile Include="System\EncryptionAndSigning.cs" />
     <Compile Include="System\ModuleManifest.cs" />
