UBT: Add an option to automatically enable edit and continue for files in the adaptive non-unity working set. E&C doesn't seem very useful for UE4 projects right now; compile time is comparable to regular build times, but it can take several minutes to apply code changes for large projects.

diff --git a/Engine/Source/Programs/UnrealBuildTool/Configuration/TargetRules.cs b/Engine/Source/Programs/UnrealBuildTool/Configuration/TargetRules.cs
--- a/Engine/Source/Programs/UnrealBuildTool/Configuration/TargetRules.cs
+++ b/Engine/Source/Programs/UnrealBuildTool/Configuration/TargetRules.cs
@@ -732,6 +732,12 @@ namespace UnrealBuildTool
 		[XmlConfigFile(Category = "BuildConfiguration")]
 		public bool bAdaptiveUnityCreatesDedicatedPCH = false;
 
+		/// <summary>
+		/// Creates a dedicated PCH for each source file in the working set, allowing faster iteration on cpp-only changes.
+		/// </summary>
+		[XmlConfigFile(Category = "BuildConfiguration")]
+		public bool bAdaptiveUnityEnablesEditAndContinue = false;
+
 		/// <summary>
 		/// The number of source files in a game module before unity build will be activated for that module.  This
 		/// allows small game modules to have faster iterative compile times for single files, at the expense of slower full
@@ -1857,6 +1863,11 @@ namespace UnrealBuildTool
 			get { return Inner.bAdaptiveUnityCreatesDedicatedPCH; }
 		}
 
+		public bool bAdaptiveUnityEnablesEditAndContinue
+		{
+			get { return Inner.bAdaptiveUnityEnablesEditAndContinue; }
+		}
+
 		public int MinGameModuleSourceFilesForUnityBuild
 		{
 			get { return Inner.MinGameModuleSourceFilesForUnityBuild; }
diff --git a/Engine/Source/Programs/UnrealBuildTool/Configuration/UEBuildModuleCPP.cs b/Engine/Source/Programs/UnrealBuildTool/Configuration/UEBuildModuleCPP.cs
--- a/Engine/Source/Programs/UnrealBuildTool/Configuration/UEBuildModuleCPP.cs
+++ b/Engine/Source/Programs/UnrealBuildTool/Configuration/UEBuildModuleCPP.cs
@@ -435,7 +435,7 @@ namespace UnrealBuildTool
 			// the module rules file.  By default, game modules only use unity of they have enough source files for that
 			// to be worthwhile.  If you have a lot of small game modules, consider specifying MinSourceFilesForUnityBuildOverride=0
 			// in the modules that you don't typically iterate on source files in very frequently.
-			int MinSourceFilesForUnityBuild = 0;
+			int MinSourceFilesForUnityBuild = 2;
 			if (Rules.MinSourceFilesForUnityBuildOverride != 0)
 			{
 				MinSourceFilesForUnityBuild = Rules.MinSourceFilesForUnityBuildOverride;
@@ -808,6 +808,10 @@ namespace UnrealBuildTool
 				{
 					AdaptiveUnityEnvironment.bOptimizeCode = false;
 				}
+				if (Target.bAdaptiveUnityEnablesEditAndContinue)
+				{
+					AdaptiveUnityEnvironment.bSupportEditAndContinue = true;
+				}
 
 				// Create a per-file PCH
 				CPPOutput AdaptiveOutput;
diff --git a/Engine/Source/Programs/UnrealBuildTool/Platform/Windows/UEBuildWindows.cs b/Engine/Source/Programs/UnrealBuildTool/Platform/Windows/UEBuildWindows.cs
--- a/Engine/Source/Programs/UnrealBuildTool/Platform/Windows/UEBuildWindows.cs
+++ b/Engine/Source/Programs/UnrealBuildTool/Platform/Windows/UEBuildWindows.cs
@@ -452,14 +452,17 @@ namespace UnrealBuildTool
 			}
 
 			// E&C support.
-			if (Target.bSupportEditAndContinue)
+			if (Target.bSupportEditAndContinue || Target.bAdaptiveUnityEnablesEditAndContinue)
 			{
 				Target.bUseIncrementalLinking = true;
-				Target.bUsePDBFiles = true;
+			}
+			if (Target.bAdaptiveUnityEnablesEditAndContinue && !Target.bAdaptiveUnityDisablesPCH && !Target.bAdaptiveUnityCreatesDedicatedPCH)
+			{
+				throw new BuildException("bAdaptiveUnityEnablesEditAndContinue requires bAdaptiveUnityDisablesPCH or bAdaptiveUnityCreatesDedicatedPCH");
 			}
 
 			// If we're using PDB files and PCHs, the generated code needs to be compiled with the same options as the PCH.
-			if(Target.bUsePDBFiles && Target.bUsePCHFiles)
+			if ((Target.bUsePDBFiles || Target.bSupportEditAndContinue) && Target.bUsePCHFiles)
 			{
 				Target.bDisableDebugInfoForGeneratedCode = false;
 			}
diff --git a/Engine/Source/Programs/UnrealBuildTool/Platform/Windows/VCToolChain.cs b/Engine/Source/Programs/UnrealBuildTool/Platform/Windows/VCToolChain.cs
--- a/Engine/Source/Programs/UnrealBuildTool/Platform/Windows/VCToolChain.cs
+++ b/Engine/Source/Programs/UnrealBuildTool/Platform/Windows/VCToolChain.cs
@@ -366,7 +366,7 @@ namespace UnrealBuildTool
 				// Store debug info in .pdb files.
 				// @todo clang: PDB files are emited from Clang but do not fully work with Visual Studio yet (breakpoints won't hit due to "symbol read error")
 				// @todo clang (update): as of clang 3.9 breakpoints work with PDBs, and the callstack is correct, so could be used for crash dumps. However debugging is still impossible due to the large amount of unreadable variables and unpredictable breakpoint stepping behaviour
-				if (CompileEnvironment.bUsePDBFiles)
+				if (CompileEnvironment.bUsePDBFiles || CompileEnvironment.bSupportEditAndContinue)
 				{
 					// Create debug info suitable for E&C if wanted.
 					if (CompileEnvironment.bSupportEditAndContinue)
@@ -996,7 +996,7 @@ namespace UnrealBuildTool
 					;
 
 				// Create PDB files if we were configured to do that.
-				if (CompileEnvironment.bUsePDBFiles)
+				if (CompileEnvironment.bUsePDBFiles || CompileEnvironment.bSupportEditAndContinue)
 				{
 					FileReference PDBLocation;
 					if (CompileEnvironment.PrecompiledHeaderAction == PrecompiledHeaderAction.Include)
diff --git a/Engine/Source/Programs/UnrealBuildTool/System/Unity.cs b/Engine/Source/Programs/UnrealBuildTool/System/Unity.cs
--- a/Engine/Source/Programs/UnrealBuildTool/System/Unity.cs
+++ b/Engine/Source/Programs/UnrealBuildTool/System/Unity.cs
@@ -270,7 +270,6 @@ namespace UnrealBuildTool
 							AdaptiveUnityBuildInfoString.Append(", " + CPPFileName);
 						}
 					}
-
 					else
 					{
 						// If adaptive unity build is enabled for this module, add this source file to the set that will invalidate the makefile
@@ -294,14 +293,23 @@ namespace UnrealBuildTool
 				{
 					if (PrintedSettingsForTargets.Add(Target.Name))
 					{
-						if (Target.bAdaptiveUnityDisablesPCH)
+						if (Target.bAdaptiveUnityCreatesDedicatedPCH)
+						{
+							Log.TraceInformation("[Adaptive unity build] Creating dedicated PCH for each excluded file. Set bAdaptiveUnityCreatesDedicatedPCH to false in BuildConfiguration.xml to change this behavior.");
+						}
+						else if (Target.bAdaptiveUnityDisablesPCH)
 						{
 							Log.TraceInformation("[Adaptive unity build] Disabling PCH for excluded files. Set bAdaptiveUnityDisablesPCH to false in BuildConfiguration.xml to change this behavior.");
 						}
+
 						if (Target.bAdaptiveUnityDisablesOptimizations)
 						{
 							Log.TraceInformation("[Adaptive unity build] Disabling optimizations for excluded files. Set bAdaptiveUnityDisablesOptimizations to false in BuildConfiguration.xml to change this behavior.");
 						}
+						if (Target.bAdaptiveUnityEnablesEditAndContinue)
+						{
+							Log.TraceInformation("[Adaptive unity build] Enabling Edit & Continue for excluded files. Set bAdaptiveUnityEnablesEditAndContinue to false in BuildConfiguration.xml to change this behavior.");
+						}
 					}
 					Log.TraceInformation(AdaptiveUnityBuildInfoString.ToString());
 				}
