UBT: Fix loading XML config files on Mono; Type.GetField(Name) does not seem to return values unless binding flags are specified.

#rnx

diff --git a/Engine/Source/Programs/UnrealBuildTool/System/XmlConfigData.cs b/Engine/Source/Programs/UnrealBuildTool/System/XmlConfigData.cs
--- a/Engine/Source/Programs/UnrealBuildTool/System/XmlConfigData.cs
+++ b/Engine/Source/Programs/UnrealBuildTool/System/XmlConfigData.cs
@@ -1,4 +1,4 @@
-﻿// Copyright 1998-2018 Epic Games, Inc. All Rights Reserved.
+// Copyright 1998-2018 Epic Games, Inc. All Rights Reserved.
 
 using System;
 using System.Collections.Generic;
@@ -95,7 +95,7 @@ namespace UnrealBuildTool
 						string FieldName = Reader.ReadString();
 
 						// Find the matching field on the output type
-						FieldInfo Field = Type.GetField(FieldName);
+						FieldInfo Field = Type.GetField(FieldName, BindingFlags.Public | BindingFlags.NonPublic | BindingFlags.Instance | BindingFlags.Static);
 						if(Field == null || Field.GetCustomAttribute<XmlConfigFileAttribute>() == null)
 						{
 							Data = null;
