Add a "SupportedPrograms" field to plugin descriptors, which allows plugins to declare which plugins they support independently of individual modules. Programs now respect the "bEnabledByDefault" setting in plugins.

Plugins that are compatible with a program now need to list that program in the SupportedPrograms list, and whitelist any modules that should load for that program.

diff --git a/Engine/Source/Programs/UnrealBuildTool/Configuration/UEBuildTarget.cs b/Engine/Source/Programs/UnrealBuildTool/Configuration/UEBuildTarget.cs
--- a/Engine/Source/Programs/UnrealBuildTool/Configuration/UEBuildTarget.cs
+++ b/Engine/Source/Programs/UnrealBuildTool/Configuration/UEBuildTarget.cs
@@ -3308,7 +3308,7 @@ namespace UnrealBuildTool
 			}
 
 			// Also synthesize references for plugins which are enabled by default
-			if (Rules.bCompileAgainstEngine)
+			if (Rules.bCompileAgainstEngine || Rules.bCompileWithPluginSupport)
 			{
 				foreach(PluginInfo Plugin in NameToInfo.Values)
 				{
@@ -3432,6 +3432,13 @@ namespace UnrealBuildTool
 					return null;
 				}
 
+				// Disable any plugins that aren't compatible with this program
+				if (Rules.Type == TargetType.Program && (Info.Descriptor.SupportedPrograms == null || !Info.Descriptor.SupportedPrograms.Contains(AppName)))
+				{
+					Log.TraceLog("Ignoring plugin '{0}' (referenced via {1}) due to absence from supported programs list.", Reference.Name, ReferenceChain);
+					return null;
+				}
+
 				// Create the new instance and add it to the cache
 				Log.TraceLog("Enabling plugin '{0}' (referenced via {1})", Reference.Name, ReferenceChain);
 				Instance = new UEBuildPlugin(Info);
diff --git a/Engine/Source/Programs/UnrealBuildTool/System/PluginDescriptor.cs b/Engine/Source/Programs/UnrealBuildTool/System/PluginDescriptor.cs
--- a/Engine/Source/Programs/UnrealBuildTool/System/PluginDescriptor.cs
+++ b/Engine/Source/Programs/UnrealBuildTool/System/PluginDescriptor.cs
@@ -119,6 +119,11 @@ namespace UnrealBuildTool
 		/// </summary>
 		public UnrealTargetPlatform[] SupportedTargetPlatforms;
 
+		/// <summary>
+		/// List of programs supported by this plugin.
+		/// </summary>
+		public string[] SupportedPrograms;
+
 		/// <summary>
 		/// List of all modules associated with this plugin
 		/// </summary>
@@ -230,6 +235,7 @@ namespace UnrealBuildTool
 			RawObject.TryGetStringField("SupportURL", out SupportURL);
 			RawObject.TryGetStringField("EngineVersion", out EngineVersion);
 			RawObject.TryGetEnumArrayField<UnrealTargetPlatform>("SupportedTargetPlatforms", out SupportedTargetPlatforms);
+			RawObject.TryGetStringArrayField("SupportedPrograms", out SupportedPrograms);
 
 			JsonObject[] ModulesArray;
 			if (RawObject.TryGetObjectArrayField("Modules", out ModulesArray))
@@ -342,6 +348,11 @@ namespace UnrealBuildTool
 				Writer.WriteEnumArrayField<UnrealTargetPlatform>("SupportedTargetPlatforms", SupportedTargetPlatforms);
 			}
 
+			if (SupportedPrograms != null && SupportedPrograms.Length > 0)
+			{
+				Writer.WriteStringArrayField("SupportedPrograms", SupportedPrograms);
+			}
+
 			ModuleDescriptor.WriteArray(Writer, "Modules", Modules);
 
 			LocalizationTargetDescriptor.WriteArray(Writer, "LocalizationTargets", LocalizationTargets);
diff --git a/Engine/Source/Runtime/Projects/Private/PluginDescriptor.cpp b/Engine/Source/Runtime/Projects/Private/PluginDescriptor.cpp
--- a/Engine/Source/Runtime/Projects/Private/PluginDescriptor.cpp
+++ b/Engine/Source/Runtime/Projects/Private/PluginDescriptor.cpp
@@ -115,6 +115,7 @@ bool FPluginDescriptor::Read(const FJsonObject& Object, FText& OutFailReason)
 	Object.TryGetStringField(TEXT("SupportURL"), SupportURL);
 	Object.TryGetStringField(TEXT("EngineVersion"), EngineVersion);
 	Object.TryGetStringArrayField(TEXT("SupportedTargetPlatforms"), SupportedTargetPlatforms);
+	Object.TryGetStringArrayField(TEXT("SupportedPrograms"), SupportedPrograms);
 
 	if (!FModuleDescriptor::ReadArray(Object, TEXT("Modules"), Modules, OutFailReason))
 	{
@@ -204,6 +205,10 @@ void FPluginDescriptor::Write(TJsonWriter<>& Writer) const
 	{
 		Writer.WriteValue(TEXT("SupportedTargetPlatforms"), SupportedTargetPlatforms);
 	}
+	if (SupportedPrograms.Num() > 0)
+	{
+		Writer.WriteValue(TEXT("SupportedPrograms"), SupportedPrograms);
+	}
 
 	FModuleDescriptor::WriteArray(Writer, TEXT("Modules"), Modules);
 
diff --git a/Engine/Source/Runtime/Projects/Private/PluginManager.cpp b/Engine/Source/Runtime/Projects/Private/PluginManager.cpp
--- a/Engine/Source/Runtime/Projects/Private/PluginManager.cpp
+++ b/Engine/Source/Runtime/Projects/Private/PluginManager.cpp
@@ -500,7 +500,6 @@ bool FPluginManager::ConfigureEnabledPlugins()
 				}
 			}
 
-#if !IS_PROGRAM || HACK_HEADER_GENERATOR
 			// Add the plugins which are enabled by default
 			for (const TPair<FString, TSharedRef<FPlugin>>& PluginPair : AllPlugins)
 			{
@@ -513,7 +512,6 @@ bool FPluginManager::ConfigureEnabledPlugins()
 					ConfiguredPluginNames.Add(PluginPair.Key);
 				}
 			}
-#endif
 		}
 #if IS_PROGRAM
 		// Programs can also define the list of enabled plugins in ini
@@ -774,6 +772,14 @@ bool FPluginManager::ConfigureEnabledPlugin(const FPluginReferenceDescriptor& Fi
 				continue;
 			}
 #endif
+			// Check that this plugin supports the current program
+#if IS_PROGRAM
+			if (!Plugin.Descriptor.SupportedPrograms.Contains(UE_APP_NAME))
+			{
+				UE_LOG(LogPluginManager, Verbose, TEXT("Ignoring plugin '%s' due to absence from the supported programs list"), *Reference.Name);
+				continue;
+			}
+#endif
 
 			// Skip loading Enterprise plugins when project is not an Enterprise project
 			if (Plugin.Type == EPluginType::Enterprise && !IProjectManager::Get().IsEnterpriseProject())
diff --git a/Engine/Source/Runtime/Projects/Public/PluginDescriptor.h b/Engine/Source/Runtime/Projects/Public/PluginDescriptor.h
--- a/Engine/Source/Runtime/Projects/Public/PluginDescriptor.h
+++ b/Engine/Source/Runtime/Projects/Public/PluginDescriptor.h
@@ -64,6 +64,9 @@ struct PROJECTS_API FPluginDescriptor
 	/** List of target platforms supported by this plugin. This list will be copied to any plugin reference from a project file, to allow filtering entire plugins from staged builds. */
 	TArray<FString> SupportedTargetPlatforms;
 
+	/** List of programs that are supported by this plugin. */
+	TArray<FString> SupportedPrograms;
+
 	/** List of all modules associated with this plugin */
 	TArray<FModuleDescriptor> Modules;
 
