UBT: Don't add import libraries as final build products unless the target is being precompiled. Prevents the need for building them for leaf nodes in the action graph.

#rnx

diff --git a/Engine/Source/Programs/UnrealBuildTool/Configuration/UEBuildBinary.cs b/Engine/Source/Programs/UnrealBuildTool/Configuration/UEBuildBinary.cs
--- a/Engine/Source/Programs/UnrealBuildTool/Configuration/UEBuildBinary.cs
+++ b/Engine/Source/Programs/UnrealBuildTool/Configuration/UEBuildBinary.cs
@@ -206,8 +206,12 @@ namespace UnrealBuildTool
 
 				if (BinaryLinkEnvironment.Platform != CppPlatform.Mac && BinaryLinkEnvironment.Platform != CppPlatform.Linux)
 				{
-					// Create the import library.
-					OutputFiles.AddRange(ToolChain.LinkAllFiles(BinaryLinkEnvironment, true, ActionGraph));
+					// Create the import libraries. They are only build products if we're precompiling, otherwise we can avoid building them for leaf nodes.
+					FileItem[] ImportLibraries = ToolChain.LinkAllFiles(BinaryLinkEnvironment, true, ActionGraph);
+					if(Target.bPrecompile)
+					{
+						OutputFiles.AddRange(ImportLibraries);
+					}
 				}
 			}
 
