UBT: Remove hack to allow Stats2.cpp to not follow IWYU convention.

#rnx

diff --git a/Engine/Source/Programs/UnrealBuildTool/Configuration/UEBuildModuleCPP.cs b/Engine/Source/Programs/UnrealBuildTool/Configuration/UEBuildModuleCPP.cs
--- a/Engine/Source/Programs/UnrealBuildTool/Configuration/UEBuildModuleCPP.cs
+++ b/Engine/Source/Programs/UnrealBuildTool/Configuration/UEBuildModuleCPP.cs
@@ -1174,7 +1174,7 @@ namespace UnrealBuildTool
 		{
 			switch(ExpectedName)
 			{
-				case "Stats2":
+//				case "Stats2":
 				case "DynamicRHI":
 				case "RHICommandList":
 				case "RHIUtilities":
diff --git a/Engine/Source/Runtime/Core/Private/Stats/Stats2.cpp b/Engine/Source/Runtime/Core/Private/Stats/Stats2.cpp
--- a/Engine/Source/Runtime/Core/Private/Stats/Stats2.cpp
+++ b/Engine/Source/Runtime/Core/Private/Stats/Stats2.cpp
@@ -1,6 +1,6 @@
 // Copyright 1998-2018 Epic Games, Inc. All Rights Reserved.
 
-#include "CoreTypes.h"
+#include "Stats/Stats2.h"
 #include "Misc/AssertionMacros.h"
 #include "HAL/UnrealMemory.h"
 #include "Templates/UnrealTemplate.h"
