Remove commented out line.

#rnx

diff --git a/Engine/Source/Programs/UnrealBuildTool/Configuration/UEBuildModuleCPP.cs b/Engine/Source/Programs/UnrealBuildTool/Configuration/UEBuildModuleCPP.cs
--- a/Engine/Source/Programs/UnrealBuildTool/Configuration/UEBuildModuleCPP.cs
+++ b/Engine/Source/Programs/UnrealBuildTool/Configuration/UEBuildModuleCPP.cs
@@ -1174,7 +1174,6 @@ namespace UnrealBuildTool
 		{
 			switch(ExpectedName)
 			{
-//				case "Stats2":
 				case "DynamicRHI":
 				case "RHICommandList":
 				case "RHIUtilities":
