Back out changelist 4224710; causing CIS errors due to receipts not matching.

#rnx

diff --git a/Engine/Source/Programs/UnrealBuildTool/Configuration/UEBuildBinary.cs b/Engine/Source/Programs/UnrealBuildTool/Configuration/UEBuildBinary.cs
--- a/Engine/Source/Programs/UnrealBuildTool/Configuration/UEBuildBinary.cs
+++ b/Engine/Source/Programs/UnrealBuildTool/Configuration/UEBuildBinary.cs
@@ -206,12 +206,8 @@ namespace UnrealBuildTool
 
 				if (BinaryLinkEnvironment.Platform != CppPlatform.Mac && BinaryLinkEnvironment.Platform != CppPlatform.Linux)
 				{
-					// Create the import libraries. They are only build products if we're precompiling, otherwise we can avoid building them for leaf nodes.
-					FileItem[] ImportLibraries = ToolChain.LinkAllFiles(BinaryLinkEnvironment, true, ActionGraph);
-					if(Target.bPrecompile)
-					{
-						OutputFiles.AddRange(ImportLibraries);
-					}
+					// Create the import library.
+					OutputFiles.AddRange(ToolChain.LinkAllFiles(BinaryLinkEnvironment, true, ActionGraph));
 				}
 			}
 
