UBT: Always add the generated code directory to the list of include paths when generating project files. It may only be created after UHT has been run.

diff --git a/Engine/Source/Programs/UnrealBuildTool/Configuration/UEBuildModuleCPP.cs b/Engine/Source/Programs/UnrealBuildTool/Configuration/UEBuildModuleCPP.cs
--- a/Engine/Source/Programs/UnrealBuildTool/Configuration/UEBuildModuleCPP.cs
+++ b/Engine/Source/Programs/UnrealBuildTool/Configuration/UEBuildModuleCPP.cs
@@ -346,6 +346,13 @@ namespace UnrealBuildTool
 			ProjectFile.AddIntelliSensePreprocessorDefinitions(ModuleCompileEnvironment.Definitions);
 			ProjectFile.AddInteliiSenseIncludePaths(ModuleCompileEnvironment.IncludePaths.SystemIncludePaths, true);
 			ProjectFile.AddInteliiSenseIncludePaths(ModuleCompileEnvironment.IncludePaths.UserIncludePaths, false);
+
+			// This directory may not exist for this module (or ever exist, if it doesn't contain any generated headers), but we want the project files
+			// to search it so we can pick up generated code definitions after UHT is run for the first time.
+			if(GeneratedCodeDirectory != null)
+			{
+				ProjectFile.AddInteliiSenseIncludePaths(new HashSet<DirectoryReference>{ GeneratedCodeDirectory }, false);
+			}
 		}
 
 		/// <summary>
