UBT: Remove legacy CppCompileEnvironment and LinkEnvironment wrappers from TargetRules that were deprecated in 4.19.

#rnx

diff --git a/Engine/Source/Programs/UnrealBuildTool/Configuration/TargetRules.cs b/Engine/Source/Programs/UnrealBuildTool/Configuration/TargetRules.cs
--- a/Engine/Source/Programs/UnrealBuildTool/Configuration/TargetRules.cs
+++ b/Engine/Source/Programs/UnrealBuildTool/Configuration/TargetRules.cs
@@ -134,106 +134,6 @@ namespace UnrealBuildTool
 			public const global::UnrealBuildTool.TargetType Program = global::UnrealBuildTool.TargetType.Program;
 		}
 
-		/// <summary>
-		/// Dummy class to maintain support for the SetupGlobalEnvironment() callback. We don't expose these classes directly any more.
-		/// </summary>
-		public class LinkEnvironmentConfiguration
-		{
-			/// <summary>
-			/// TargetRules instance to forward settings to.
-			/// </summary>
-			TargetRules Inner;
-
-			/// <summary>
-			/// Constructor.
-			/// </summary>
-			/// <param name="Inner">The target rules object. Fields on this object are treated as aliases to fields on this rules object.</param>
-			public LinkEnvironmentConfiguration(TargetRules Inner)
-			{
-				this.Inner = Inner;
-			}
-
-			/// <summary>
-			/// Settings exposed from the TargetRules instance
-			/// </summary>
-			#region Forwarded fields
-			#if !__MonoCS__
-			#pragma warning disable CS1591
-			#endif
-
-			[Obsolete("LinkEnvironmentConfiguration.bHasExports is deprecated in the 4.19 release. Modify TargetRules.bHasExports instead.")]
-			public bool bHasExports
-			{
-				get { return Inner.bHasExports; }
-				set { Inner.bHasExports = value; }
-			}
-
-			[Obsolete("LinkEnvironmentConfiguration.bIsBuildingConsoleApplication is deprecated in the 4.19 release. Modify TargetRules.bIsBuildingConsoleApplication instead.")]
-			public bool bIsBuildingConsoleApplication
-			{
-				get { return Inner.bIsBuildingConsoleApplication; }
-				set { Inner.bIsBuildingConsoleApplication = value; }
-			}
-
-			[Obsolete("LinkEnvironmentConfiguration.bDisableSymbolCache is deprecated in the 4.19 release. Modify TargetRules.bDisableSymbolCache instead.")]
-			public bool bDisableSymbolCache
-			{
-				get { return Inner.bDisableSymbolCache; }
-				set { Inner.bDisableSymbolCache = value; }
-			}
-
-			#if !__MonoCS__
-			#pragma warning restore CS1591
-			#endif
-			#endregion
-		}
-
-		/// <summary>
-		/// Dummy class to maintain support for the SetupGlobalEnvironment() callback. We don't expose these classes directly any more.
-		/// </summary>
-		public class CPPEnvironmentConfiguration
-		{
-			/// <summary>
-			/// TargetRules instance to forward settings to.
-			/// </summary>
-			TargetRules Inner;
-
-			/// <summary>
-			/// Constructor.
-			/// </summary>
-			/// <param name="Inner">The target rules object. Fields on this object are treated as aliases to fields on this rules object.</param>
-			public CPPEnvironmentConfiguration(TargetRules Inner)
-			{
-				this.Inner = Inner;
-			}
-
-			/// <summary>
-			/// Settings exposed from the TargetRules instance
-			/// </summary>
-			#region Forwarded fields
-			#if !__MonoCS__
-			#pragma warning disable CS1591
-			#endif
-
-			[Obsolete("CPPEnvironmentConfiguration.bEnableOSX109Support is deprecated in the 4.19 release. Modify TargetRules.GlobalDefinitions instead.")]
-			public List<string> Definitions
-			{
-				get { return Inner.GlobalDefinitions; }
-			}
-
-			[Obsolete("CPPEnvironmentConfiguration.bEnableOSX109Support is deprecated in the 4.19 release. Modify TargetRules.bEnableOSX109Support directly instead.")]
-			public bool bEnableOSX109Support
-			{
-				get { return Inner.bEnableOSX109Support; }
-				set { Inner.bEnableOSX109Support = value; }
-			}
-
-#if !__MonoCS__
-#pragma warning restore CS1591
-#endif
-			#endregion
-		}
-
 		/// <summary>
 		/// The name of this target.
 		/// </summary>
@@ -371,13 +271,6 @@ namespace UnrealBuildTool
 		/// </summary>
 		public bool? bBuildInSolutionByDefault = null;
 
-		/// <summary>
-		/// Output the executable to the engine binaries folder.
-		/// </summary>
-		[RequiresUniqueBuildEnvironment]
-		[Obsolete("bOutputToEngineBinaries is deprecated in 4.19. The output directory for binaries is determined automatically based on the location of the target.")]
-		public bool bOutputToEngineBinaries = false;
-
 		/// <summary>
 		/// Whether this target should be compiled as a DLL.  Requires LinkType to be set to TargetLinkType.Monolithic.
 		/// </summary>
@@ -1080,15 +973,6 @@ namespace UnrealBuildTool
 		/// </summary>
 		public List<string> ProjectDefinitions = new List<string>();
 
-		/// <summary>
-		/// Wrapper around GlobalDefinitions for people just stripping CPPEnvironmentConfiguration from variable names due to deprecation in 4.18.
-		/// </summary>
-		[Obsolete("Definitions is deprecated in the 4.19 release. Use GlobalDefinitions instead.")]
-		public List<string> Definitions
-		{
-			get { return GlobalDefinitions; }
-		}
-
 		/// <summary>
 		/// Specifies the name of the launch module. For modular builds, this is the module that is compiled into the target's executable.
 		/// </summary>
@@ -1423,23 +1307,6 @@ namespace UnrealBuildTool
 		{
 			get { return ProjectFileGenerator.bGenerateProjectFiles; }
 		}
-
-		/// <summary>
-		/// Setup the global environment for building this target
-		/// IMPORTANT: Game targets will *not* have this function called if they use the shared build environment.
-		/// See ShouldUseSharedBuildEnvironment().
-		/// </summary>
-		/// <param name="Target">The target information - such as platform and configuration</param>
-		/// <param name="OutLinkEnvironmentConfiguration">Output link environment settings</param>
-		/// <param name="OutCPPEnvironmentConfiguration">Output compile environment settings</param>
-		[ObsoleteOverride("SetupGlobalEnvironment() has been deprecated in the 4.19 release. Please set options from the constructor instead.")]
-		public virtual void SetupGlobalEnvironment(
-			TargetInfo Target,
-			ref LinkEnvironmentConfiguration OutLinkEnvironmentConfiguration,
-			ref CPPEnvironmentConfiguration OutCPPEnvironmentConfiguration
-			)
-		{
-		}
 	}
 
 	/// <summary>
@@ -1588,12 +1455,6 @@ namespace UnrealBuildTool
 			get { return Inner.bBuildInSolutionByDefault; }
 		}
 
-		[Obsolete("bOutputToEngineBinaries is deprecated in 4.19. The output directory for binaries is determined automatically based on the location of the target.")]
-		public bool bOutputToEngineBinaries
-		{
-			get { return Inner.bOutputToEngineBinaries; }
-		}
-
 		public string ExeBinariesSubFolder
 		{
 			get { return Inner.ExeBinariesSubFolder; }
diff --git a/Engine/Source/Programs/UnrealBuildTool/Configuration/UEBuildTarget.cs b/Engine/Source/Programs/UnrealBuildTool/Configuration/UEBuildTarget.cs
--- a/Engine/Source/Programs/UnrealBuildTool/Configuration/UEBuildTarget.cs
+++ b/Engine/Source/Programs/UnrealBuildTool/Configuration/UEBuildTarget.cs
@@ -450,14 +450,8 @@ namespace UnrealBuildTool
 				}
 			}
 
-			// Invoke the legacy callback for configuring the global environment
-			if(RulesObject.BuildEnvironment == TargetBuildEnvironment.Unique)
-			{
-				TargetRules.CPPEnvironmentConfiguration CppEnvironment = new TargetRules.CPPEnvironmentConfiguration(RulesObject);
-				TargetRules.LinkEnvironmentConfiguration LinkEnvironment = new TargetRules.LinkEnvironmentConfiguration(RulesObject);
-				RulesObject.SetupGlobalEnvironment(new TargetInfo(new ReadOnlyTargetRules(RulesObject)), ref LinkEnvironment, ref CppEnvironment);
-			}
-			else
+			// If we're using the shared build environment, make sure all the settings are valid
+			if(RulesObject.BuildEnvironment == TargetBuildEnvironment.Shared)
 			{
 				ValidateSharedEnvironment(RulesAssembly, Desc.Name, RulesObject);
 			}
diff --git a/Engine/Source/Programs/UnrealBuildTool/ProjectFiles/VisualStudioCode/VSCodeProjectFileGenerator.cs b/Engine/Source/Programs/UnrealBuildTool/ProjectFiles/VisualStudioCode/VSCodeProjectFileGenerator.cs
--- a/Engine/Source/Programs/UnrealBuildTool/ProjectFiles/VisualStudioCode/VSCodeProjectFileGenerator.cs
+++ b/Engine/Source/Programs/UnrealBuildTool/ProjectFiles/VisualStudioCode/VSCodeProjectFileGenerator.cs
@@ -365,10 +365,6 @@ namespace UnrealBuildTool
 				{
 					foreach (ProjectTarget Target in Project.ProjectTargets)
 					{
-						TargetRules.CPPEnvironmentConfiguration CppEnvironment = new TargetRules.CPPEnvironmentConfiguration(Target.TargetRules);
-						TargetRules.LinkEnvironmentConfiguration LinkEnvironment = new TargetRules.LinkEnvironmentConfiguration(Target.TargetRules);
-						Target.TargetRules.SetupGlobalEnvironment(new TargetInfo(new ReadOnlyTargetRules(Target.TargetRules)), ref LinkEnvironment, ref CppEnvironment);
-
 						Array Configs = Enum.GetValues(typeof(UnrealTargetConfiguration));
 						List<UnrealTargetPlatform> Platforms = new List<UnrealTargetPlatform>(UEBuildTarget.GetSupportedPlatforms(Target.TargetRules));
 
@@ -395,7 +391,7 @@ namespace UnrealBuildTool
 											Config = Config,
 											UProjectFile = Target.UnrealProjectFilePath,
 											OutputType = ProjectData.EOutputType.Exe,
-											OutputFile = GetExecutableFilename(Project, Target, Platform, Config, LinkEnvironment),
+											OutputFile = GetExecutableFilename(Project, Target, Platform, Config),
 											CSharpInfo = null
 										});
 									}
@@ -863,7 +859,7 @@ namespace UnrealBuildTool
 			return Result;
 		}
 
-		private FileReference GetExecutableFilename(ProjectFile Project, ProjectTarget Target, UnrealTargetPlatform Platform, UnrealTargetConfiguration Configuration, TargetRules.LinkEnvironmentConfiguration LinkEnvironment)
+		private FileReference GetExecutableFilename(ProjectFile Project, ProjectTarget Target, UnrealTargetPlatform Platform, UnrealTargetConfiguration Configuration)
 		{
 			TargetRules TargetRulesObject = Target.TargetRules;
 			FileReference TargetFilePath = Target.TargetFilePath;
