UBT: Fix editor targets with unique build environment having the wrong executable path in generated project files. Move move logic to configure target rules post-construction by the rules assembly to ensure it's valid.

diff --git a/Engine/Source/Programs/UnrealBuildTool/Configuration/TargetInfo.cs b/Engine/Source/Programs/UnrealBuildTool/Configuration/TargetInfo.cs
--- a/Engine/Source/Programs/UnrealBuildTool/Configuration/TargetInfo.cs
+++ b/Engine/Source/Programs/UnrealBuildTool/Configuration/TargetInfo.cs
@@ -1,4 +1,4 @@
-﻿// Copyright 1998-2018 Epic Games, Inc. All Rights Reserved.
+// Copyright 1998-2018 Epic Games, Inc. All Rights Reserved.
 
 using System;
 using System.Collections.Generic;
@@ -63,19 +63,5 @@ namespace UnrealBuildTool
 			this.ProjectFile = ProjectFile;
 			this.Version = Version;
 		}
-
-		/// <summary>
-		/// Constructs a TargetInfo for passing to the ModuleRules constructor.
-		/// </summary>
-		/// <param name="Rules">The constructed target rules object. This is null when passed into a TargetRules constructor, but should be valid at all other times.</param>
-		public TargetInfo(ReadOnlyTargetRules Rules)
-		{
-			this.Name = Rules.Name;
-			this.Platform = Rules.Platform;
-			this.Configuration = Rules.Configuration;
-			this.Architecture = Rules.Architecture;
-			this.ProjectFile = Rules.ProjectFile;
-			this.Version = Rules.Version;
-		}
 	}
 }
diff --git a/Engine/Source/Programs/UnrealBuildTool/Configuration/UEBuildTarget.cs b/Engine/Source/Programs/UnrealBuildTool/Configuration/UEBuildTarget.cs
--- a/Engine/Source/Programs/UnrealBuildTool/Configuration/UEBuildTarget.cs
+++ b/Engine/Source/Programs/UnrealBuildTool/Configuration/UEBuildTarget.cs
@@ -412,7 +412,7 @@ namespace UnrealBuildTool
 			}
 
 			FileReference TargetFileName;
-			TargetRules RulesObject = RulesAssembly.CreateTargetRules(Desc.Name, Desc.Platform, Desc.Configuration, Desc.Architecture, Desc.ProjectFile, Version, out TargetFileName);
+			TargetRules RulesObject = RulesAssembly.CreateTargetRules(Desc.Name, Desc.Platform, Desc.Configuration, Desc.Architecture, Desc.ProjectFile, Version, Arguments, out TargetFileName);
 			if ((ProjectFileGenerator.bGenerateProjectFiles == false) && !GetSupportedPlatforms(RulesObject).Contains(Desc.Platform))
 			{
 				throw new BuildException("{0} does not support the {1} platform.", Desc.Name, Desc.Platform.ToString());
@@ -421,35 +421,6 @@ namespace UnrealBuildTool
 			// Now that we found the actual Editor target, make sure we're no longer using the old TargetName (which is the Game target)
 			Desc.Name = RulesObject.Name;
 
-			// Parse any additional command-line arguments. These override default settings specified in config files or the .target.cs files.
-			foreach(object ConfigurableObject in RulesObject.GetConfigurableObjects())
-			{
-				CommandLine.ParseArguments(Arguments, ConfigurableObject);
-			}
-
-			// Set the final value for the link type in the target rules
-			if(RulesObject.LinkType == TargetLinkType.Default)
-			{
-				throw new BuildException("TargetRules.LinkType should be inferred from TargetType");
-			}
-
-			// Set the default value for whether to use the shared build environment
-			if(RulesObject.BuildEnvironment == TargetBuildEnvironment.Default)
-			{
-				if(RulesObject.Type == TargetType.Program && Desc.ProjectFile != null && TargetFileName.IsUnderDirectory(Desc.ProjectFile.Directory))
-				{
-					RulesObject.BuildEnvironment = TargetBuildEnvironment.Unique;
-				}
-				else if(UnrealBuildTool.IsEngineInstalled() || RulesObject.LinkType != TargetLinkType.Monolithic)
-				{
-					RulesObject.BuildEnvironment = TargetBuildEnvironment.Shared;
-				}
-				else
-				{
-					RulesObject.BuildEnvironment = TargetBuildEnvironment.Unique;
-				}
-			}
-
 			// If we're using the shared build environment, make sure all the settings are valid
 			if(RulesObject.BuildEnvironment == TargetBuildEnvironment.Shared)
 			{
@@ -472,31 +443,6 @@ namespace UnrealBuildTool
 				RulesObject.GlobalDefinitions.Add("USE_MALLOC_PROFILER=1");
 			}
 
-			// handle some special case defines (so build system can pass -DEFINE as normal instead of needing
-			// to know about special parameters)
-			foreach (string Define in RulesObject.GlobalDefinitions)
-			{
-				switch (Define)
-				{
-					case "WITH_EDITOR=0":
-						RulesObject.bBuildEditor = false;
-						break;
-
-					case "WITH_EDITORONLY_DATA=0":
-						RulesObject.bBuildWithEditorOnlyData = false;
-						break;
-
-					// Memory profiler doesn't work if frame pointers are omitted
-					case "USE_MALLOC_PROFILER=1":
-						RulesObject.bOmitFramePointers = false;
-						break;
-
-					case "WITH_LEAN_AND_MEAN_UE=1":
-						RulesObject.bCompileLeanAndMeanUE = true;
-						break;
-				}
-			}
-
 			// If we're precompiling, generate a list of all the files that we depend on
 			if (RulesObject.bPrecompile)
 			{
@@ -532,22 +478,6 @@ namespace UnrealBuildTool
 				RulesObject.bDisableLinking = true;
 			}
 
-			// Lean and mean means no Editor and other frills.
-			if (RulesObject.bCompileLeanAndMeanUE)
-			{
-				RulesObject.bBuildEditor = false;
-				RulesObject.bBuildDeveloperTools = false;
-				RulesObject.bCompileSimplygon = false;
-				RulesObject.bCompileSimplygonSSF = false;
-				RulesObject.bCompileSpeedTree = false;
-			}
-
-			// Automatically include CoreUObject
-			if (RulesObject.bCompileAgainstEngine)
-			{
-				RulesObject.bCompileAgainstCoreUObject = true;
-			}
-
 			// Disable editor when its not needed
 			UEBuildPlatform BuildPlatform = UEBuildPlatform.GetBuildPlatform(RulesObject.Platform);
 			if (BuildPlatform.ShouldNotBuildEditor(Desc.Platform, Desc.Configuration) == true)
@@ -561,24 +491,6 @@ namespace UnrealBuildTool
 				RulesObject.bBuildRequiresCookedData = true;
 			}
 
-			// Must have editor only data if building the editor.
-			if (RulesObject.bBuildEditor)
-			{
-				RulesObject.bBuildWithEditorOnlyData = true;
-			}
-
-			// Apply the override to force debug info to be enabled
-			if (RulesObject.bForceDebugInfo)
-			{
-				RulesObject.bDisableDebugInfo = false;
-				RulesObject.bOmitPCDebugInfoInDevelopment = false;
-			}
-
-			if (!RulesObject.bAllowGeneratedIniWhenCooked)
-			{
-				RulesObject.GlobalDefinitions.Add("DISABLE_GENERATED_INI_WHEN_COOKED=1");
-			}
-
 			// Allow the platform to finalize the settings
 			UEBuildPlatform Platform = UEBuildPlatform.GetBuildPlatform(RulesObject.Platform);
 			Platform.ValidateTarget(RulesObject);
@@ -664,7 +576,7 @@ namespace UnrealBuildTool
 			}
 
 			// Create the target rules for it
-			TargetRules BaseRules = RulesAssembly.CreateTargetRules(BaseTargetName, ThisRules.Platform, ThisRules.Configuration, ThisRules.Architecture, null, ThisRules.Version);
+			TargetRules BaseRules = RulesAssembly.CreateTargetRules(BaseTargetName, ThisRules.Platform, ThisRules.Configuration, ThisRules.Architecture, null, ThisRules.Version, null);
 
 			// Get all the configurable objects
 			object[] BaseObjects = BaseRules.GetConfigurableObjects().ToArray();
diff --git a/Engine/Source/Programs/UnrealBuildTool/ProjectFiles/ProjectFileGenerator.cs b/Engine/Source/Programs/UnrealBuildTool/ProjectFiles/ProjectFileGenerator.cs
--- a/Engine/Source/Programs/UnrealBuildTool/ProjectFiles/ProjectFileGenerator.cs
+++ b/Engine/Source/Programs/UnrealBuildTool/ProjectFiles/ProjectFileGenerator.cs
@@ -2107,7 +2107,7 @@ namespace UnrealBuildTool
 
 					// Create target rules for all of the platforms and configuration combinations that we want to enable support for.
 					// Just use the current platform as we only need to recover the target type and both should be supported for all targets...
-					TargetRules TargetRulesObject = RulesAssembly.CreateTargetRules(TargetName, BuildHostPlatform.Current.Platform, UnrealTargetConfiguration.Development, "", CheckProjectFile, Version);
+					TargetRules TargetRulesObject = RulesAssembly.CreateTargetRules(TargetName, BuildHostPlatform.Current.Platform, UnrealTargetConfiguration.Development, "", CheckProjectFile, Version, null);
 
 					bool IsProgramTarget = false;
 
@@ -2249,7 +2249,7 @@ namespace UnrealBuildTool
 							ProjectFilePath = ProjectFilePath,
 							UnrealProjectFilePath = CheckProjectFile,
 							SupportedPlatforms = UEBuildTarget.GetSupportedPlatforms(TargetRulesObject).Where(x => UEBuildPlatform.GetBuildPlatform(x, true) != null).ToArray(),
-							CreateRulesDelegate = (Platform, Configuration) => RulesAssembly.CreateTargetRules(TargetName, Platform, Configuration, "", CheckProjectFile, Version)
+							CreateRulesDelegate = (Platform, Configuration) => RulesAssembly.CreateTargetRules(TargetName, Platform, Configuration, "", CheckProjectFile, Version, null)
                         };
 
 					if (TargetName == "ShaderCompileWorker")		// @todo projectfiles: Ideally, the target rules file should set this
diff --git a/Engine/Source/Programs/UnrealBuildTool/ProjectFiles/VisualStudio/VCProject.cs b/Engine/Source/Programs/UnrealBuildTool/ProjectFiles/VisualStudio/VCProject.cs
--- a/Engine/Source/Programs/UnrealBuildTool/ProjectFiles/VisualStudio/VCProject.cs
+++ b/Engine/Source/Programs/UnrealBuildTool/ProjectFiles/VisualStudio/VCProject.cs
@@ -1189,14 +1189,9 @@ namespace UnrealBuildTool
 
 					// Get the executable name (minus any platform or config suffixes)
 					string BaseExeName = TargetName;
-					if (!bShouldCompileMonolithic && TargetRulesObject.Type != TargetType.Program)
+					if (!bShouldCompileMonolithic && TargetRulesObject.Type != TargetType.Program && TargetRulesObject.BuildEnvironment != TargetBuildEnvironment.Unique)
 					{
-						// Figure out what the compiled binary will be called so that we can point the IDE to the correct file
-						string TargetConfigurationName = TargetRulesObject.Type.ToString();
-						if (TargetConfigurationName != TargetType.Game.ToString() && TargetConfigurationName != TargetType.Program.ToString())
-						{
-							BaseExeName = "UE4" + TargetConfigurationName;
-						}
+						BaseExeName = "UE4" + TargetRulesObject.Type.ToString();
 					}
 
 					// Make the output file path
diff --git a/Engine/Source/Programs/UnrealBuildTool/System/RulesAssembly.cs b/Engine/Source/Programs/UnrealBuildTool/System/RulesAssembly.cs
--- a/Engine/Source/Programs/UnrealBuildTool/System/RulesAssembly.cs
+++ b/Engine/Source/Programs/UnrealBuildTool/System/RulesAssembly.cs
@@ -388,8 +388,9 @@ namespace UnrealBuildTool
 		/// </summary>
 		/// <param name="TypeName">Type name of the target rules</param>
 		/// <param name="TargetInfo">Target configuration information to pass to the constructor</param>
+		/// <param name="Arguments">Command line arguments for this target</param>
 		/// <returns>Instance of the corresponding TargetRules</returns>
-		protected TargetRules CreateTargetRulesInstance(string TypeName, TargetInfo TargetInfo)
+		protected TargetRules CreateTargetRulesInstance(string TypeName, TargetInfo TargetInfo, string[] Arguments)
 		{
 			// The build module must define a type named '<TargetName>Target' that derives from our 'TargetRules' type.  
 			Type RulesType = CompiledAssembly.GetType(TypeName);
@@ -419,7 +420,75 @@ namespace UnrealBuildTool
 				throw new BuildException(Ex, "Unable to instantiate instance of '{0}' object type from compiled assembly '{1}'.  Unreal Build Tool creates an instance of your module's 'Rules' object in order to find out about your module's requirements.  The CLR exception details may provide more information:  {2}", TypeName, Path.GetFileNameWithoutExtension(CompiledAssembly.Location), Ex.ToString());
 			}
 
+			// Set the default overriddes for the configured target type
 			Rules.SetOverridesForTargetType();
+
+			// Parse any additional command-line arguments. These override default settings specified in config files or the .target.cs files.
+			if(Arguments != null)
+			{
+				foreach(object ConfigurableObject in Rules.GetConfigurableObjects())
+				{
+					CommandLine.ParseArguments(Arguments, ConfigurableObject);
+				}
+			}
+
+			// Set the final value for the link type in the target rules
+			if(Rules.LinkType == TargetLinkType.Default)
+			{
+				throw new BuildException("TargetRules.LinkType should be inferred from TargetType");
+			}
+
+			// Set the default value for whether to use the shared build environment
+			if(Rules.BuildEnvironment == TargetBuildEnvironment.Default)
+			{
+				if(Rules.Type == TargetType.Program && TargetInfo.ProjectFile != null && TargetNameToTargetFile[Rules.Name].IsUnderDirectory(TargetInfo.ProjectFile.Directory))
+				{
+					Rules.BuildEnvironment = TargetBuildEnvironment.Unique;
+				}
+				else if(UnrealBuildTool.IsEngineInstalled() || Rules.LinkType != TargetLinkType.Monolithic)
+				{
+					Rules.BuildEnvironment = TargetBuildEnvironment.Shared;
+				}
+				else
+				{
+					Rules.BuildEnvironment = TargetBuildEnvironment.Unique;
+				}
+			}
+
+			// Lean and mean means no Editor and other frills.
+			if (Rules.bCompileLeanAndMeanUE)
+			{
+				Rules.bBuildEditor = false;
+				Rules.bBuildDeveloperTools = false;
+				Rules.bCompileSimplygon = false;
+				Rules.bCompileSimplygonSSF = false;
+				Rules.bCompileSpeedTree = false;
+			}
+
+			// Automatically include CoreUObject
+			if (Rules.bCompileAgainstEngine)
+			{
+				Rules.bCompileAgainstCoreUObject = true;
+			}
+
+			// Must have editor only data if building the editor.
+			if (Rules.bBuildEditor)
+			{
+				Rules.bBuildWithEditorOnlyData = true;
+			}
+
+			// Apply the override to force debug info to be enabled
+			if (Rules.bForceDebugInfo)
+			{
+				Rules.bDisableDebugInfo = false;
+				Rules.bOmitPCDebugInfoInDevelopment = false;
+			}
+
+			if (!Rules.bAllowGeneratedIniWhenCooked)
+			{
+				Rules.GlobalDefinitions.Add("DISABLE_GENERATED_INI_WHEN_COOKED=1");
+			}
+
 			return Rules;
 		}
 
@@ -431,12 +500,13 @@ namespace UnrealBuildTool
 		/// <param name="Configuration">The configuration the target is being built for</param>
 		/// <param name="Architecture">The architecture the target is being built for</param>
 		/// <param name="ProjectFile">The project containing the target being built</param>
+		/// <param name="Arguments">Command line arguments for the target</param>
 		/// <param name="Version">The current build version</param>
 		/// <returns>The build target rules for the specified target</returns>
-		public TargetRules CreateTargetRules(string TargetName, UnrealTargetPlatform Platform, UnrealTargetConfiguration Configuration, string Architecture, FileReference ProjectFile, ReadOnlyBuildVersion Version)
+		public TargetRules CreateTargetRules(string TargetName, UnrealTargetPlatform Platform, UnrealTargetConfiguration Configuration, string Architecture, FileReference ProjectFile, ReadOnlyBuildVersion Version, string[] Arguments)
 		{
 			FileReference TargetFileName;
-			return CreateTargetRules(TargetName, Platform, Configuration, Architecture, ProjectFile, Version, out TargetFileName);
+			return CreateTargetRules(TargetName, Platform, Configuration, Architecture, ProjectFile, Version, Arguments, out TargetFileName);
 		}
 
 		/// <summary>
@@ -448,9 +518,10 @@ namespace UnrealBuildTool
 		/// <param name="Architecture">Architecture being built</param>
 		/// <param name="ProjectFile">Path to the project file for this target</param>
 		/// <param name="Version">The current build version</param>
+		/// <param name="Arguments">Command line arguments for this target</param>
 		/// <param name="TargetFileName">The original source file name of the Target.cs file for this target</param>
 		/// <returns>The build target rules for the specified target</returns>
-		public TargetRules CreateTargetRules(string TargetName, UnrealTargetPlatform Platform, UnrealTargetConfiguration Configuration, string Architecture, FileReference ProjectFile, ReadOnlyBuildVersion Version, out FileReference TargetFileName)
+		public TargetRules CreateTargetRules(string TargetName, UnrealTargetPlatform Platform, UnrealTargetConfiguration Configuration, string Architecture, FileReference ProjectFile, ReadOnlyBuildVersion Version, string[] Arguments, out FileReference TargetFileName)
 		{
 			bool bFoundTargetName = TargetNameToTargetFile.ContainsKey(TargetName);
 			if (bFoundTargetName == false)
@@ -476,7 +547,7 @@ namespace UnrealBuildTool
 				}
 				else
 				{
-					return Parent.CreateTargetRules(TargetName, Platform, Configuration, Architecture, ProjectFile, Version, out TargetFileName);
+					return Parent.CreateTargetRules(TargetName, Platform, Configuration, Architecture, ProjectFile, Version, Arguments, out TargetFileName);
 				}
 			}
 
@@ -487,7 +558,7 @@ namespace UnrealBuildTool
 			string TargetTypeName = TargetName + "Target";
 
 			// The build module must define a type named '<TargetName>Target' that derives from our 'TargetRules' type.  
-			return CreateTargetRulesInstance(TargetTypeName, new TargetInfo(TargetName, Platform, Configuration, Architecture, ProjectFile, Version));
+			return CreateTargetRulesInstance(TargetTypeName, new TargetInfo(TargetName, Platform, Configuration, Architecture, ProjectFile, Version), Arguments);
 		}
 
 		/// <summary>
@@ -506,7 +577,7 @@ namespace UnrealBuildTool
 			List<string> Matches = new List<string>();
 			foreach(KeyValuePair<string, FileReference> TargetPair in TargetNameToTargetFile)
 			{
-				TargetRules Rules = CreateTargetRulesInstance(TargetPair.Key + "Target", new TargetInfo(TargetPair.Key, Platform, Configuration, Architecture, ProjectFile, Version));
+				TargetRules Rules = CreateTargetRulesInstance(TargetPair.Key + "Target", new TargetInfo(TargetPair.Key, Platform, Configuration, Architecture, ProjectFile, Version), null);
 				if(Rules.Type == Type)
 				{
 					Matches.Add(TargetPair.Key);
