UBT: Remove the UEBuildPlatform.ShouldNotBuildEditor() hook for target platforms. We shouldn't be modifying a target's build environment to disable the editor; it is invalid to build the editor for these target platforms at all, and this is already enforced by the GetSupportedPlatforms() function.

#rnx

diff --git a/Engine/Source/Programs/UnrealBuildTool/Configuration/UEBuildPlatform.cs b/Engine/Source/Programs/UnrealBuildTool/Configuration/UEBuildPlatform.cs
--- a/Engine/Source/Programs/UnrealBuildTool/Configuration/UEBuildPlatform.cs
+++ b/Engine/Source/Programs/UnrealBuildTool/Configuration/UEBuildPlatform.cs
@@ -615,17 +615,6 @@ namespace UnrealBuildTool
 			throw new BuildException("GetDebugInfoExtensions for {0} not handled in {1}", InBinaryType.ToString(), this.ToString());
 		}
 
-		/// <summary>
-		/// Whether the editor should be built for this platform or not
-		/// </summary>
-		/// <param name="InPlatform"> The UnrealTargetPlatform being built</param>
-		/// <param name="InConfiguration">The UnrealTargetConfiguration being built</param>
-		/// <returns>bool   true if the editor should be built, false if not</returns>
-		public virtual bool ShouldNotBuildEditor(UnrealTargetPlatform InPlatform, UnrealTargetConfiguration InConfiguration)
-		{
-			return false;
-		}
-
 		/// <summary>
 		/// Whether this build should support ONLY cooked data or not
 		/// </summary>
diff --git a/Engine/Source/Programs/UnrealBuildTool/Configuration/UEBuildTarget.cs b/Engine/Source/Programs/UnrealBuildTool/Configuration/UEBuildTarget.cs
--- a/Engine/Source/Programs/UnrealBuildTool/Configuration/UEBuildTarget.cs
+++ b/Engine/Source/Programs/UnrealBuildTool/Configuration/UEBuildTarget.cs
@@ -478,14 +478,8 @@ namespace UnrealBuildTool
 				RulesObject.bDisableLinking = true;
 			}
 
-			// Disable editor when its not needed
-			UEBuildPlatform BuildPlatform = UEBuildPlatform.GetBuildPlatform(RulesObject.Platform);
-			if (BuildPlatform.ShouldNotBuildEditor(Desc.Platform, Desc.Configuration) == true)
-			{
-				RulesObject.bBuildEditor = false;
-			}
-
 			// Disable the DDC and a few other things related to preparing assets
+			UEBuildPlatform BuildPlatform = UEBuildPlatform.GetBuildPlatform(RulesObject.Platform);
 			if (BuildPlatform.BuildRequiresCookedData(Desc.Platform, Desc.Configuration) == true)
 			{
 				RulesObject.bBuildRequiresCookedData = true;
diff --git a/Engine/Source/Programs/UnrealBuildTool/Platform/Android/UEBuildAndroid.cs b/Engine/Source/Programs/UnrealBuildTool/Platform/Android/UEBuildAndroid.cs
--- a/Engine/Source/Programs/UnrealBuildTool/Platform/Android/UEBuildAndroid.cs
+++ b/Engine/Source/Programs/UnrealBuildTool/Platform/Android/UEBuildAndroid.cs
@@ -217,11 +217,6 @@ namespace UnrealBuildTool
 			return true;
 		}
 
-		public override bool ShouldNotBuildEditor(UnrealTargetPlatform InPlatform, UnrealTargetConfiguration InConfiguration)
-		{
-			return true;
-		}
-
 		public override bool BuildRequiresCookedData(UnrealTargetPlatform InPlatform, UnrealTargetConfiguration InConfiguration)
 		{
 			return true;
diff --git a/Engine/Source/Programs/UnrealBuildTool/Platform/HTML5/UEBuildHTML5.cs b/Engine/Source/Programs/UnrealBuildTool/Platform/HTML5/UEBuildHTML5.cs
--- a/Engine/Source/Programs/UnrealBuildTool/Platform/HTML5/UEBuildHTML5.cs
+++ b/Engine/Source/Programs/UnrealBuildTool/Platform/HTML5/UEBuildHTML5.cs
@@ -120,17 +120,6 @@ namespace UnrealBuildTool
 			return true;
 		}
 
-		/// <summary>
-		/// Whether the editor should be built for this platform or not
-		/// </summary>
-		/// <param name="InPlatform"> The UnrealTargetPlatform being built</param>
-		/// <param name="InConfiguration">The UnrealTargetConfiguration being built</param>
-		/// <returns>bool   true if the editor should be built, false if not</returns>
-		public override bool ShouldNotBuildEditor(UnrealTargetPlatform InPlatform, UnrealTargetConfiguration InConfiguration)
-		{
-			return true;
-		}
-
 		public override bool BuildRequiresCookedData(UnrealTargetPlatform InPlatform, UnrealTargetConfiguration InConfiguration)
 		{
 			return true;
diff --git a/Engine/Source/Programs/UnrealBuildTool/Platform/IOS/UEBuildIOS.cs b/Engine/Source/Programs/UnrealBuildTool/Platform/IOS/UEBuildIOS.cs
--- a/Engine/Source/Programs/UnrealBuildTool/Platform/IOS/UEBuildIOS.cs
+++ b/Engine/Source/Programs/UnrealBuildTool/Platform/IOS/UEBuildIOS.cs
@@ -746,17 +746,6 @@ namespace UnrealBuildTool
 			return HasCustomIcons(ProjectDirectoryName);
 		}
 
-		/// <summary>
-		/// Whether the editor should be built for this platform or not
-		/// </summary>
-		/// <param name="InPlatform"> The UnrealTargetPlatform being built</param>
-		/// <param name="InConfiguration">The UnrealTargetConfiguration being built</param>
-		/// <returns>bool   true if the editor should be built, false if not</returns>
-		public override bool ShouldNotBuildEditor(UnrealTargetPlatform InPlatform, UnrealTargetConfiguration InConfiguration)
-		{
-			return true;
-		}
-
 		public override bool BuildRequiresCookedData(UnrealTargetPlatform InPlatform, UnrealTargetConfiguration InConfiguration)
 		{
 			return true; // for iOS can only run cooked. this is mostly for testing console code paths.
diff --git a/Engine/Source/Programs/UnrealBuildTool/Platform/Windows/UEBuildWindows.cs b/Engine/Source/Programs/UnrealBuildTool/Platform/Windows/UEBuildWindows.cs
--- a/Engine/Source/Programs/UnrealBuildTool/Platform/Windows/UEBuildWindows.cs
+++ b/Engine/Source/Programs/UnrealBuildTool/Platform/Windows/UEBuildWindows.cs
@@ -1290,17 +1290,6 @@ namespace UnrealBuildTool
 			return new string [] {};
 		}
 
-		/// <summary>
-		/// Whether the editor should be built for this platform or not
-		/// </summary>
-		/// <param name="InPlatform"> The UnrealTargetPlatform being built</param>
-		/// <param name="InConfiguration">The UnrealTargetConfiguration being built</param>
-		/// <returns>bool   true if the editor should be built, false if not</returns>
-		public override bool ShouldNotBuildEditor(UnrealTargetPlatform InPlatform, UnrealTargetConfiguration InConfiguration)
-		{
-			return false;
-		}
-
 		public override bool BuildRequiresCookedData(UnrealTargetPlatform InPlatform, UnrealTargetConfiguration InConfiguration)
 		{
 			return false;
