UBT: Move another target setting into the rules assembly.

#rnx

diff --git a/Engine/Source/Programs/UnrealBuildTool/Configuration/UEBuildTarget.cs b/Engine/Source/Programs/UnrealBuildTool/Configuration/UEBuildTarget.cs
--- a/Engine/Source/Programs/UnrealBuildTool/Configuration/UEBuildTarget.cs
+++ b/Engine/Source/Programs/UnrealBuildTool/Configuration/UEBuildTarget.cs
@@ -436,13 +436,6 @@ namespace UnrealBuildTool
 				}
 			}
 
-			// Setup the malloc profiler
-			if (RulesObject.bUseMallocProfiler)
-			{
-				RulesObject.bOmitFramePointers = false;
-				RulesObject.GlobalDefinitions.Add("USE_MALLOC_PROFILER=1");
-			}
-
 			// If we're precompiling, generate a list of all the files that we depend on
 			if (RulesObject.bPrecompile)
 			{
diff --git a/Engine/Source/Programs/UnrealBuildTool/System/RulesAssembly.cs b/Engine/Source/Programs/UnrealBuildTool/System/RulesAssembly.cs
--- a/Engine/Source/Programs/UnrealBuildTool/System/RulesAssembly.cs
+++ b/Engine/Source/Programs/UnrealBuildTool/System/RulesAssembly.cs
@@ -484,6 +484,14 @@ namespace UnrealBuildTool
 				Rules.bOmitPCDebugInfoInDevelopment = false;
 			}
 
+			// Setup the malloc profiler
+			if (Rules.bUseMallocProfiler)
+			{
+				Rules.bOmitFramePointers = false;
+				Rules.GlobalDefinitions.Add("USE_MALLOC_PROFILER=1");
+			}
+
+			// Set a macro if we allow using generated inis
 			if (!Rules.bAllowGeneratedIniWhenCooked)
 			{
 				Rules.GlobalDefinitions.Add("DISABLE_GENERATED_INI_WHEN_COOKED=1");
