UBT: Store the module rules file location on the ModuleRules instance, as well as the plugin that it was created from. Also expose the plugin directory as a property on the ModuleRules instance.

diff --git a/Engine/Source/Programs/UnrealBuildTool/Configuration/ModuleRules.cs b/Engine/Source/Programs/UnrealBuildTool/Configuration/ModuleRules.cs
--- a/Engine/Source/Programs/UnrealBuildTool/Configuration/ModuleRules.cs
+++ b/Engine/Source/Programs/UnrealBuildTool/Configuration/ModuleRules.cs
@@ -234,6 +234,30 @@ namespace UnrealBuildTool
 			}
 		}
 
+		/// <summary>
+		/// Name of this module
+		/// </summary>
+		public string Name
+		{
+			get;
+			internal set;
+		}
+
+		/// <summary>
+		/// File containing this module
+		/// </summary>
+		internal FileReference File;
+
+		/// <summary>
+		/// Directory containing this module
+		/// </summary>
+		internal DirectoryReference Directory;
+
+		/// <summary>
+		/// Plugin containing this module
+		/// </summary>
+		internal PluginInfo Plugin;
+
 		/// <summary>
 		/// Rules for the target that this module belongs to
 		/// </summary>
@@ -536,6 +560,24 @@ namespace UnrealBuildTool
 			}
 		}
 
+		/// <summary>
+		/// Property for the directory containing this plugin. Useful for adding paths to third party dependencies.
+		/// </summary>
+		public string PluginDirectory
+		{
+			get
+			{
+				if(Plugin == null)
+				{
+					throw new BuildException("Module '{0}' does not belong to a plugin; PluginDirectory property is invalid.", Name);
+				}
+				else
+				{
+					return Plugin.Directory.FullName;
+				}
+			}
+		}
+
 		/// <summary>
 		/// Property for the directory containing this module. Useful for adding paths to third party dependencies.
 		/// </summary>
@@ -543,7 +585,7 @@ namespace UnrealBuildTool
 		{
 			get
 			{
-				return Path.GetDirectoryName(RulesCompiler.GetFileNameFromType(GetType()));
+				return Directory.FullName;
 			}
 		}
 
diff --git a/Engine/Source/Programs/UnrealBuildTool/Configuration/UEBuildTarget.cs b/Engine/Source/Programs/UnrealBuildTool/Configuration/UEBuildTarget.cs
--- a/Engine/Source/Programs/UnrealBuildTool/Configuration/UEBuildTarget.cs
+++ b/Engine/Source/Programs/UnrealBuildTool/Configuration/UEBuildTarget.cs
@@ -1235,8 +1235,7 @@ namespace UnrealBuildTool
 				}
 
 				// Create the module rules
-				FileReference ModuleRulesFileName;
-				ModuleRules Rules = CreateModuleRulesAndSetDefaults(Module.Name, "external file list option", out ModuleRulesFileName);
+				ModuleRules Rules = CreateModuleRulesAndSetDefaults(Module.Name, "external file list option");
 
 				// Add Additional Bundle Resources for all modules
 				foreach (UEBuildBundleResource Resource in Rules.AdditionalBundleResources)
@@ -1261,7 +1260,7 @@ namespace UnrealBuildTool
 				}
 
 				// Add the rules file itself
-				Files.Add(ModuleRulesFileName);
+				Files.Add(Rules.File);
 
 				// Get a list of all the library paths
 				List<string> LibraryPaths = new List<string>();
@@ -2878,13 +2877,11 @@ namespace UnrealBuildTool
 				// Create rules for each remaining module, and check that it's set to be compiled
 				foreach(string FilteredModuleName in FilteredModuleNames)
 				{
-					FileReference ModuleFileName = null;
-
 					// Try to create the rules object, but catch any exceptions if it fails. Some modules (eg. SQLite) may determine that they are unavailable in the constructor.
 					ModuleRules ModuleRules;
 					try
 					{
-						ModuleRules = RulesAssembly.CreateModuleRules(FilteredModuleName, this.Rules, "precompile option", out ModuleFileName);
+						ModuleRules = RulesAssembly.CreateModuleRules(FilteredModuleName, this.Rules, "precompile option");
 					}
 					catch (BuildException)
 					{
@@ -2892,7 +2889,7 @@ namespace UnrealBuildTool
 					}
 
 					// Figure out if it can be precompiled
-					if (ModuleRules != null && ModuleRules.IsValidForTarget(ModuleFileName))
+					if (ModuleRules != null && ModuleRules.IsValidForTarget(ModuleRules.File))
 					{
 						ValidModuleNames.Add(FilteredModuleName);
 					}
@@ -2950,25 +2947,23 @@ namespace UnrealBuildTool
 		/// <summary>
 		/// Finds the base output directory for build products of the given module
 		/// </summary>
-		/// <param name="ModuleName">Name of the module</param>
 		/// <param name="ModuleRules">The rules object created for this module</param>
-		/// <param name="Plugin">The plugin that this module belongs to</param>
 		/// <returns>The base output directory for compiled object files for this module</returns>
-		private DirectoryReference GetBaseOutputDirectory(string ModuleName, ModuleRules ModuleRules, PluginInfo Plugin)
+		private DirectoryReference GetBaseOutputDirectory(ModuleRules ModuleRules)
 		{
 			// Get the root output directory and base name (target name/app name) for this binary
 			DirectoryReference BaseOutputDirectory;
-			if (Plugin != null)
+			if (ModuleRules.Plugin != null)
 			{
-				BaseOutputDirectory = Plugin.Directory;
+				BaseOutputDirectory = ModuleRules.Plugin.Directory;
 			}
-			else if (RulesAssembly.IsGameModule(ModuleName) || !bUseSharedBuildEnvironment)
+			else if (RulesAssembly.IsGameModule(ModuleRules.Name) || !bUseSharedBuildEnvironment)
 			{
 				BaseOutputDirectory = ProjectDirectory;
 			}
 			else
 			{
-				if (RulesAssembly.GetModuleFileName(ModuleName).IsUnderDirectory(UnrealBuildTool.EnterpriseDirectory))
+				if (RulesAssembly.GetModuleFileName(ModuleRules.Name).IsUnderDirectory(UnrealBuildTool.EnterpriseDirectory))
 				{
 					BaseOutputDirectory = UnrealBuildTool.EnterpriseDirectory;
 				}
@@ -2983,18 +2978,16 @@ namespace UnrealBuildTool
 		/// <summary>
 		/// Finds the base output directory for a module
 		/// </summary>
-		/// <param name="ModuleName">Name of the module</param>
 		/// <param name="ModuleRules">The rules object created for this module</param>
-		/// <param name="Plugin">The plugin that this module belongs to</param>
 		/// <returns>The output directory for compiled object files for this module</returns>
-		private DirectoryReference GetModuleIntermediateDirectory(string ModuleName, ModuleRules ModuleRules, PluginInfo Plugin)
+		private DirectoryReference GetModuleIntermediateDirectory(ModuleRules ModuleRules)
 		{
 			// Get the root output directory and base name (target name/app name) for this binary
-			DirectoryReference BaseOutputDirectory = GetBaseOutputDirectory(ModuleName, ModuleRules, Plugin);
+			DirectoryReference BaseOutputDirectory = GetBaseOutputDirectory(ModuleRules);
 
 			// Get the configuration that this module will be built in. Engine modules compiled in DebugGame will use Development.
 			UnrealTargetConfiguration ModuleConfiguration = Configuration;
-			if (Configuration == UnrealTargetConfiguration.DebugGame && !RulesAssembly.IsGameModule(ModuleName) && !ModuleName.Equals(Rules.LaunchModuleName, StringComparison.InvariantCultureIgnoreCase))
+			if (Configuration == UnrealTargetConfiguration.DebugGame && !RulesAssembly.IsGameModule(ModuleRules.Name) && !ModuleRules.Name.Equals(Rules.LaunchModuleName, StringComparison.InvariantCultureIgnoreCase))
 			{
 				ModuleConfiguration = UnrealTargetConfiguration.Development;
 			}
@@ -3008,7 +3001,7 @@ namespace UnrealBuildTool
 				IntermediateDirectory = DirectoryReference.Combine(IntermediateDirectory, ModuleRules.BinariesSubFolder);
 			}
 
-			return DirectoryReference.Combine(IntermediateDirectory, ModuleRules.ShortName ?? ModuleName);
+			return DirectoryReference.Combine(IntermediateDirectory, ModuleRules.ShortName ?? ModuleRules.Name);
 		}
 
 		/// <summary>
@@ -3018,15 +3011,8 @@ namespace UnrealBuildTool
 		/// <returns>The new binary. This has not been added to the target.</returns>
 		private UEBuildBinary CreateDynamicLibraryForModule(UEBuildModuleCPP Module)
 		{
-			// Get the plugin info for this module
-			PluginInfo Plugin = null;
-			if (Module.RulesFile != null)
-			{
-				RulesAssembly.TryGetPluginForModule(Module.RulesFile, out Plugin);
-			}
-
 			// Get the root output directory and base name (target name/app name) for this binary
-			DirectoryReference BaseOutputDirectory = GetBaseOutputDirectory(Module.Name, Module.Rules, Plugin);
+			DirectoryReference BaseOutputDirectory = GetBaseOutputDirectory(Module.Rules);
 			DirectoryReference OutputDirectory = DirectoryReference.Combine(BaseOutputDirectory, "Binaries", Platform.ToString());
 
 			// Append a subdirectory if the module rules specifies one
@@ -3801,10 +3787,10 @@ namespace UnrealBuildTool
         /// <summary>
         /// Create a rules object for the given module, and set any default values for this target
         /// </summary>
-        private ModuleRules CreateModuleRulesAndSetDefaults(string ModuleName, string ReferenceChain, out FileReference ModuleFileName)
+        private ModuleRules CreateModuleRulesAndSetDefaults(string ModuleName, string ReferenceChain)
 		{
 			// Create the rules from the assembly
-			ModuleRules RulesObject = RulesAssembly.CreateModuleRules(ModuleName, Rules, ReferenceChain, out ModuleFileName);
+			ModuleRules RulesObject = RulesAssembly.CreateModuleRules(ModuleName, Rules, ReferenceChain);
 
 			// Reads additional dependencies array for project module from project file and fills PrivateDependencyModuleNames. 
 			if (ProjectDescriptor != null && ProjectDescriptor.Modules != null)
@@ -3834,7 +3820,7 @@ namespace UnrealBuildTool
 				// Make sure that engine modules use shared PCHs or have an explicit private PCH
 				if(RulesObject.PCHUsage == ModuleRules.PCHUsageMode.NoSharedPCHs && RulesObject.PrivatePCHHeaderFile == null)
 				{
-					if(ProjectFile == null || !ModuleFileName.IsUnderDirectory(ProjectFile.Directory))
+					if(ProjectFile == null || !RulesObject.File.IsUnderDirectory(ProjectFile.Directory))
 					{
 						Log.TraceWarning("{0} module has shared PCHs disabled, but does not have a private PCH set", ModuleName);
 					}
@@ -3845,8 +3831,7 @@ namespace UnrealBuildTool
 				{
 					if(RulesObject.bUseBackwardsCompatibleDefaults && !Rules.bIWYU)
 					{
-						PluginInfo Plugin;
-						if(RulesAssembly.TryGetPluginForModule(ModuleFileName, out Plugin))
+						if(RulesObject.Plugin != null)
 						{
 							// Game plugin.  Enable shared PCHs by default, since they aren't typically large enough to warrant their own PCH.
 							RulesObject.PCHUsage = ModuleRules.PCHUsageMode.UseSharedPCHs;
@@ -3891,36 +3876,30 @@ namespace UnrealBuildTool
 			{
 				// @todo projectfiles: Cross-platform modules can appear here during project generation, but they may have already
 				//   been filtered out by the project generator.  This causes the projects to not be added to directories properly.
-				FileReference ModuleFileName;
-				ModuleRules RulesObject = CreateModuleRulesAndSetDefaults(ModuleName, ReferenceChain, out ModuleFileName);
-				DirectoryReference ModuleDirectory = ModuleFileName.Directory;
+				ModuleRules RulesObject = CreateModuleRulesAndSetDefaults(ModuleName, ReferenceChain);
+				DirectoryReference ModuleDirectory = RulesObject.File.Directory;
 
 				// Get the type of module we're creating
 				UHTModuleType? ModuleType = null;
 
-				// Get the plugin for this module
-				PluginInfo Plugin;
-				if(RulesAssembly.TryGetPluginForModule(ModuleFileName, out Plugin))
+				// Clear the bUsePrecompiled flag if we're compiling a foreign plugin; since it's treated like an engine module, it will default to true in an installed build.
+				if(RulesObject.Plugin != null && RulesObject.Plugin.File == ForeignPlugin)
 				{
-					// Clear the bUsePrecompiled flag if we're compiling a foreign plugin; since it's treated like an engine module, it will default to true in an installed build.
-					if(Plugin.File == ForeignPlugin)
-					{
-						RulesObject.bUsePrecompiled = false;
-					}
+					RulesObject.bUsePrecompiled = false;
 				}
 
 				// Get the module descriptor for this module if it's a plugin
 				ModuleDescriptor PluginModuleDesc = null;
-				if (Plugin != null)
+				if (RulesObject.Plugin != null)
 				{
-					PluginModuleDesc = Plugin.Descriptor.Modules.FirstOrDefault(x => x.Name == ModuleName);
+					PluginModuleDesc = RulesObject.Plugin.Descriptor.Modules.FirstOrDefault(x => x.Name == ModuleName);
 					if (PluginModuleDesc != null && PluginModuleDesc.Type == ModuleHostType.Program)
 					{
 						ModuleType = UHTModuleType.Program;
 					}
 				}
 
-				if (UnrealBuildTool.IsUnderAnEngineDirectory(ModuleFileName.Directory))
+				if (UnrealBuildTool.IsUnderAnEngineDirectory(RulesObject.File.Directory))
 				{
 					if (RulesObject.Type == ModuleRules.ModuleType.External)
 					{
@@ -3935,13 +3914,13 @@ namespace UnrealBuildTool
 
 						if (!ModuleType.HasValue)
 						{
-							if (ModuleFileName.IsUnderDirectory(UnrealBuildTool.EngineDirectory))
+							if (RulesObject.File.IsUnderDirectory(UnrealBuildTool.EngineDirectory))
 							{
-								ModuleType = ExternalExecution.GetEngineModuleTypeBasedOnLocation(UnrealBuildTool.EngineSourceDirectory, ModuleFileName);
+								ModuleType = ExternalExecution.GetEngineModuleTypeBasedOnLocation(UnrealBuildTool.EngineSourceDirectory, RulesObject.File);
 							}
-							else if (ModuleFileName.IsUnderDirectory(UnrealBuildTool.EnterpriseSourceDirectory))
+							else if (RulesObject.File.IsUnderDirectory(UnrealBuildTool.EnterpriseSourceDirectory))
 							{
-								ModuleType = ExternalExecution.GetEngineModuleTypeBasedOnLocation(UnrealBuildTool.EnterpriseSourceDirectory, ModuleFileName);
+								ModuleType = ExternalExecution.GetEngineModuleTypeBasedOnLocation(UnrealBuildTool.EnterpriseSourceDirectory, RulesObject.File);
 							}
 						}
 					}
@@ -3980,14 +3959,14 @@ namespace UnrealBuildTool
 
 				if (!ModuleType.HasValue)
 				{
-					throw new BuildException("Unable to determine module type for {0}\n(referenced via {1})", ModuleFileName, ReferenceChain);
+					throw new BuildException("Unable to determine module type for {0}\n(referenced via {1})", RulesObject.File, ReferenceChain);
 				}
 
 				// Get the base directory for paths referenced by the module. If the module's under the UProject source directory use that, otherwise leave it relative to the Engine source directory.
 				if (ProjectFile != null)
 				{
 					DirectoryReference ProjectSourceDirectoryName = DirectoryReference.Combine(ProjectFile.Directory, "Source");
-					if (ModuleFileName.IsUnderDirectory(ProjectSourceDirectoryName))
+					if (RulesObject.File.IsUnderDirectory(ProjectSourceDirectoryName))
 					{
 						RulesObject.PublicIncludePaths = CombinePathList(ProjectSourceDirectoryName, RulesObject.PublicIncludePaths);
 						RulesObject.PrivateIncludePaths = CombinePathList(ProjectSourceDirectoryName, RulesObject.PrivateIncludePaths);
@@ -4001,15 +3980,15 @@ namespace UnrealBuildTool
 				DirectoryReference GeneratedCodeDirectory = null;
 				if (RulesObject.Type != ModuleRules.ModuleType.External)
 				{
-					if (Plugin != null)
+					if (RulesObject.Plugin != null)
 					{
-						GeneratedCodeDirectory = Plugin.Directory;
+						GeneratedCodeDirectory = RulesObject.Plugin.Directory;
 					}
-					else if (bUseSharedBuildEnvironment && ModuleFileName.IsUnderDirectory(UnrealBuildTool.EngineDirectory))
+					else if (bUseSharedBuildEnvironment && RulesObject.File.IsUnderDirectory(UnrealBuildTool.EngineDirectory))
 					{
 						GeneratedCodeDirectory = UnrealBuildTool.EngineDirectory;
 					}
-					else if (bUseSharedBuildEnvironment && UnrealBuildTool.IsUnderAnEngineDirectory(ModuleFileName.Directory))
+					else if (bUseSharedBuildEnvironment && UnrealBuildTool.IsUnderAnEngineDirectory(RulesObject.File.Directory))
 					{
 						GeneratedCodeDirectory = UnrealBuildTool.EnterpriseDirectory;
 					}
@@ -4021,13 +4000,13 @@ namespace UnrealBuildTool
 				}
 
 				// For legacy modules, add a bunch of default include paths.
-				if (RulesObject.Type == ModuleRules.ModuleType.CPlusPlus && RulesObject.bAddDefaultIncludePaths && (Plugin != null || (ProjectFile != null && ModuleFileName.IsUnderDirectory(ProjectFile.Directory))))
+				if (RulesObject.Type == ModuleRules.ModuleType.CPlusPlus && RulesObject.bAddDefaultIncludePaths && (RulesObject.Plugin != null || (ProjectFile != null && RulesObject.File.IsUnderDirectory(ProjectFile.Directory))))
 				{
 					// Add the module source directory 
 					DirectoryReference BaseSourceDirectory;
-					if (Plugin != null)
+					if (RulesObject.Plugin != null)
 					{
-						BaseSourceDirectory = DirectoryReference.Combine(Plugin.Directory, "Source");
+						BaseSourceDirectory = DirectoryReference.Combine(RulesObject.Plugin.Directory, "Source");
 					}
 					else
 					{
@@ -4035,7 +4014,7 @@ namespace UnrealBuildTool
 					}
 
 					// If it's a game module (plugin or otherwise), add the root source directory to the include paths.
-					if (ModuleFileName.IsUnderDirectory(TargetRulesFile.Directory) || (Plugin != null && Plugin.LoadedFrom == PluginLoadedFrom.Project))
+					if (RulesObject.File.IsUnderDirectory(TargetRulesFile.Directory) || (RulesObject.Plugin != null && RulesObject.Plugin.LoadedFrom == PluginLoadedFrom.Project))
 					{
 						if(DirectoryReference.Exists(BaseSourceDirectory))
 						{
@@ -4077,10 +4056,10 @@ namespace UnrealBuildTool
 				if (RulesObject.Type == ModuleRules.ModuleType.CPlusPlus)
 				{
 					// So all we care about are the game module and/or plugins.
-					if (bDiscoverFiles && !UnrealBuildTool.IsUnderAnInstalledDirectory(ModuleFileName.Directory))
+					if (bDiscoverFiles && !UnrealBuildTool.IsUnderAnInstalledDirectory(RulesObject.File.Directory))
 					{
 						List<FileReference> SourceFilePaths = new List<FileReference>();
-						SourceFilePaths = SourceFileSearch.FindModuleSourceFiles(ModuleRulesFile: ModuleFileName);
+						SourceFilePaths = SourceFileSearch.FindModuleSourceFiles(ModuleRulesFile: RulesObject.File);
 						FoundSourceFiles = GetCPlusPlusFilesToBuild(SourceFilePaths, ModuleDirectory, Platform);
 					}
 				}
@@ -4106,9 +4085,9 @@ namespace UnrealBuildTool
 					{
 						Variables["ProjectDir"] = ProjectDirectory.FullName;
 					}
-					if (Plugin != null)
+					if (RulesObject.Plugin != null)
 					{
-						Variables["PluginDir"] = Plugin.Directory.FullName;
+						Variables["PluginDir"] = RulesObject.Plugin.Directory.FullName;
 					}
 
 					// Convert them into concrete file lists. Ignore anything that still hasn't been resolved.
@@ -4131,7 +4110,7 @@ namespace UnrealBuildTool
 				}
 
 				// Now, go ahead and create the module builder instance
-				Module = InstantiateModule(RulesObject, ModuleName, ModuleType.Value, ModuleDirectory, GeneratedCodeDirectory, FoundSourceFiles, bBuildFiles, ModuleFileName, RuntimeDependencies, Plugin);
+				Module = InstantiateModule(RulesObject, ModuleType.Value, GeneratedCodeDirectory, FoundSourceFiles, bBuildFiles, RuntimeDependencies);
 				Modules.Add(Module.Name, Module);
 				FlatModuleCsData.Add(new FlatModuleCsDataType(Module.Name, (Module.RulesFile == null) ? null : Module.RulesFile.FullName, RulesObject.ExternalDependencies));
 			}
@@ -4156,39 +4135,35 @@ namespace UnrealBuildTool
 
 		protected UEBuildModule InstantiateModule(
 			ModuleRules RulesObject,
-			string ModuleName,
 			UHTModuleType ModuleType,
-			DirectoryReference ModuleDirectory,
 			DirectoryReference GeneratedCodeDirectory,
 			List<FileItem> ModuleSourceFiles,
 			bool bBuildSourceFiles,
-			FileReference InRulesFile,
-			List<RuntimeDependency> InRuntimeDependencies,
-			PluginInfo Plugin)
+			List<RuntimeDependency> InRuntimeDependencies)
 		{
 			switch (RulesObject.Type)
 			{
 				case ModuleRules.ModuleType.CPlusPlus:
 					return new UEBuildModuleCPP(
-							InName: ModuleName,
+							InName: RulesObject.Name,
 							InType: ModuleType,
-							InModuleDirectory: ModuleDirectory,
-							InIntermediateDirectory: GetModuleIntermediateDirectory(ModuleName, RulesObject, Plugin),
+							InModuleDirectory: RulesObject.Directory,
+							InIntermediateDirectory: GetModuleIntermediateDirectory(RulesObject),
 							InGeneratedCodeDirectory: GeneratedCodeDirectory,
 							InSourceFiles: ModuleSourceFiles,
 							InRules: RulesObject,
 							bInBuildSourceFiles: bBuildSourceFiles,
-							InRulesFile: InRulesFile,
+							InRulesFile: RulesObject.File,
 							InRuntimeDependencies: InRuntimeDependencies
 						);
 
 				case ModuleRules.ModuleType.External:
 					return new UEBuildModuleExternal(
-							InName: ModuleName,
+							InName: RulesObject.Name,
 							InType: ModuleType,
-							InModuleDirectory: ModuleDirectory,
+							InModuleDirectory: RulesObject.Directory,
 							InRules: RulesObject,
-							InRulesFile: InRulesFile,
+							InRulesFile: RulesObject.File,
 							InRuntimeDependencies: InRuntimeDependencies
 						);
 
diff --git a/Engine/Source/Programs/UnrealBuildTool/System/RulesAssembly.cs b/Engine/Source/Programs/UnrealBuildTool/System/RulesAssembly.cs
--- a/Engine/Source/Programs/UnrealBuildTool/System/RulesAssembly.cs
+++ b/Engine/Source/Programs/UnrealBuildTool/System/RulesAssembly.cs
@@ -301,33 +301,20 @@ namespace UnrealBuildTool
 			}
 		}
 
-		/// <summary>
-		/// Creates an instance of a module rules descriptor object for the specified module name
-		/// </summary>
-		/// <param name="ModuleName">Name of the module</param>
-		/// <param name="ReferenceChain">Chain of references leading to this module</param>
-		/// <param name="Target">Information about the target associated with this module</param>
-		/// <returns>Compiled module rule info</returns>
-		public ModuleRules CreateModuleRules(string ModuleName, ReadOnlyTargetRules Target, string ReferenceChain)
-		{
-			FileReference ModuleFileName;
-			return CreateModuleRules(ModuleName, Target, ReferenceChain, out ModuleFileName);
-		}
-
 		/// <summary>
 		/// Creates an instance of a module rules descriptor object for the specified module name
 		/// </summary>
 		/// <param name="ModuleName">Name of the module</param>
 		/// <param name="Target">Information about the target associated with this module</param>
 		/// <param name="ReferenceChain">Chain of references leading to this module</param>
-		/// <param name="ModuleFileName">The original source file name for the Module.cs file for this module</param>
 		/// <returns>Compiled module rule info</returns>
-		public ModuleRules CreateModuleRules(string ModuleName, ReadOnlyTargetRules Target, string ReferenceChain, out FileReference ModuleFileName)
+		public ModuleRules CreateModuleRules(string ModuleName, ReadOnlyTargetRules Target, string ReferenceChain)
 		{
 			// Currently, we expect the user's rules object type name to be the same as the module name
 			string ModuleTypeName = ModuleName;
 
 			// Make sure the module file is known to us
+			FileReference ModuleFileName;
 			if (!ModuleNameToModuleFile.TryGetValue(ModuleName, out ModuleFileName))
 			{
 				if (Parent == null)
@@ -336,7 +323,7 @@ namespace UnrealBuildTool
 				}
 				else
 				{
-					return Parent.CreateModuleRules(ModuleName, Target, ReferenceChain, out ModuleFileName);
+					return Parent.CreateModuleRules(ModuleName, Target, ReferenceChain);
 				}
 			}
 
@@ -352,6 +339,10 @@ namespace UnrealBuildTool
 			{
 				// Create an uninitialized ModuleRules object and set some defaults.
 				ModuleRules RulesObject = (ModuleRules)FormatterServices.GetUninitializedObject(RulesObjectType);
+				RulesObject.Name = ModuleName;
+				RulesObject.File = ModuleFileName;
+				RulesObject.Directory = ModuleFileName.Directory;
+				ModuleFileToPluginInfo.TryGetValue(RulesObject.File, out RulesObject.Plugin);
 				RulesObject.bTreatAsEngineModule = bContainsEngineModules;
 				RulesObject.bUseBackwardsCompatibleDefaults = bUseBackwardsCompatibleDefaults && Target.bUseBackwardsCompatibleDefaults;
 				RulesObject.bPrecompile = RulesObject.bTreatAsEngineModule && Target.bPrecompile;
