UBT: Consolidate functionality for UHT module setup in ExternalExecution.cs.

diff --git a/Engine/Source/Programs/UnrealBuildTool/Configuration/UEBuildModule.cs b/Engine/Source/Programs/UnrealBuildTool/Configuration/UEBuildModule.cs
--- a/Engine/Source/Programs/UnrealBuildTool/Configuration/UEBuildModule.cs
+++ b/Engine/Source/Programs/UnrealBuildTool/Configuration/UEBuildModule.cs
@@ -22,11 +22,6 @@ namespace UnrealBuildTool
 		/// </summary>
 		public readonly string Name;
 
-		/// <summary>
-		/// The type of module being built. Used to switch between debug/development and precompiled/source configurations.
-		/// </summary>
-		public UHTModuleType Type;
-
 		/// <summary>
 		/// The rules for this module
 		/// </summary>
@@ -161,15 +156,13 @@ namespace UnrealBuildTool
 		/// Constructor
 		/// </summary>
 		/// <param name="InName">Name of the module</param>
-		/// <param name="InType">Type of the module, for UHT</param>
 		/// <param name="InModuleDirectory">Base directory for the module</param>
 		/// <param name="InRules">Rules for this module</param>
 		/// <param name="InRulesFile">Path to the rules file</param>
 		/// <param name="InRuntimeDependencies">List of runtime dependencies</param>
-		public UEBuildModule(string InName, UHTModuleType InType, DirectoryReference InModuleDirectory, ModuleRules InRules, FileReference InRulesFile, List<RuntimeDependency> InRuntimeDependencies)
+		public UEBuildModule(string InName, DirectoryReference InModuleDirectory, ModuleRules InRules, FileReference InRulesFile, List<RuntimeDependency> InRuntimeDependencies)
 		{
 			Name = InName;
-			Type = InType;
 			ModuleDirectory = InModuleDirectory;
 			Rules = InRules;
 			RulesFile = InRulesFile;
@@ -767,7 +760,6 @@ namespace UnrealBuildTool
 		public virtual void ExportJson(JsonWriter Writer)
 		{
 			Writer.WriteValue("Name", Name);
-			Writer.WriteValue("Type", Type.ToString());
 			Writer.WriteValue("Directory", ModuleDirectory.FullName);
 			Writer.WriteValue("Rules", RulesFile.FullName);
 			Writer.WriteValue("PCHUsage", Rules.PCHUsage.ToString());
diff --git a/Engine/Source/Programs/UnrealBuildTool/Configuration/UEBuildModuleCPP.cs b/Engine/Source/Programs/UnrealBuildTool/Configuration/UEBuildModuleCPP.cs
--- a/Engine/Source/Programs/UnrealBuildTool/Configuration/UEBuildModuleCPP.cs
+++ b/Engine/Source/Programs/UnrealBuildTool/Configuration/UEBuildModuleCPP.cs
@@ -233,7 +233,6 @@ namespace UnrealBuildTool
 
 		public UEBuildModuleCPP(
 			string InName,
-			UHTModuleType InType,
 			DirectoryReference InModuleDirectory,
 			DirectoryReference InIntermediateDirectory,
 			DirectoryReference InGeneratedCodeDirectory,
@@ -245,7 +244,6 @@ namespace UnrealBuildTool
 			)
 			: base(
 					InName,
-					InType,
 					InModuleDirectory,
 					InRules,
 					InRulesFile,
@@ -1334,20 +1332,6 @@ namespace UnrealBuildTool
 			return CompileEnvironment;
 		}
 
-		public class UHTModuleInfoCacheType
-		{
-			public UHTModuleInfoCacheType(IEnumerable<string> InHeaderFilenames, UHTModuleInfo InInfo)
-			{
-				HeaderFilenames = InHeaderFilenames;
-				Info = InInfo;
-			}
-
-			public IEnumerable<string> HeaderFilenames = null;
-			public UHTModuleInfo Info = null;
-		}
-
-		private UHTModuleInfoCacheType UHTModuleInfoCache = null;
-
 		/// Total time spent generating PCHs for modules (not actually compiling, but generating the PCH's input data)
 		public static double TotalPCHGenTime = 0.0;
 
@@ -1355,45 +1339,6 @@ namespace UnrealBuildTool
 		public static double TotalPCHCacheTime = 0.0;
 
 
-		/// <summary>
-		/// If any of this module's source files contain UObject definitions, this will return those header files back to the caller
-		/// </summary>
-		/// <returns></returns>
-		public UHTModuleInfoCacheType GetCachedUHTModuleInfo(EGeneratedCodeVersion GeneratedCodeVersion)
-		{
-			if (UHTModuleInfoCache == null)
-			{
-				List<FileReference> HeaderFiles = new List<FileReference>();
-				FileSystemName[] ExcludedFolders = UEBuildPlatform.GetBuildPlatform(Rules.Target.Platform, true).GetExcludedFolderNames();
-				FindHeaders(new DirectoryInfo(ModuleDirectory.FullName), ExcludedFolders, HeaderFiles);
-				UHTModuleInfo Info = ExternalExecution.CreateUHTModuleInfo(HeaderFiles, Name, RulesFile, ModuleDirectory, Type, GeneratedCodeVersion);
-				UHTModuleInfoCache = new UHTModuleInfoCacheType(Info.PublicUObjectHeaders.Concat(Info.PublicUObjectClassesHeaders).Concat(Info.PrivateUObjectHeaders).Select(x => x.AbsolutePath).ToList(), Info);
-			}
-
-			return UHTModuleInfoCache;
-		}
-
-		/// <summary>
-		/// Find all the headers under the given base directory, excluding any other platform folders.
-		/// </summary>
-		/// <param name="BaseDir">Base directory to search</param>
-		/// <param name="ExcludeFolders">Array of folders to exclude</param>
-		/// <param name="Headers">Receives the list of headers that was found</param>
-		static void FindHeaders(DirectoryInfo BaseDir, FileSystemName[] ExcludeFolders, List<FileReference> Headers)
-		{
-			if (!ExcludeFolders.Any(x => x.DisplayName.Equals(BaseDir.Name, StringComparison.InvariantCultureIgnoreCase)))
-			{
-				foreach (DirectoryInfo SubDir in BaseDir.EnumerateDirectories())
-				{
-					FindHeaders(SubDir, ExcludeFolders, Headers);
-				}
-				foreach (FileInfo File in BaseDir.EnumerateFiles("*.h"))
-				{
-					Headers.Add(new FileReference(File));
-				}
-			}
-		}
-
 		public override void GetAllDependencyModules(List<UEBuildModule> ReferencedModules, HashSet<UEBuildModule> IgnoreReferencedModules, bool bIncludeDynamicallyLoaded, bool bForceCircular, bool bOnlyDirectDependencies)
 		{
 			List<UEBuildModule> AllDependencyModules = new List<UEBuildModule>();
diff --git a/Engine/Source/Programs/UnrealBuildTool/Configuration/UEBuildModuleExternal.cs b/Engine/Source/Programs/UnrealBuildTool/Configuration/UEBuildModuleExternal.cs
--- a/Engine/Source/Programs/UnrealBuildTool/Configuration/UEBuildModuleExternal.cs
+++ b/Engine/Source/Programs/UnrealBuildTool/Configuration/UEBuildModuleExternal.cs
@@ -15,7 +15,6 @@ namespace UnrealBuildTool
 	class UEBuildModuleExternal : UEBuildModule
 	{
 		public UEBuildModuleExternal(
-			UHTModuleType InType,
 			string InName,
 			DirectoryReference InModuleDirectory,
 			ModuleRules InRules,
@@ -23,7 +22,6 @@ namespace UnrealBuildTool
 			List<RuntimeDependency> InRuntimeDependencies
 			)
 			: base(
-				InType: InType,
 				InName: InName,
 				InModuleDirectory: InModuleDirectory,
 				InRules: InRules,
diff --git a/Engine/Source/Programs/UnrealBuildTool/Configuration/UEBuildTarget.cs b/Engine/Source/Programs/UnrealBuildTool/Configuration/UEBuildTarget.cs
--- a/Engine/Source/Programs/UnrealBuildTool/Configuration/UEBuildTarget.cs
+++ b/Engine/Source/Programs/UnrealBuildTool/Configuration/UEBuildTarget.cs
@@ -2083,7 +2083,7 @@ namespace UnrealBuildTool
 					NameToFlatModuleData[FlatModuleData.ModuleName] = FlatModuleData;
 				}
 
-				ExternalExecution.SetupUObjectModules(ModulesToGenerateHeadersFor, Rules, GlobalCompileEnvironment, UObjectModules, NameToFlatModuleData, Rules.GeneratedCodeVersion, bIsAssemblingBuild);
+				ExternalExecution.SetupUObjectModules(ModulesToGenerateHeadersFor, Rules, ProjectDescriptor, GlobalCompileEnvironment, UObjectModules, NameToFlatModuleData, Rules.GeneratedCodeVersion, bIsAssemblingBuild);
 
 				// NOTE: Even in Gather mode, we need to run UHT to make sure the files exist for the static action graph to be setup correctly.  This is because UHT generates .cpp
 				// files that are injected as top level prerequisites.  If UHT only emitted included header files, we wouldn't need to run it during the Gather phase at all.
@@ -3879,89 +3879,12 @@ namespace UnrealBuildTool
 				ModuleRules RulesObject = CreateModuleRulesAndSetDefaults(ModuleName, ReferenceChain);
 				DirectoryReference ModuleDirectory = RulesObject.File.Directory;
 
-				// Get the type of module we're creating
-				UHTModuleType? ModuleType = null;
-
 				// Clear the bUsePrecompiled flag if we're compiling a foreign plugin; since it's treated like an engine module, it will default to true in an installed build.
 				if(RulesObject.Plugin != null && RulesObject.Plugin.File == ForeignPlugin)
 				{
 					RulesObject.bUsePrecompiled = false;
 				}
 
-				// Get the module descriptor for this module if it's a plugin
-				ModuleDescriptor PluginModuleDesc = null;
-				if (RulesObject.Plugin != null)
-				{
-					PluginModuleDesc = RulesObject.Plugin.Descriptor.Modules.FirstOrDefault(x => x.Name == ModuleName);
-					if (PluginModuleDesc != null && PluginModuleDesc.Type == ModuleHostType.Program)
-					{
-						ModuleType = UHTModuleType.Program;
-					}
-				}
-
-				if (UnrealBuildTool.IsUnderAnEngineDirectory(RulesObject.File.Directory))
-				{
-					if (RulesObject.Type == ModuleRules.ModuleType.External)
-					{
-						ModuleType = UHTModuleType.EngineThirdParty;
-					}
-					else
-					{
-						if (!ModuleType.HasValue && PluginModuleDesc != null)
-						{
-							ModuleType = ExternalExecution.GetEngineModuleTypeFromDescriptor(PluginModuleDesc);
-						}
-
-						if (!ModuleType.HasValue)
-						{
-							if (RulesObject.File.IsUnderDirectory(UnrealBuildTool.EngineDirectory))
-							{
-								ModuleType = ExternalExecution.GetEngineModuleTypeBasedOnLocation(UnrealBuildTool.EngineSourceDirectory, RulesObject.File);
-							}
-							else if (RulesObject.File.IsUnderDirectory(UnrealBuildTool.EnterpriseSourceDirectory))
-							{
-								ModuleType = ExternalExecution.GetEngineModuleTypeBasedOnLocation(UnrealBuildTool.EnterpriseSourceDirectory, RulesObject.File);
-							}
-						}
-					}
-				}
-				else
-				{
-					if (RulesObject.Type == ModuleRules.ModuleType.External)
-					{
-						ModuleType = UHTModuleType.GameThirdParty;
-					}
-					else
-					{
-						if (!ModuleType.HasValue && PluginModuleDesc != null)
-						{
-							ModuleType = ExternalExecution.GetGameModuleTypeFromDescriptor(PluginModuleDesc);
-						}
-
-						if (!ModuleType.HasValue)
-						{
-							if (ProjectDescriptor != null)
-							{
-								ModuleDescriptor ProjectModule = (ProjectDescriptor.Modules == null)? null : ProjectDescriptor.Modules.FirstOrDefault(x => x.Name == ModuleName);
-								if (ProjectModule != null)
-								{
-									ModuleType = UHTModuleTypeExtensions.GameModuleTypeFromHostType(ProjectModule.Type) ?? UHTModuleType.GameRuntime;
-								}
-								else
-								{
-									// No descriptor file or module was not on the list
-									ModuleType = UHTModuleType.GameRuntime;
-								}
-							}
-						}
-					}
-				}
-
-				if (!ModuleType.HasValue)
-				{
-					throw new BuildException("Unable to determine module type for {0}\n(referenced via {1})", RulesObject.File, ReferenceChain);
-				}
-
 				// Get the base directory for paths referenced by the module. If the module's under the UProject source directory use that, otherwise leave it relative to the Engine source directory.
 				if (ProjectFile != null)
 				{
@@ -4018,8 +3941,8 @@ namespace UnrealBuildTool
 					{
 						if(DirectoryReference.Exists(BaseSourceDirectory))
 						{
-						RulesObject.PublicIncludePaths.Add(NormalizeIncludePath(BaseSourceDirectory));
-					}
+							RulesObject.PublicIncludePaths.Add(NormalizeIncludePath(BaseSourceDirectory));
+						}
 					}
 
 					// Resolve private include paths against the project source root
@@ -4110,7 +4033,7 @@ namespace UnrealBuildTool
 				}
 
 				// Now, go ahead and create the module builder instance
-				Module = InstantiateModule(RulesObject, ModuleType.Value, GeneratedCodeDirectory, FoundSourceFiles, bBuildFiles, RuntimeDependencies);
+				Module = InstantiateModule(RulesObject, GeneratedCodeDirectory, FoundSourceFiles, bBuildFiles, RuntimeDependencies);
 				Modules.Add(Module.Name, Module);
 				FlatModuleCsData.Add(new FlatModuleCsDataType(Module.Name, (Module.RulesFile == null) ? null : Module.RulesFile.FullName, RulesObject.ExternalDependencies));
 			}
@@ -4135,7 +4058,6 @@ namespace UnrealBuildTool
 
 		protected UEBuildModule InstantiateModule(
 			ModuleRules RulesObject,
-			UHTModuleType ModuleType,
 			DirectoryReference GeneratedCodeDirectory,
 			List<FileItem> ModuleSourceFiles,
 			bool bBuildSourceFiles,
@@ -4146,7 +4068,6 @@ namespace UnrealBuildTool
 				case ModuleRules.ModuleType.CPlusPlus:
 					return new UEBuildModuleCPP(
 							InName: RulesObject.Name,
-							InType: ModuleType,
 							InModuleDirectory: RulesObject.Directory,
 							InIntermediateDirectory: GetModuleIntermediateDirectory(RulesObject),
 							InGeneratedCodeDirectory: GeneratedCodeDirectory,
@@ -4160,7 +4081,6 @@ namespace UnrealBuildTool
 				case ModuleRules.ModuleType.External:
 					return new UEBuildModuleExternal(
 							InName: RulesObject.Name,
-							InType: ModuleType,
 							InModuleDirectory: RulesObject.Directory,
 							InRules: RulesObject,
 							InRulesFile: RulesObject.File,
diff --git a/Engine/Source/Programs/UnrealBuildTool/System/ExternalExecution.cs b/Engine/Source/Programs/UnrealBuildTool/System/ExternalExecution.cs
--- a/Engine/Source/Programs/UnrealBuildTool/System/ExternalExecution.cs
+++ b/Engine/Source/Programs/UnrealBuildTool/System/ExternalExecution.cs
@@ -378,7 +378,7 @@ namespace UnrealBuildTool
 		{
 		}
 
-		public static UHTModuleType GetEngineModuleTypeFromDescriptor(ModuleDescriptor Module)
+		static UHTModuleType GetEngineModuleTypeFromDescriptor(ModuleDescriptor Module)
 		{
             UHTModuleType? Type = UHTModuleTypeExtensions.EngineModuleTypeFromHostType(Module.Type);
             if (Type == null)
@@ -388,7 +388,7 @@ namespace UnrealBuildTool
             return Type.GetValueOrDefault();
         }
 
-		public static UHTModuleType GetGameModuleTypeFromDescriptor(ModuleDescriptor Module)
+		static UHTModuleType GetGameModuleTypeFromDescriptor(ModuleDescriptor Module)
 		{
             UHTModuleType? Type = UHTModuleTypeExtensions.GameModuleTypeFromHostType(Module.Type);
             if (Type == null)
@@ -398,7 +398,7 @@ namespace UnrealBuildTool
             return Type.GetValueOrDefault();
         }
 
-		public static UHTModuleType? GetEngineModuleTypeBasedOnLocation(DirectoryReference SourceDirectory, FileReference ModuleFileName)
+		static UHTModuleType? GetEngineModuleTypeBasedOnLocation(DirectoryReference SourceDirectory, FileReference ModuleFileName)
 		{
 			if (ModuleFileName.IsUnderDirectory(DirectoryReference.Combine(SourceDirectory, "Runtime")))
 			{
@@ -433,7 +433,7 @@ namespace UnrealBuildTool
 		/// remain in their same relative order within the original Nodes sequence.
 		/// </summary>
 		/// <param name="NodeList">The list of nodes to sort.</param>
-		public static void StableTopologicalSort(List<UEBuildModuleCPP> NodeList)
+		static void StableTopologicalSort(List<UEBuildModuleCPP> NodeList)
 		{
 			int            NodeCount = NodeList.Count;
 
@@ -484,22 +484,142 @@ namespace UnrealBuildTool
 			return Dependencies.Contains(ToModule);
 		}
 
-		public static void SetupUObjectModules(IEnumerable<UEBuildModuleCPP> ModulesToGenerateHeadersFor, ReadOnlyTargetRules Target, CppCompileEnvironment GlobalCompileEnvironment, List<UHTModuleInfo> UObjectModules, Dictionary<string, FlatModuleCsDataType> FlatModuleCsData, EGeneratedCodeVersion GeneratedCodeVersion, bool bIsAssemblingBuild)
+		/// <summary>
+		/// Gets the module type for a given rules object
+		/// </summary>
+		/// <param name="RulesObject">The rules object</param>
+		/// <param name="ProjectDescriptor">Descriptor for the project being built</param>
+		/// <returns>The module type</returns>
+		static UHTModuleType GetModuleType(ModuleRules RulesObject, ProjectDescriptor ProjectDescriptor)
+		{
+			// Get the type of module we're creating
+			UHTModuleType? ModuleType = null;
+
+			// Get the module descriptor for this module if it's a plugin
+			ModuleDescriptor PluginModuleDesc = null;
+			if (RulesObject.Plugin != null)
+			{
+				PluginModuleDesc = RulesObject.Plugin.Descriptor.Modules.FirstOrDefault(x => x.Name == RulesObject.Name);
+				if (PluginModuleDesc != null && PluginModuleDesc.Type == ModuleHostType.Program)
+				{
+					ModuleType = UHTModuleType.Program;
+				}
+			}
+
+			if (UnrealBuildTool.IsUnderAnEngineDirectory(RulesObject.File.Directory))
+			{
+				if (RulesObject.Type == ModuleRules.ModuleType.External)
+				{
+					ModuleType = UHTModuleType.EngineThirdParty;
+				}
+				else
+				{
+					if (!ModuleType.HasValue && PluginModuleDesc != null)
+					{
+						ModuleType = ExternalExecution.GetEngineModuleTypeFromDescriptor(PluginModuleDesc);
+					}
+
+					if (!ModuleType.HasValue)
+					{
+						if (RulesObject.File.IsUnderDirectory(UnrealBuildTool.EngineDirectory))
+						{
+							ModuleType = ExternalExecution.GetEngineModuleTypeBasedOnLocation(UnrealBuildTool.EngineSourceDirectory, RulesObject.File);
+						}
+						else if (RulesObject.File.IsUnderDirectory(UnrealBuildTool.EnterpriseSourceDirectory))
+						{
+							ModuleType = ExternalExecution.GetEngineModuleTypeBasedOnLocation(UnrealBuildTool.EnterpriseSourceDirectory, RulesObject.File);
+						}
+					}
+				}
+			}
+			else
+			{
+				if (RulesObject.Type == ModuleRules.ModuleType.External)
+				{
+					ModuleType = UHTModuleType.GameThirdParty;
+				}
+				else
+				{
+					if (!ModuleType.HasValue && PluginModuleDesc != null)
+					{
+						ModuleType = ExternalExecution.GetGameModuleTypeFromDescriptor(PluginModuleDesc);
+					}
+
+					if (!ModuleType.HasValue)
+					{
+						if (ProjectDescriptor != null)
+						{
+							ModuleDescriptor ProjectModule = (ProjectDescriptor.Modules == null)? null : ProjectDescriptor.Modules.FirstOrDefault(x => x.Name == RulesObject.Name);
+							if (ProjectModule != null)
+							{
+								ModuleType = UHTModuleTypeExtensions.GameModuleTypeFromHostType(ProjectModule.Type) ?? UHTModuleType.GameRuntime;
+							}
+							else
+							{
+								// No descriptor file or module was not on the list
+								ModuleType = UHTModuleType.GameRuntime;
+							}
+						}
+					}
+				}
+			}
+
+			if (!ModuleType.HasValue)
+			{
+				throw new BuildException("Unable to determine UHT module type for {0}", RulesObject.File);
+			}
+
+			return ModuleType.Value;
+		}
+
+		/// <summary>
+		/// Find all the headers under the given base directory, excluding any other platform folders.
+		/// </summary>
+		/// <param name="BaseDir">Base directory to search</param>
+		/// <param name="ExcludeFolders">Array of folders to exclude</param>
+		/// <param name="Headers">Receives the list of headers that was found</param>
+		static void FindHeaders(DirectoryInfo BaseDir, FileSystemName[] ExcludeFolders, List<FileReference> Headers)
+		{
+			if (!ExcludeFolders.Any(x => x.DisplayName.Equals(BaseDir.Name, StringComparison.InvariantCultureIgnoreCase)))
+			{
+				foreach (DirectoryInfo SubDir in BaseDir.EnumerateDirectories())
+				{
+					FindHeaders(SubDir, ExcludeFolders, Headers);
+				}
+				foreach (FileInfo File in BaseDir.EnumerateFiles("*.h"))
+				{
+					Headers.Add(new FileReference(File));
+				}
+			}
+		}
+
+		public static void SetupUObjectModules(IEnumerable<UEBuildModuleCPP> ModulesToGenerateHeadersFor, ReadOnlyTargetRules Target, ProjectDescriptor ProjectDescriptor, CppCompileEnvironment GlobalCompileEnvironment, List<UHTModuleInfo> UObjectModules, Dictionary<string, FlatModuleCsDataType> FlatModuleCsData, EGeneratedCodeVersion GeneratedCodeVersion, bool bIsAssemblingBuild)
 		{
 			DateTime UObjectDiscoveryStartTime = DateTime.UtcNow;
 
+			// Find the type of each module
+			Dictionary<UEBuildModuleCPP, UHTModuleType> ModuleToType = new Dictionary<UEBuildModuleCPP, UHTModuleType>();
+			foreach(UEBuildModuleCPP Module in ModulesToGenerateHeadersFor)
+			{
+				ModuleToType[Module] = GetModuleType(Module.Rules, ProjectDescriptor);
+			}
+
 			// Sort modules by type, then by dependency
-			List<UEBuildModuleCPP> ModulesSortedByType = ModulesToGenerateHeadersFor.OrderBy(c => c.Type).ToList();
+			List<UEBuildModuleCPP> ModulesSortedByType = ModulesToGenerateHeadersFor.OrderBy(c => ModuleToType[c]).ToList();
 			StableTopologicalSort(ModulesSortedByType);
 
+			FileSystemName[] ExcludedFolders = UEBuildPlatform.GetBuildPlatform(Target.Platform, true).GetExcludedFolderNames();
 			foreach (UEBuildModuleCPP Module in ModulesSortedByType)
 			{
-				UEBuildModuleCPP.UHTModuleInfoCacheType UHTModuleInfo = Module.GetCachedUHTModuleInfo(GeneratedCodeVersion);
-				if (UHTModuleInfo.Info.PublicUObjectClassesHeaders.Count > 0 || UHTModuleInfo.Info.PrivateUObjectHeaders.Count > 0 || UHTModuleInfo.Info.PublicUObjectHeaders.Count > 0)
+				List<FileReference> HeaderFiles = new List<FileReference>();
+				FindHeaders(new DirectoryInfo(Module.ModuleDirectory.FullName), ExcludedFolders, HeaderFiles);
+
+				UHTModuleInfo Info = ExternalExecution.CreateUHTModuleInfo(HeaderFiles, Module.Name, Module.RulesFile, Module.ModuleDirectory, ModuleToType[Module], GeneratedCodeVersion);
+				if (Info.PublicUObjectClassesHeaders.Count > 0 || Info.PrivateUObjectHeaders.Count > 0 || Info.PublicUObjectHeaders.Count > 0)
 				{
 					// If we've got this far and there are no source files then it's likely we're installed and ignoring
 					// engine files, so we don't need a .gen.cpp either
-					UHTModuleInfo.Info.GeneratedCPPFilenameBase = Path.Combine(Module.GeneratedCodeDirectory.FullName, UHTModuleInfo.Info.ModuleName) + ".gen";
+					Info.GeneratedCPPFilenameBase = Path.Combine(Module.GeneratedCodeDirectory.FullName, Info.ModuleName) + ".gen";
 					if (Module.SourceFilesToBuild.Count != 0)
 					{
 						Module.GeneratedCodeWildcard = Path.Combine(Module.GeneratedCodeDirectory.FullName, "*.gen.cpp");
@@ -508,7 +628,7 @@ namespace UnrealBuildTool
 					// If we're running in "gather" mode only, we'll go ahead and cache PCH information for each module right now, so that we don't
 					// have to do it in the assembling phase.  It's OK for gathering to take a bit longer, even if UObject headers are not out of
 					// date in order to save a lot of time in the assembling runs.
-					UHTModuleInfo.Info.PCH = "";
+					Info.PCH = "";
 					if (!bIsAssemblingBuild)
 					{
 						// We need to figure out which PCH header this module is including, so that UHT can inject an include statement for it into any .cpp files it is synthesizing
@@ -516,14 +636,14 @@ namespace UnrealBuildTool
 						Module.CachePCHUsageForModuleSourceFiles(Target, ModuleCompileEnvironment);
 						if (Module.ProcessedDependencies.UniquePCHHeaderFile != null)
 						{
-							UHTModuleInfo.Info.PCH = Module.ProcessedDependencies.UniquePCHHeaderFile.AbsolutePath;
+							Info.PCH = Module.ProcessedDependencies.UniquePCHHeaderFile.AbsolutePath;
 						}
 					}
 
-					UObjectModules.Add(UHTModuleInfo.Info);
+					UObjectModules.Add(Info);
 					FlatModuleCsData[Module.Name].ModuleSourceFolder = Module.ModuleDirectory;
-					FlatModuleCsData[Module.Name].UHTHeaderNames = UHTModuleInfo.HeaderFilenames.ToList();
-					Log.TraceVerbose("Detected UObject module: " + UHTModuleInfo.Info.ModuleName);
+					FlatModuleCsData[Module.Name].UHTHeaderNames = Info.PublicUObjectHeaders.Concat(Info.PublicUObjectClassesHeaders).Concat(Info.PrivateUObjectHeaders).Select(x => x.AbsolutePath).ToList();
+					Log.TraceVerbose("Detected UObject module: " + Info.ModuleName);
 				}
 				else
 				{
