UBT: Modules must now explicitly specify the path to the header used to generate a PCH if one is desired, rather than the header being determined automatically by attempting to parse the source code. Now that PCHs are force-included anyway, this removes a lot of dependencies inside UBT.

diff --git a/Engine/Source/Programs/UnrealBuildTool/Configuration/UEBuildModuleCPP.cs b/Engine/Source/Programs/UnrealBuildTool/Configuration/UEBuildModuleCPP.cs
--- a/Engine/Source/Programs/UnrealBuildTool/Configuration/UEBuildModuleCPP.cs
+++ b/Engine/Source/Programs/UnrealBuildTool/Configuration/UEBuildModuleCPP.cs
@@ -407,7 +407,7 @@ namespace UnrealBuildTool
 
 			{
 				// Process all of the header file dependencies for this module
-				this.CachePCHUsageForModuleSourceFiles(Target, ModuleCompileEnvironment);
+				CheckFirstIncludeMatchesEachCppFile(Target, ModuleCompileEnvironment);
 
 				// Make sure our RC files have cached includes.  
 				foreach (FileItem RCFile in SourceFilesToBuild.RCFiles)
@@ -486,8 +486,8 @@ namespace UnrealBuildTool
 			CppCompileEnvironment CompileEnvironment = ModuleCompileEnvironment;
 			if (Target.bUsePCHFiles)
 			{
-				// If this module has an explicit PCH, use that
-				if(Rules.PrivatePCHHeaderFile != null)
+				// If this module doesn't need a shared PCH, configure that
+				if(Rules.PrivatePCHHeaderFile != null && (Rules.PCHUsage == ModuleRules.PCHUsageMode.NoSharedPCHs || Rules.PCHUsage == ModuleRules.PCHUsageMode.UseExplicitOrSharedPCHs))
 				{
 					PrecompiledHeaderInstance Instance = CreatePrivatePCH(ToolChain, FileItem.GetItemByFileReference(FileReference.Combine(ModuleDirectory, Rules.PrivatePCHHeaderFile)), CompileEnvironment, ActionGraph);
 
@@ -593,8 +593,7 @@ namespace UnrealBuildTool
 					foreach (string GeneratedFilename in GeneratedFiles)
 					{
 						FileItem GeneratedCppFileItem = FileItem.GetItemByPath(GeneratedFilename);
-
-						CachePCHUsageForModuleSourceFile(CompileEnvironment, GeneratedCppFileItem);
+						GeneratedCppFileItem.CachedIncludePaths = ModuleCompileEnvironment.IncludePaths;
 
 						// @todo ubtmake: Check for ALL other places where we might be injecting .cpp or .rc files for compiling without caching CachedCPPIncludeInfo first (anything platform specific?)
 						GeneratedFileItems.Add(GeneratedCppFileItem);
@@ -994,34 +993,14 @@ namespace UnrealBuildTool
 			}
 		}
 
-		public static FileItem CachePCHUsageForModuleSourceFile(CppCompileEnvironment ModuleCompileEnvironment, FileItem CPPFile)
-		{
-			if (!CPPFile.bExists)
-			{
-				throw new BuildException("Required source file not found: " + CPPFile.AbsolutePath);
-			}
-
-			DateTime PCHCacheTimerStart = DateTime.UtcNow;
-
-			// Store the module compile environment along with the .cpp file.  This is so that we can use it later on when looking
-			// for header dependencies
-			CPPFile.CachedIncludePaths = ModuleCompileEnvironment.IncludePaths;
-
-			FileItem PCHFile = ModuleCompileEnvironment.Headers.CachePCHUsageForCPPFile(CPPFile, ModuleCompileEnvironment.IncludePaths, ModuleCompileEnvironment.Platform);
-
-			if (UnrealBuildTool.bPrintPerformanceInfo)
-			{
-				double PCHCacheTime = (DateTime.UtcNow - PCHCacheTimerStart).TotalSeconds;
-				TotalPCHCacheTime += PCHCacheTime;
-			}
-
-			return PCHFile;
-		}
-
-
-		public void CachePCHUsageForModuleSourceFiles(ReadOnlyTargetRules Target, CppCompileEnvironment ModuleCompileEnvironment)
+		/// <summary>
+		/// Checks that the first header included by the source files in this module all include the same header
+		/// </summary>
+		/// <param name="Target">The target being compiled</param>
+		/// <param name="ModuleCompileEnvironment">Compile environment for the module</param>
+		private void CheckFirstIncludeMatchesEachCppFile(ReadOnlyTargetRules Target, CppCompileEnvironment ModuleCompileEnvironment)
 		{
-			if(Rules == null || Rules.PCHUsage == ModuleRules.PCHUsageMode.UseExplicitOrSharedPCHs || Rules.PrivatePCHHeaderFile != null)
+			if(Rules.PCHUsage == ModuleRules.PCHUsageMode.UseExplicitOrSharedPCHs)
 			{
 				if(InvalidIncludeDirectiveMessages == null)
 				{
@@ -1068,111 +1047,6 @@ namespace UnrealBuildTool
 					}
 				}
 			}
-			else
-			{
-				if (ProcessedDependencies == null)
-				{
-					DateTime PCHCacheTimerStart = DateTime.UtcNow;
-
-					bool bFoundAProblemWithPCHs = false;
-
-					FileItem UniquePCH = null;
-					foreach (FileItem CPPFile in SourceFilesFound.CPPFiles)	// @todo ubtmake: We're not caching CPPEnvironments for .c/.mm files, etc.  Even though they don't use PCHs, they still have #includes!  This can break dependency checking!
-					{
-						// Store the module compile environment along with the .cpp file.  This is so that we can use it later on when looking
-						// for header dependencies
-						CPPFile.CachedIncludePaths = ModuleCompileEnvironment.IncludePaths;
-
-						// Find headers used by the source file.
-						FileItem PCH = ModuleCompileEnvironment.Headers.CachePCHUsageForCPPFile(CPPFile, ModuleCompileEnvironment.IncludePaths, ModuleCompileEnvironment.Platform);
-						if (PCH == null)
-						{
-							throw new BuildException("Source file \"{0}\" is not including any headers.  We expect all modules to include a header file for precompiled header generation.  Please add an #include statement.", CPPFile.AbsolutePath);
-						}
-
-						if (UniquePCH == null)
-						{
-							UniquePCH = PCH;
-						}
-						else if (!UniquePCH.Info.Name.Equals(PCH.Info.Name, StringComparison.InvariantCultureIgnoreCase))		// @todo ubtmake: We do a string compare on the file name (not path) here, because sometimes the include resolver will pick an Intermediate copy of a PCH header file and throw off our comparisons
-						{
-							// OK, looks like we have multiple source files including a different header file first.  We'll keep track of this and print out
-							// helpful information afterwards.
-							bFoundAProblemWithPCHs = true;
-						}
-					}
-
-					ProcessedDependencies = new ProcessedDependenciesClass { UniquePCHHeaderFile = UniquePCH };
-
-
-					if (bFoundAProblemWithPCHs)
-					{
-						// Map from pch header string to the source files that use that PCH
-						Dictionary<FileReference, List<FileItem>> UsageMapPCH = new Dictionary<FileReference, List<FileItem>>();
-						foreach (FileItem CPPFile in SourceFilesToBuild.CPPFiles)
-						{
-							// Create a new entry if not in the pch usage map
-							List<FileItem> Files;
-							if (!UsageMapPCH.TryGetValue(CPPFile.PrecompiledHeaderIncludeFilename, out Files))
-							{
-								Files = new List<FileItem>();
-								UsageMapPCH.Add(CPPFile.PrecompiledHeaderIncludeFilename, Files);
-							}
-							Files.Add(CPPFile);
-						}
-
-						if (UnrealBuildTool.bPrintDebugInfo)
-						{
-							Log.TraceVerbose("{0} PCH files for module {1}:", UsageMapPCH.Count, Name);
-							int MostFilesIncluded = 0;
-							foreach (KeyValuePair<FileReference, List<FileItem>> CurPCH in UsageMapPCH)
-							{
-								if (CurPCH.Value.Count > MostFilesIncluded)
-								{
-									MostFilesIncluded = CurPCH.Value.Count;
-								}
-
-								Log.TraceVerbose("   {0}  ({1} files including it: {2}, ...)", CurPCH.Key, CurPCH.Value.Count, CurPCH.Value[0].AbsolutePath);
-							}
-						}
-
-						if (UsageMapPCH.Count > 1)
-						{
-							// Keep track of the PCH file that is most used within this module
-							FileReference MostFilesAreIncludingPCH = null;
-							int MostFilesIncluded = 0;
-							foreach (KeyValuePair<FileReference, List<FileItem>> CurPCH in UsageMapPCH.Where(PCH => PCH.Value.Count > MostFilesIncluded))
-							{
-								MostFilesAreIncludingPCH = CurPCH.Key;
-								MostFilesIncluded = CurPCH.Value.Count;
-							}
-
-							// Find all of the files that are not including our "best" PCH header
-							StringBuilder FilesNotIncludingBestPCH = new StringBuilder();
-							foreach (KeyValuePair<FileReference, List<FileItem>> CurPCH in UsageMapPCH.Where(PCH => PCH.Key != MostFilesAreIncludingPCH))
-							{
-								foreach (FileItem SourceFile in CurPCH.Value)
-								{
-									FilesNotIncludingBestPCH.AppendFormat("{0} (including {1})\n", SourceFile.AbsolutePath, CurPCH.Key);
-								}
-							}
-
-							// Bail out and let the user know which source files may need to be fixed up
-							throw new BuildException(
-								"All source files in module \"{0}\" must include the same precompiled header first.  Currently \"{1}\" is included by most of the source files.  The following source files are not including \"{1}\" as their first include:\n\n{2}\n\nTo compile this module without implicit precompiled headers, add \"PCHUsage = ModuleRules.PCHUsageMode.UseExplicitOrSharedPCHs;\" to {0}.build.cs.",
-								Name,
-								MostFilesAreIncludingPCH,
-								FilesNotIncludingBestPCH);
-						}
-					}
-
-					if (UnrealBuildTool.bPrintPerformanceInfo)
-					{
-						double PCHCacheTime = (DateTime.UtcNow - PCHCacheTimerStart).TotalSeconds;
-						TotalPCHCacheTime += PCHCacheTime;
-					}
-				}
-			}
 		}
 
 		private bool IgnoreMismatchedHeader(string ExpectedName)
@@ -1332,13 +1206,6 @@ namespace UnrealBuildTool
 			return CompileEnvironment;
 		}
 
-		/// Total time spent generating PCHs for modules (not actually compiling, but generating the PCH's input data)
-		public static double TotalPCHGenTime = 0.0;
-
-		/// Time spent caching which PCH header is included by each module and source file
-		public static double TotalPCHCacheTime = 0.0;
-
-
 		public override void GetAllDependencyModules(List<UEBuildModule> ReferencedModules, HashSet<UEBuildModule> IgnoreReferencedModules, bool bIncludeDynamicallyLoaded, bool bForceCircular, bool bOnlyDirectDependencies)
 		{
 			List<UEBuildModule> AllDependencyModules = new List<UEBuildModule>();
diff --git a/Engine/Source/Programs/UnrealBuildTool/Configuration/UEBuildTarget.cs b/Engine/Source/Programs/UnrealBuildTool/Configuration/UEBuildTarget.cs
--- a/Engine/Source/Programs/UnrealBuildTool/Configuration/UEBuildTarget.cs
+++ b/Engine/Source/Programs/UnrealBuildTool/Configuration/UEBuildTarget.cs
@@ -3848,6 +3848,13 @@ namespace UnrealBuildTool
 						RulesObject.PCHUsage = ModuleRules.PCHUsageMode.UseExplicitOrSharedPCHs;
 					}
 				}
+
+				// If we can't use a shared PCH, check there's a private PCH set
+				if(RulesObject.PCHUsage == ModuleRules.PCHUsageMode.NoSharedPCHs && RulesObject.PrivatePCHHeaderFile == null)
+				{
+					Log.TraceWarning("Modules must specify which file to use as a precompiled header in their build.cs file (eg. PrivatePCHHeaderFile = \"Private/{0}PrivatePCH.h\") or enable shared PCHs (eg. PCHUsage = PCHUsage.UseSharedPCHs) from UE 4.21 or later. '{0}' will be compiled using a shared PCH instead.", ModuleName);
+					RulesObject.PCHUsage = ModuleRules.PCHUsageMode.UseSharedPCHs;
+				}
 			}
 			return RulesObject;
 		}
diff --git a/Engine/Source/Programs/UnrealBuildTool/System/CPPHeaders.cs b/Engine/Source/Programs/UnrealBuildTool/System/CPPHeaders.cs
--- a/Engine/Source/Programs/UnrealBuildTool/System/CPPHeaders.cs
+++ b/Engine/Source/Programs/UnrealBuildTool/System/CPPHeaders.cs
@@ -293,8 +293,7 @@ namespace UnrealBuildTool
 					{
 						Dependencies.Add(IncludedFile.Location);
 					}
-					FileReference PCHName = SourceFile.PrecompiledHeaderIncludeFilename;
-					FlatCPPIncludeDependencyCache.SetDependenciesForFile(SourceFile.Location, PCHName, Dependencies);
+					FlatCPPIncludeDependencyCache.SetDependenciesForFile(SourceFile.Location, Dependencies);
 				}
 			}
 
@@ -421,57 +420,6 @@ namespace UnrealBuildTool
 			return DirectlyIncludedFiles;
 		}
 
-		public FileItem CachePCHUsageForCPPFile(FileItem CPPFile, CppIncludePaths IncludePaths, CppPlatform Platform)
-		{
-			// @todo ubtmake: We don't really need to scan every file looking for PCH headers, just need one.  The rest is just for error checking.
-			// @todo ubtmake: We don't need all of the direct includes either.  We just need the first, unless we want to check for errors.
-			List<DependencyInclude> DirectIncludeFilenames = GetDirectIncludeDependencies(CPPFile, bOnlyCachedDependencies: false);
-			if (UnrealBuildTool.bPrintDebugInfo)
-			{
-				Log.TraceVerbose("Found direct includes for {0}: {1}", Path.GetFileName(CPPFile.AbsolutePath), string.Join(", ", DirectIncludeFilenames.Select(F => F.IncludeName)));
-			}
-
-			if (DirectIncludeFilenames.Count == 0)
-			{
-				return null;
-			}
-
-			DependencyInclude FirstInclude = DirectIncludeFilenames[0];
-
-			// Resolve the PCH header to an absolute path.
-			// Check NullOrEmpty here because if the file could not be resolved we need to throw an exception
-			if (FirstInclude.IncludeResolvedNameIfSuccessful != null &&
-				// ignore any preexisting resolve cache if we are not configured to use it.
-				bUseIncludeDependencyResolveCache &&
-				// if we are testing the resolve cache, we force UBT to resolve every time to look for conflicts
-				!bTestIncludeDependencyResolveCache)
-			{
-				CPPFile.PrecompiledHeaderIncludeFilename = FirstInclude.IncludeResolvedNameIfSuccessful;
-				return FileItem.GetItemByFileReference(CPPFile.PrecompiledHeaderIncludeFilename);
-			}
-
-			// search the include paths to resolve the file.
-			string FirstIncludeName = FirstInclude.IncludeName;
-			UEBuildPlatform BuildPlatform = UEBuildPlatform.GetBuildPlatformForCPPTargetPlatform(Platform);
-			// convert back from relative to host path if needed
-			if (!BuildPlatform.UseAbsolutePathsInUnityFiles())
-			{
-				FirstIncludeName = RemoteExports.UnconvertPath(FirstIncludeName);
-			}
-
-			FileItem PrecompiledHeaderIncludeFile = CPPHeaders.FindIncludedFile(CPPFile.Location, FirstIncludeName, IncludePaths);
-			if (PrecompiledHeaderIncludeFile == null)
-			{
-				FirstIncludeName = RemoteExports.UnconvertPath(FirstInclude.IncludeName);
-				throw new BuildException("The first include statement in source file '{0}' is trying to include the file '{1}' as the precompiled header, but that file could not be located in any of the module's include search paths.", CPPFile.AbsolutePath, FirstIncludeName);
-			}
-
-			IncludeDependencyCache.CacheResolvedIncludeFullPath(CPPFile, 0, PrecompiledHeaderIncludeFile.Location, bUseIncludeDependencyResolveCache, bTestIncludeDependencyResolveCache);
-			CPPFile.PrecompiledHeaderIncludeFilename = PrecompiledHeaderIncludeFile.Location;
-
-			return PrecompiledHeaderIncludeFile;
-		}
-
 		public static double TotalTimeSpentGettingIncludes = 0.0;
 		public static int TotalIncludesRequested = 0;
 		public static double DirectIncludeCacheMissesTotalTime = 0.0;
diff --git a/Engine/Source/Programs/UnrealBuildTool/System/ExternalExecution.cs b/Engine/Source/Programs/UnrealBuildTool/System/ExternalExecution.cs
--- a/Engine/Source/Programs/UnrealBuildTool/System/ExternalExecution.cs
+++ b/Engine/Source/Programs/UnrealBuildTool/System/ExternalExecution.cs
@@ -181,11 +181,6 @@ namespace UnrealBuildTool
 		/// </summary>
 		public List<FileItem> PrivateUObjectHeaders;
 
-		/// <summary>
-		/// Module PCH absolute path
-		/// </summary>
-		public string PCH;
-
 		/// <summary>
 		/// Base (i.e. extensionless) path+filename of the .gen files
 		/// </summary>
@@ -209,7 +204,6 @@ namespace UnrealBuildTool
 			PublicUObjectClassesHeaders = (List<FileItem>)Info.GetValue("cl", typeof(List<FileItem>));
 			PublicUObjectHeaders = (List<FileItem>)Info.GetValue("pu", typeof(List<FileItem>));
 			PrivateUObjectHeaders = (List<FileItem>)Info.GetValue("pr", typeof(List<FileItem>));
-			PCH = Info.GetString("pc");
 			GeneratedCPPFilenameBase = Info.GetString("ge");
 			GeneratedCodeVersion = (EGeneratedCodeVersion)Info.GetInt32("gv");
 		}
@@ -223,7 +217,6 @@ namespace UnrealBuildTool
 			Info.AddValue("cl", PublicUObjectClassesHeaders);
 			Info.AddValue("pu", PublicUObjectHeaders);
 			Info.AddValue("pr", PrivateUObjectHeaders);
-			Info.AddValue("pc", PCH);
 			Info.AddValue("ge", GeneratedCPPFilenameBase);
 			Info.AddValue("gv", (int)GeneratedCodeVersion);
 		}
@@ -273,7 +266,6 @@ namespace UnrealBuildTool
 			public List<string> ClassesHeaders;
 			public List<string> PublicHeaders;
 			public List<string> PrivateHeaders;
-			public string PCH;
 			public string GeneratedCPPFilenameBase;
 			public bool SaveExportedHeaders;
 			public EGeneratedCodeVersion UHTGeneratedCodeVersion;
@@ -302,7 +294,6 @@ namespace UnrealBuildTool
 				ClassesHeaders = Info.PublicUObjectClassesHeaders.Select((Header) => Header.AbsolutePath).ToList(),
 				PublicHeaders = Info.PublicUObjectHeaders.Select((Header) => Header.AbsolutePath).ToList(),
 				PrivateHeaders = Info.PrivateUObjectHeaders.Select((Header) => Header.AbsolutePath).ToList(),
-				PCH = Info.PCH,
 				GeneratedCPPFilenameBase = Info.GeneratedCPPFilenameBase,
 				SaveExportedHeaders = !bUsePrecompiled || !Info.ModuleDirectory.IsUnderDirectory(UnrealBuildTool.EngineDirectory),
 				UHTGeneratedCodeVersion = Info.GeneratedCodeVersion,
@@ -625,21 +616,6 @@ namespace UnrealBuildTool
 						Module.GeneratedCodeWildcard = Path.Combine(Module.GeneratedCodeDirectory.FullName, "*.gen.cpp");
 					}
 
-					// If we're running in "gather" mode only, we'll go ahead and cache PCH information for each module right now, so that we don't
-					// have to do it in the assembling phase.  It's OK for gathering to take a bit longer, even if UObject headers are not out of
-					// date in order to save a lot of time in the assembling runs.
-					Info.PCH = "";
-					if (!bIsAssemblingBuild)
-					{
-						// We need to figure out which PCH header this module is including, so that UHT can inject an include statement for it into any .cpp files it is synthesizing
-						CppCompileEnvironment ModuleCompileEnvironment = Module.CreateModuleCompileEnvironment(Target, GlobalCompileEnvironment);
-						Module.CachePCHUsageForModuleSourceFiles(Target, ModuleCompileEnvironment);
-						if (Module.ProcessedDependencies.UniquePCHHeaderFile != null)
-						{
-							Info.PCH = Module.ProcessedDependencies.UniquePCHHeaderFile.AbsolutePath;
-						}
-					}
-
 					UObjectModules.Add(Info);
 					FlatModuleCsData[Module.Name].ModuleSourceFolder = Module.ModuleDirectory;
 					FlatModuleCsData[Module.Name].UHTHeaderNames = Info.PublicUObjectHeaders.Concat(Info.PublicUObjectClassesHeaders).Concat(Info.PrivateUObjectHeaders).Select(x => x.AbsolutePath).ToList();
@@ -1142,31 +1118,6 @@ namespace UnrealBuildTool
 
 				// ensure the headers are up to date
 				bool bUHTNeedsToRun = (BuildConfiguration.bForceHeaderGeneration || !bHaveHeaderTool || AreGeneratedCodeFilesOutOfDate(BuildConfiguration, UObjectModules, HeaderToolTimestamp, Target.bUsePrecompiled, HotReload, bIsGatheringBuild, bIsAssemblingBuild));
-				if (bUHTNeedsToRun || bIsGatheringBuild)
-				{
-					// Since code files are definitely out of date, we'll now finish computing information about the UObject modules for UHT.  We
-					// want to save this work until we know that UHT actually needs to be run to speed up best-case iteration times.
-					if (bIsGatheringBuild)		// In assembler-only mode, PCH info is loaded from our UBTMakefile!
-					{
-						foreach (UHTModuleInfo UHTModuleInfo in UObjectModules)
-						{
-							// Only cache the PCH name if we don't already have one.  When running in 'gather only' mode, this will have already been cached
-							if (string.IsNullOrEmpty(UHTModuleInfo.PCH))
-							{
-								UHTModuleInfo.PCH = "";
-
-								// We need to figure out which PCH header this module is including, so that UHT can inject an include statement for it into any .cpp files it is synthesizing
-								UEBuildModuleCPP DependencyModuleCPP = (UEBuildModuleCPP)Target.GetModuleByName(UHTModuleInfo.ModuleName);
-								CppCompileEnvironment ModuleCompileEnvironment = DependencyModuleCPP.CreateModuleCompileEnvironment(Target.Rules, GlobalCompileEnvironment);
-								DependencyModuleCPP.CachePCHUsageForModuleSourceFiles(Target.Rules, ModuleCompileEnvironment);
-								if (DependencyModuleCPP.ProcessedDependencies != null && DependencyModuleCPP.ProcessedDependencies.UniquePCHHeaderFile != null)
-								{
-									UHTModuleInfo.PCH = DependencyModuleCPP.ProcessedDependencies.UniquePCHHeaderFile.AbsolutePath;
-								}
-							}
-						}
-					}
-				}
 
 				// Get the file containing dependencies for the generated code
 				FileReference ExternalDependenciesFile = ModuleInfoFileName.ChangeExtension(".deps");
diff --git a/Engine/Source/Programs/UnrealBuildTool/System/FileItem.cs b/Engine/Source/Programs/UnrealBuildTool/System/FileItem.cs
--- a/Engine/Source/Programs/UnrealBuildTool/System/FileItem.cs
+++ b/Engine/Source/Programs/UnrealBuildTool/System/FileItem.cs
@@ -73,16 +73,6 @@ namespace UnrealBuildTool
 		private CppIncludePaths CachedIncludePathsValue;
 
 
-		///
-		/// Preparation only (not serialized)
-		/// 
-
-		/// <summary>
-		/// The PCH file that this file will use
-		/// </summary>
-		public FileReference PrecompiledHeaderIncludeFilename;
-
-
 		///
 		/// Transients (not serialized)
 		///
diff --git a/Engine/Source/Programs/UnrealBuildTool/System/FlatCPPIncludeDepencencyCache.cs b/Engine/Source/Programs/UnrealBuildTool/System/FlatCPPIncludeDepencencyCache.cs
--- a/Engine/Source/Programs/UnrealBuildTool/System/FlatCPPIncludeDepencencyCache.cs
+++ b/Engine/Source/Programs/UnrealBuildTool/System/FlatCPPIncludeDepencencyCache.cs
@@ -15,14 +15,15 @@ namespace UnrealBuildTool
 	/// </summary>
 	class FlatCPPIncludeDependencyInfo
 	{
-		/// The PCH header this file is dependent on.
-		public FileReference PCHName;
-
+		/// <summary>
 		/// List of files this file includes, excluding headers that were included from a PCH.
+		/// </summary>
 		public List<FileReference> Includes;
 
+		/// <summary>
 		/// Transient cached list of FileItems for all of the includes for a specific files.  This just saves a bunch of string
 		/// hash lookups as we locate FileItems for files that we've already requested dependencies for
+		/// </summary>
 		public List<FileItem> IncludeFileItems;
 	}
 
@@ -193,7 +194,6 @@ namespace UnrealBuildTool
 				Writer.Write(Pair.Key);
 
 				FlatCPPIncludeDependencyInfo Info = Pair.Value;
-				Writer.Write(Info.PCHName, FileToUniqueId);
 
 				Writer.Write(Info.Includes.Count);
 				foreach (FileReference Include in Info.Includes)
@@ -224,9 +224,6 @@ namespace UnrealBuildTool
 
 				FlatCPPIncludeDependencyInfo Info = new FlatCPPIncludeDependencyInfo();
 
-				// Read the PCH name
-				Info.PCHName = Reader.ReadFileReference(UniqueFiles);
-
 				// Read the includes 
 				int NumIncludes = Reader.ReadInt32();
 				Info.Includes = new List<FileReference>(NumIncludes);
@@ -247,12 +244,10 @@ namespace UnrealBuildTool
 		/// Sets the new dependencies for the specified file
 		/// </summary>
 		/// <param name="AbsoluteFilePath">File to set</param>
-		/// <param name="PCHName">The PCH for this file</param>
 		/// <param name="Dependencies">List of source dependencies</param>
-		public void SetDependenciesForFile(FileReference AbsoluteFilePath, FileReference PCHName, List<FileReference> Dependencies)
+		public void SetDependenciesForFile(FileReference AbsoluteFilePath, List<FileReference> Dependencies)
 		{
 			FlatCPPIncludeDependencyInfo DependencyInfo = new FlatCPPIncludeDependencyInfo();
-			DependencyInfo.PCHName = PCHName;	// @todo ubtmake: Not actually used yet.  The idea is to use this to reduce the number of indirect includes we need to store in the cache.
 			DependencyInfo.Includes = Dependencies;
 			DependencyInfo.IncludeFileItems = null;
 
diff --git a/Engine/Source/Programs/UnrealBuildTool/System/Unity.cs b/Engine/Source/Programs/UnrealBuildTool/System/Unity.cs
--- a/Engine/Source/Programs/UnrealBuildTool/System/Unity.cs
+++ b/Engine/Source/Programs/UnrealBuildTool/System/Unity.cs
@@ -366,9 +366,7 @@ namespace UnrealBuildTool
 				// Cache information about the unity .cpp dependencies
 				// @todo ubtmake urgent: Fails when building remotely for Mac because unity .cpp has an include for a PCH on the REMOTE machine
 				FileItem FirstCppFile = UnityFile.Files.First();
-				UEBuildModuleCPP.CachePCHUsageForModuleSourceFile(CompileEnvironment, FirstCppFile);
 				UnityCPPFile.CachedIncludePaths = FirstCppFile.CachedIncludePaths;
-				UnityCPPFile.PrecompiledHeaderIncludeFilename = FirstCppFile.PrecompiledHeaderIncludeFilename;
 			}
 
 			return NewCPPFiles;
diff --git a/Engine/Source/Programs/UnrealBuildTool/UnrealBuildTool.cs b/Engine/Source/Programs/UnrealBuildTool/UnrealBuildTool.cs
--- a/Engine/Source/Programs/UnrealBuildTool/UnrealBuildTool.cs
+++ b/Engine/Source/Programs/UnrealBuildTool/UnrealBuildTool.cs
@@ -930,8 +930,6 @@ namespace UnrealBuildTool
 						Log.TraceInformation("GetIncludes time: " + CPPHeaders.TotalTimeSpentGettingIncludes + "s (" + CPPHeaders.TotalIncludesRequested + " includes)");
 						Log.TraceInformation("DirectIncludes cache miss time: " + CPPHeaders.DirectIncludeCacheMissesTotalTime + "s (" + CPPHeaders.TotalDirectIncludeCacheMisses + " misses)");
 						Log.TraceInformation("FindIncludePaths calls: " + CPPHeaders.TotalFindIncludedFileCalls + " (" + CPPHeaders.IncludePathSearchAttempts + " searches)");
-						Log.TraceInformation("PCH gen time: " + UEBuildModuleCPP.TotalPCHGenTime + "s");
-						Log.TraceInformation("PCH cache time: " + UEBuildModuleCPP.TotalPCHCacheTime + "s");
 						Log.TraceInformation("Deep C++ include scan time: " + UnrealBuildTool.TotalDeepIncludeScanTime + "s");
 						Log.TraceInformation("Include Resolves: {0} ({1} misses, {2:0.00}%)", CPPHeaders.TotalDirectIncludeResolves, CPPHeaders.TotalDirectIncludeResolveCacheMisses, (float)CPPHeaders.TotalDirectIncludeResolveCacheMisses / (float)CPPHeaders.TotalDirectIncludeResolves * 100);
 						Log.TraceInformation("Total FileItems: {0} ({1} missing)", FileItem.TotalFileItemCount, FileItem.MissingFileItemCount);
diff --git a/Engine/Source/Programs/UnrealHeaderTool/Private/CodeGenerator.cpp b/Engine/Source/Programs/UnrealHeaderTool/Private/CodeGenerator.cpp
--- a/Engine/Source/Programs/UnrealHeaderTool/Private/CodeGenerator.cpp
+++ b/Engine/Source/Programs/UnrealHeaderTool/Private/CodeGenerator.cpp
@@ -5178,22 +5178,13 @@ FNativeClassHeaderGenerator::FNativeClassHeaderGenerator(
 
 	const FManifestModule* ModuleInfo = GPackageToManifestModuleMap.FindChecked(Package);
 
-	// Find other includes to put at the top of the .cpp
-	FUHTStringBuilder OtherIncludes;
-	if (ModuleInfo->PCH.Len())
-	{
-		FString PCH = ModuleInfo->PCH;
-		ConvertToBuildIncludePath(Package, PCH);
-		OtherIncludes.Logf(TEXT("#include \"%s\"") LINE_TERMINATOR, *PCH);
-	}
-
 	// Generate CPP files
 	TArray<FString> GeneratedCPPNames;
 	for (const TPair<FUnrealSourceFile*, FGeneratedCPP>& GeneratedCPP : GeneratedCPPs)
 	{
 		FUHTStringBuilder FileText;
 
-		FString GeneratedIncludes = OtherIncludes;
+		FString GeneratedIncludes;
 		for (const FString& RelativeInclude : GeneratedCPP.Value.RelativeIncludes)
 		{
 			GeneratedIncludes += FString::Printf(TEXT("#include \"%s\"\r\n"), *RelativeInclude);
diff --git a/Engine/Source/Programs/UnrealHeaderTool/Private/Manifest.cpp b/Engine/Source/Programs/UnrealHeaderTool/Private/Manifest.cpp
--- a/Engine/Source/Programs/UnrealHeaderTool/Private/Manifest.cpp
+++ b/Engine/Source/Programs/UnrealHeaderTool/Private/Manifest.cpp
@@ -133,7 +133,6 @@ FManifest FManifest::LoadFromFile(const FString& Filename)
 		GetJsonFieldValue(ClassesHeaders,                        ModuleObj, TEXT("ClassesHeaders"),           *Outer);
 		GetJsonFieldValue(PublicHeaders,                         ModuleObj, TEXT("PublicHeaders"),            *Outer);
 		GetJsonFieldValue(PrivateHeaders,                        ModuleObj, TEXT("PrivateHeaders"),           *Outer);
-		GetJsonFieldValue(KnownModule.PCH,                       ModuleObj, TEXT("PCH"),                      *Outer);
 		GetJsonFieldValue(KnownModule.GeneratedCPPFilenameBase,  ModuleObj, TEXT("GeneratedCPPFilenameBase"), *Outer);
 		GetJsonFieldValue(GeneratedCodeVersionString, ModuleObj, TEXT("UHTGeneratedCodeVersion"), *Outer);
 		KnownModule.GeneratedCodeVersion = ToGeneratedCodeVersion(GeneratedCodeVersionString);
@@ -246,7 +245,6 @@ bool FManifestModule::IsCompatibleWith(const FManifestModule& ManifestModule)
 		&& PublicUObjectClassesHeaders == ManifestModule.PublicUObjectClassesHeaders
 		&& PublicUObjectHeaders == ManifestModule.PublicUObjectHeaders
 		&& PrivateUObjectHeaders == ManifestModule.PrivateUObjectHeaders
-		&& PCH == ManifestModule.PCH
 		&& GeneratedCPPFilenameBase == ManifestModule.GeneratedCPPFilenameBase
 		&& SaveExportedHeaders == ManifestModule.SaveExportedHeaders
 		&& GeneratedCodeVersion == ManifestModule.GeneratedCodeVersion;
diff --git a/Engine/Source/Programs/UnrealHeaderTool/Private/Manifest.h b/Engine/Source/Programs/UnrealHeaderTool/Private/Manifest.h
--- a/Engine/Source/Programs/UnrealHeaderTool/Private/Manifest.h
+++ b/Engine/Source/Programs/UnrealHeaderTool/Private/Manifest.h
@@ -37,9 +37,6 @@ struct FManifestModule
 	/** List of C++ private header files with UObjects in them */
 	TArray<FString> PrivateUObjectHeaders;
 
-	/** Absolute path to the module's PCH */
-	FString PCH;
-
 	/** Base (i.e. extensionless) path+filename of where to write out the module's .generated.* files */
 	FString GeneratedCPPFilenameBase;
 
@@ -66,7 +63,6 @@ struct FManifestModule
 		Ar << ManifestModule.PublicUObjectClassesHeaders;
 		Ar << ManifestModule.PublicUObjectHeaders;
 		Ar << ManifestModule.PrivateUObjectHeaders;
-		Ar << ManifestModule.PCH;
 		Ar << ManifestModule.GeneratedCPPFilenameBase;
 		Ar << ManifestModule.SaveExportedHeaders;
 		Ar << ManifestModule.GeneratedCodeVersion;
