UBT: Remove unused lists inside UEBuildModuleCPP.SourceFilesClass.

#rnx

diff --git a/Engine/Source/Programs/UnrealBuildTool/Configuration/UEBuildModuleCPP.cs b/Engine/Source/Programs/UnrealBuildTool/Configuration/UEBuildModuleCPP.cs
--- a/Engine/Source/Programs/UnrealBuildTool/Configuration/UEBuildModuleCPP.cs
+++ b/Engine/Source/Programs/UnrealBuildTool/Configuration/UEBuildModuleCPP.cs
@@ -18,25 +18,21 @@ namespace UnrealBuildTool
 	{
 		public class SourceFilesClass
 		{
-			public readonly List<FileItem> MissingFiles = new List<FileItem>();
 			public readonly List<FileItem> CPPFiles = new List<FileItem>();
 			public readonly List<FileItem> CFiles = new List<FileItem>();
 			public readonly List<FileItem> CCFiles = new List<FileItem>();
 			public readonly List<FileItem> MMFiles = new List<FileItem>();
 			public readonly List<FileItem> RCFiles = new List<FileItem>();
-			public readonly List<FileItem> OtherFiles = new List<FileItem>();
 
 			public int Count
 			{
 				get
 				{
-					return MissingFiles.Count +
-						   CPPFiles.Count +
+					return CPPFiles.Count +
 						   CFiles.Count +
 						   CCFiles.Count +
 						   MMFiles.Count +
-						   RCFiles.Count +
-						   OtherFiles.Count;
+						   RCFiles.Count;
 				}
 			}
 
@@ -57,13 +53,11 @@ namespace UnrealBuildTool
 			/// <param name="Other">Source object.</param>
 			public void CopyFrom(SourceFilesClass Other)
 			{
-				CopyFromListToList(Other.MissingFiles, MissingFiles);
 				CopyFromListToList(Other.CPPFiles, CPPFiles);
 				CopyFromListToList(Other.CFiles, CFiles);
 				CopyFromListToList(Other.CCFiles, CCFiles);
 				CopyFromListToList(Other.MMFiles, MMFiles);
 				CopyFromListToList(Other.RCFiles, RCFiles);
-				CopyFromListToList(Other.OtherFiles, OtherFiles);
 			}
 		}
 
@@ -140,11 +134,7 @@ namespace UnrealBuildTool
 			foreach (FileItem SourceFile in InSourceFiles)
 			{
 				string Extension = Path.GetExtension(SourceFile.AbsolutePath).ToUpperInvariant();
-				if (!SourceFile.bExists)
-				{
-					OutSourceFiles.MissingFiles.Add(SourceFile);
-				}
-				else if (Extension == ".CPP")
+				if (Extension == ".CPP")
 				{
 					OutSourceFiles.CPPFiles.Add(SourceFile);
 				}
@@ -164,10 +154,6 @@ namespace UnrealBuildTool
 				{
 					OutSourceFiles.RCFiles.Add(SourceFile);
 				}
-				else
-				{
-					OutSourceFiles.OtherFiles.Add(SourceFile);
-				}
 			}
 		}
 
@@ -395,28 +381,16 @@ namespace UnrealBuildTool
 				return LinkInputFiles;
 			}
 
-			// Throw an error if the module's source file list referenced any non-existent files.
-			if (SourceFilesToBuild.MissingFiles.Count > 0)
-			{
-				throw new BuildException(
-					"UBT ERROR: Module \"{0}\" references non-existent files:\n{1} (perhaps a file was added to the project but not checked in)",
-					Name,
-					string.Join("\n", SourceFilesToBuild.MissingFiles.Select(M => M.AbsolutePath))
-				);
-			}
+			// Process all of the header file dependencies for this module
+			CheckFirstIncludeMatchesEachCppFile(Target, ModuleCompileEnvironment);
 
+			// Make sure our RC files have cached includes.  
+			foreach (FileItem RCFile in SourceFilesToBuild.RCFiles)
 			{
-				// Process all of the header file dependencies for this module
-				CheckFirstIncludeMatchesEachCppFile(Target, ModuleCompileEnvironment);
-
-				// Make sure our RC files have cached includes.  
-				foreach (FileItem RCFile in SourceFilesToBuild.RCFiles)
+				// The default resource file (PCLaunch.rc) is created in a module-agnostic way, so we want to avoid overriding the include paths for it
+				if(RCFile.CachedIncludePaths == null)
 				{
-					// The default resource file (PCLaunch.rc) is created in a module-agnostic way, so we want to avoid overriding the include paths for it
-					if(RCFile.CachedIncludePaths == null)
-					{
-						RCFile.CachedIncludePaths = ModuleCompileEnvironment.IncludePaths;
-					}
+					RCFile.CachedIncludePaths = ModuleCompileEnvironment.IncludePaths;
 				}
 			}
 
