UBT: Remove some legacy code from auto-detecting PCHs.

#rnx

diff --git a/Engine/Source/Programs/UnrealBuildTool/Configuration/UEBuildModuleCPP.cs b/Engine/Source/Programs/UnrealBuildTool/Configuration/UEBuildModuleCPP.cs
--- a/Engine/Source/Programs/UnrealBuildTool/Configuration/UEBuildModuleCPP.cs
+++ b/Engine/Source/Programs/UnrealBuildTool/Configuration/UEBuildModuleCPP.cs
@@ -91,31 +91,11 @@ namespace UnrealBuildTool
 		/// </summary>
 		public string GeneratedCodeWildcard;
 
-		public class ProcessedDependenciesClass
-		{
-			/// <summary>
-			/// The file, if any, which is used as the unique PCH for this module
-			/// </summary>
-			public FileItem UniquePCHHeaderFile = null;
-		}
-
-		/// <summary>
-		/// The processed dependencies for the class
-		/// </summary>
-		public ProcessedDependenciesClass ProcessedDependencies = null;
-
 		/// <summary>
 		/// List of invalid include directives. These are buffered up and output before we start compiling.
 		/// </summary>
 		public List<string> InvalidIncludeDirectiveMessages;
 
-		public IEnumerable<string> FindGeneratedCppFiles()
-		{
-			return ((null == GeneratedCodeDirectory) || !DirectoryReference.Exists(GeneratedCodeDirectory))
-				? Enumerable.Empty<string>()
-				: DirectoryReference.EnumerateFiles(GeneratedCodeDirectory, "*.gen.cpp", SearchOption.TopDirectoryOnly).Select((Dir) => Dir.FullName);
-		}
-
 		protected override void GetReferencedDirectories(HashSet<DirectoryReference> Directories)
 		{
 			base.GetReferencedDirectories(Directories);
@@ -517,20 +497,6 @@ namespace UnrealBuildTool
 						LinkInputFiles.AddRange(Instance.Output.ObjectFiles);
 					}
 				}
-
-				// If there was one header that was included first by enough C++ files, use it as the precompiled header. Only use precompiled headers for projects with enough files to make the PCH creation worthwhile.
-				if (CompileEnvironment.PrecompiledHeaderFile == null && SourceFilesToBuild.CPPFiles.Count >= MinFilesUsingPrecompiledHeader && ProcessedDependencies != null)
-				{
-					PrecompiledHeaderInstance Instance = CreatePrivatePCH(ToolChain, ProcessedDependencies.UniquePCHHeaderFile, CompileEnvironment, ActionGraph);
-
-					CompileEnvironment = new CppCompileEnvironment(CompileEnvironment);
-					CompileEnvironment.Definitions.Clear();
-					CompileEnvironment.PrecompiledHeaderAction = PrecompiledHeaderAction.Include;
-					CompileEnvironment.PrecompiledHeaderIncludeFilename = Instance.HeaderFile.Location;
-					CompileEnvironment.PrecompiledHeaderFile = Instance.Output.PrecompiledHeaderFile;
-
-					LinkInputFiles.AddRange(Instance.Output.ObjectFiles);
-				}
 			}
 
 			// Write all the definitions to a separate file
