Fixup precompiled header setting for samples and games.

#rnx

diff --git a/Engine/Source/Programs/UnrealBuildTool/Configuration/UEBuildTarget.cs b/Engine/Source/Programs/UnrealBuildTool/Configuration/UEBuildTarget.cs
--- a/Engine/Source/Programs/UnrealBuildTool/Configuration/UEBuildTarget.cs
+++ b/Engine/Source/Programs/UnrealBuildTool/Configuration/UEBuildTarget.cs
@@ -3852,7 +3852,7 @@ namespace UnrealBuildTool
 				// If we can't use a shared PCH, check there's a private PCH set
 				if(RulesObject.PCHUsage == ModuleRules.PCHUsageMode.NoSharedPCHs && RulesObject.PrivatePCHHeaderFile == null)
 				{
-					Log.TraceWarning("Modules must specify which file to use as a precompiled header in their build.cs file (eg. PrivatePCHHeaderFile = \"Private/{0}PrivatePCH.h\") or enable shared PCHs (eg. PCHUsage = PCHUsage.UseSharedPCHs) from UE 4.21 or later. '{0}' will be compiled using a shared PCH instead.", ModuleName);
+					Log.WriteLine(LogEventType.Warning, LogFormatOptions.NoSeverityPrefix, "{0}: warning: Modules should specify which file to use as a precompiled header (eg. PrivatePCHHeaderFile = \"Private/{1}PrivatePCH.h\") or enable shared PCHs (eg. PCHUsage = PCHUsageMode.UseSharedPCHs). '{1}' will be compiled using a shared PCH instead.", RulesObject.File, RulesObject.Name);
 					RulesObject.PCHUsage = ModuleRules.PCHUsageMode.UseSharedPCHs;
 				}
 			}
