UAT: Remove the FileSystemName class, now that we don't need to normalize strings before comparison.

#rnx

diff --git a/Engine/Source/Programs/AutomationTool/AutomationUtils/DeploymentContext.cs b/Engine/Source/Programs/AutomationTool/AutomationUtils/DeploymentContext.cs
--- a/Engine/Source/Programs/AutomationTool/AutomationUtils/DeploymentContext.cs
+++ b/Engine/Source/Programs/AutomationTool/AutomationUtils/DeploymentContext.cs
@@ -229,7 +229,7 @@ public class DeploymentContext //: ProjectParams
 	/// <summary>
 	/// List of restricted folder names which are not permitted in staged build
 	/// </summary>
-	public HashSet<FileSystemName> RestrictedFolderNames = new HashSet<FileSystemName>();
+	public HashSet<string> RestrictedFolderNames = new HashSet<string>(StringComparer.OrdinalIgnoreCase);
 
 	/// <summary>
 	/// List of directories to remap during the stage
@@ -422,7 +422,7 @@ public class DeploymentContext //: ProjectParams
 			RestrictedFolderNames.ExceptWith(PlatformExports.GetIncludedFolderNames(StagePlatform));
 		}
 		RestrictedFolderNames.UnionWith(FileFilter.RestrictedFolderNames);
-		RestrictedFolderNames.Remove(new FileSystemName(StageTargetPlatform.IniPlatformType.ToString()));
+		RestrictedFolderNames.Remove(StageTargetPlatform.IniPlatformType.ToString());
 
 		// Read the game config files
 		ConfigHierarchy GameConfig = ConfigCache.ReadHierarchy(ConfigHierarchyType.Game, ProjectRoot, InTargetPlatform.PlatformType);
@@ -549,7 +549,7 @@ public class DeploymentContext //: ProjectParams
 		{
 			foreach(DirectoryReference SubDir in DirectoryReference.EnumerateDirectories(BaseDir))
 			{
-				FileSystemName Name = new FileSystemName(SubDir);
+				string Name = SubDir.GetDirectoryName();
 				if(!RestrictedFolderNames.Contains(Name))
 				{
 					FindFilesToStageInternal(SubDir, Pattern, Option, Files);
diff --git a/Engine/Source/Programs/AutomationTool/AutomationUtils/StagedFileReference.cs b/Engine/Source/Programs/AutomationTool/AutomationUtils/StagedFileReference.cs
--- a/Engine/Source/Programs/AutomationTool/AutomationUtils/StagedFileReference.cs
+++ b/Engine/Source/Programs/AutomationTool/AutomationUtils/StagedFileReference.cs
@@ -54,19 +54,19 @@ public class StagedFileReference : StagedFileSystemReference, IEquatable<StagedF
 	/// </summary>
 	/// <param name="OtherName"></param>
 	/// <returns></returns>
-	public bool ContainsName(FileSystemName OtherName)
+	public bool ContainsName(string OtherName)
 	{
 		int StartIdx = 0;
 		for(;;)
 		{
-			int Idx = Name.IndexOf(OtherName.DisplayName, StartIdx, FileSystemReference.Comparison);
+			int Idx = Name.IndexOf(OtherName, StartIdx, FileSystemReference.Comparison);
 			if(Idx == -1)
 			{
 				return false;
 			}
 			if (Idx == 0 || Name[Idx - 1] == '/')
 			{
-				int EndIdx = Idx + OtherName.DisplayName.Length;
+				int EndIdx = Idx + OtherName.Length;
 				if(EndIdx < Name.Length && Name[EndIdx] == '/')
 				{
 					return true;
diff --git a/Engine/Source/Programs/AutomationTool/BuildGraph/BuildGraph.cs b/Engine/Source/Programs/AutomationTool/BuildGraph/BuildGraph.cs
--- a/Engine/Source/Programs/AutomationTool/BuildGraph/BuildGraph.cs
+++ b/Engine/Source/Programs/AutomationTool/BuildGraph/BuildGraph.cs
@@ -125,8 +125,8 @@ namespace AutomationTool
 			DefaultProperties["RootDir"] = CommandUtils.RootDirectory.FullName;
 			DefaultProperties["IsBuildMachine"] = IsBuildMachine ? "true" : "false";
 			DefaultProperties["HostPlatform"] = HostPlatform.Current.HostEditorPlatform.ToString();
-			DefaultProperties["RestrictedFolderNames"] = String.Join(";", FileFilter.RestrictedFolderNames.Select(x => x.DisplayName));
-			DefaultProperties["RestrictedFolderFilter"] = String.Join(";", FileFilter.RestrictedFolderNames.Select(x => String.Format(".../{0}/...", x.DisplayName)));
+			DefaultProperties["RestrictedFolderNames"] = String.Join(";", FileFilter.RestrictedFolderNames);
+			DefaultProperties["RestrictedFolderFilter"] = String.Join(";", FileFilter.RestrictedFolderNames.Select(x => String.Format(".../{0}/...", x)));
 
 			// Attempt to read existing Build Version information
 			BuildVersion Version;
diff --git a/Engine/Source/Programs/AutomationTool/Scripts/CheckCsprojDotNetVersion.cs b/Engine/Source/Programs/AutomationTool/Scripts/CheckCsprojDotNetVersion.cs
--- a/Engine/Source/Programs/AutomationTool/Scripts/CheckCsprojDotNetVersion.cs
+++ b/Engine/Source/Programs/AutomationTool/Scripts/CheckCsprojDotNetVersion.cs
@@ -1,4 +1,4 @@
-﻿// Copyright 1998-2018 Epic Games, Inc. All Rights Reserved.
+// Copyright 1998-2018 Epic Games, Inc. All Rights Reserved.
 using System;
 using System.Collections.Generic;
 using System.Text;
@@ -29,8 +29,8 @@ class CheckCsprojDotNetVersion : BuildCommand
         Regex DotNetCoreRegex = new Regex("<TargetFramework>(netcoreapp2.0|netstandard2.0)<\\/TargetFramework>");
         foreach (FileReference CsProj in DirectoryReference.EnumerateFiles(EngineDir, "*.csproj", SearchOption.AllDirectories))
         {
-            if (CsProj.ContainsName(new FileSystemName("ThirdParty"), EngineDir) ||
-                (CsProj.ContainsName(new FileSystemName("UE4TemplateProject"), EngineDir) && CsProj.GetFileName().Equals("ProjectTemplate.csproj")) ||
+            if (CsProj.ContainsName("ThirdParty", EngineDir) ||
+                (CsProj.ContainsName("UE4TemplateProject", EngineDir) && CsProj.GetFileName().Equals("ProjectTemplate.csproj")) ||
                 CsProj.GetFileNameWithoutExtension().ToLower().Contains("_mono") ||
                 CsProj.GetFileNameWithoutExtension().ToLower().Contains("unrealvs"))
 
diff --git a/Engine/Source/Programs/AutomationTool/Scripts/CheckRestrictedFolders.cs b/Engine/Source/Programs/AutomationTool/Scripts/CheckRestrictedFolders.cs
--- a/Engine/Source/Programs/AutomationTool/Scripts/CheckRestrictedFolders.cs
+++ b/Engine/Source/Programs/AutomationTool/Scripts/CheckRestrictedFolders.cs
@@ -26,7 +26,7 @@ namespace AutomationTool
 			}
 
 			// Find a list of restricted folders, and remove any names which are explicitly whitelisted
-			HashSet<string> RestrictedNames = new HashSet<string>(FileFilter.RestrictedFolderNames.Select(x => x.DisplayName), StringComparer.InvariantCultureIgnoreCase);
+			HashSet<string> RestrictedNames = new HashSet<string>(FileFilter.RestrictedFolderNames, StringComparer.InvariantCultureIgnoreCase);
 			foreach (string AllowParam in ParseParamValues("Allow"))
 			{
 				RestrictedNames.ExceptWith(AllowParam.Split('+'));
diff --git a/Engine/Source/Programs/AutomationTool/Scripts/CopyBuildToStagingDirectory.Automation.cs b/Engine/Source/Programs/AutomationTool/Scripts/CopyBuildToStagingDirectory.Automation.cs
--- a/Engine/Source/Programs/AutomationTool/Scripts/CopyBuildToStagingDirectory.Automation.cs
+++ b/Engine/Source/Programs/AutomationTool/Scripts/CopyBuildToStagingDirectory.Automation.cs
@@ -831,7 +831,7 @@ public partial class Project : CommandUtils
 
 		// Make sure there are no restricted folders in the output
 		HashSet<StagedFileReference> RestrictedFiles = new HashSet<StagedFileReference>();
-		foreach (FileSystemName RestrictedName in SC.RestrictedFolderNames)
+		foreach (string RestrictedName in SC.RestrictedFolderNames)
 		{
 			RestrictedFiles.UnionWith(SC.FilesToStage.UFSFiles.Keys.Where(x => x.ContainsName(RestrictedName)));
 			RestrictedFiles.UnionWith(SC.FilesToStage.NonUFSFiles.Keys.Where(x => x.ContainsName(RestrictedName)));
@@ -843,7 +843,7 @@ public partial class Project : CommandUtils
 		if (RestrictedFiles.Count > 0)
 		{
 			List<string> RestrictedNames = new List<string>();
-			foreach(FileSystemName RestrictedFolderName in SC.RestrictedFolderNames)
+			foreach(string RestrictedFolderName in SC.RestrictedFolderNames)
 			{
 				if(RestrictedFiles.Any(x => x.ContainsName(RestrictedFolderName)))
 				{
@@ -858,7 +858,7 @@ public partial class Project : CommandUtils
 				Message.AppendFormat("\n{0}", RestrictedFile);
 			}
 			Message.Append("\n[Restrictions]");
-			foreach (FileSystemName RestrictedName in SC.RestrictedFolderNames)
+			foreach (string RestrictedName in SC.RestrictedFolderNames)
 			{
 				Message.AppendFormat("\n{0}", RestrictedName);
 			}
diff --git a/Engine/Source/Programs/AutomationTool/Scripts/ListThirdPartySoftware.Automation.cs b/Engine/Source/Programs/AutomationTool/Scripts/ListThirdPartySoftware.Automation.cs
--- a/Engine/Source/Programs/AutomationTool/Scripts/ListThirdPartySoftware.Automation.cs
+++ b/Engine/Source/Programs/AutomationTool/Scripts/ListThirdPartySoftware.Automation.cs
@@ -94,7 +94,7 @@ class ListThirdPartySoftware : BuildCommand
 
 			// Get the platforms to exclude
 			List<UnrealTargetPlatform> SupportedPlatforms = new List<UnrealTargetPlatform> { (UnrealTargetPlatform)Enum.Parse(typeof(UnrealTargetPlatform), Object.GetStringField("Platform")) };
-			FileSystemName[] ExcludePlatformNames = Utils.MakeListOfUnsupportedPlatforms(SupportedPlatforms).Select(x => new FileSystemName(x)).ToArray();
+			string[] ExcludePlatformNames = Utils.MakeListOfUnsupportedPlatforms(SupportedPlatforms).ToArray();
 
 			// Find all the TPS files under the engine directory which match
 			foreach(DirectoryReference DirectoryToScan in SortedDirectoriesToScan)
diff --git a/Engine/Source/Programs/AutomationTool/Scripts/RebasePublicIncludePaths.cs b/Engine/Source/Programs/AutomationTool/Scripts/RebasePublicIncludePaths.cs
--- a/Engine/Source/Programs/AutomationTool/Scripts/RebasePublicIncludePaths.cs
+++ b/Engine/Source/Programs/AutomationTool/Scripts/RebasePublicIncludePaths.cs
@@ -106,10 +106,10 @@ namespace AutomationTool
 			}
 
 			// List of folders to exclude from updates
-			FileSystemName[] ExcludeFoldersFromUpdate =
+			string[] ExcludeFoldersFromUpdate =
 			{
-				new FileSystemName("Intermediate"),
-				new FileSystemName("ThirdParty")
+				"Intermediate",
+				"ThirdParty"
 			};
 			
 			// Enumerate all the files to update
diff --git a/Engine/Source/Programs/DotNETCommon/DotNETUtilities/DotNETUtilities.csproj b/Engine/Source/Programs/DotNETCommon/DotNETUtilities/DotNETUtilities.csproj
--- a/Engine/Source/Programs/DotNETCommon/DotNETUtilities/DotNETUtilities.csproj
+++ b/Engine/Source/Programs/DotNETCommon/DotNETUtilities/DotNETUtilities.csproj
@@ -71,7 +71,6 @@
     <Compile Include="FileFilterNode.cs" />
     <Compile Include="FilePattern.cs" />
     <Compile Include="FileReference.cs" />
-    <Compile Include="FileSystemName.cs" />
     <Compile Include="FileSystemReference.cs" />
     <Compile Include="JsonObject.cs" />
     <Compile Include="Log.cs" />
diff --git a/Engine/Source/Programs/DotNETCommon/DotNETUtilities/FileFilter.cs b/Engine/Source/Programs/DotNETCommon/DotNETUtilities/FileFilter.cs
--- a/Engine/Source/Programs/DotNETCommon/DotNETUtilities/FileFilter.cs
+++ b/Engine/Source/Programs/DotNETCommon/DotNETUtilities/FileFilter.cs
@@ -37,16 +37,16 @@ namespace Tools.DotNETCommon
 		/// <summary>
 		/// List of all confidential folder names
 		/// </summary>
-		public static readonly FileSystemName[] RestrictedFolderNames =
+		public static readonly string[] RestrictedFolderNames =
 		{
-			new FileSystemName("EpicInternal"),
-			new FileSystemName("CarefullyRedist"),
-			new FileSystemName("NotForLicensees"),
-			new FileSystemName("NoRedist"),
-			new FileSystemName("PS4"),
-			new FileSystemName("XboxOne"),
-			new FileSystemName("Switch"),
-			new FileSystemName("Quail")
+			"EpicInternal",
+			"CarefullyRedist",
+			"NotForLicensees",
+			"NoRedist",
+			"PS4",
+			"XboxOne",
+			"Switch",
+			"Quail"
 		};
 
 		/// <summary>
@@ -363,7 +363,7 @@ namespace Tools.DotNETCommon
 		/// </summary>
 		public void ExcludeRestrictedFolders()
 		{
-			foreach(FileSystemName RestrictedFolderName in RestrictedFolderNames)
+			foreach(string RestrictedFolderName in RestrictedFolderNames)
 			{
 				AddRule(String.Format(".../{0}/...", RestrictedFolderName), FileFilterType.Exclude);
 			}
diff --git a/Engine/Source/Programs/DotNETCommon/DotNETUtilities/FileSystemName.cs b/Engine/Source/Programs/DotNETCommon/DotNETUtilities/FileSystemName.cs
deleted file mode 100644
diff --git a/Engine/Source/Programs/DotNETCommon/DotNETUtilities/FileSystemReference.cs b/Engine/Source/Programs/DotNETCommon/DotNETUtilities/FileSystemReference.cs
--- a/Engine/Source/Programs/DotNETCommon/DotNETUtilities/FileSystemReference.cs
+++ b/Engine/Source/Programs/DotNETCommon/DotNETUtilities/FileSystemReference.cs
@@ -157,7 +157,7 @@ namespace Tools.DotNETCommon
 		/// <param name="Name">Name to check for</param>
 		/// <param name="Offset">Offset within the string to start the search</param>
 		/// <returns>True if the given name is found within the path</returns>
-		public bool ContainsName(FileSystemName Name, int Offset)
+		public bool ContainsName(string Name, int Offset)
 		{
 			return ContainsName(Name, Offset, FullName.Length - Offset);
 		}
@@ -169,10 +169,10 @@ namespace Tools.DotNETCommon
 		/// <param name="Offset">Offset within the string to start the search</param>
 		/// <param name="Length">Length of the substring to search</param>
 		/// <returns>True if the given name is found within the path</returns>
-		public bool ContainsName(FileSystemName Name, int Offset, int Length)
+		public bool ContainsName(string Name, int Offset, int Length)
 		{
 			// Check the substring to search is at least long enough to contain a match
-			if(Length < Name.DisplayName.Length)
+			if(Length < Name.Length)
 			{
 				return false;
 			}
@@ -182,21 +182,21 @@ namespace Tools.DotNETCommon
 			for(;;)
 			{
 				// Find the next occurrence
-				MatchIdx = FullName.IndexOf(Name.DisplayName, MatchIdx, Offset + Length - MatchIdx, Comparison);
+				MatchIdx = FullName.IndexOf(Name, MatchIdx, Offset + Length - MatchIdx, Comparison);
 				if(MatchIdx == -1)
 				{
 					return false;
 				}
 
 				// Check if the substring is a directory
-				int MatchEndIdx = MatchIdx + Name.DisplayName.Length;
+				int MatchEndIdx = MatchIdx + Name.Length;
 				if(FullName[MatchIdx - 1] == Path.DirectorySeparatorChar && (MatchEndIdx == FullName.Length || FullName[MatchEndIdx] == Path.DirectorySeparatorChar))
 				{
 					return true;
 				}
 
 				// Move past the string that didn't match
-				MatchIdx += Name.DisplayName.Length;
+				MatchIdx += Name.Length;
 			}
 		}
 
@@ -206,7 +206,7 @@ namespace Tools.DotNETCommon
 		/// <param name="Name">Name of a subfolder to also check for</param>
 		/// <param name="BaseDir">Base directory to check against</param>
 		/// <returns>True if the path is under the given directory</returns>
-		public bool ContainsName(FileSystemName Name, DirectoryReference BaseDir)
+		public bool ContainsName(string Name, DirectoryReference BaseDir)
 		{
 			// Check that this is under the base directory
 			if(!IsUnderDirectory(BaseDir))
@@ -225,7 +225,7 @@ namespace Tools.DotNETCommon
 		/// <param name="Names">Names of subfolders to also check for</param>
 		/// <param name="BaseDir">Base directory to check against</param>
 		/// <returns>True if the path is under the given directory</returns>
-		public bool ContainsAnyNames(FileSystemName[] Names, DirectoryReference BaseDir)
+		public bool ContainsAnyNames(IEnumerable<string> Names, DirectoryReference BaseDir)
 		{
 			// Check that this is under the base directory
 			if(!IsUnderDirectory(BaseDir))
diff --git a/Engine/Source/Programs/UnrealBuildTool/Configuration/UEBuildModuleCPP.cs b/Engine/Source/Programs/UnrealBuildTool/Configuration/UEBuildModuleCPP.cs
--- a/Engine/Source/Programs/UnrealBuildTool/Configuration/UEBuildModuleCPP.cs
+++ b/Engine/Source/Programs/UnrealBuildTool/Configuration/UEBuildModuleCPP.cs
@@ -272,7 +272,7 @@ namespace UnrealBuildTool
 			{
 				PublicIncludePaths.Add(PublicDirectory);
 
-				FileSystemName[] ExcludedFolderNames = UEBuildPlatform.GetBuildPlatform(Rules.Target.Platform).GetExcludedFolderNames();
+				string[] ExcludedFolderNames = UEBuildPlatform.GetBuildPlatform(Rules.Target.Platform).GetExcludedFolderNames();
 				foreach (DirectoryReference PublicSubDirectory in DirectoryLookupCache.EnumerateDirectoriesRecursively(PublicDirectory))
 				{
 					if(!PublicSubDirectory.ContainsAnyNames(ExcludedFolderNames, PublicDirectory))
diff --git a/Engine/Source/Programs/UnrealBuildTool/Configuration/UEBuildPlatform.cs b/Engine/Source/Programs/UnrealBuildTool/Configuration/UEBuildPlatform.cs
--- a/Engine/Source/Programs/UnrealBuildTool/Configuration/UEBuildPlatform.cs
+++ b/Engine/Source/Programs/UnrealBuildTool/Configuration/UEBuildPlatform.cs
@@ -30,17 +30,17 @@ namespace UnrealBuildTool
 		/// <summary>
 		/// All the platform folder names
 		/// </summary>
-		private static FileSystemName[] CachedPlatformFolderNames;
+		private static string[] CachedPlatformFolderNames;
 
 		/// <summary>
 		/// Cached copy of the list of folders to include for this platform
 		/// </summary>
-		private FileSystemName[] CachedIncludedFolderNames;
+		private string[] CachedIncludedFolderNames;
 
 		/// <summary>
 		/// Cached copy of the list of folders to exclude for this platform
 		/// </summary>
-		private FileSystemName[] CachedExcludedFolderNames;
+		private string[] CachedExcludedFolderNames;
 
 		/// <summary>
 		/// Constructor.
@@ -57,25 +57,25 @@ namespace UnrealBuildTool
 		/// Gets an array of all platform folder names
 		/// </summary>
 		/// <returns>Array of platform folders</returns>
-		public static FileSystemName[] GetPlatformFolderNames()
+		public static string[] GetPlatformFolderNames()
 		{
 			if(CachedPlatformFolderNames == null)
 			{
-				List<FileSystemName> PlatformFolderNames = new List<FileSystemName>();
+				List<string> PlatformFolderNames = new List<string>();
 
 				// Find all the platform folders to exclude from the list of precompiled modules
 				foreach (UnrealTargetPlatform TargetPlatform in Enum.GetValues(typeof(UnrealTargetPlatform)))
 				{
 					if (TargetPlatform != UnrealTargetPlatform.Unknown)
 					{
-						PlatformFolderNames.Add(new FileSystemName(TargetPlatform.ToString()));
+						PlatformFolderNames.Add(TargetPlatform.ToString());
 					}
 				}
 
 				// Also exclude all the platform groups that this platform is not a part of
 				foreach (UnrealPlatformGroup PlatformGroup in Enum.GetValues(typeof(UnrealPlatformGroup)))
 				{
-					PlatformFolderNames.Add(new FileSystemName(PlatformGroup.ToString()));
+					PlatformFolderNames.Add(PlatformGroup.ToString());
 				}
 
 				// Save off the list as an array
@@ -87,16 +87,16 @@ namespace UnrealBuildTool
 		/// <summary>
 		/// Finds a list of folder names to include when building for this platform
 		/// </summary>
-		virtual public FileSystemName[] GetIncludedFolderNames()
+		virtual public string[] GetIncludedFolderNames()
 		{
 			if(CachedIncludedFolderNames == null)
 			{
-				List<FileSystemName> Names = new List<FileSystemName>();
+				List<string> Names = new List<string>();
 
-				Names.Add(new FileSystemName(Platform.ToString()));
+				Names.Add(Platform.ToString());
 				foreach(UnrealPlatformGroup Group in UEBuildPlatform.GetPlatformGroups(Platform))
 				{
-					Names.Add(new FileSystemName(Group.ToString()));
+					Names.Add(Group.ToString());
 				}
 
 				CachedIncludedFolderNames = Names.ToArray();
@@ -107,7 +107,7 @@ namespace UnrealBuildTool
 		/// <summary>
 		/// Finds a list of folder names to exclude when building for this platform
 		/// </summary>
-		public FileSystemName[] GetExcludedFolderNames()
+		public string[] GetExcludedFolderNames()
 		{
 			if(CachedExcludedFolderNames == null)
 			{
diff --git a/Engine/Source/Programs/UnrealBuildTool/Configuration/UEBuildTarget.cs b/Engine/Source/Programs/UnrealBuildTool/Configuration/UEBuildTarget.cs
--- a/Engine/Source/Programs/UnrealBuildTool/Configuration/UEBuildTarget.cs
+++ b/Engine/Source/Programs/UnrealBuildTool/Configuration/UEBuildTarget.cs
@@ -2781,12 +2781,12 @@ namespace UnrealBuildTool
 				RulesAssembly.GetAllModuleNames(ModuleNames);
 
 				// Find all the platform folders to exclude from the list of valid modules
-				List<FileSystemName> ExcludeFolders = new List<FileSystemName>();
+				List<string> ExcludeFolders = new List<string>();
 				foreach (UnrealTargetPlatform TargetPlatform in Enum.GetValues(typeof(UnrealTargetPlatform)))
 				{
 					if (TargetPlatform != Platform)
 					{
-						ExcludeFolders.Add(new FileSystemName(TargetPlatform.ToString()));
+						ExcludeFolders.Add(TargetPlatform.ToString());
 					}
 				}
 
@@ -2796,12 +2796,12 @@ namespace UnrealBuildTool
 				{
 					if (!IncludePlatformGroups.Contains(PlatformGroup))
 					{
-						ExcludeFolders.Add(new FileSystemName(PlatformGroup.ToString()));
+						ExcludeFolders.Add(PlatformGroup.ToString());
 					}
 				}
 
 				// Create an array of folders to exclude
-				FileSystemName[] ExcludeFoldersArray = ExcludeFolders.ToArray();
+				string[] ExcludeFoldersArray = ExcludeFolders.ToArray();
 
 				// Find all the directories containing engine modules that may be compatible with this target
 				List<DirectoryReference> Directories = new List<DirectoryReference>();
@@ -3137,7 +3137,7 @@ namespace UnrealBuildTool
 			}
 
 			// Get an array of folders to filter out
-			FileSystemName[] ExcludeFolders = MissingPlatforms.Select(x => new FileSystemName(x.ToString())).ToArray();
+			string[] ExcludeFolders = MissingPlatforms.Select(x => x.ToString()).ToArray();
 
 			// Set of all the plugins that have been referenced
 			HashSet<string> ReferencedNames = new HashSet<string>(StringComparer.InvariantCultureIgnoreCase);
@@ -3241,7 +3241,7 @@ namespace UnrealBuildTool
 		/// <param name="NameToInstance">Map from plugin name to instance of it</param>
 		/// <param name="NameToInfo">Map from plugin name to information</param>
 		/// <returns>Instance of the plugin, or null if it should not be used</returns>
-		private UEBuildPlugin AddPlugin(PluginReferenceDescriptor Reference, string ReferenceChain, FileSystemName[] ExcludeFolders, Dictionary<string, UEBuildPlugin> NameToInstance, Dictionary<string, PluginInfo> NameToInfo)
+		private UEBuildPlugin AddPlugin(PluginReferenceDescriptor Reference, string ReferenceChain, string[] ExcludeFolders, Dictionary<string, UEBuildPlugin> NameToInstance, Dictionary<string, PluginInfo> NameToInfo)
 		{
 			// Ignore disabled references
 			if(!Reference.bEnabled)
@@ -3381,7 +3381,7 @@ namespace UnrealBuildTool
 		/// <summary>
 		/// Checks whether a plugin path contains a platform directory fragment
 		/// </summary>
-		private bool ShouldExcludePlugin(PluginInfo Plugin, FileSystemName[] ExcludeFolders)
+		private bool ShouldExcludePlugin(PluginInfo Plugin, string[] ExcludeFolders)
 		{
 			if (Plugin.LoadedFrom == PluginLoadedFrom.Engine)
 			{
diff --git a/Engine/Source/Programs/UnrealBuildTool/Platform/Linux/LinuxToolChain.cs b/Engine/Source/Programs/UnrealBuildTool/Platform/Linux/LinuxToolChain.cs
--- a/Engine/Source/Programs/UnrealBuildTool/Platform/Linux/LinuxToolChain.cs
+++ b/Engine/Source/Programs/UnrealBuildTool/Platform/Linux/LinuxToolChain.cs
@@ -1383,7 +1383,7 @@ namespace UnrealBuildTool
 
 			foreach (FileReference RuntimeDependency in LinkEnvironment.RuntimeDependencies)
 			{
-				if (RuntimeDependency.ContainsName(new FileSystemName("Binaries"), 0) && RuntimeDependency.GetExtension() == ".so" && Path.GetDirectoryName(RuntimeDependency.FullName) != Path.GetDirectoryName(OutputFile.AbsolutePath))
+				if (RuntimeDependency.ContainsName("Binaries", 0) && RuntimeDependency.GetExtension() == ".so" && Path.GetDirectoryName(RuntimeDependency.FullName) != Path.GetDirectoryName(OutputFile.AbsolutePath))
 				{
 					string RelativeRootPath = RuntimeDependency.Directory.MakeRelativeTo(UnrealBuildTool.RootDirectory);
 					// We're assuming that the binary will be placed according to our ProjectName/Binaries/Platform scheme
diff --git a/Engine/Source/Programs/UnrealBuildTool/ProjectFiles/ProjectFileGenerator.cs b/Engine/Source/Programs/UnrealBuildTool/ProjectFiles/ProjectFileGenerator.cs
--- a/Engine/Source/Programs/UnrealBuildTool/ProjectFiles/ProjectFileGenerator.cs
+++ b/Engine/Source/Programs/UnrealBuildTool/ProjectFiles/ProjectFileGenerator.cs
@@ -428,7 +428,7 @@ namespace UnrealBuildTool
 		/// </summary>
 		void DiscoverCSharpProgramProjects(MasterProjectFolder ProgramsFolder)
 		{
-			FileSystemName[] UnsupportedPlatformNames = Utils.MakeListOfUnsupportedPlatforms(SupportedPlatforms).Select(x => new FileSystemName(x)).ToArray();
+			string[] UnsupportedPlatformNames = Utils.MakeListOfUnsupportedPlatforms(SupportedPlatforms).ToArray();
 
 			List<FileReference> FoundProjects = new List<FileReference>();
 			DirectoryReference EngineProgramsSource = DirectoryReference.Combine(UnrealBuildTool.EngineDirectory, "Source", "Programs");
diff --git a/Engine/Source/Programs/UnrealBuildTool/System/ExternalExecution.cs b/Engine/Source/Programs/UnrealBuildTool/System/ExternalExecution.cs
--- a/Engine/Source/Programs/UnrealBuildTool/System/ExternalExecution.cs
+++ b/Engine/Source/Programs/UnrealBuildTool/System/ExternalExecution.cs
@@ -569,9 +569,9 @@ namespace UnrealBuildTool
 		/// <param name="BaseDir">Base directory to search</param>
 		/// <param name="ExcludeFolders">Array of folders to exclude</param>
 		/// <param name="Headers">Receives the list of headers that was found</param>
-		static void FindHeaders(DirectoryInfo BaseDir, FileSystemName[] ExcludeFolders, List<FileReference> Headers)
+		static void FindHeaders(DirectoryInfo BaseDir, string[] ExcludeFolders, List<FileReference> Headers)
 		{
-			if (!ExcludeFolders.Any(x => x.DisplayName.Equals(BaseDir.Name, StringComparison.InvariantCultureIgnoreCase)))
+			if (!ExcludeFolders.Any(x => x.Equals(BaseDir.Name, StringComparison.InvariantCultureIgnoreCase)))
 			{
 				foreach (DirectoryInfo SubDir in BaseDir.EnumerateDirectories())
 				{
@@ -599,7 +599,7 @@ namespace UnrealBuildTool
 			List<UEBuildModuleCPP> ModulesSortedByType = ModulesToGenerateHeadersFor.OrderBy(c => ModuleToType[c]).ToList();
 			StableTopologicalSort(ModulesSortedByType);
 
-			FileSystemName[] ExcludedFolders = UEBuildPlatform.GetBuildPlatform(Target.Platform, true).GetExcludedFolderNames();
+			string[] ExcludedFolders = UEBuildPlatform.GetBuildPlatform(Target.Platform, true).GetExcludedFolderNames();
 			foreach (UEBuildModuleCPP Module in ModulesSortedByType)
 			{
 				List<FileReference> HeaderFiles = new List<FileReference>();
diff --git a/Engine/Source/Programs/UnrealBuildTool/System/PlatformExports.cs b/Engine/Source/Programs/UnrealBuildTool/System/PlatformExports.cs
--- a/Engine/Source/Programs/UnrealBuildTool/System/PlatformExports.cs
+++ b/Engine/Source/Programs/UnrealBuildTool/System/PlatformExports.cs
@@ -72,7 +72,7 @@ namespace UnrealBuildTool
 		/// Returns an array of all platform folder names
 		/// </summary>
 		/// <returns>All platform folder names</returns>
-		public static FileSystemName[] GetPlatformFolderNames()
+		public static string[] GetPlatformFolderNames()
 		{
 			return UEBuildPlatform.GetPlatformFolderNames();
 		}
@@ -82,7 +82,7 @@ namespace UnrealBuildTool
 		/// </summary>
 		/// <param name="Platform">The platform to get the included folder names for</param>
 		/// <returns>All platform folder names</returns>
-		public static FileSystemName[] GetIncludedFolderNames(UnrealTargetPlatform Platform)
+		public static string[] GetIncludedFolderNames(UnrealTargetPlatform Platform)
 		{
 			UEBuildPlatform BuildPlatform = UEBuildPlatform.GetBuildPlatform(Platform, false);
 			return BuildPlatform.GetIncludedFolderNames();
@@ -93,7 +93,7 @@ namespace UnrealBuildTool
 		/// </summary>
 		/// <param name="Platform">The platform to get the excluded folder names for</param>
 		/// <returns>Array of folder names</returns>
-		public static FileSystemName[] GetExcludedFolderNames(UnrealTargetPlatform Platform)
+		public static string[] GetExcludedFolderNames(UnrealTargetPlatform Platform)
 		{
 			UEBuildPlatform BuildPlatform = UEBuildPlatform.GetBuildPlatform(Platform, false);
 			return BuildPlatform.GetExcludedFolderNames();
diff --git a/Engine/Source/Programs/UnrealBuildTool/System/RestrictedFolders.cs b/Engine/Source/Programs/UnrealBuildTool/System/RestrictedFolders.cs
--- a/Engine/Source/Programs/UnrealBuildTool/System/RestrictedFolders.cs
+++ b/Engine/Source/Programs/UnrealBuildTool/System/RestrictedFolders.cs
@@ -1,4 +1,4 @@
-﻿// Copyright 1998-2018 Epic Games, Inc. All Rights Reserved.
+// Copyright 1998-2018 Epic Games, Inc. All Rights Reserved.
 
 using System;
 using System.Collections.Generic;
@@ -59,7 +59,7 @@ namespace UnrealBuildTool
 		/// <summary>
 		/// Keeps a cache of folder substrings and flags
 		/// </summary>
-		static readonly Tuple<FileSystemName, RestrictedFolder>[] RestrictedFolderNames;
+		static readonly Tuple<string, RestrictedFolder>[] RestrictedFolderNames;
 		
 		/// <summary>
 		/// Static constructor
@@ -69,10 +69,10 @@ namespace UnrealBuildTool
 			string[] Names = Enum.GetNames(typeof(RestrictedFolder));
 			RestrictedFolder[] Values = (RestrictedFolder[])Enum.GetValues(typeof(RestrictedFolder));
 
-			RestrictedFolderNames = new Tuple<FileSystemName, RestrictedFolder>[Names.Length];
+			RestrictedFolderNames = new Tuple<string, RestrictedFolder>[Names.Length];
 			for(int Idx = 0; Idx < Names.Length; Idx++)
 			{
-				RestrictedFolderNames[Idx] = Tuple.Create(new FileSystemName(Names[Idx]), Values[Idx]);
+				RestrictedFolderNames[Idx] = Tuple.Create(Names[Idx], Values[Idx]);
 			}
 		}
 
@@ -87,7 +87,7 @@ namespace UnrealBuildTool
 			List<RestrictedFolder> Folders = new List<RestrictedFolder>();
 			if(OtherDir.IsUnderDirectory(BaseDir))
 			{
-				foreach(Tuple<FileSystemName, RestrictedFolder> Pair in RestrictedFolderNames)
+				foreach(Tuple<string, RestrictedFolder> Pair in RestrictedFolderNames)
 				{
 					if(OtherDir.ContainsName(Pair.Item1, BaseDir.FullName.Length))
 					{
