UBT: Remove some redundant parameters from UEBuildModule/UEBuildModuleCPP/UEBuildModuleExternal constructors.

#rnx

diff --git a/Engine/Source/Programs/UnrealBuildTool/Configuration/UEBuildModule.cs b/Engine/Source/Programs/UnrealBuildTool/Configuration/UEBuildModule.cs
--- a/Engine/Source/Programs/UnrealBuildTool/Configuration/UEBuildModule.cs
+++ b/Engine/Source/Programs/UnrealBuildTool/Configuration/UEBuildModule.cs
@@ -18,29 +18,38 @@ namespace UnrealBuildTool
 	abstract class UEBuildModule
 	{
 		/// <summary>
-		/// The name that uniquely identifies the module.
+		/// The rules for this module
 		/// </summary>
-		public readonly string Name;
+		public readonly ModuleRules Rules;
 
 		/// <summary>
-		/// The rules for this module
+		/// The name that uniquely identifies the module.
 		/// </summary>
-		public ModuleRules Rules;
+		public string Name
+		{
+			get { return Rules.Name; }
+		}
 
 		/// <summary>
 		/// Path to the module directory
 		/// </summary>
-		public readonly DirectoryReference ModuleDirectory;
+		public DirectoryReference ModuleDirectory
+		{
+			get { return Rules.Directory; }
+		}
 
 		/// <summary>
-		/// Is this module allowed to be redistributed.
+		/// The name of the .Build.cs file this module was created from, if any
 		/// </summary>
-		private readonly bool? IsRedistributableOverride;
+		public FileReference RulesFile
+		{
+			get { return Rules.File; }
+		}
 
 		/// <summary>
-		/// The name of the .Build.cs file this module was created from, if any
+		/// Is this module allowed to be redistributed.
 		/// </summary>
-		public FileReference RulesFile;
+		private readonly bool? IsRedistributableOverride;
 
 		/// <summary>
 		/// The binary the module will be linked into for the current target.  Only set after UEBuildBinary.BindModules is called.
@@ -155,43 +164,37 @@ namespace UnrealBuildTool
 		/// <summary>
 		/// Constructor
 		/// </summary>
-		/// <param name="InName">Name of the module</param>
-		/// <param name="InModuleDirectory">Base directory for the module</param>
-		/// <param name="InRules">Rules for this module</param>
-		/// <param name="InRulesFile">Path to the rules file</param>
-		/// <param name="InRuntimeDependencies">List of runtime dependencies</param>
-		public UEBuildModule(string InName, DirectoryReference InModuleDirectory, ModuleRules InRules, FileReference InRulesFile, List<RuntimeDependency> InRuntimeDependencies)
+		/// <param name="Rules">Rules for this module</param>
+		/// <param name="RuntimeDependencies">List of runtime dependencies</param>
+		public UEBuildModule(ModuleRules Rules, List<RuntimeDependency> RuntimeDependencies)
 		{
-			Name = InName;
-			ModuleDirectory = InModuleDirectory;
-			Rules = InRules;
-			RulesFile = InRulesFile;
+			this.Rules = Rules;
 
 			ModuleApiDefine = Name.ToUpperInvariant() + "_API";
 
-			PublicDefinitions = HashSetFromOptionalEnumerableStringParameter(InRules.PublicDefinitions);
-			PublicIncludePaths = CreateDirectoryHashSet(InRules.PublicIncludePaths);
-			PublicSystemIncludePaths = CreateDirectoryHashSet(InRules.PublicSystemIncludePaths);
-			PublicLibraryPaths = CreateDirectoryHashSet(InRules.PublicLibraryPaths);
-			PublicAdditionalLibraries = HashSetFromOptionalEnumerableStringParameter(InRules.PublicAdditionalLibraries);
-			PublicFrameworks = HashSetFromOptionalEnumerableStringParameter(InRules.PublicFrameworks);
-			PublicWeakFrameworks = HashSetFromOptionalEnumerableStringParameter(InRules.PublicWeakFrameworks);
-			PublicAdditionalFrameworks = InRules.PublicAdditionalFrameworks == null ? new HashSet<UEBuildFramework>() : new HashSet<UEBuildFramework>(InRules.PublicAdditionalFrameworks);
-			PublicAdditionalShadowFiles = HashSetFromOptionalEnumerableStringParameter(InRules.PublicAdditionalShadowFiles);
-			PublicAdditionalBundleResources = InRules.AdditionalBundleResources == null ? new HashSet<UEBuildBundleResource>() : new HashSet<UEBuildBundleResource>(InRules.AdditionalBundleResources);
-			PublicDelayLoadDLLs = HashSetFromOptionalEnumerableStringParameter(InRules.PublicDelayLoadDLLs);
+			PublicDefinitions = HashSetFromOptionalEnumerableStringParameter(Rules.PublicDefinitions);
+			PublicIncludePaths = CreateDirectoryHashSet(Rules.PublicIncludePaths);
+			PublicSystemIncludePaths = CreateDirectoryHashSet(Rules.PublicSystemIncludePaths);
+			PublicLibraryPaths = CreateDirectoryHashSet(Rules.PublicLibraryPaths);
+			PublicAdditionalLibraries = HashSetFromOptionalEnumerableStringParameter(Rules.PublicAdditionalLibraries);
+			PublicFrameworks = HashSetFromOptionalEnumerableStringParameter(Rules.PublicFrameworks);
+			PublicWeakFrameworks = HashSetFromOptionalEnumerableStringParameter(Rules.PublicWeakFrameworks);
+			PublicAdditionalFrameworks = Rules.PublicAdditionalFrameworks == null ? new HashSet<UEBuildFramework>() : new HashSet<UEBuildFramework>(Rules.PublicAdditionalFrameworks);
+			PublicAdditionalShadowFiles = HashSetFromOptionalEnumerableStringParameter(Rules.PublicAdditionalShadowFiles);
+			PublicAdditionalBundleResources = Rules.AdditionalBundleResources == null ? new HashSet<UEBuildBundleResource>() : new HashSet<UEBuildBundleResource>(Rules.AdditionalBundleResources);
+			PublicDelayLoadDLLs = HashSetFromOptionalEnumerableStringParameter(Rules.PublicDelayLoadDLLs);
 			if(Rules.bUsePrecompiled)
 			{
 				PrivateIncludePaths = new HashSet<DirectoryReference>();
 			}
 			else
 			{
-				PrivateIncludePaths = CreateDirectoryHashSet(InRules.PrivateIncludePaths);
+				PrivateIncludePaths = CreateDirectoryHashSet(Rules.PrivateIncludePaths);
 			}
-			RuntimeDependencies = InRuntimeDependencies;
-			IsRedistributableOverride = InRules.IsRedistributableOverride;
+			this.RuntimeDependencies = RuntimeDependencies;
+			IsRedistributableOverride = Rules.IsRedistributableOverride;
 
-			WhitelistRestrictedFolders = new HashSet<DirectoryReference>(InRules.WhitelistRestrictedFolders.Select(x => DirectoryReference.Combine(ModuleDirectory, x)));
+			WhitelistRestrictedFolders = new HashSet<DirectoryReference>(Rules.WhitelistRestrictedFolders.Select(x => DirectoryReference.Combine(ModuleDirectory, x)));
 		}
 
 		/// <summary>
diff --git a/Engine/Source/Programs/UnrealBuildTool/Configuration/UEBuildModuleCPP.cs b/Engine/Source/Programs/UnrealBuildTool/Configuration/UEBuildModuleCPP.cs
--- a/Engine/Source/Programs/UnrealBuildTool/Configuration/UEBuildModuleCPP.cs
+++ b/Engine/Source/Programs/UnrealBuildTool/Configuration/UEBuildModuleCPP.cs
@@ -198,31 +198,22 @@ namespace UnrealBuildTool
 
 
 		public UEBuildModuleCPP(
-			string InName,
-			DirectoryReference InModuleDirectory,
-			DirectoryReference InIntermediateDirectory,
-			DirectoryReference InGeneratedCodeDirectory,
-			IEnumerable<FileItem> InSourceFiles,
-			ModuleRules InRules,
-			bool bInBuildSourceFiles,
-			FileReference InRulesFile,
-			List<RuntimeDependency> InRuntimeDependencies
+			ModuleRules Rules,
+			DirectoryReference IntermediateDirectory,
+			DirectoryReference GeneratedCodeDirectory,
+			IEnumerable<FileItem> SourceFiles,
+			bool bBuildSourceFiles,
+			List<RuntimeDependency> RuntimeDependencies
 			)
-			: base(
-					InName,
-					InModuleDirectory,
-					InRules,
-					InRulesFile,
-					InRuntimeDependencies
-				)
+			: base(Rules, RuntimeDependencies)
 		{
-			IntermediateDirectory = InIntermediateDirectory;
-			GeneratedCodeDirectory = InGeneratedCodeDirectory;
+			this.IntermediateDirectory = IntermediateDirectory;
+			this.GeneratedCodeDirectory = GeneratedCodeDirectory;
 
-			SourceFiles = InSourceFiles.ToList();
+			this.SourceFiles = SourceFiles.ToList();
 
-			CategorizeSourceFiles(InSourceFiles, SourceFilesFound);
-			if (bInBuildSourceFiles)
+			CategorizeSourceFiles(SourceFiles, SourceFilesFound);
+			if (bBuildSourceFiles)
 			{
 				SourceFilesToBuild.CopyFrom(SourceFilesFound);
 			}
diff --git a/Engine/Source/Programs/UnrealBuildTool/Configuration/UEBuildModuleExternal.cs b/Engine/Source/Programs/UnrealBuildTool/Configuration/UEBuildModuleExternal.cs
--- a/Engine/Source/Programs/UnrealBuildTool/Configuration/UEBuildModuleExternal.cs
+++ b/Engine/Source/Programs/UnrealBuildTool/Configuration/UEBuildModuleExternal.cs
@@ -1,4 +1,4 @@
-﻿// Copyright 1998-2018 Epic Games, Inc. All Rights Reserved.
+// Copyright 1998-2018 Epic Games, Inc. All Rights Reserved.
 
 using System;
 using System.Collections.Generic;
@@ -14,20 +14,8 @@ namespace UnrealBuildTool
 	/// </summary>
 	class UEBuildModuleExternal : UEBuildModule
 	{
-		public UEBuildModuleExternal(
-			string InName,
-			DirectoryReference InModuleDirectory,
-			ModuleRules InRules,
-			FileReference InRulesFile,
-			List<RuntimeDependency> InRuntimeDependencies
-			)
-			: base(
-				InName: InName,
-				InModuleDirectory: InModuleDirectory,
-				InRules: InRules,
-				InRulesFile: InRulesFile,
-				InRuntimeDependencies: InRuntimeDependencies
-				)
+		public UEBuildModuleExternal(ModuleRules Rules, List<RuntimeDependency> RuntimeDependencies)
+			: base(Rules, RuntimeDependencies)
 		{
 		}
 
diff --git a/Engine/Source/Programs/UnrealBuildTool/Configuration/UEBuildTarget.cs b/Engine/Source/Programs/UnrealBuildTool/Configuration/UEBuildTarget.cs
--- a/Engine/Source/Programs/UnrealBuildTool/Configuration/UEBuildTarget.cs
+++ b/Engine/Source/Programs/UnrealBuildTool/Configuration/UEBuildTarget.cs
@@ -4068,30 +4068,24 @@ namespace UnrealBuildTool
 			DirectoryReference GeneratedCodeDirectory,
 			List<FileItem> ModuleSourceFiles,
 			bool bBuildSourceFiles,
-			List<RuntimeDependency> InRuntimeDependencies)
+			List<RuntimeDependency> RuntimeDependencies)
 		{
 			switch (RulesObject.Type)
 			{
 				case ModuleRules.ModuleType.CPlusPlus:
 					return new UEBuildModuleCPP(
-							InName: RulesObject.Name,
-							InModuleDirectory: RulesObject.Directory,
-							InIntermediateDirectory: GetModuleIntermediateDirectory(RulesObject),
-							InGeneratedCodeDirectory: GeneratedCodeDirectory,
-							InSourceFiles: ModuleSourceFiles,
-							InRules: RulesObject,
-							bInBuildSourceFiles: bBuildSourceFiles,
-							InRulesFile: RulesObject.File,
-							InRuntimeDependencies: InRuntimeDependencies
+							Rules: RulesObject,
+							IntermediateDirectory: GetModuleIntermediateDirectory(RulesObject),
+							GeneratedCodeDirectory: GeneratedCodeDirectory,
+							SourceFiles: ModuleSourceFiles,
+							bBuildSourceFiles: bBuildSourceFiles,
+							RuntimeDependencies: RuntimeDependencies
 						);
 
 				case ModuleRules.ModuleType.External:
 					return new UEBuildModuleExternal(
-							InName: RulesObject.Name,
-							InModuleDirectory: RulesObject.Directory,
-							InRules: RulesObject,
-							InRulesFile: RulesObject.File,
-							InRuntimeDependencies: InRuntimeDependencies
+							Rules: RulesObject,
+							RuntimeDependencies: RuntimeDependencies
 						);
 
 				default:
