Bump the FlatCPPIncludeDependencyCache version, to prevent errors trying to load old files.

#rnx

diff --git a/Engine/Source/Programs/UnrealBuildTool/System/FlatCPPIncludeDepencencyCache.cs b/Engine/Source/Programs/UnrealBuildTool/System/FlatCPPIncludeDepencencyCache.cs
--- a/Engine/Source/Programs/UnrealBuildTool/System/FlatCPPIncludeDepencencyCache.cs
+++ b/Engine/Source/Programs/UnrealBuildTool/System/FlatCPPIncludeDepencencyCache.cs
@@ -36,7 +36,7 @@ namespace UnrealBuildTool
 		/// <summary>
 		/// The version number for binary serialization
 		/// </summary>
-		const int FileVersion = 2;
+		const int FileVersion = 3;
 
 		/// <summary>
 		/// The file signature for binary serialization
