UBT: Added support for expanding UE-specific variables in include paths and library paths: $(EngineDir), $(ProjectDir), $(PluginDir), $(ModuleDir).

diff --git a/Engine/Source/Programs/UnrealBuildTool/Configuration/UEBuildModule.cs b/Engine/Source/Programs/UnrealBuildTool/Configuration/UEBuildModule.cs
--- a/Engine/Source/Programs/UnrealBuildTool/Configuration/UEBuildModule.cs
+++ b/Engine/Source/Programs/UnrealBuildTool/Configuration/UEBuildModule.cs
@@ -246,13 +246,7 @@ namespace UnrealBuildTool
 			{
 				foreach(string InputString in InEnumerableStrings)
 				{
-					string ExpandedString = Utils.ExpandVariables(InputString);
-					if(ExpandedString.Contains("$("))
-					{
-						throw new BuildException("Unable to expand variable in '{0}'", InputString);
-					}
-
-					DirectoryReference Dir = new DirectoryReference(ExpandedString);
+					DirectoryReference Dir = new DirectoryReference(ExpandPathVariables(InputString, null));
 					if(DirectoryReference.Exists(Dir))
 					{
 						Directories.Add(Dir);
@@ -548,6 +542,92 @@ namespace UnrealBuildTool
 			}
 		}
 
+		/// <summary>
+		/// Expand path variables within the context of this module
+		/// </summary>
+		/// <param name="Path">Path to expand variables within</param>
+		/// <param name="Binary">The source binary containing the module. May be null.</param>
+		/// <returns>The path with variables expanded</returns>
+		private string ExpandPathVariables(string Path, UEBuildBinary Binary)
+		{
+			if(Path.StartsWith("$(", StringComparison.Ordinal))
+			{
+				int StartIdx = 2;
+				for(int EndIdx = StartIdx; EndIdx < Path.Length; EndIdx++)
+				{
+					if(Path[EndIdx] == ')')
+					{
+						if(MatchVariableName(Path, StartIdx, EndIdx, "EngineDir"))
+						{
+							Path = UnrealBuildTool.EngineDirectory + Path.Substring(EndIdx + 1);
+						}
+						else if(MatchVariableName(Path, StartIdx, EndIdx, "ProjectDir"))
+						{
+							if(Rules.Target.ProjectFile == null)
+							{
+								Path = UnrealBuildTool.EngineDirectory + Path.Substring(EndIdx + 1);
+							}
+							else
+							{
+								Path = Rules.Target.ProjectFile.Directory + Path.Substring(EndIdx + 1);
+							}
+						}
+						else if(MatchVariableName(Path, StartIdx, EndIdx, "ModuleDir"))
+						{
+							Path = Rules.ModuleDirectory + Path.Substring(EndIdx + 1);
+						}
+						else if(MatchVariableName(Path, StartIdx, EndIdx, "PluginDir"))
+						{
+							Path = Rules.PluginDirectory + Path.Substring(EndIdx + 1);
+						}
+						else if(Binary != null && MatchVariableName(Path, StartIdx, EndIdx, "BinaryDir"))
+						{
+							Path = Binary.OutputFilePaths[0].Directory.FullName + Path.Substring(EndIdx + 1);
+						}
+						else
+						{
+							string Name = Path.Substring(StartIdx, EndIdx - StartIdx);
+							string Value = Environment.GetEnvironmentVariable(Name);
+							if(String.IsNullOrEmpty(Value))
+							{
+								throw new BuildException("Environment variable '{0}' is not defined (referenced by {1})", Name, Rules.File);
+							}
+							Path = Value + Path.Substring(EndIdx + 1);
+						}
+						break;
+					}
+				}
+			}
+			return Path;
+		}
+
+		/// <summary>
+		/// Match a variable name within a path
+		/// </summary>
+		/// <param name="Path">The path variable</param>
+		/// <param name="StartIdx">Start index of the substring to match</param>
+		/// <param name="EndIdx">End index of the substring to match</param>
+		/// <param name="Name">Variable name to compare against</param>
+		/// <returns>True if the variable name matches</returns>
+		private bool MatchVariableName(string Path, int StartIdx, int EndIdx, string Name)
+		{
+			return Name.Length == EndIdx - StartIdx && String.Compare(Path, StartIdx, Name, 0, EndIdx - StartIdx) == 0;
+		}
+
+		/// <summary>
+		/// Expand path variables within the context of this module
+		/// </summary>
+		/// <param name="Paths">Path to expand variables within</param>
+		/// <param name="Binary">The source binary containing the module. May be null.</param>
+		/// <returns>The path with variables expanded</returns>
+		private IEnumerable<string> ExpandPathVariables(IEnumerable<string> Paths, UEBuildBinary Binary)
+		{
+			foreach(string Path in Paths)
+			{
+				yield return ExpandPathVariables(Path, Binary);
+			}
+		}
+
 		/// <summary>
 		/// Sets up the environment for linking any module that includes the public interface of this module.
 		/// </summary>
