Fix exceptions in log messages having trailing newlines.

#rnx

diff --git a/Engine/Source/Programs/UnrealBuildTool/Configuration/UEBuildTarget.cs b/Engine/Source/Programs/UnrealBuildTool/Configuration/UEBuildTarget.cs
--- a/Engine/Source/Programs/UnrealBuildTool/Configuration/UEBuildTarget.cs
+++ b/Engine/Source/Programs/UnrealBuildTool/Configuration/UEBuildTarget.cs
@@ -1186,7 +1186,7 @@ namespace UnrealBuildTool
 					}
 					catch (Exception Ex)
 					{
-						throw new BuildException(Ex, "Unable to delete {0} ({1})", DirectoryToDelete, Ex.Message);
+						throw new BuildException(Ex, "Unable to delete {0} ({1})", DirectoryToDelete, Ex.Message.TrimEnd());
 					}
 				}
 			}
@@ -1202,7 +1202,7 @@ namespace UnrealBuildTool
 					}
 					catch (Exception Ex)
 					{
-						throw new BuildException(Ex, "Unable to delete {0} ({1})", FileToDelete, Ex.Message);
+						throw new BuildException(Ex, "Unable to delete {0} ({1})", FileToDelete, Ex.Message.TrimEnd());
 					}
 				}
 			}
@@ -2260,7 +2260,7 @@ namespace UnrealBuildTool
 							}
 							catch (Exception Ex)
 							{
-								Log.TraceError("Unable to delete {0} ({1})", ExistingFile, Ex.Message);
+								Log.TraceError("Unable to delete {0} ({1})", ExistingFile, Ex.Message.TrimEnd());
 							}
 						}
 					}
diff --git a/Engine/Source/Programs/UnrealBuildTool/System/ActionGraph.cs b/Engine/Source/Programs/UnrealBuildTool/System/ActionGraph.cs
--- a/Engine/Source/Programs/UnrealBuildTool/System/ActionGraph.cs
+++ b/Engine/Source/Programs/UnrealBuildTool/System/ActionGraph.cs
@@ -629,7 +629,7 @@ namespace UnrealBuildTool
 										catch (Exception Ex)
 										{
 											// Ignore all exceptions
-											Log.TraceInformation("Unable to delete old hot reload file: \"{0}\". Error: {1}", JunkFile, Ex.Message);
+											Log.TraceInformation("Unable to delete old hot reload file: \"{0}\". Error: {1}", JunkFile, Ex.Message.TrimEnd());
 										}
 
 										// Delete the PDB file.
@@ -644,7 +644,7 @@ namespace UnrealBuildTool
 											catch (Exception Ex)
 											{
 												// Ignore all exceptions
-												Log.TraceInformation("Unable to delete old hot reload file: \"{0}\". Error: {1}", JunkPDBFile, Ex.Message);
+												Log.TraceInformation("Unable to delete old hot reload file: \"{0}\". Error: {1}", JunkPDBFile, Ex.Message.TrimEnd());
 											}
 										}
 									}
diff --git a/Engine/Source/Programs/UnrealBuildTool/System/JunkDeleter.cs b/Engine/Source/Programs/UnrealBuildTool/System/JunkDeleter.cs
--- a/Engine/Source/Programs/UnrealBuildTool/System/JunkDeleter.cs
+++ b/Engine/Source/Programs/UnrealBuildTool/System/JunkDeleter.cs
@@ -182,7 +182,7 @@ namespace UnrealBuildTool
 			catch (Exception Ex)
 			{
 				// Ignore all exceptions
-				Log.TraceInformation("Unable to delete junk directory: \"{0}\". Error: {1}", DirectoryPath, Ex.Message);
+				Log.TraceInformation("Unable to delete junk directory: \"{0}\". Error: {1}", DirectoryPath, Ex.Message.TrimEnd());
 			}
 		}
 
@@ -200,7 +200,7 @@ namespace UnrealBuildTool
 			catch (Exception Ex)
 			{
 				// Ingore all exceptions
-				Log.TraceInformation("Unable to delete junk file: \"{0}\". Error: {1}", Filename, Ex.Message);
+				Log.TraceInformation("Unable to delete junk file: \"{0}\". Error: {1}", Filename, Ex.Message.TrimEnd());
 			}
 		}
 	}
