UBT: Always force include a PCH for generated code if there's one set; the code may depend on it to compile.

#rnx

diff --git a/Engine/Source/Programs/UnrealBuildTool/Configuration/UEBuildModuleCPP.cs b/Engine/Source/Programs/UnrealBuildTool/Configuration/UEBuildModuleCPP.cs
--- a/Engine/Source/Programs/UnrealBuildTool/Configuration/UEBuildModuleCPP.cs
+++ b/Engine/Source/Programs/UnrealBuildTool/Configuration/UEBuildModuleCPP.cs
@@ -519,6 +519,13 @@ namespace UnrealBuildTool
 						GeneratedCPPCompileEnvironment.bCreateDebugInfo = false;
 					}
 
+					// Always force include the PCH, even if PCHs are disabled, for generated code. Legacy code can rely on PCHs being included to compile correctly, and this used to be done by UHT manually including it.
+					if(GeneratedCPPCompileEnvironment.PrecompiledHeaderFile == null && Rules.PrivatePCHHeaderFile != null && Rules.PCHUsage != ModuleRules.PCHUsageMode.UseExplicitOrSharedPCHs)
+					{
+						GeneratedCPPCompileEnvironment = new CppCompileEnvironment(GeneratedCPPCompileEnvironment);
+						GeneratedCPPCompileEnvironment.ForceIncludeFiles.Add(FileItem.GetExistingItemByFileReference(FileReference.Combine(ModuleDirectory, Rules.PrivatePCHHeaderFile)));
+					}
+
 					// Compile all the generated files
 					List<FileItem> GeneratedFileItems = new List<FileItem>();
 					foreach (string GeneratedFilename in GeneratedFiles)
diff --git a/Engine/Source/Programs/UnrealBuildTool/Configuration/UEBuildTarget.cs b/Engine/Source/Programs/UnrealBuildTool/Configuration/UEBuildTarget.cs
--- a/Engine/Source/Programs/UnrealBuildTool/Configuration/UEBuildTarget.cs
+++ b/Engine/Source/Programs/UnrealBuildTool/Configuration/UEBuildTarget.cs
@@ -3852,7 +3852,7 @@ namespace UnrealBuildTool
 				// If we can't use a shared PCH, check there's a private PCH set
 				if(RulesObject.PCHUsage == ModuleRules.PCHUsageMode.NoSharedPCHs && RulesObject.PrivatePCHHeaderFile == null)
 				{
-					Log.TraceWarning(RulesObject.File, "Modules should specify which file to use as a precompiled header (eg. PrivatePCHHeaderFile = \"Private/{0}PrivatePCH.h\") or enable shared PCHs (eg. PCHUsage = PCHUsageMode.UseSharedPCHs). '{0}' will be compiled using a shared PCH instead.", RulesObject.Name);
+					Log.TraceWarning(RulesObject.File, "As of UE 4.21, modules should specify which file to use as a precompiled header (eg. PrivatePCHHeaderFile = \"Private/{0}PrivatePCH.h\") or enable shared PCHs (eg. PCHUsage = PCHUsageMode.UseSharedPCHs). '{0}' will be compiled using a shared PCH instead.", RulesObject.Name);
 					RulesObject.PCHUsage = ModuleRules.PCHUsageMode.UseSharedPCHs;
 				}
 			}
