Enable PCHs on IOS. Reduces build time by ~25%.

diff --git a/Engine/Source/Programs/UnrealBuildTool/Platform/IOS/IOSToolChain.cs b/Engine/Source/Programs/UnrealBuildTool/Platform/IOS/IOSToolChain.cs
--- a/Engine/Source/Programs/UnrealBuildTool/Platform/IOS/IOSToolChain.cs
+++ b/Engine/Source/Programs/UnrealBuildTool/Platform/IOS/IOSToolChain.cs
@@ -691,6 +691,7 @@ namespace UnrealBuildTool
 			foreach (FileItem SourceFile in InputFiles)
 			{
 				Action CompileAction = ActionGraph.Add(ActionType.Compile);
+				string FilePCHArguments = "";
 				string FileArguments = "";
 				string Extension = Path.GetExtension(SourceFile.AbsolutePath).ToUpperInvariant();
 
@@ -723,7 +724,7 @@ namespace UnrealBuildTool
 					FileArguments += GetRTTIFlag(CompileEnvironment);
 
 					// only use PCH for .cpp files
-					FileArguments += PCHArguments;
+					FilePCHArguments = PCHArguments;
 				}
 
 				// Add the C++ source file and its included files to the prerequisite item list.
@@ -774,7 +775,7 @@ namespace UnrealBuildTool
 					CompileAction.ActionHandler = new Action.BlockingActionHandler(RPCUtilHelper.RPCActionHandler);
 				}
 
-				string AllArgs = Arguments + FileArguments + CompileEnvironment.AdditionalArguments;
+				string AllArgs = FilePCHArguments + Arguments + FileArguments + CompileEnvironment.AdditionalArguments;
 /*				string SourceText = System.IO.File.ReadAllText(SourceFile.AbsolutePath);
 				if (CompileEnvironment.bOptimizeForSize && (SourceFile.AbsolutePath.Contains("ElementBatcher.cpp") || SourceText.Contains("ElementBatcher.cpp") || SourceFile.AbsolutePath.Contains("AnimationRuntime.cpp") || SourceText.Contains("AnimationRuntime.cpp")
 					|| SourceFile.AbsolutePath.Contains("AnimEncoding.cpp") || SourceText.Contains("AnimEncoding.cpp") || SourceFile.AbsolutePath.Contains("TextRenderComponent.cpp") || SourceText.Contains("TextRenderComponent.cpp")
diff --git a/Engine/Source/Programs/UnrealBuildTool/Platform/IOS/UEBuildIOS.cs b/Engine/Source/Programs/UnrealBuildTool/Platform/IOS/UEBuildIOS.cs
--- a/Engine/Source/Programs/UnrealBuildTool/Platform/IOS/UEBuildIOS.cs
+++ b/Engine/Source/Programs/UnrealBuildTool/Platform/IOS/UEBuildIOS.cs
@@ -553,8 +553,6 @@ namespace UnrealBuildTool
 
 		public override void ValidateTarget(TargetRules Target)
 		{
-			Target.bUsePCHFiles = false;
-
 			// we assume now we are building with IOS8 or later
 			if (Target.bCompileAgainstEngine)
 			{
