UBT: Add better messaging for modules that need to have a private PCH set. Now detects the likely PCH using the same method as legacy code and includes it as a suggestion.

diff --git a/Engine/Source/Programs/DotNETCommon/DotNETUtilities/Log.cs b/Engine/Source/Programs/DotNETCommon/DotNETUtilities/Log.cs
--- a/Engine/Source/Programs/DotNETCommon/DotNETUtilities/Log.cs
+++ b/Engine/Source/Programs/DotNETCommon/DotNETUtilities/Log.cs
@@ -730,6 +730,31 @@ namespace Tools.DotNETCommon
 			WriteLinePrivate(1, true, LogEventType.Warning, LogFormatOptions.None, Format, Args);
 		}
 
+		/// <summary>
+		/// Writes a warning message to the console.
+		/// </summary>
+		/// <param name="File">The file containing the error</param>
+		/// <param name="Format">Message format string</param>
+		/// <param name="Args">Optional arguments</param>
+		[MethodImplAttribute(MethodImplOptions.NoInlining)]
+		public static void TraceWarningOnce(FileReference File, string Format, params object[] Args)
+		{
+			WriteLinePrivate(1, true, LogEventType.Warning, LogFormatOptions.NoSeverityPrefix, "{0}: warning: {1}", File, String.Format(Format, Args));
+		}
+
+		/// <summary>
+		/// Writes a warning message to the console.
+		/// </summary>
+		/// <param name="File">The file containing the error</param>
+		/// <param name="Line">Line number of the error</param>
+		/// <param name="Format">Message format string</param>
+		/// <param name="Args">Optional arguments</param>
+		[MethodImplAttribute(MethodImplOptions.NoInlining)]
+		public static void TraceWarningOnce(FileReference File, int Line, string Format, params object[] Args)
+		{
+			WriteLinePrivate(1, true, LogEventType.Warning, LogFormatOptions.NoSeverityPrefix, "{0}({1}): warning: {2}", File, Line, String.Format(Format, Args));
+		}
+
 		/// <summary>
 		/// Writes a very verbose message to the console.
 		/// </summary>
diff --git a/Engine/Source/Programs/UnrealBuildTool/Configuration/UEBuildTarget.cs b/Engine/Source/Programs/UnrealBuildTool/Configuration/UEBuildTarget.cs
--- a/Engine/Source/Programs/UnrealBuildTool/Configuration/UEBuildTarget.cs
+++ b/Engine/Source/Programs/UnrealBuildTool/Configuration/UEBuildTarget.cs
@@ -3850,10 +3850,32 @@ namespace UnrealBuildTool
 				}
 
 				// If we can't use a shared PCH, check there's a private PCH set
-				if(RulesObject.PCHUsage == ModuleRules.PCHUsageMode.NoSharedPCHs && RulesObject.PrivatePCHHeaderFile == null)
+				if(RulesObject.PCHUsage != ModuleRules.PCHUsageMode.UseExplicitOrSharedPCHs && RulesObject.PrivatePCHHeaderFile == null)
 				{
-					Log.TraceWarning(RulesObject.File, "As of UE 4.21, modules should specify which file to use as a precompiled header (eg. PrivatePCHHeaderFile = \"Private/{0}PrivatePCH.h\") or enable shared PCHs (eg. PCHUsage = PCHUsageMode.UseSharedPCHs). '{0}' will be compiled using a shared PCH instead.", RulesObject.Name);
-					RulesObject.PCHUsage = ModuleRules.PCHUsageMode.UseSharedPCHs;
+					// Try to figure out the legacy PCH file
+					FileReference CppFile = DirectoryReference.EnumerateFiles(RulesObject.Directory, "*.cpp", SearchOption.AllDirectories).FirstOrDefault();
+					if(CppFile != null)
+					{
+						List<DependencyInclude> Includes = CPPHeaders.GetUncachedDirectIncludeDependencies(CppFile, ProjectFile);
+						if(Includes.Count > 0)
+						{
+							FileReference PchIncludeFile = DirectoryReference.EnumerateFiles(RulesObject.Directory, Path.GetFileName(Includes[0].IncludeName), SearchOption.AllDirectories).FirstOrDefault();
+							if(PchIncludeFile != null)
+							{
+								RulesObject.PrivatePCHHeaderFile = PchIncludeFile.MakeRelativeTo(RulesObject.Directory).Replace(Path.DirectorySeparatorChar, '/');
+							}
+						}
+					}
+
+					// Print a suggestion for which file to include
+					if(RulesObject.PrivatePCHHeaderFile == null)
+					{
+						Log.TraceWarningOnce(RulesObject.File, "Modules must specify an explicit precompiled header (eg. PrivatePCHHeaderFile = \"Private/{0}PrivatePCH.h\") from UE 4.21 onwards.", ModuleName);
+					}
+					else
+					{
+						Log.TraceWarningOnce(RulesObject.File, "Modules must specify an explicit precompiled header (eg. PrivatePCHHeaderFile = \"{0}\") from UE 4.21 onwards.", RulesObject.PrivatePCHHeaderFile);
+					}
 				}
 			}
 			return RulesObject;
diff --git a/Engine/Source/Programs/UnrealBuildTool/ProjectFiles/VisualStudio/VCProject.cs b/Engine/Source/Programs/UnrealBuildTool/ProjectFiles/VisualStudio/VCProject.cs
--- a/Engine/Source/Programs/UnrealBuildTool/ProjectFiles/VisualStudio/VCProject.cs
+++ b/Engine/Source/Programs/UnrealBuildTool/ProjectFiles/VisualStudio/VCProject.cs
@@ -752,8 +752,7 @@ namespace UnrealBuildTool
 					if (bWritePerFilePCHInfo && VCFileType == "ClCompile")
 					{
 						FileReference TruePath = FileReference.Combine(ProjectFilePath.Directory, AliasedFile.FileSystemPath);
-						FileItem SourceFile = FileItem.GetItemByFileReference(TruePath);
-						List <DependencyInclude> DirectlyIncludedFilenames = CPPHeaders.GetUncachedDirectIncludeDependencies(SourceFile, null);
+						List <DependencyInclude> DirectlyIncludedFilenames = CPPHeaders.GetUncachedDirectIncludeDependencies(TruePath, null);
 						if (DirectlyIncludedFilenames.Count > 0)
 						{
 							PCHFileName = DirectlyIncludedFilenames[0].IncludeName;
diff --git a/Engine/Source/Programs/UnrealBuildTool/System/CPPHeaders.cs b/Engine/Source/Programs/UnrealBuildTool/System/CPPHeaders.cs
--- a/Engine/Source/Programs/UnrealBuildTool/System/CPPHeaders.cs
+++ b/Engine/Source/Programs/UnrealBuildTool/System/CPPHeaders.cs
@@ -490,7 +490,7 @@ namespace UnrealBuildTool
 			DateTime TimerStartTime = DateTime.UtcNow;
 			++CPPHeaders.TotalDirectIncludeCacheMisses;
 
-			Result = GetUncachedDirectIncludeDependencies(CPPFile, ProjectFile);
+			Result = GetUncachedDirectIncludeDependencies(CPPFile.Location, ProjectFile);
 
 			// Populate cache with results.
 			IncludeDependencyCache.SetDependencyInfo(CPPFile, Result);
@@ -500,12 +500,12 @@ namespace UnrealBuildTool
 			return Result;
 		}
 
-		public static List<DependencyInclude> GetUncachedDirectIncludeDependencies(FileItem CPPFile, FileReference ProjectFile)
+		public static List<DependencyInclude> GetUncachedDirectIncludeDependencies(FileReference SourceFile, FileReference ProjectFile)
 		{
 			List<DependencyInclude> Result = new List<DependencyInclude>();
 
 			// Get the adjusted filename
-			string FileToRead = CPPFile.AbsolutePath;
+			string FileToRead = SourceFile.FullName;
 
 			// Read lines from the C++ file.
 			string FileContents = GetFileContents(FileToRead);
@@ -536,30 +536,30 @@ namespace UnrealBuildTool
 						EndIndex = FileContents.Length;
 					}
 
-					Result.AddRange(CollectHeaders(ProjectFile, CPPFile, FileToRead, FileContents, InstalledFolder, StartIndex, EndIndex));
+					Result.AddRange(CollectHeaders(ProjectFile, SourceFile, FileToRead, FileContents, InstalledFolder, StartIndex, EndIndex));
 
 					StartIndex = EndIndex + 1;
 				}
 			}
 			else
 			{
-				Result = CollectHeaders(ProjectFile, CPPFile, FileToRead, FileContents, InstalledFolder, 0, FileContents.Length);
+				Result = CollectHeaders(ProjectFile, SourceFile, FileToRead, FileContents, InstalledFolder, 0, FileContents.Length);
 			}
 
 			return Result;
 		}
 
 		/// <summary>
-		/// Collects all header files included in a CPPFile
+		/// Collects all header files included in a source file
 		/// </summary>
 		/// <param name="ProjectFile"></param>
-		/// <param name="CPPFile"></param>
+		/// <param name="SourceFile"></param>
 		/// <param name="FileToRead"></param>
 		/// <param name="FileContents"></param>
 		/// <param name="InstalledFolder"></param>
 		/// <param name="StartIndex"></param>
 		/// <param name="EndIndex"></param>
-		private static List<DependencyInclude> CollectHeaders(FileReference ProjectFile, FileItem CPPFile, string FileToRead, string FileContents, string InstalledFolder, int StartIndex, int EndIndex)
+		private static List<DependencyInclude> CollectHeaders(FileReference ProjectFile, FileReference SourceFile, string FileToRead, string FileContents, string InstalledFolder, int StartIndex, int EndIndex)
 		{
 			List<DependencyInclude> Result = new List<DependencyInclude>();
 
@@ -586,8 +586,7 @@ namespace UnrealBuildTool
 			}
 
 			// also look for #import in objective C files
-			string Ext = Path.GetExtension(CPPFile.AbsolutePath).ToUpperInvariant();
-			if (Ext == ".MM" || Ext == ".M")
+			if (SourceFile.HasExtension(".MM") || SourceFile.HasExtension(".M"))
 			{
 				M = MMHeaderRegex.Match(FileContents, StartIndex, EndIndex - StartIndex);
 				Captures = M.Groups["HeaderFile"].Captures;
