UBT: Simplify handling of installed/precompiled builds. Settings for whether a folder is installed/read-only or not is now stored on the RulesAssembly instance, allowing multiple things to be configured separately and stacked together (eg. engine/enterprise/project). RulesAssembly.IsReadOnly() allows determining if a flie can be modified or not and replaces many previous IsXXXInstalledCalls(), and traverses the chain of assemblies.

diff --git a/Engine/Source/Programs/UnrealBuildTool/Configuration/BuildConfiguration.cs b/Engine/Source/Programs/UnrealBuildTool/Configuration/BuildConfiguration.cs
--- a/Engine/Source/Programs/UnrealBuildTool/Configuration/BuildConfiguration.cs
+++ b/Engine/Source/Programs/UnrealBuildTool/Configuration/BuildConfiguration.cs
@@ -27,6 +27,12 @@ namespace UnrealBuildTool
 		[CommandLine("-SkipActionHistory", Value = "false")]
 		public bool bUseActionHistory = true;
 
+		/// <summary>
+		/// Use existing static libraries for all engine modules in this target.
+		/// </summary>
+		[CommandLine("-UsePrecompiled")]
+		public bool bUsePrecompiled = false;
+
 		/// <summary>
 		/// Whether debug info should be written to the console.
 		/// </summary>
diff --git a/Engine/Source/Programs/UnrealBuildTool/Configuration/ModuleRules.cs b/Engine/Source/Programs/UnrealBuildTool/Configuration/ModuleRules.cs
--- a/Engine/Source/Programs/UnrealBuildTool/Configuration/ModuleRules.cs
+++ b/Engine/Source/Programs/UnrealBuildTool/Configuration/ModuleRules.cs
@@ -619,7 +619,7 @@ namespace UnrealBuildTool
 		/// <param name="ModuleNames">The names of the modules to add</param>
 		public void AddEngineThirdPartyPrivateStaticDependencies(ReadOnlyTargetRules Target, params string[] ModuleNames)
 		{
-			if (!UnrealBuildTool.IsEngineInstalled() || Target.LinkType == TargetLinkType.Monolithic)
+			if (!bUsePrecompiled || Target.LinkType == TargetLinkType.Monolithic)
 			{
 				PrivateDependencyModuleNames.AddRange(ModuleNames);
 			}
@@ -635,7 +635,7 @@ namespace UnrealBuildTool
 		/// <param name="ModuleNames">The names of the modules to add</param>
 		public void AddEngineThirdPartyPrivateDynamicDependencies(ReadOnlyTargetRules Target, params string[] ModuleNames)
 		{
-			if (!UnrealBuildTool.IsEngineInstalled() || Target.LinkType == TargetLinkType.Monolithic)
+			if (!bUsePrecompiled || Target.LinkType == TargetLinkType.Monolithic)
 			{
 				PrivateIncludePathModuleNames.AddRange(ModuleNames);
 				DynamicallyLoadedModuleNames.AddRange(ModuleNames);
diff --git a/Engine/Source/Programs/UnrealBuildTool/Configuration/TargetDescriptor.cs b/Engine/Source/Programs/UnrealBuildTool/Configuration/TargetDescriptor.cs
--- a/Engine/Source/Programs/UnrealBuildTool/Configuration/TargetDescriptor.cs
+++ b/Engine/Source/Programs/UnrealBuildTool/Configuration/TargetDescriptor.cs
@@ -44,9 +44,10 @@ namespace UnrealBuildTool
 		/// Parse a list of target descriptors from the command line
 		/// </summary>
 		/// <param name="Arguments">Command-line arguments</param>
+		/// <param name="bUsePrecompiled">Whether to use a precompiled engine distribution</param>
 		/// <param name="ProjectFile">The project file, if already set. May be updated if not.</param>
 		/// <returns>List of target descriptors</returns>
-		public static List<TargetDescriptor> ParseCommandLine(string[] Arguments, ref FileReference ProjectFile)
+		public static List<TargetDescriptor> ParseCommandLine(string[] Arguments, bool bUsePrecompiled, ref FileReference ProjectFile)
 		{
 			UnrealTargetPlatform Platform = UnrealTargetPlatform.Unknown;
 			UnrealTargetConfiguration Configuration = UnrealTargetConfiguration.Unknown;
@@ -168,7 +169,7 @@ namespace UnrealBuildTool
 				}
 				else
 				{
-					TargetNames.Add(RulesCompiler.CreateProjectRulesAssembly(ProjectFile).GetTargetNameByType(Type, Platform, Configuration, Architecture, ProjectFile, new ReadOnlyBuildVersion(BuildVersion.ReadDefault())));
+					TargetNames.Add(RulesCompiler.CreateProjectRulesAssembly(ProjectFile, bUsePrecompiled).GetTargetNameByType(Type, Platform, Configuration, Architecture, ProjectFile, new ReadOnlyBuildVersion(BuildVersion.ReadDefault())));
 				}
 			}
 
diff --git a/Engine/Source/Programs/UnrealBuildTool/Configuration/TargetRules.cs b/Engine/Source/Programs/UnrealBuildTool/Configuration/TargetRules.cs
--- a/Engine/Source/Programs/UnrealBuildTool/Configuration/TargetRules.cs
+++ b/Engine/Source/Programs/UnrealBuildTool/Configuration/TargetRules.cs
@@ -570,12 +570,6 @@ namespace UnrealBuildTool
 		[CommandLine("-Precompile")]
 		public bool bPrecompile = false;
 
-		/// <summary>
-		/// Use existing static libraries for all engine modules in this target.
-		/// </summary>
-		[CommandLine("-UsePrecompiled")]
-		public bool bUsePrecompiled = false;
-
 		/// <summary>
 		/// Whether we should compile with support for OS X 10.9 Mavericks. Used for some tools that we need to be compatible with this version of OS X.
 		/// </summary>
@@ -1111,9 +1105,6 @@ namespace UnrealBuildTool
 				UEBuildPlatform.GetBuildPlatform(Platform).ResetTarget(this);
 			}
 
-			// If the engine is installed, always use precompiled libraries
-			bUsePrecompiled = UnrealBuildTool.IsEngineInstalled();
-
 			// Check that the appropriate headers exist to enable Simplygon
 			if(bCompileSimplygon)
 			{
@@ -1690,11 +1681,6 @@ namespace UnrealBuildTool
 			get { return Inner.bPrecompile; }
 		}
 
-		public bool bUsePrecompiled
-		{
-			get { return Inner.bUsePrecompiled; }
-		}
-
 		public bool bEnableOSX109Support
 		{
 			get { return Inner.bEnableOSX109Support; }
diff --git a/Engine/Source/Programs/UnrealBuildTool/Configuration/UEBuildModuleCPP.cs b/Engine/Source/Programs/UnrealBuildTool/Configuration/UEBuildModuleCPP.cs
--- a/Engine/Source/Programs/UnrealBuildTool/Configuration/UEBuildModuleCPP.cs
+++ b/Engine/Source/Programs/UnrealBuildTool/Configuration/UEBuildModuleCPP.cs
@@ -602,7 +602,7 @@ namespace UnrealBuildTool
 			HeaderFile.CachedIncludePaths = CompileEnvironment.IncludePaths;
 
 			DirectoryReference PrecompiledHeaderDir;
-			if(Target.Rules.bUsePrecompiled)
+			if(Rules.bUsePrecompiled)
 			{
 				PrecompiledHeaderDir = DirectoryReference.Combine(Target.ProjectIntermediateDirectory, Name);
 			}
diff --git a/Engine/Source/Programs/UnrealBuildTool/Configuration/UEBuildTarget.cs b/Engine/Source/Programs/UnrealBuildTool/Configuration/UEBuildTarget.cs
--- a/Engine/Source/Programs/UnrealBuildTool/Configuration/UEBuildTarget.cs
+++ b/Engine/Source/Programs/UnrealBuildTool/Configuration/UEBuildTarget.cs
@@ -380,25 +380,26 @@ namespace UnrealBuildTool
 		/// <param name="Desc">Information about the target</param>
 		/// <param name="Arguments">Command line arguments</param>
 		/// <param name="bCompilingSingleFile">Whether we're compiling a single file</param>
+		/// <param name="bUsePrecompiled">Whether to use a precompiled engine/enterprise build</param>
 		/// <param name="Version">The current build version</param>
 		/// <returns>The build target object for the specified build rules source file</returns>
-		public static UEBuildTarget CreateTarget(TargetDescriptor Desc, string[] Arguments, bool bCompilingSingleFile, ReadOnlyBuildVersion Version)
+		public static UEBuildTarget CreateTarget(TargetDescriptor Desc, string[] Arguments, bool bCompilingSingleFile, bool bUsePrecompiled, ReadOnlyBuildVersion Version)
 		{
 			DateTime CreateTargetStartTime = DateTime.UtcNow;
 
 			RulesAssembly RulesAssembly;
 			if (Desc.ProjectFile != null)
 			{
-				RulesAssembly = RulesCompiler.CreateProjectRulesAssembly(Desc.ProjectFile);
+				RulesAssembly = RulesCompiler.CreateProjectRulesAssembly(Desc.ProjectFile, bUsePrecompiled);
 			}
 			else
 			{
-				RulesAssembly = RulesCompiler.CreateEngineRulesAssembly();
+				RulesAssembly = RulesCompiler.CreateEngineRulesAssembly(bUsePrecompiled);
 
 				if (RulesAssembly.GetTargetFileName(Desc.Name) == null)
 				{
 					// Target isn't part of the engine assembly, try the enterprise assembly
-					RulesAssembly EnterpriseRulesAssembly = RulesCompiler.CreateEnterpriseRulesAssembly();
+					RulesAssembly EnterpriseRulesAssembly = RulesCompiler.CreateEnterpriseRulesAssembly(bUsePrecompiled);
 
 					if (EnterpriseRulesAssembly != null)
 					{
@@ -718,11 +719,6 @@ namespace UnrealBuildTool
 		/// </summary>
 		public bool bPrecompile;
 
-		/// <summary>
-		/// Whether to use precompiled engine modules
-		/// </summary>
-		public bool bUsePrecompiled;
-
 		/// <summary>
 		/// Identifies whether the project contains a script plugin. This will cause UHT to be rebuilt, even in installed builds.
 		/// </summary>
@@ -871,7 +867,6 @@ namespace UnrealBuildTool
 			OutputPaths = (List<FileReference>)Info.GetValue("op", typeof(List<FileReference>));
 			VersionFile = (FileReference)Info.GetValue("vf", typeof(FileReference));
 			bPrecompile = Info.GetBoolean("pc");
-			bUsePrecompiled = Info.GetBoolean("up");
 			OnlyModules = (List<OnlyModule>)Info.GetValue("om", typeof(List<OnlyModule>));
 			bCompileMonolithic = Info.GetBoolean("cm");
 			FlatModuleCsData = (List<FlatModuleCsDataType>)Info.GetValue("fm", typeof(List<FlatModuleCsDataType>));
@@ -902,7 +897,6 @@ namespace UnrealBuildTool
 			Info.AddValue("op", OutputPaths);
 			Info.AddValue("vf", VersionFile);
 			Info.AddValue("pc", bPrecompile);
-			Info.AddValue("up", bUsePrecompiled);
 			Info.AddValue("om", OnlyModules);
 			Info.AddValue("cm", bCompileMonolithic);
 			Info.AddValue("fm", FlatModuleCsData);
@@ -935,7 +929,6 @@ namespace UnrealBuildTool
 			RulesAssembly = InRulesAssembly;
 			TargetType = Rules.Type;
 			bPrecompile = InRules.bPrecompile;
-			bUsePrecompiled = InRules.bUsePrecompiled;
 			ForeignPlugin = InDesc.ForeignPlugin;
 			ForceReceiptFileName = InDesc.ForceReceiptFileName;
 
@@ -1089,23 +1082,7 @@ namespace UnrealBuildTool
 			}
 
 			// If we're running a precompiled build, remove anything under the engine folder
-			if (bUsePrecompiled)
-			{
-				BaseDirs.RemoveAll(x => x.IsUnderDirectory(UnrealBuildTool.EngineDirectory));
-			}
-
-			// If we're in an installed enterprise build, remove anything under the enterprise folder
-			if (UnrealBuildTool.IsEnterpriseInstalled())
-			{
-				BaseDirs.RemoveAll(x => x.IsUnderDirectory(UnrealBuildTool.EnterpriseDirectory));
-			}
-
-			// If we're in an installed project build, only allow cleaning stuff that's under the mod directories
-			if (UnrealBuildTool.IsProjectInstalled())
-			{
-				List<DirectoryReference> ModDirs = EnabledPlugins.Where(x => x.Type == PluginType.Mod).Select(x => x.Directory).ToList();
-				BaseDirs.RemoveAll(x => !ModDirs.Any(y => x.IsUnderDirectory(y)));
-			}
+			BaseDirs.RemoveAll(x => RulesAssembly.IsReadOnly(x));
 
 			// Get all the names which can prefix build products
 			List<string> NamePrefixes = new List<string>();
@@ -1208,7 +1185,7 @@ namespace UnrealBuildTool
 			}
 
 			// Finally clean UnrealHeaderTool if this target uses CoreUObject modules and we're not cleaning UHT already and we want UHT to be cleaned.
-			if (bIncludeUnrealHeaderTool && !bUsePrecompiled && TargetName != "UnrealHeaderTool")
+			if (bIncludeUnrealHeaderTool && !RulesAssembly.IsReadOnly(ExternalExecution.GetHeaderToolReceiptFile(UnrealTargetConfiguration.Development)) && TargetName != "UnrealHeaderTool")
 			{
 				ExternalExecution.RunExternalDotNETExecutable(UnrealBuildTool.GetUBTPath(), String.Format("UnrealHeaderTool {0} {1} -NoMutex -Clean -IgnoreJunk -NoLog", BuildHostPlatform.Current.Platform, UnrealTargetConfiguration.Development));
 			}
@@ -1709,7 +1686,7 @@ namespace UnrealBuildTool
 				UEBuildPlatform BuildPlatform = UEBuildPlatform.GetBuildPlatform(Platform);
 				if (OnlyModules == null || OnlyModules.Count == 0)
 				{
-					if(!IsFileInstalled(ReceiptFileName))
+					if(!RulesAssembly.IsReadOnly(ReceiptFileName))
 					{
 						DirectoryReference.CreateDirectory(ReceiptFileName.Directory);
 						Receipt.Write(ReceiptFileName, UnrealBuildTool.EngineDirectory, ProjectDirectory);
@@ -1718,7 +1695,7 @@ namespace UnrealBuildTool
 				if (ForceReceiptFileName != null)
 				{
 					FileReference ForceReceiptFile = new FileReference(ForceReceiptFileName);
-					if(!IsFileInstalled(ForceReceiptFile))
+					if(!RulesAssembly.IsReadOnly(ForceReceiptFile))
 					{
 						DirectoryReference.CreateDirectory(ForceReceiptFile.Directory);
 						Receipt.Write(ForceReceiptFile, UnrealBuildTool.EngineDirectory, ProjectDirectory);
@@ -1726,7 +1703,7 @@ namespace UnrealBuildTool
 				}
 				if(VersionFile != null)
 				{
-					if(!IsFileInstalled(VersionFile))
+					if(!RulesAssembly.IsReadOnly(VersionFile))
 					{
 						DirectoryReference.CreateDirectory(VersionFile.Directory);
 
@@ -1745,7 +1722,7 @@ namespace UnrealBuildTool
 			{
 				foreach (KeyValuePair<FileReference, ModuleManifest> FileNameToVersionManifest in FileReferenceToModuleManifestPairs)
 				{
-					if(!IsFileInstalled(FileNameToVersionManifest.Key))
+					if(!RulesAssembly.IsReadOnly(FileNameToVersionManifest.Key))
 					{
 						// Write the manifest out to a string buffer, then only write it to disk if it's changed.
 						string OutputText;
@@ -1764,24 +1741,6 @@ namespace UnrealBuildTool
 			}
 		}
 
-		/// <summary>
-		/// Checks whether the given file is under an installed directory, and should not be overridden
-		/// </summary>
-		/// <param name="File">File to test</param>
-		/// <returns>True if the file is part of the installed distribution, false otherwise</returns>
-		bool IsFileInstalled(FileReference File)
-		{
-			if(UnrealBuildTool.IsEngineInstalled() && File.IsUnderDirectory(UnrealBuildTool.EngineDirectory))
-			{
-				return true;
-			}
-			if(UnrealBuildTool.IsProjectInstalled() && ProjectFile != null && File.IsUnderDirectory(ProjectFile.Directory))
-			{
-				return true;
-			}
-			return false;
-		}
-
 		/// <summary>
 		/// Gathers dependency modules for given binaries list.
 		/// </summary>
@@ -2083,7 +2042,7 @@ namespace UnrealBuildTool
 					NameToFlatModuleData[FlatModuleData.ModuleName] = FlatModuleData;
 				}
 
-				ExternalExecution.SetupUObjectModules(ModulesToGenerateHeadersFor, Rules, ProjectDescriptor, GlobalCompileEnvironment, UObjectModules, NameToFlatModuleData, Rules.GeneratedCodeVersion, bIsAssemblingBuild);
+				ExternalExecution.SetupUObjectModules(ModulesToGenerateHeadersFor, Rules.Platform, ProjectDescriptor, UObjectModules, NameToFlatModuleData, Rules.GeneratedCodeVersion, bIsAssemblingBuild);
 
 				// NOTE: Even in Gather mode, we need to run UHT to make sure the files exist for the static action graph to be setup correctly.  This is because UHT generates .cpp
 				// files that are injected as top level prerequisites.  If UHT only emitted included header files, we wouldn't need to run it during the Gather phase at all.
@@ -4008,7 +3967,7 @@ namespace UnrealBuildTool
 				if (RulesObject.Type == ModuleRules.ModuleType.CPlusPlus)
 				{
 					// So all we care about are the game module and/or plugins.
-					if (bDiscoverFiles && !UnrealBuildTool.IsUnderAnInstalledDirectory(RulesObject.File.Directory))
+					if (bDiscoverFiles && !RulesObject.bUsePrecompiled)
 					{
 						List<FileReference> SourceFilePaths = new List<FileReference>();
 						SourceFilePaths = SourceFileSearch.FindModuleSourceFiles(ModuleRulesFile: RulesObject.File);
diff --git a/Engine/Source/Programs/UnrealBuildTool/Platform/IOS/IOSExports.cs b/Engine/Source/Programs/UnrealBuildTool/Platform/IOS/IOSExports.cs
--- a/Engine/Source/Programs/UnrealBuildTool/Platform/IOS/IOSExports.cs
+++ b/Engine/Source/Programs/UnrealBuildTool/Platform/IOS/IOSExports.cs
@@ -181,11 +181,11 @@ namespace UnrealBuildTool
 			List<TargetDescriptor> TargetDescs = new List<TargetDescriptor>();
 			foreach (string[] TargetSetting in TargetSettings)
 			{
-				TargetDescs.AddRange(TargetDescriptor.ParseCommandLine(TargetSetting, ref ProjectFile));
+				TargetDescs.AddRange(TargetDescriptor.ParseCommandLine(TargetSetting, false, ref ProjectFile));
 			}
 			foreach (TargetDescriptor TargetDesc in TargetDescs)
 			{
-				UEBuildTarget Target = UEBuildTarget.CreateTarget(TargetDesc, Arguments, false, Version);
+				UEBuildTarget Target = UEBuildTarget.CreateTarget(TargetDesc, Arguments, false, false, Version);
 				if (Target == null)
 				{
 					continue;
diff --git a/Engine/Source/Programs/UnrealBuildTool/ProjectFiles/ProjectFileGenerator.cs b/Engine/Source/Programs/UnrealBuildTool/ProjectFiles/ProjectFileGenerator.cs
--- a/Engine/Source/Programs/UnrealBuildTool/ProjectFiles/ProjectFileGenerator.cs
+++ b/Engine/Source/Programs/UnrealBuildTool/ProjectFiles/ProjectFileGenerator.cs
@@ -2093,16 +2093,16 @@ namespace UnrealBuildTool
 					{
 						if(TargetFilePath.IsUnderDirectory(UnrealBuildTool.EnterpriseDirectory))
 						{
-							RulesAssembly = RulesCompiler.CreateEnterpriseRulesAssembly();
+							RulesAssembly = RulesCompiler.CreateEnterpriseRulesAssembly(false);
 						}
 						else
 						{
-							RulesAssembly = RulesCompiler.CreateEngineRulesAssembly();
+							RulesAssembly = RulesCompiler.CreateEngineRulesAssembly(false);
 						}
 					}
 					else
 					{
-						RulesAssembly = RulesCompiler.CreateProjectRulesAssembly(CheckProjectFile);
+						RulesAssembly = RulesCompiler.CreateProjectRulesAssembly(CheckProjectFile, false);
 					}
 
 					// Create target rules for all of the platforms and configuration combinations that we want to enable support for.
diff --git a/Engine/Source/Programs/UnrealBuildTool/System/ExternalExecution.cs b/Engine/Source/Programs/UnrealBuildTool/System/ExternalExecution.cs
--- a/Engine/Source/Programs/UnrealBuildTool/System/ExternalExecution.cs
+++ b/Engine/Source/Programs/UnrealBuildTool/System/ExternalExecution.cs
@@ -191,8 +191,22 @@ namespace UnrealBuildTool
 		/// </summary>
 		public EGeneratedCodeVersion GeneratedCodeVersion;
 
-		public UHTModuleInfo()
+		/// <summary>
+		/// Whether this module is read-only
+		/// </summary>
+		public bool bIsReadOnly;
+
+		public UHTModuleInfo(string ModuleName, FileReference ModuleRulesFile, DirectoryReference ModuleDirectory, UHTModuleType ModuleType, List<FileItem> PublicUObjectClassesHeaders, List<FileItem> PublicUObjectHeaders, List<FileItem> PrivateUObjectHeaders, EGeneratedCodeVersion GeneratedCodeVersion, bool bIsReadOnly)
 		{
+			this.ModuleName = ModuleName;
+			this.ModuleRulesFile = ModuleRulesFile;
+			this.ModuleDirectory = ModuleDirectory;
+			this.ModuleType = ModuleType.ToString();
+			this.PublicUObjectClassesHeaders = PublicUObjectClassesHeaders;
+			this.PublicUObjectHeaders = PublicUObjectHeaders;
+			this.PrivateUObjectHeaders = PrivateUObjectHeaders;
+			this.GeneratedCodeVersion = GeneratedCodeVersion;
+			this.bIsReadOnly = bIsReadOnly;
 		}
 
 		public UHTModuleInfo(SerializationInfo Info, StreamingContext Context)
@@ -206,6 +220,7 @@ namespace UnrealBuildTool
 			PrivateUObjectHeaders = (List<FileItem>)Info.GetValue("pr", typeof(List<FileItem>));
 			GeneratedCPPFilenameBase = Info.GetString("ge");
 			GeneratedCodeVersion = (EGeneratedCodeVersion)Info.GetInt32("gv");
+			bIsReadOnly = Info.GetBoolean("ro");
 		}
 
 		public void GetObjectData(SerializationInfo Info, StreamingContext Context)
@@ -219,6 +234,7 @@ namespace UnrealBuildTool
 			Info.AddValue("pr", PrivateUObjectHeaders);
 			Info.AddValue("ge", GeneratedCPPFilenameBase);
 			Info.AddValue("gv", (int)GeneratedCodeVersion);
+			Info.AddValue("ro", bIsReadOnly);
 		}
 
 		public override string ToString()
@@ -256,7 +272,7 @@ namespace UnrealBuildTool
 
 	struct UHTManifest
 	{
-		public struct Module
+		public class Module
 		{
 			public string Name;
 			public string ModuleType;
@@ -270,42 +286,43 @@ namespace UnrealBuildTool
 			public bool SaveExportedHeaders;
 			public EGeneratedCodeVersion UHTGeneratedCodeVersion;
 
+			public Module(UHTModuleInfo Info)
+			{
+				Name = Info.ModuleName;
+				ModuleType = Info.ModuleType;
+				BaseDirectory = Info.ModuleDirectory.FullName;
+				IncludeBase = Info.ModuleDirectory.ParentDirectory.FullName;
+				OutputDirectory = Path.GetDirectoryName(Info.GeneratedCPPFilenameBase);
+				ClassesHeaders = Info.PublicUObjectClassesHeaders.Select((Header) => Header.AbsolutePath).ToList();
+				PublicHeaders = Info.PublicUObjectHeaders.Select((Header) => Header.AbsolutePath).ToList();
+				PrivateHeaders = Info.PrivateUObjectHeaders.Select((Header) => Header.AbsolutePath).ToList();
+				GeneratedCPPFilenameBase = Info.GeneratedCPPFilenameBase;
+				SaveExportedHeaders = !Info.bIsReadOnly;
+				UHTGeneratedCodeVersion = Info.GeneratedCodeVersion;
+			}
+
 			public override string ToString()
 			{
 				return Name;
 			}
 		}
 
-		public UHTManifest(UEBuildTarget Target, string InRootLocalPath, string InRootBuildPath, string InExternalDependenciesFile, IEnumerable<UHTModuleInfo> ModuleInfo, bool bUsePrecompiled)
+		public bool IsGameTarget;     // True if the current target is a game target
+		public string RootLocalPath;    // The engine path on the local machine
+		public string RootBuildPath;    // The engine path on the build machine, if different (e.g. Mac/iOS builds)
+		public string TargetName;       // Name of the target currently being compiled
+		public string ExternalDependenciesFile; // File to contain additional dependencies that the generated code depends on
+		public List<Module> Modules;
+
+		public UHTManifest(UEBuildTarget Target, string InRootLocalPath, string InRootBuildPath, string InExternalDependenciesFile, List<Module> InModules)
 		{
 			IsGameTarget = (Target.TargetType != TargetType.Program);
 			RootLocalPath = InRootLocalPath;
 			RootBuildPath = InRootBuildPath;
 			TargetName = Target.GetTargetName();
 			ExternalDependenciesFile = InExternalDependenciesFile;
-
-			Modules = ModuleInfo.Select(Info => new Module
-			{
-				Name = Info.ModuleName,
-				ModuleType = Info.ModuleType,
-				BaseDirectory = Info.ModuleDirectory.FullName,
-				IncludeBase = Info.ModuleDirectory.ParentDirectory.FullName,
-				OutputDirectory = Path.GetDirectoryName(Info.GeneratedCPPFilenameBase),
-				ClassesHeaders = Info.PublicUObjectClassesHeaders.Select((Header) => Header.AbsolutePath).ToList(),
-				PublicHeaders = Info.PublicUObjectHeaders.Select((Header) => Header.AbsolutePath).ToList(),
-				PrivateHeaders = Info.PrivateUObjectHeaders.Select((Header) => Header.AbsolutePath).ToList(),
-				GeneratedCPPFilenameBase = Info.GeneratedCPPFilenameBase,
-				SaveExportedHeaders = !bUsePrecompiled || !Info.ModuleDirectory.IsUnderDirectory(UnrealBuildTool.EngineDirectory),
-				UHTGeneratedCodeVersion = Info.GeneratedCodeVersion,
-			}).ToList();
+			Modules = InModules;
 		}
-
-		public bool IsGameTarget;     // True if the current target is a game target
-		public string RootLocalPath;    // The engine path on the local machine
-		public string RootBuildPath;    // The engine path on the build machine, if different (e.g. Mac/iOS builds)
-		public string TargetName;       // Name of the target currently being compiled
-		public string ExternalDependenciesFile; // File to contain additional dependencies that the generated code depends on
-		public List<Module> Modules;
 	}
 
 
@@ -318,12 +335,12 @@ namespace UnrealBuildTool
 		/// Generates a UHTModuleInfo for a particular named module under a directory.
 		/// </summary>
 		/// <returns></returns>
-		public static UHTModuleInfo CreateUHTModuleInfo(IEnumerable<FileReference> HeaderFilenames, string ModuleName, FileReference ModuleRulesFile, DirectoryReference ModuleDirectory, UHTModuleType ModuleType, EGeneratedCodeVersion GeneratedCodeVersion)
+		public static UHTModuleInfo CreateUHTModuleInfo(IEnumerable<FileReference> HeaderFilenames, string ModuleName, FileReference ModuleRulesFile, DirectoryReference ModuleDirectory, UHTModuleType ModuleType, EGeneratedCodeVersion GeneratedCodeVersion, bool bIsReadOnly)
 		{
 			DirectoryReference ClassesFolder = DirectoryReference.Combine(ModuleDirectory, "Classes");
 			DirectoryReference PublicFolder = DirectoryReference.Combine(ModuleDirectory, "Public");
 
-			List<FileItem> AllClassesHeaders = new List<FileItem>();
+			List<FileItem> PublicClassesUObjectHeaders = new List<FileItem>();
 			List<FileItem> PublicUObjectHeaders = new List<FileItem>();
 			List<FileItem> PrivateUObjectHeaders = new List<FileItem>();
 
@@ -337,7 +354,7 @@ namespace UnrealBuildTool
 				{
 					if (UObjectHeaderFileItem.Location.IsUnderDirectory(ClassesFolder))
 					{
-						AllClassesHeaders.Add(UObjectHeaderFileItem);
+						PublicClassesUObjectHeaders.Add(UObjectHeaderFileItem);
 					}
 					else if (UObjectHeaderFileItem.Location.IsUnderDirectory(PublicFolder))
 					{
@@ -350,19 +367,7 @@ namespace UnrealBuildTool
 				}
 			}
 
-			UHTModuleInfo Result = new UHTModuleInfo
-			{
-				ModuleName = ModuleName,
-				ModuleRulesFile = ModuleRulesFile,
-				ModuleDirectory = ModuleDirectory,
-				ModuleType = ModuleType.ToString(),
-				PublicUObjectClassesHeaders = AllClassesHeaders,
-				PublicUObjectHeaders = PublicUObjectHeaders,
-				PrivateUObjectHeaders = PrivateUObjectHeaders,
-				GeneratedCodeVersion = GeneratedCodeVersion
-			};
-
-			return Result;
+			return new UHTModuleInfo(ModuleName, ModuleRulesFile, ModuleDirectory, ModuleType, PublicClassesUObjectHeaders, PublicUObjectHeaders, PrivateUObjectHeaders, GeneratedCodeVersion, bIsReadOnly);
 		}
 
 		static ExternalExecution()
@@ -584,7 +589,7 @@ namespace UnrealBuildTool
 			}
 		}
 
-		public static void SetupUObjectModules(IEnumerable<UEBuildModuleCPP> ModulesToGenerateHeadersFor, ReadOnlyTargetRules Target, ProjectDescriptor ProjectDescriptor, CppCompileEnvironment GlobalCompileEnvironment, List<UHTModuleInfo> UObjectModules, Dictionary<string, FlatModuleCsDataType> FlatModuleCsData, EGeneratedCodeVersion GeneratedCodeVersion, bool bIsAssemblingBuild)
+		public static void SetupUObjectModules(IEnumerable<UEBuildModuleCPP> ModulesToGenerateHeadersFor, UnrealTargetPlatform Platform, ProjectDescriptor ProjectDescriptor, List<UHTModuleInfo> UObjectModules, Dictionary<string, FlatModuleCsDataType> FlatModuleCsData, EGeneratedCodeVersion GeneratedCodeVersion, bool bIsAssemblingBuild)
 		{
 			DateTime UObjectDiscoveryStartTime = DateTime.UtcNow;
 
@@ -599,13 +604,13 @@ namespace UnrealBuildTool
 			List<UEBuildModuleCPP> ModulesSortedByType = ModulesToGenerateHeadersFor.OrderBy(c => ModuleToType[c]).ToList();
 			StableTopologicalSort(ModulesSortedByType);
 
-			string[] ExcludedFolders = UEBuildPlatform.GetBuildPlatform(Target.Platform, true).GetExcludedFolderNames();
+			string[] ExcludedFolders = UEBuildPlatform.GetBuildPlatform(Platform, true).GetExcludedFolderNames();
 			foreach (UEBuildModuleCPP Module in ModulesSortedByType)
 			{
 				List<FileReference> HeaderFiles = new List<FileReference>();
 				FindHeaders(new DirectoryInfo(Module.ModuleDirectory.FullName), ExcludedFolders, HeaderFiles);
 
-				UHTModuleInfo Info = ExternalExecution.CreateUHTModuleInfo(HeaderFiles, Module.Name, Module.RulesFile, Module.ModuleDirectory, ModuleToType[Module], GeneratedCodeVersion);
+				UHTModuleInfo Info = ExternalExecution.CreateUHTModuleInfo(HeaderFiles, Module.Name, Module.RulesFile, Module.ModuleDirectory, ModuleToType[Module], GeneratedCodeVersion, Module.Rules.bUsePrecompiled);
 				if (Info.PublicUObjectClassesHeaders.Count > 0 || Info.PrivateUObjectHeaders.Count > 0 || Info.PublicUObjectHeaders.Count > 0)
 				{
 					// If we've got this far and there are no source files then it's likely we're installed and ignoring
@@ -624,7 +629,7 @@ namespace UnrealBuildTool
 				else
 				{
 					// Remove any stale generated code directory
-					if(Module.GeneratedCodeDirectory != null && (!UnrealBuildTool.IsUnderAnInstalledDirectory(Module.GeneratedCodeDirectory)))
+					if(Module.GeneratedCodeDirectory != null && !Module.Rules.bUsePrecompiled)
 					{
 						if (DirectoryReference.Exists(Module.GeneratedCodeDirectory))
 						{
@@ -646,7 +651,7 @@ namespace UnrealBuildTool
 		/// Gets the path to the receipt for UHT
 		/// </summary>
 		/// <returns>Path to the UHT receipt</returns>
-		static FileReference GetHeaderToolReceiptFile(UnrealTargetConfiguration Configuration)
+		public static FileReference GetHeaderToolReceiptFile(UnrealTargetConfiguration Configuration)
 		{
 			return TargetReceipt.GetDefaultPath(UnrealBuildTool.EngineDirectory, "UnrealHeaderTool", BuildHostPlatform.Current.Platform, Configuration, "");
 		}
@@ -811,12 +816,11 @@ namespace UnrealBuildTool
 		/// <param name="BuildConfiguration">Build configuration</param>
 		/// <param name="UObjectModules">Modules that we generate headers for</param>
 		/// <param name="HeaderToolTimestamp">Timestamp for UHT</param>
-		/// <param name="bUsePrecompiled">Whether to use precompiled engine modules</param>
 		/// <param name="HotReload">The hot reload state</param>
 		/// <param name="bIsGatheringBuild"></param>
 		/// <param name="bIsAssemblingBuild"></param>
 		/// <returns>True if the code files are out of date</returns>
-		private static bool AreGeneratedCodeFilesOutOfDate(BuildConfiguration BuildConfiguration, List<UHTModuleInfo> UObjectModules, DateTime HeaderToolTimestamp, bool bUsePrecompiled, EHotReload HotReload, bool bIsGatheringBuild, bool bIsAssemblingBuild)
+		private static bool AreGeneratedCodeFilesOutOfDate(BuildConfiguration BuildConfiguration, List<UHTModuleInfo> UObjectModules, DateTime HeaderToolTimestamp, EHotReload HotReload, bool bIsGatheringBuild, bool bIsAssemblingBuild)
 		{
 			// Get CoreUObject.init.gen.cpp timestamp.  If the source files are older than the CoreUObject generated code, we'll
 			// need to regenerate code for the module
@@ -841,13 +845,7 @@ namespace UnrealBuildTool
 			foreach (UHTModuleInfo Module in UObjectModules)
 			{
 				// If we're using a precompiled engine, skip checking timestamps for modules that are under the engine directory
-				if (bUsePrecompiled && Module.ModuleDirectory.IsUnderDirectory(UnrealBuildTool.EngineDirectory))
-				{
-					continue;
-				}
-
-				// Skip checking timestamps for modules that are under an installed directory
-				if (UnrealBuildTool.IsUnderAnInstalledDirectory(Module.ModuleDirectory))
+				if (Module.bIsReadOnly)
 				{
 					continue;
 				}
@@ -987,53 +985,44 @@ namespace UnrealBuildTool
 		/// <summary>
 		/// Updates the intermediate include directory timestamps of all the passed in UObject modules
 		/// </summary>
-		private static void UpdateDirectoryTimestamps(List<UHTModuleInfo> UObjectModules, bool bUsePrecompiled)
+		private static void UpdateDirectoryTimestamps(List<UHTModuleInfo> UObjectModules)
 		{
 			foreach (UHTModuleInfo Module in UObjectModules)
 			{
-				string GeneratedCodeDirectory = Path.GetDirectoryName(Module.GeneratedCPPFilenameBase);
-				DirectoryInfo GeneratedCodeDirectoryInfo = new DirectoryInfo(GeneratedCodeDirectory);
-
-				try
+				if(!Module.bIsReadOnly)
 				{
-					if (GeneratedCodeDirectoryInfo.Exists)
-					{
-						// Don't write anything to the engine directory if we're using precompiled
-						if (bUsePrecompiled && Module.ModuleDirectory.IsUnderDirectory(UnrealBuildTool.EngineDirectory))
-						{
-							continue;
-						}
+					string GeneratedCodeDirectory = Path.GetDirectoryName(Module.GeneratedCPPFilenameBase);
+					DirectoryInfo GeneratedCodeDirectoryInfo = new DirectoryInfo(GeneratedCodeDirectory);
 
-						// Don't write anything to installed directories if we're running an installed build
-						if (UnrealBuildTool.IsUnderAnInstalledDirectory(Module.ModuleDirectory))
+					try
+					{
+						if (GeneratedCodeDirectoryInfo.Exists)
 						{
-							continue;
-						}
-
-						// Touch the include directory since we have technically 'generated' the headers
-						// However, the headers might not be touched at all since that would cause the compiler to recompile everything
-						// We can't alter the directory timestamp directly, because this may throw exceptions when the directory is
-						// open in visual studio or windows explorer, so instead we create a blank file that will change the timestamp for us
-						FileReference TimestampFile = FileReference.Combine(new DirectoryReference(GeneratedCodeDirectoryInfo.FullName), "Timestamp");
+							// Touch the include directory since we have technically 'generated' the headers
+							// However, the headers might not be touched at all since that would cause the compiler to recompile everything
+							// We can't alter the directory timestamp directly, because this may throw exceptions when the directory is
+							// open in visual studio or windows explorer, so instead we create a blank file that will change the timestamp for us
+							FileReference TimestampFile = FileReference.Combine(new DirectoryReference(GeneratedCodeDirectoryInfo.FullName), "Timestamp");
+
+							// Save all of the UObject files to a timestamp file.  We'll load these on the next run to see if any new
+							// files with UObject classes were deleted, so that we'll know to run UHT even if the timestamps of all
+							// of the other source files were unchanged
+							{
+								List<string> AllUObjectFiles = new List<string>();
+								AllUObjectFiles.AddRange(Module.PublicUObjectClassesHeaders.ConvertAll(Item => Item.AbsolutePath));
+								AllUObjectFiles.AddRange(Module.PublicUObjectHeaders.ConvertAll(Item => Item.AbsolutePath));
+								AllUObjectFiles.AddRange(Module.PrivateUObjectHeaders.ConvertAll(Item => Item.AbsolutePath));
+								FileReference.WriteAllLines(TimestampFile, AllUObjectFiles);
+							}
 
-						// Save all of the UObject files to a timestamp file.  We'll load these on the next run to see if any new
-						// files with UObject classes were deleted, so that we'll know to run UHT even if the timestamps of all
-						// of the other source files were unchanged
-						{
-							List<string> AllUObjectFiles = new List<string>();
-							AllUObjectFiles.AddRange(Module.PublicUObjectClassesHeaders.ConvertAll(Item => Item.AbsolutePath));
-							AllUObjectFiles.AddRange(Module.PublicUObjectHeaders.ConvertAll(Item => Item.AbsolutePath));
-							AllUObjectFiles.AddRange(Module.PrivateUObjectHeaders.ConvertAll(Item => Item.AbsolutePath));
-							FileReference.WriteAllLines(TimestampFile, AllUObjectFiles);
+							// Because new .cpp and .h files may have been generated by UHT, invalidate the DirectoryLookupCache
+							DirectoryLookupCache.InvalidateCachedDirectory(new DirectoryReference(GeneratedCodeDirectoryInfo.FullName));
 						}
-
-                        // Because new .cpp and .h files may have been generated by UHT, invalidate the DirectoryLookupCache
-                        DirectoryLookupCache.InvalidateCachedDirectory(new DirectoryReference(GeneratedCodeDirectoryInfo.FullName));
 					}
-				}
-				catch (Exception Exception)
-				{
-					throw new BuildException(Exception, "Couldn't touch header directories: " + Exception.Message);
+					catch (Exception Exception)
+					{
+						throw new BuildException(Exception, "Couldn't touch header directories: " + Exception.Message);
+					}
 				}
 			}
 		}
@@ -1117,7 +1106,7 @@ namespace UnrealBuildTool
 				bool bHaveHeaderTool = !bIsBuildingUHT && GetHeaderToolTimestamp(UHTConfig, out HeaderToolTimestamp);
 
 				// ensure the headers are up to date
-				bool bUHTNeedsToRun = (BuildConfiguration.bForceHeaderGeneration || !bHaveHeaderTool || AreGeneratedCodeFilesOutOfDate(BuildConfiguration, UObjectModules, HeaderToolTimestamp, Target.bUsePrecompiled, HotReload, bIsGatheringBuild, bIsAssemblingBuild));
+				bool bUHTNeedsToRun = (BuildConfiguration.bForceHeaderGeneration || !bHaveHeaderTool || AreGeneratedCodeFilesOutOfDate(BuildConfiguration, UObjectModules, HeaderToolTimestamp, HotReload, bIsGatheringBuild, bIsAssemblingBuild));
 
 				// Get the file containing dependencies for the generated code
 				FileReference ExternalDependenciesFile = ModuleInfoFileName.ChangeExtension(".deps");
@@ -1128,7 +1117,13 @@ namespace UnrealBuildTool
 				}
 
 				// @todo ubtmake: Optimization: Ideally we could avoid having to generate this data in the case where UHT doesn't even need to run!  Can't we use the existing copy?  (see below use of Manifest)
-				UHTManifest Manifest = new UHTManifest(Target, RootLocalPath, UEBuildPlatform.GetBuildPlatform(Target.Platform).ConvertPath(RootLocalPath + '\\'), ExternalDependenciesFile.FullName, UObjectModules, Target.bUsePrecompiled);
+
+				List<UHTManifest.Module> Modules = new List<UHTManifest.Module>();
+				foreach(UHTModuleInfo UObjectModule in UObjectModules)
+				{
+					Modules.Add(new UHTManifest.Module(UObjectModule));
+				}
+				UHTManifest Manifest = new UHTManifest(Target, RootLocalPath, UEBuildPlatform.GetBuildPlatform(Target.Platform).ConvertPath(RootLocalPath + '\\'), ExternalDependenciesFile.FullName, Modules);
 
 				if (!bIsBuildingUHT && bUHTNeedsToRun)
 				{
@@ -1292,7 +1287,7 @@ namespace UnrealBuildTool
 				}
 
 				// touch the directories
-				UpdateDirectoryTimestamps(UObjectModules, Target.bUsePrecompiled);
+				UpdateDirectoryTimestamps(UObjectModules);
 
 				Progress.Write(3, 3);
 			}
diff --git a/Engine/Source/Programs/UnrealBuildTool/System/RulesAssembly.cs b/Engine/Source/Programs/UnrealBuildTool/System/RulesAssembly.cs
--- a/Engine/Source/Programs/UnrealBuildTool/System/RulesAssembly.cs
+++ b/Engine/Source/Programs/UnrealBuildTool/System/RulesAssembly.cs
@@ -22,6 +22,11 @@ namespace UnrealBuildTool
 		/// </summary>
 		private Assembly CompiledAssembly;
 
+		/// <summary>
+		/// The base directory for this assembly
+		/// </summary>
+		private DirectoryReference BaseDir;
+
 		/// <summary>
 		/// All the plugins included in this assembly
 		/// </summary>
@@ -59,9 +64,9 @@ namespace UnrealBuildTool
 		private bool bUseBackwardsCompatibleDefaults;
 
 		/// <summary>
-		/// Whether the modules and targets in this assembly are installed.
+		/// Whether the modules and targets in this assembly are read-only
 		/// </summary>
-		private bool bInstalled;
+		private bool bReadOnly;
 
 		/// <summary>
 		/// The parent rules assembly that this assembly inherits. Game assemblies inherit the engine assembly, and the engine assembly inherits nothing.
@@ -71,6 +76,7 @@ namespace UnrealBuildTool
 		/// <summary>
 		/// Constructor. Compiles a rules assembly from the given source files.
 		/// </summary>
+		/// <param name="BaseDir">The base directory for this assembly</param>
 		/// <param name="Plugins">All the plugins included in this assembly</param>
 		/// <param name="ModuleFiles">List of module files to compile</param>
 		/// <param name="TargetFiles">List of target files to compile</param>
@@ -78,15 +84,16 @@ namespace UnrealBuildTool
 		/// <param name="AssemblyFileName">The output path for the compiled assembly</param>
 		/// <param name="bContainsEngineModules">Whether this assembly contains engine modules. Used to initialize the default value for ModuleRules.bTreatAsEngineModule.</param>
 		/// <param name="bUseBackwardsCompatibleDefaults">Whether modules in this assembly should use backwards-compatible defaults.</param>
-		/// <param name="bInstalled">Whether the modules and targets in this assembly are installed, and should be created with the bUsePrecompiled flag set</param> 
+		/// <param name="bReadOnly">Whether the modules and targets in this assembly are installed, and should be created with the bUsePrecompiled flag set</param> 
 		/// <param name="Parent">The parent rules assembly</param>
-		public RulesAssembly(IReadOnlyList<PluginInfo> Plugins, List<FileReference> ModuleFiles, List<FileReference> TargetFiles, Dictionary<FileReference, PluginInfo> ModuleFileToPluginInfo, FileReference AssemblyFileName, bool bContainsEngineModules, bool bUseBackwardsCompatibleDefaults, bool bInstalled, RulesAssembly Parent)
+		public RulesAssembly(DirectoryReference BaseDir, IReadOnlyList<PluginInfo> Plugins, List<FileReference> ModuleFiles, List<FileReference> TargetFiles, Dictionary<FileReference, PluginInfo> ModuleFileToPluginInfo, FileReference AssemblyFileName, bool bContainsEngineModules, bool bUseBackwardsCompatibleDefaults, bool bReadOnly, RulesAssembly Parent)
 		{
+			this.BaseDir = BaseDir;
 			this.Plugins = Plugins;
 			this.ModuleFileToPluginInfo = ModuleFileToPluginInfo;
 			this.bContainsEngineModules = bContainsEngineModules;
 			this.bUseBackwardsCompatibleDefaults = bUseBackwardsCompatibleDefaults;
-			this.bInstalled = bInstalled;
+			this.bReadOnly = bReadOnly;
 			this.Parent = Parent;
 
 			// Find all the source files
@@ -157,6 +164,27 @@ namespace UnrealBuildTool
 			}
 		}
 
+		/// <summary>
+		/// Determines if the given path is read-only
+		/// </summary>
+		/// <param name="Location">The location to check</param>
+		/// <returns>True if the path is read-only, false otherwise</returns>
+		public bool IsReadOnly(FileSystemReference Location)
+		{
+			if(Location.IsUnderDirectory(BaseDir))
+			{
+				return bReadOnly;
+			}
+			else if(Parent != null)
+			{
+				return Parent.IsReadOnly(Location);
+			}
+			else
+			{
+				return false;
+			}
+		}
+
 		/// <summary>
 		/// Finds all the preprocessor definitions that need to be set for the current engine.
 		/// </summary>
@@ -346,7 +374,7 @@ namespace UnrealBuildTool
 				RulesObject.bTreatAsEngineModule = bContainsEngineModules;
 				RulesObject.bUseBackwardsCompatibleDefaults = bUseBackwardsCompatibleDefaults && Target.bUseBackwardsCompatibleDefaults;
 				RulesObject.bPrecompile = RulesObject.bTreatAsEngineModule && Target.bPrecompile;
-				RulesObject.bUsePrecompiled = (RulesObject.bTreatAsEngineModule && Target.bUsePrecompiled) || bInstalled;
+				RulesObject.bUsePrecompiled = bReadOnly;
 
 				// Call the constructor
 				ConstructorInfo Constructor = RulesObjectType.GetConstructor(new Type[] { typeof(ReadOnlyTargetRules) });
diff --git a/Engine/Source/Programs/UnrealBuildTool/System/RulesCompiler.cs b/Engine/Source/Programs/UnrealBuildTool/System/RulesCompiler.cs
--- a/Engine/Source/Programs/UnrealBuildTool/System/RulesCompiler.cs
+++ b/Engine/Source/Programs/UnrealBuildTool/System/RulesCompiler.cs
@@ -1,4 +1,4 @@
-﻿// Copyright 1998-2018 Epic Games, Inc. All Rights Reserved.
+// Copyright 1998-2018 Epic Games, Inc. All Rights Reserved.
 
 using System;
 using System.IO;
@@ -258,13 +258,14 @@ namespace UnrealBuildTool
 		/// <summary>
 		/// Creates the engine rules assembly
 		/// </summary>
+		/// <param name="bUsePrecompiled">Whether to use a precompiled engine</param>
 		/// <returns>New rules assembly</returns>
-		public static RulesAssembly CreateEngineRulesAssembly()
+		public static RulesAssembly CreateEngineRulesAssembly(bool bUsePrecompiled)
 		{
 			if (EngineRulesAssembly == null)
 			{
 				IReadOnlyList<PluginInfo> IncludedPlugins = Plugins.ReadEnginePlugins(UnrealBuildTool.EngineDirectory);
-				EngineRulesAssembly = CreateEngineOrEnterpriseRulesAssembly(UnrealBuildTool.EngineDirectory, ProjectFileGenerator.EngineProjectFileNameBase, IncludedPlugins, UnrealBuildTool.IsEngineInstalled(), null);
+				EngineRulesAssembly = CreateEngineOrEnterpriseRulesAssembly(UnrealBuildTool.EngineDirectory, ProjectFileGenerator.EngineProjectFileNameBase, IncludedPlugins, UnrealBuildTool.IsEngineInstalled() || bUsePrecompiled, null);
 			}
 			return EngineRulesAssembly;
 		}
@@ -272,13 +273,14 @@ namespace UnrealBuildTool
 		/// <summary>
 		/// Creates the enterprise rules assembly
 		/// </summary>
+		/// <param name="bUsePrecompiled">Whether to use a precompiled enterprise and engine folder</param>
 		/// <returns>New rules assembly. Returns null if the enterprise directory is unavailable.</returns>
-		public static RulesAssembly CreateEnterpriseRulesAssembly()
+		public static RulesAssembly CreateEnterpriseRulesAssembly(bool bUsePrecompiled)
 		{
 			if (EnterpriseRulesAssembly == null && DirectoryReference.Exists(UnrealBuildTool.EnterpriseDirectory))
 			{
 				IReadOnlyList<PluginInfo> IncludedPlugins = Plugins.ReadEnterprisePlugins(UnrealBuildTool.EnterpriseDirectory);
-				EnterpriseRulesAssembly = CreateEngineOrEnterpriseRulesAssembly(UnrealBuildTool.EnterpriseDirectory, ProjectFileGenerator.EnterpriseProjectFileNameBase, IncludedPlugins, false, CreateEngineRulesAssembly());
+				EnterpriseRulesAssembly = CreateEngineOrEnterpriseRulesAssembly(UnrealBuildTool.EnterpriseDirectory, ProjectFileGenerator.EnterpriseProjectFileNameBase, IncludedPlugins, UnrealBuildTool.IsEnterpriseInstalled() || bUsePrecompiled, CreateEngineRulesAssembly(bUsePrecompiled));
 			}
 			return EnterpriseRulesAssembly;
 		}
@@ -289,10 +291,10 @@ namespace UnrealBuildTool
 		/// <param name="RootDirectory">The root directory to create rules for</param>
 		/// <param name="AssemblyPrefix">A prefix for the assembly file name</param>
 		/// <param name="Plugins">List of plugins to include in this assembly</param>
-		/// <param name="bInstalled">Whether the assembly should be marked as installed</param>
+		/// <param name="bReadOnly">Whether the assembly should be marked as installed</param>
 		/// <param name="Parent">The parent rules assembly</param>
 		/// <returns>New rules assembly</returns>
-		private static RulesAssembly CreateEngineOrEnterpriseRulesAssembly(DirectoryReference RootDirectory, string AssemblyPrefix, IReadOnlyList<PluginInfo> Plugins, bool bInstalled, RulesAssembly Parent)
+		private static RulesAssembly CreateEngineOrEnterpriseRulesAssembly(DirectoryReference RootDirectory, string AssemblyPrefix, IReadOnlyList<PluginInfo> Plugins, bool bReadOnly, RulesAssembly Parent)
 		{
 			DirectoryReference SourceDirectory = DirectoryReference.Combine(RootDirectory, "Source");
 			DirectoryReference ProgramsDirectory = DirectoryReference.Combine(SourceDirectory, "Programs");
@@ -313,7 +315,7 @@ namespace UnrealBuildTool
 
 			// Create the assembly
 			FileReference EngineAssemblyFileName = FileReference.Combine(RootDirectory, "Intermediate", "Build", "BuildRules", AssemblyPrefix + "Rules" + FrameworkAssemblyExtension);
-			RulesAssembly EngineAssembly = new RulesAssembly(Plugins, EngineModuleFiles, new List<FileReference>(), ModuleFileToPluginInfo, EngineAssemblyFileName, bContainsEngineModules: true, bUseBackwardsCompatibleDefaults: false, bInstalled: bInstalled, Parent: Parent);
+			RulesAssembly EngineAssembly = new RulesAssembly(RootDirectory, Plugins, EngineModuleFiles, new List<FileReference>(), ModuleFileToPluginInfo, EngineAssemblyFileName, bContainsEngineModules: true, bUseBackwardsCompatibleDefaults: false, bReadOnly: bReadOnly, Parent: Parent);
 
 			// Find all the rules files
 			List<FileReference> ProgramModuleFiles = new List<FileReference>(FindAllRulesFiles(ProgramsDirectory, RulesFileType.Module));
@@ -321,7 +323,7 @@ namespace UnrealBuildTool
 
 			// Create a path to the assembly that we'll either load or compile
 			FileReference ProgramAssemblyFileName = FileReference.Combine(RootDirectory, "Intermediate", "Build", "BuildRules", AssemblyPrefix + "ProgramRules" + FrameworkAssemblyExtension);
-			RulesAssembly ProgramAssembly = new RulesAssembly(new List<PluginInfo>().AsReadOnly(), ProgramModuleFiles, ProgramTargetFiles, new Dictionary<FileReference, PluginInfo>(), ProgramAssemblyFileName, bContainsEngineModules: false, bUseBackwardsCompatibleDefaults: false, bInstalled: bInstalled, Parent: EngineAssembly);
+			RulesAssembly ProgramAssembly = new RulesAssembly(RootDirectory, new List<PluginInfo>().AsReadOnly(), ProgramModuleFiles, ProgramTargetFiles, new Dictionary<FileReference, PluginInfo>(), ProgramAssemblyFileName, bContainsEngineModules: false, bUseBackwardsCompatibleDefaults: false, bReadOnly: bReadOnly, Parent: EngineAssembly);
 
 			// Return the combined assembly
 			return ProgramAssembly;
@@ -331,8 +333,9 @@ namespace UnrealBuildTool
 		/// Creates a rules assembly with the given parameters.
 		/// </summary>
 		/// <param name="ProjectFileName">The project file to create rules for. Null for the engine.</param>
+		/// <param name="bUsePrecompiled">Whether to use a precompiled engine</param>
 		/// <returns>New rules assembly</returns>
-		public static RulesAssembly CreateProjectRulesAssembly(FileReference ProjectFileName)
+		public static RulesAssembly CreateProjectRulesAssembly(FileReference ProjectFileName, bool bUsePrecompiled)
 		{
 			// Check if there's an existing assembly for this project
 			RulesAssembly ProjectRulesAssembly;
@@ -344,11 +347,11 @@ namespace UnrealBuildTool
 				RulesAssembly Parent;
 				if (Project.IsEnterpriseProject)
 				{
-					Parent = CreateEnterpriseRulesAssembly();
+					Parent = CreateEnterpriseRulesAssembly(bUsePrecompiled);
 				}
 				else
 				{
-					Parent = CreateEngineRulesAssembly();
+					Parent = CreateEngineRulesAssembly(bUsePrecompiled);
 				}
 
 				// Find all the rules under the project source directory
@@ -383,7 +386,7 @@ namespace UnrealBuildTool
 
 				// Compile the assembly
 				FileReference AssemblyFileName = FileReference.Combine(ProjectDirectory, "Intermediate", "Build", "BuildRules", ProjectFileName.GetFileNameWithoutExtension() + "ModuleRules" + FrameworkAssemblyExtension);
-				ProjectRulesAssembly = new RulesAssembly(ProjectPlugins, ModuleFiles, TargetFiles, ModuleFileToPluginInfo, AssemblyFileName, bContainsEngineModules: false, bUseBackwardsCompatibleDefaults: true, bInstalled: UnrealBuildTool.IsProjectInstalled(), Parent: Parent);
+				ProjectRulesAssembly = new RulesAssembly(ProjectDirectory, ProjectPlugins, ModuleFiles, TargetFiles, ModuleFileToPluginInfo, AssemblyFileName, bContainsEngineModules: false, bUseBackwardsCompatibleDefaults: true, bReadOnly: UnrealBuildTool.IsProjectInstalled(), Parent: Parent);
 				LoadedAssemblyMap.Add(ProjectFileName, ProjectRulesAssembly);
 			}
 			return ProjectRulesAssembly;
@@ -419,7 +422,7 @@ namespace UnrealBuildTool
 
 				// Compile the assembly
 				FileReference AssemblyFileName = FileReference.Combine(PluginFileName.Directory, "Intermediate", "Build", "BuildRules", Path.GetFileNameWithoutExtension(PluginFileName.FullName) + "ModuleRules" + FrameworkAssemblyExtension);
-				PluginRulesAssembly = new RulesAssembly(ForeignPlugins, ModuleFiles, TargetFiles, ModuleFileToPluginInfo, AssemblyFileName, bContainsEngineModules, bUseBackwardsCompatibleDefaults: false, bInstalled: false, Parent: Parent);
+				PluginRulesAssembly = new RulesAssembly(PluginFileName.Directory, ForeignPlugins, ModuleFiles, TargetFiles, ModuleFileToPluginInfo, AssemblyFileName, bContainsEngineModules, bUseBackwardsCompatibleDefaults: false, bReadOnly: false, Parent: Parent);
 				LoadedAssemblyMap.Add(PluginFileName, PluginRulesAssembly);
 			}
 			return PluginRulesAssembly;
diff --git a/Engine/Source/Programs/UnrealBuildTool/UnrealBuildTool.cs b/Engine/Source/Programs/UnrealBuildTool/UnrealBuildTool.cs
--- a/Engine/Source/Programs/UnrealBuildTool/UnrealBuildTool.cs
+++ b/Engine/Source/Programs/UnrealBuildTool/UnrealBuildTool.cs
@@ -220,22 +220,6 @@ namespace UnrealBuildTool
 				InDirectory.IsUnderDirectory(UnrealBuildTool.EnterprisePluginsDirectory) || InDirectory.IsUnderDirectory(UnrealBuildTool.EnterpriseIntermediateDirectory);
 		}
 
-		/// <summary>
-		/// Determines whether a directory is part of an installed directory
-		/// </summary>
-		/// <param name="InDirectory"></param>
-		/// <returns>true if the directory is under an installed directory, false if not</returns>
-		static public bool IsUnderAnInstalledDirectory(DirectoryReference InDirectory)
-		{
-			// Enterprise modules are considered as engine modules
-			bool bIsUnderEngine = InDirectory.IsUnderDirectory(UnrealBuildTool.EngineDirectory);
-
-			bool bIsUnderEnterprise = InDirectory.IsUnderDirectory(UnrealBuildTool.EnterpriseSourceDirectory) ||
-				InDirectory.IsUnderDirectory(UnrealBuildTool.EnterprisePluginsDirectory) || InDirectory.IsUnderDirectory(UnrealBuildTool.EnterpriseIntermediateDirectory);
-
-			return (IsEngineInstalled() && bIsUnderEngine) || (IsEnterpriseInstalled() && bIsUnderEnterprise);
-		}
-
 		public static void RegisterAllUBTClasses(SDKOutputLevel OutputLevel, bool bValidatingPlatforms)
 		{
 			// Find and register all tool chains and build platforms that are present
@@ -1132,7 +1116,7 @@ namespace UnrealBuildTool
 
 					foreach (string[] TargetSetting in TargetSettings)
 					{
-						TargetDescs.AddRange(TargetDescriptor.ParseCommandLine(TargetSetting, ref ProjectFile));
+						TargetDescs.AddRange(TargetDescriptor.ParseCommandLine(TargetSetting, BuildConfiguration.bUsePrecompiled, ref ProjectFile));
 					}
 
 					if (UnrealBuildTool.bPrintPerformanceInfo)
@@ -1393,7 +1377,7 @@ namespace UnrealBuildTool
 						Targets = new List<UEBuildTarget>();
 						foreach (TargetDescriptor TargetDesc in TargetDescs)
 						{
-							UEBuildTarget Target = UEBuildTarget.CreateTarget(TargetDesc, Arguments, BuildConfiguration.SingleFileToCompile != null, Version);
+							UEBuildTarget Target = UEBuildTarget.CreateTarget(TargetDesc, Arguments, BuildConfiguration.SingleFileToCompile != null, BuildConfiguration.bUsePrecompiled, Version);
 							if ((Target == null) && (BuildConfiguration.bCleanProject))
 							{
 								continue;
