UBT: Runtime dependencies can now be copied to output directories as part of the build. When adding a runtime dependency, an optional source location can be specified to copy from. Both the source and target paths can use variables can be used as part of the path, eg. $(OutputDir), $(ModuleDir), $(PluginDir).

Example usage (from a .build.cs file):

RuntimeDependencies.Add("$(OutputDir)/Foo.dll", "$(PluginDir)/Source/ThirdParty/Foo.dll", StagedFileType.NonUFS);

diff --git a/Engine/Source/Programs/UnrealBuildTool/Configuration/ModuleRules.cs b/Engine/Source/Programs/UnrealBuildTool/Configuration/ModuleRules.cs
--- a/Engine/Source/Programs/UnrealBuildTool/Configuration/ModuleRules.cs
+++ b/Engine/Source/Programs/UnrealBuildTool/Configuration/ModuleRules.cs
@@ -130,6 +130,11 @@ namespace UnrealBuildTool
 			/// </summary>
 			public string Path;
 
+			/// <summary>
+			/// The initial location for this file. It will be copied to Path at build time, ready for staging.
+			/// </summary>
+			public string SourcePath;
+
 			/// <summary>
 			/// How to stage this file.
 			/// </summary>
@@ -145,6 +150,19 @@ namespace UnrealBuildTool
 				Path = InPath;
 				Type = InType;
 			}
+
+			/// <summary>
+			/// Constructor
+			/// </summary>
+			/// <param name="InPath">Path to the runtime dependency</param>
+			/// <param name="InSourcePath">Source path for the file in the working tree</param>
+			/// <param name="InType">How to stage the given path</param>
+			public RuntimeDependency(string InPath, string InSourcePath, StagedFileType InType = StagedFileType.NonUFS)
+			{
+				Path = InPath;
+				SourcePath = InSourcePath;
+				Type = InType;
+			}
 		}
 
 		/// <summary>
@@ -184,6 +202,17 @@ namespace UnrealBuildTool
 				Inner.Add(new RuntimeDependency(InPath, InType));
 			}
 
+			/// <summary>
+			/// Add a runtime dependency to the list
+			/// </summary>
+			/// <param name="InPath">Path to the runtime dependency. May include wildcards.</param>
+			/// <param name="InSourcePath">Source path for the file to be added as a dependency. May include wildcards.</param>
+			/// <param name="InType">How to stage this file</param>
+			public void Add(string InPath, string InSourcePath, StagedFileType InType = StagedFileType.NonUFS)
+			{
+				Inner.Add(new RuntimeDependency(InPath, InSourcePath, InType));
+			}
+
 			/// <summary>
 			/// Add a runtime dependency to the list
 			/// </summary>
diff --git a/Engine/Source/Programs/UnrealBuildTool/Configuration/UEBuildBinary.cs b/Engine/Source/Programs/UnrealBuildTool/Configuration/UEBuildBinary.cs
--- a/Engine/Source/Programs/UnrealBuildTool/Configuration/UEBuildBinary.cs
+++ b/Engine/Source/Programs/UnrealBuildTool/Configuration/UEBuildBinary.cs
@@ -236,6 +236,94 @@ namespace UnrealBuildTool
 			return OutputFiles;
 		}
 
+		/// <summary>
+		/// Gets all the runtime dependencies Copies all the runtime dependencies from any modules in 
+		/// </summary>
+		/// <param name="RuntimeDependencies">The output list of runtime dependencies, mapping target file to type</param>
+		/// <param name="SourceFiles">Receives the list of source files that were copied</param>
+		/// <param name="ActionGraph">Actions to be executed</param>
+		public IEnumerable<FileItem> PrepareRuntimeDependencies(List<RuntimeDependency> RuntimeDependencies, List<FileReference> SourceFiles, ActionGraph ActionGraph)
+		{
+			List<FileItem> CopiedFiles = new List<FileItem>();
+			foreach(UEBuildModule Module in Modules)
+			{
+				foreach (ModuleRules.RuntimeDependency Dependency in Module.Rules.RuntimeDependencies.Inner)
+				{
+					if(Dependency.SourcePath == null)
+					{
+						// Expand the target path
+						string ExpandedPath = Module.ExpandPathVariables(Dependency.Path, this);
+						if (FileFilter.FindWildcardIndex(ExpandedPath) == -1)
+						{
+							RuntimeDependencies.Add(new RuntimeDependency(new FileReference(ExpandedPath), Dependency.Type));
+						}
+						else
+						{
+							RuntimeDependencies.AddRange(FileFilter.ResolveWildcard(ExpandedPath).Select(x => new RuntimeDependency(x, Dependency.Type)));
+						}
+					}
+					else
+					{
+						// Parse the source and target patterns
+						FilePattern SourcePattern = new FilePattern(UnrealBuildTool.EngineSourceDirectory, Module.ExpandPathVariables(Dependency.SourcePath, this));
+						FilePattern TargetPattern = new FilePattern(UnrealBuildTool.EngineSourceDirectory, Module.ExpandPathVariables(Dependency.Path, this));
+
+						// Resolve all the wildcards between the source and target paths
+						Dictionary<FileReference, FileReference> Mapping;
+						try
+						{
+							Mapping = FilePattern.CreateMapping(null, ref SourcePattern, ref TargetPattern);
+						}
+						catch(FilePatternException Ex)
+						{
+							throw new BuildException(Ex, "While creating runtime dependencies for module {0}", Module.Name);
+						}
+
+						// Add actions to copy everything
+						foreach(KeyValuePair<FileReference, FileReference> Pair in Mapping)
+						{
+							CopiedFiles.Add(CreateCopyAction(Pair.Value, Pair.Key, ActionGraph));
+							SourceFiles.Add(Pair.Value);
+							RuntimeDependencies.Add(new RuntimeDependency(Pair.Key, Dependency.Type));
+						}
+					}
+				}
+			}
+			return CopiedFiles;
+		}
+
+		/// <summary>
+		/// Creates an action which copies a file from one location to another
+		/// </summary>
+		/// <param name="SourceFile">The source file location</param>
+		/// <param name="TargetFile">The target file location</param>
+		/// <param name="ActionGraph">The action graph</param>
+		/// <returns>File item for the output file</returns>
+		static FileItem CreateCopyAction(FileReference SourceFile, FileReference TargetFile, ActionGraph ActionGraph)
+		{
+			FileItem SourceFileItem = FileItem.GetItemByFileReference(SourceFile);
+			FileItem TargetFileItem = FileItem.GetItemByFileReference(TargetFile);
+
+			Action CopyAction = ActionGraph.Add(ActionType.BuildProject);
+			CopyAction.CommandDescription = "Copy";
+			if(BuildHostPlatform.Current.Platform == UnrealTargetPlatform.Win64)
+			{
+				CopyAction.CommandPath = "cmd.exe";
+				CopyAction.CommandArguments = String.Format("/C \"copy /Y \"{0}\" \"{1}\" 1>nul\"", SourceFile, TargetFile);
+			}
+			else
+			{
+				CopyAction.CommandPath = "/bin/sh";
+				CopyAction.CommandArguments = String.Format("cp -f {0} {1}", Utils.EscapeShellArgument(SourceFile.FullName), Utils.EscapeShellArgument(TargetFile.FullName));
+			}
+			CopyAction.PrerequisiteItems.Add(SourceFileItem);
+			CopyAction.ProducedItems.Add(TargetFileItem);
+			CopyAction.DeleteItems.Add(TargetFileItem);
+			CopyAction.StatusDescription = TargetFileItem.Location.GetFileName();
+			CopyAction.bCanExecuteRemotely = false;
+
+			return TargetFileItem;
+		}
 
 		/// <summary>
 		/// Called to allow the binary to modify the link environment of a different binary containing 
diff --git a/Engine/Source/Programs/UnrealBuildTool/Configuration/UEBuildModule.cs b/Engine/Source/Programs/UnrealBuildTool/Configuration/UEBuildModule.cs
--- a/Engine/Source/Programs/UnrealBuildTool/Configuration/UEBuildModule.cs
+++ b/Engine/Source/Programs/UnrealBuildTool/Configuration/UEBuildModule.cs
@@ -151,11 +151,6 @@ namespace UnrealBuildTool
 		/// </summary>
 		protected List<UEBuildModule> DynamicallyLoadedModules;
 
-		/// <summary>
-		/// Files which this module depends on at runtime.
-		/// </summary>
-		public List<RuntimeDependency> RuntimeDependencies;
-
 		/// <summary>
 		/// Set of all whitelisted restricted folder references
 		/// </summary>
@@ -165,8 +160,7 @@ namespace UnrealBuildTool
 		/// Constructor
 		/// </summary>
 		/// <param name="Rules">Rules for this module</param>
-		/// <param name="RuntimeDependencies">List of runtime dependencies</param>
-		public UEBuildModule(ModuleRules Rules, List<RuntimeDependency> RuntimeDependencies)
+		public UEBuildModule(ModuleRules Rules)
 		{
 			this.Rules = Rules;
 
@@ -191,7 +185,6 @@ namespace UnrealBuildTool
 			{
 				PrivateIncludePaths = CreateDirectoryHashSet(Rules.PrivateIncludePaths);
 			}
-			this.RuntimeDependencies = RuntimeDependencies;
 			IsRedistributableOverride = Rules.IsRedistributableOverride;
 
 			WhitelistRestrictedFolders = new HashSet<DirectoryReference>(Rules.WhitelistRestrictedFolders.Select(x => DirectoryReference.Combine(ModuleDirectory, x)));
@@ -548,7 +541,7 @@ namespace UnrealBuildTool
 		/// <param name="Path">Path to expand variables within</param>
 		/// <param name="Binary">The source binary containing the module. May be null.</param>
 		/// <returns>The path with variables expanded</returns>
-		private string ExpandPathVariables(string Path, UEBuildBinary Binary)
+		public string ExpandPathVariables(string Path, UEBuildBinary Binary)
 		{
 			if(Path.StartsWith("$(", StringComparison.Ordinal))
 			{
@@ -580,7 +573,7 @@ namespace UnrealBuildTool
 						{
 							Path = Rules.PluginDirectory + Path.Substring(EndIdx + 1);
 						}
-						else if(Binary != null && MatchVariableName(Path, StartIdx, EndIdx, "BinaryDir"))
+						else if(Binary != null && MatchVariableName(Path, StartIdx, EndIdx, "OutputDir"))
 						{
 							Path = Binary.OutputFilePaths[0].Directory.FullName + Path.Substring(EndIdx + 1);
 						}
@@ -877,16 +870,6 @@ namespace UnrealBuildTool
 				Writer.WriteValue(ModuleName);
 			}
 			Writer.WriteArrayEnd();
-
-			Writer.WriteArrayStart("RuntimeDependencies");
-			foreach(RuntimeDependency RuntimeDependency in RuntimeDependencies)
-			{
-				Writer.WriteObjectStart();
-				Writer.WriteValue("Path", RuntimeDependency.Path.FullName);
-				Writer.WriteValue("Type", RuntimeDependency.Type.ToString());
-				Writer.WriteObjectEnd();
-			}
-			Writer.WriteArrayEnd();
 		}
 
 		/// <summary>
diff --git a/Engine/Source/Programs/UnrealBuildTool/Configuration/UEBuildModuleCPP.cs b/Engine/Source/Programs/UnrealBuildTool/Configuration/UEBuildModuleCPP.cs
--- a/Engine/Source/Programs/UnrealBuildTool/Configuration/UEBuildModuleCPP.cs
+++ b/Engine/Source/Programs/UnrealBuildTool/Configuration/UEBuildModuleCPP.cs
@@ -197,15 +197,8 @@ namespace UnrealBuildTool
         };
 
 
-		public UEBuildModuleCPP(
-			ModuleRules Rules,
-			DirectoryReference IntermediateDirectory,
-			DirectoryReference GeneratedCodeDirectory,
-			IEnumerable<FileItem> SourceFiles,
-			bool bBuildSourceFiles,
-			List<RuntimeDependency> RuntimeDependencies
-			)
-			: base(Rules, RuntimeDependencies)
+		public UEBuildModuleCPP(ModuleRules Rules, DirectoryReference IntermediateDirectory, DirectoryReference GeneratedCodeDirectory, IEnumerable<FileItem> SourceFiles, bool bBuildSourceFiles)
+			: base(Rules)
 		{
 			this.IntermediateDirectory = IntermediateDirectory;
 			this.GeneratedCodeDirectory = GeneratedCodeDirectory;
diff --git a/Engine/Source/Programs/UnrealBuildTool/Configuration/UEBuildModuleExternal.cs b/Engine/Source/Programs/UnrealBuildTool/Configuration/UEBuildModuleExternal.cs
--- a/Engine/Source/Programs/UnrealBuildTool/Configuration/UEBuildModuleExternal.cs
+++ b/Engine/Source/Programs/UnrealBuildTool/Configuration/UEBuildModuleExternal.cs
@@ -14,8 +14,8 @@ namespace UnrealBuildTool
 	/// </summary>
 	class UEBuildModuleExternal : UEBuildModule
 	{
-		public UEBuildModuleExternal(ModuleRules Rules, List<RuntimeDependency> RuntimeDependencies)
-			: base(Rules, RuntimeDependencies)
+		public UEBuildModuleExternal(ModuleRules Rules)
+			: base(Rules)
 		{
 		}
 
diff --git a/Engine/Source/Programs/UnrealBuildTool/Configuration/UEBuildTarget.cs b/Engine/Source/Programs/UnrealBuildTool/Configuration/UEBuildTarget.cs
--- a/Engine/Source/Programs/UnrealBuildTool/Configuration/UEBuildTarget.cs
+++ b/Engine/Source/Programs/UnrealBuildTool/Configuration/UEBuildTarget.cs
@@ -1197,9 +1197,10 @@ namespace UnrealBuildTool
 		/// Writes a list of all the externally referenced files required to use the precompiled data for this target
 		/// </summary>
 		/// <param name="Location">Path to the dependency list</param>
-		void WriteDependencyList(FileReference Location)
+		/// <param name="RuntimeDependencySourceFiles">All the source files for runtime dependencies. This is only evaluated at build time.</param>
+		void WriteDependencyList(FileReference Location, IEnumerable<FileReference> RuntimeDependencySourceFiles)
 		{
-			HashSet<FileReference> Files = new HashSet<FileReference>();
+			HashSet<FileReference> Files = new HashSet<FileReference>(RuntimeDependencySourceFiles);
 
 			// Get the platform we're building for
 			UEBuildPlatform BuildPlatform = UEBuildPlatform.GetBuildPlatform(Platform);
@@ -1304,12 +1305,6 @@ namespace UnrealBuildTool
 						}
 					}
 				}
-
-				// Add all of the runtime dependencies
-				foreach(RuntimeDependency RuntimeDependency in Module.RuntimeDependencies)
-				{
-					Files.Add(RuntimeDependency.Path);
-				}
 			}
 
 			// Write the file
@@ -1396,7 +1391,11 @@ namespace UnrealBuildTool
 		/// <summary>
 		/// Prepare all the receipts this target (all the .target and .modules files). See the VersionManifest class for an explanation of what these files are.
 		/// </summary>
-		void PrepareReceipts(UEToolChain ToolChain, List<KeyValuePair<FileReference, BuildProductType>> BuildProducts, EHotReload HotReload)
+		/// <param name="ToolChain">The toolchain used to build the target</param>
+		/// <param name="BuildProducts">Artifacts from the build</param>
+		/// <param name="RuntimeDependencies">Output runtime dependencies</param>
+		/// <param name="HotReload">The hot-reload mode</param>
+		void PrepareReceipts(UEToolChain ToolChain, List<KeyValuePair<FileReference, BuildProductType>> BuildProducts, List<RuntimeDependency> RuntimeDependencies, EHotReload HotReload)
 		{
 			// If linking is disabled, don't generate any receipt
 			if(Rules.bDisableLinking)
@@ -1466,6 +1465,16 @@ namespace UnrealBuildTool
 				}
 			}
 
+			// Add all the other runtime dependencies
+			HashSet<FileReference> UniqueRuntimeDependencyFiles = new HashSet<FileReference>();
+			foreach(RuntimeDependency RuntimeDependency in RuntimeDependencies)
+			{
+				if(UniqueRuntimeDependencyFiles.Add(RuntimeDependency.Path))
+				{
+					Receipt.RuntimeDependencies.Add(RuntimeDependency);
+				}
+			}
+
 			// Find all the modules which are part of this target
 			HashSet<UEBuildModule> UniqueLinkedModules = new HashSet<UEBuildModule>();
 			foreach (UEBuildBinary Binary in Binaries)
@@ -1474,10 +1483,6 @@ namespace UnrealBuildTool
 				{
 					if (UniqueLinkedModules.Add(Module))
 					{
-						foreach (RuntimeDependency RuntimeDependency in Module.RuntimeDependencies)
-						{
-							Receipt.RuntimeDependencies.Add(RuntimeDependency.Path, RuntimeDependency.Type);
-						}
 						Receipt.AdditionalProperties.AddRange(Module.Rules.AdditionalPropertiesForReceipt.Inner);
 					}
 				}
@@ -2102,6 +2107,14 @@ namespace UnrealBuildTool
 				OutputItems.AddRange(Binary.Build(Rules, TargetToolChain, GlobalCompileEnvironment, GlobalLinkEnvironment, SharedPCHs, WorkingSet, ActionGraph));
 			}
 
+			// Prepare all the runtime dependencies, copying them from their source folders if necessary
+			List<FileReference> RuntimeDependencySourceFiles = new List<FileReference>();
+			List<RuntimeDependency> RuntimeDependencies = new List<RuntimeDependency>();
+			foreach(UEBuildBinary Binary in Binaries)
+			{
+				OutputItems.AddRange(Binary.PrepareRuntimeDependencies(RuntimeDependencies, RuntimeDependencySourceFiles, ActionGraph));
+			}
+
 			// If we're just precompiling a plugin, only include output items which are under that directory
 			if(ForeignPlugin != null)
 			{
@@ -2123,7 +2136,7 @@ namespace UnrealBuildTool
 			// Create a receipt for the target
 			if (!ProjectFileGenerator.bGenerateProjectFiles)
 			{
-				PrepareReceipts(TargetToolChain, BuildProducts, HotReload);
+				PrepareReceipts(TargetToolChain, BuildProducts, RuntimeDependencies, HotReload);
 			}
 
 			// Make sure all the checked headers were valid
@@ -2141,7 +2154,7 @@ namespace UnrealBuildTool
 			// Build a list of all the externally files
 			foreach(FileReference DependencyListFileName in Rules.DependencyListFileNames)
 			{
-				WriteDependencyList(DependencyListFileName);
+				WriteDependencyList(DependencyListFileName, RuntimeDependencySourceFiles);
 			}
 
 			// If we're only generating the manifest, return now
@@ -3984,44 +3997,8 @@ namespace UnrealBuildTool
 				// locations.
 				UEBuildPlatform.PlatformModifyHostModuleRules(ModuleName, RulesObject, Rules);
 
-				// Expand the list of runtime dependencies, and update the run-time dependencies path to remove $(PluginDir) and replace with a full path. When the 
-				// receipt is saved it'll be converted to a $(ProjectDir) or $(EngineDir) equivalent.
-				List<RuntimeDependency> RuntimeDependencies = new List<RuntimeDependency>();
-				if(RulesObject.RuntimeDependencies.Inner.Count > 0)
-				{
-					// Get all the valid variables which can be expanded for this module
-					Dictionary<string, string> Variables = new Dictionary<string, string>();
-					Variables["EngineDir"] = UnrealBuildTool.EngineDirectory.FullName;
-					if (ProjectFile != null)
-					{
-						Variables["ProjectDir"] = ProjectDirectory.FullName;
-					}
-					if (RulesObject.Plugin != null)
-					{
-						Variables["PluginDir"] = RulesObject.Plugin.Directory.FullName;
-					}
-
-					// Convert them into concrete file lists. Ignore anything that still hasn't been resolved.
-					foreach (ModuleRules.RuntimeDependency Dependency in RulesObject.RuntimeDependencies.Inner)
-					{
-						string ExpandedPath = Utils.ExpandVariables(Dependency.Path, Variables);
-						if (!ExpandedPath.StartsWith("$("))
-						{
-							int WildcardIdx = FileFilter.FindWildcardIndex(ExpandedPath);
-							if (WildcardIdx == -1)
-							{
-								RuntimeDependencies.Add(new RuntimeDependency(new FileReference(ExpandedPath), Dependency.Type));
-							}
-							else
-							{
-								RuntimeDependencies.AddRange(FileFilter.ResolveWildcard(ExpandedPath).Select(x => new RuntimeDependency(x, Dependency.Type)));
-							}
-						}
-					}
-				}
-
 				// Now, go ahead and create the module builder instance
-				Module = InstantiateModule(RulesObject, GeneratedCodeDirectory, FoundSourceFiles, bBuildFiles, RuntimeDependencies);
+				Module = InstantiateModule(RulesObject, GeneratedCodeDirectory, FoundSourceFiles, bBuildFiles);
 				Modules.Add(Module.Name, Module);
 				FlatModuleCsData.Add(new FlatModuleCsDataType(Module.Name, (Module.RulesFile == null) ? null : Module.RulesFile.FullName, RulesObject.ExternalDependencies));
 			}
@@ -4048,8 +4025,7 @@ namespace UnrealBuildTool
 			ModuleRules RulesObject,
 			DirectoryReference GeneratedCodeDirectory,
 			List<FileItem> ModuleSourceFiles,
-			bool bBuildSourceFiles,
-			List<RuntimeDependency> RuntimeDependencies)
+			bool bBuildSourceFiles)
 		{
 			switch (RulesObject.Type)
 			{
@@ -4059,15 +4035,11 @@ namespace UnrealBuildTool
 							IntermediateDirectory: GetModuleIntermediateDirectory(RulesObject),
 							GeneratedCodeDirectory: GeneratedCodeDirectory,
 							SourceFiles: ModuleSourceFiles,
-							bBuildSourceFiles: bBuildSourceFiles,
-							RuntimeDependencies: RuntimeDependencies
+							bBuildSourceFiles: bBuildSourceFiles
 						);
 
 				case ModuleRules.ModuleType.External:
-					return new UEBuildModuleExternal(
-							Rules: RulesObject,
-							RuntimeDependencies: RuntimeDependencies
-						);
+					return new UEBuildModuleExternal(RulesObject);
 
 				default:
 					throw new BuildException("Unrecognized module type specified by 'Rules' object {0}", RulesObject.ToString());
