Expose a flag for whether the engine is installed, to fix issues generating project files.

#rnx

diff --git a/Engine/Source/Programs/UnrealBuildTool/Configuration/TargetRules.cs b/Engine/Source/Programs/UnrealBuildTool/Configuration/TargetRules.cs
--- a/Engine/Source/Programs/UnrealBuildTool/Configuration/TargetRules.cs
+++ b/Engine/Source/Programs/UnrealBuildTool/Configuration/TargetRules.cs
@@ -1298,6 +1298,15 @@ namespace UnrealBuildTool
 		{
 			get { return ProjectFileGenerator.bGenerateProjectFiles; }
 		}
+
+		/// <summary>
+		/// Expose a setting for whether or not the engine is installed
+		/// </summary>
+		/// <returns>Flag for whether the engine is installed</returns>
+		public bool bIsEngineInstalled
+		{
+			get { return UnrealBuildTool.IsEngineInstalled(); }
+		}
 	}
 
 	/// <summary>
