Fix null reference generating receipts when UBT makefiles are active.

#rnx

diff --git a/Engine/Source/Programs/UnrealBuildTool/Configuration/UEBuildTarget.cs b/Engine/Source/Programs/UnrealBuildTool/Configuration/UEBuildTarget.cs
--- a/Engine/Source/Programs/UnrealBuildTool/Configuration/UEBuildTarget.cs
+++ b/Engine/Source/Programs/UnrealBuildTool/Configuration/UEBuildTarget.cs
@@ -1691,7 +1691,7 @@ namespace UnrealBuildTool
 				UEBuildPlatform BuildPlatform = UEBuildPlatform.GetBuildPlatform(Platform);
 				if (OnlyModules == null || OnlyModules.Count == 0)
 				{
-					if(!RulesAssembly.IsReadOnly(ReceiptFileName))
+					if(!IsFileInstalled(ReceiptFileName))
 					{
 						DirectoryReference.CreateDirectory(ReceiptFileName.Directory);
 						Receipt.Write(ReceiptFileName, UnrealBuildTool.EngineDirectory, ProjectDirectory);
@@ -1700,7 +1700,7 @@ namespace UnrealBuildTool
 				if (ForceReceiptFileName != null)
 				{
 					FileReference ForceReceiptFile = new FileReference(ForceReceiptFileName);
-					if(!RulesAssembly.IsReadOnly(ForceReceiptFile))
+					if(!IsFileInstalled(ForceReceiptFile))
 					{
 						DirectoryReference.CreateDirectory(ForceReceiptFile.Directory);
 						Receipt.Write(ForceReceiptFile, UnrealBuildTool.EngineDirectory, ProjectDirectory);
@@ -1708,7 +1708,7 @@ namespace UnrealBuildTool
 				}
 				if(VersionFile != null)
 				{
-					if(!RulesAssembly.IsReadOnly(VersionFile))
+					if(!IsFileInstalled(VersionFile))
 					{
 						DirectoryReference.CreateDirectory(VersionFile.Directory);
 
@@ -1727,7 +1727,7 @@ namespace UnrealBuildTool
 			{
 				foreach (KeyValuePair<FileReference, ModuleManifest> FileNameToVersionManifest in FileReferenceToModuleManifestPairs)
 				{
-					if(!RulesAssembly.IsReadOnly(FileNameToVersionManifest.Key))
+					if(!IsFileInstalled(FileNameToVersionManifest.Key))
 					{
 						// Write the manifest out to a string buffer, then only write it to disk if it's changed.
 						string OutputText;
@@ -1746,6 +1746,24 @@ namespace UnrealBuildTool
 			}
 		}
 
+		/// <summary>
+		/// Checks whether the given file is under an installed directory, and should not be overridden
+		/// </summary>
+		/// <param name="File">File to test</param>
+		/// <returns>True if the file is part of the installed distribution, false otherwise</returns>
+		bool IsFileInstalled(FileReference File)
+		{
+			if(UnrealBuildTool.IsEngineInstalled() && File.IsUnderDirectory(UnrealBuildTool.EngineDirectory))
+			{
+				return true;
+			}
+			if(UnrealBuildTool.IsProjectInstalled() && ProjectFile != null && File.IsUnderDirectory(ProjectFile.Directory))
+			{
+				return true;
+			}
+			return false;
+		}
+
 		/// <summary>
 		/// Gathers dependency modules for given binaries list.
 		/// </summary>
