UBT: Exclude mobile pipeline caches from generated project files. Causes huge slowdown when using 'Find in Files' through the IDE.

diff --git a/Engine/Source/Programs/UnrealBuildTool/ProjectFiles/ProjectFileGenerator.cs b/Engine/Source/Programs/UnrealBuildTool/ProjectFiles/ProjectFileGenerator.cs
--- a/Engine/Source/Programs/UnrealBuildTool/ProjectFiles/ProjectFileGenerator.cs
+++ b/Engine/Source/Programs/UnrealBuildTool/ProjectFiles/ProjectFileGenerator.cs
@@ -1189,6 +1189,7 @@ namespace UnrealBuildTool
 						SubdirectoryNamesToExclude.Add("Receipts");
 						SubdirectoryNamesToExclude.Add("Scripts");
 						SubdirectoryNamesToExclude.Add("FileOpenOrder");
+						SubdirectoryNamesToExclude.Add("PipelineCaches");
 
 						ProjectFile GameProjectFile = GameFolderAndProjectFile.Value;
 						GameProjectFile.AddFilesToProject( SourceFileSearch.FindFiles( GameBuildDirectory, SubdirectoryNamesToExclude ), GameProjectDirectory );
