UBT: Include action number in parallel executor output.

diff --git a/Engine/Source/Programs/UnrealBuildTool/Executors/ParallelExecutor.cs b/Engine/Source/Programs/UnrealBuildTool/Executors/ParallelExecutor.cs
--- a/Engine/Source/Programs/UnrealBuildTool/Executors/ParallelExecutor.cs
+++ b/Engine/Source/Programs/UnrealBuildTool/Executors/ParallelExecutor.cs
@@ -183,19 +183,20 @@ namespace UnrealBuildTool
 										CompletedThread.Join();
 										ExecutingActions.Remove(CompletedAction);
 
+										// Update the progress
+										NumCompletedActions++;
+										ProgressWriter.Write(NumCompletedActions, InputActions.Count);
+
 										// Write it to the log
 										if(CompletedAction.LogLines.Count > 0)
 										{
-											foreach(string LogLine in CompletedAction.LogLines)
+											Log.TraceInformation("[{0}/{1}] {2}", NumCompletedActions, InputActions.Count, CompletedAction.LogLines[0]);
+											for(int LineIdx = 1; LineIdx < CompletedAction.LogLines.Count; LineIdx++)
 											{
-												Log.TraceInformation(LogLine);
+												Log.TraceInformation("{0}", CompletedAction.LogLines[LineIdx]);
 											}
 										}
 
-										// Update the progress
-										NumCompletedActions++;
-										ProgressWriter.Write(NumCompletedActions, InputActions.Count);
-
 										// Check the exit code
 										if(CompletedAction.ExitCode == 0)
 										{
