UBT: Fix spelling of "Intellisense".

#rnx

diff --git a/Engine/Source/Programs/UnrealBuildTool/Configuration/UEBuildModuleCPP.cs b/Engine/Source/Programs/UnrealBuildTool/Configuration/UEBuildModuleCPP.cs
--- a/Engine/Source/Programs/UnrealBuildTool/Configuration/UEBuildModuleCPP.cs
+++ b/Engine/Source/Programs/UnrealBuildTool/Configuration/UEBuildModuleCPP.cs
@@ -292,14 +292,14 @@ namespace UnrealBuildTool
 		{
 			CppCompileEnvironment ModuleCompileEnvironment = CreateModuleCompileEnvironment(Target, BinaryCompileEnvironment);
 			ProjectFile.AddIntelliSensePreprocessorDefinitions(ModuleCompileEnvironment.Definitions);
-			ProjectFile.AddInteliiSenseIncludePaths(ModuleCompileEnvironment.IncludePaths.SystemIncludePaths, true);
-			ProjectFile.AddInteliiSenseIncludePaths(ModuleCompileEnvironment.IncludePaths.UserIncludePaths, false);
+			ProjectFile.AddIntelliSenseIncludePaths(ModuleCompileEnvironment.IncludePaths.SystemIncludePaths, true);
+			ProjectFile.AddIntelliSenseIncludePaths(ModuleCompileEnvironment.IncludePaths.UserIncludePaths, false);
 
 			// This directory may not exist for this module (or ever exist, if it doesn't contain any generated headers), but we want the project files
 			// to search it so we can pick up generated code definitions after UHT is run for the first time.
 			if(GeneratedCodeDirectory != null)
 			{
-				ProjectFile.AddInteliiSenseIncludePaths(new HashSet<DirectoryReference>{ GeneratedCodeDirectory }, false);
+				ProjectFile.AddIntelliSenseIncludePaths(new HashSet<DirectoryReference>{ GeneratedCodeDirectory }, false);
 			}
 		}
 
diff --git a/Engine/Source/Programs/UnrealBuildTool/ProjectFiles/Project.cs b/Engine/Source/Programs/UnrealBuildTool/ProjectFiles/Project.cs
--- a/Engine/Source/Programs/UnrealBuildTool/ProjectFiles/Project.cs
+++ b/Engine/Source/Programs/UnrealBuildTool/ProjectFiles/Project.cs
@@ -318,7 +318,7 @@ namespace UnrealBuildTool
 		/// </summary>
 		/// <param name="NewIncludePaths">List of include paths to add</param>
 		/// <param name="bAddingSystemIncludes"></param>
-		public void AddInteliiSenseIncludePaths(HashSet<DirectoryReference> NewIncludePaths, bool bAddingSystemIncludes)
+		public void AddIntelliSenseIncludePaths(HashSet<DirectoryReference> NewIncludePaths, bool bAddingSystemIncludes)
 		{
 			if (ProjectFileGenerator.OnlyGenerateIntelliSenseDataForProject == null ||
 				ProjectFileGenerator.OnlyGenerateIntelliSenseDataForProject == this)
