UBT: Remote compilation now uploads the entire workspace to the remote Mac and executes a separate remote instance of UBT rather than synchronizing individual actions. This makes the remote compile codepath much simpler, and removes a lot of special cases that exist to support it previously.

The list of files to be transferred to the remote are listed as rsync filter rules in Engine/Build/Rsync/RsyncEngine.txt and RsyncProject.txt, which are applied to the root engine directory and project directory respectively. Projects that need to customize which files are uploaded can add their own <ProjectDir>/Build/Rsync/RsyncProject.txt file, which will be included in the filter before the default version.

diff --git a/Engine/Build/BatchFiles/Mac/Build.sh b/Engine/Build/BatchFiles/Mac/Build.sh
--- a/Engine/Build/BatchFiles/Mac/Build.sh
+++ b/Engine/Build/BatchFiles/Mac/Build.sh
@@ -39,15 +39,15 @@ case $ACTION in
 		case $2 in 
 			"iphoneos"|"IOS")
 		        Platform="IOS"
-				AdditionalFlags+=" -deploy -nocreatestub "
+				AdditionalFlags+=" -deploy "
 			;; 
 			"appletvos")
 		        Platform="TVOS"
-				AdditionalFlags+=" -deploy -nocreatestub "
+				AdditionalFlags+=" -deploy "
 			;; 
   			"iphonesimulator")
 		        	Platform="IOS"
-		         	AdditionalFlags+=" -deploy -simulator -nocreatestub"
+		         	AdditionalFlags+=" -deploy -simulator"
 			;;
 			"macosx")
 					Platform="Mac"
@@ -81,12 +81,12 @@ case $ACTION in
 		case $3 in 
 			"iphoneos"|"IOS")
 	        	        Platform="IOS"
-				AdditionalFlags+=" -nocreatestub"
+				AdditionalFlags+=" "
 			;;
 			"iphonesimulator")
 			        Platform="IOS"
 		        	AdditionalFlags+=" -simulator"
-				AdditionalFlags+=" -nocreatestub"
+				AdditionalFlags+=" "
 			;; 
 			"macosx")
 					Platform="Mac"
diff --git a/Engine/Build/Rsync/RsyncEngine.txt b/Engine/Build/Rsync/RsyncEngine.txt
new file mode 100644
--- /dev/null
+++ b/Engine/Build/Rsync/RsyncEngine.txt
@@ -0,0 +1,284 @@
+# This file contains filter rules for files to transfer to the remote Mac for compilation. Since UHT is also executed remotely, it needs everything to require that on Mac.
+# 
+# Syntax is described here: https://download.samba.org/pub/rsync/rsync.html
+# 
+# Note that:
+#  * For any directory or file that is checked against this list, patterns are compared in order and the first matching rule applies (ie. subsequent inversions of that rule will be ignored)
+#  * Directories will not be recursed through unless their parent directories match a rule.
+#  * If directories match a protect rule, it applies to the entire directory and short-circuits traversal of the contents.
+
+# Don't remove build products from previous builds
++ /Binaries/
+protect /Binaries/Mac/**
+protect /Binaries/IOS/**
+protect /Binaries/TVOS/**
+protect /Build/IOS/*.PackageVersionCounter
++ /Intermediate/
++ /Intermediate/Build/
+protect /Intermediate/Build/Mac/**
+protect /Intermediate/Build/IOS/**
+protect /Intermediate/Build/TVOS/**
+protect /Intermediate/IOS/**
+protect /Intermediate/Remote/**
+protect /**/Binaries/Mac/**
+protect /**/Binaries/IOS/**
+protect /**/Binaries/TVOS/**
+protect /**/Intermediate/Build/Mac/**
+protect /**/Intermediate/Build/IOS/**
+protect /**/Intermediate/Build/TVOS/**
+protect /Saved/**
++ /Programs/
+protect /Programs/UnrealBuildTool/**
+protect /Programs/UnrealHeaderTool/Intermediate/**
+protect /Programs/UnrealHeaderTool/Saved/**
+
+# Don't overwrite symlinks created by SetupMono.sh
+protect /Binaries/ThirdParty/Mono/Mac/bin/mono
+protect /Binaries/ThirdParty/Mono/Mac/lib/mono/4.5/Accessibility.dll
+protect /Binaries/ThirdParty/Mono/Mac/lib/mono/4.5/Commons.Xml.Relaxng.dll
+protect /Binaries/ThirdParty/Mono/Mac/lib/mono/4.5/cscompmgd.dll
+protect /Binaries/ThirdParty/Mono/Mac/lib/mono/4.5/CustomMarshalers.dll
+protect /Binaries/ThirdParty/Mono/Mac/lib/mono/4.5/I18N.CJK.dll
+protect /Binaries/ThirdParty/Mono/Mac/lib/mono/4.5/I18N.dll
+protect /Binaries/ThirdParty/Mono/Mac/lib/mono/4.5/I18N.MidEast.dll
+protect /Binaries/ThirdParty/Mono/Mac/lib/mono/4.5/I18N.Other.dll
+protect /Binaries/ThirdParty/Mono/Mac/lib/mono/4.5/I18N.Rare.dll
+protect /Binaries/ThirdParty/Mono/Mac/lib/mono/4.5/I18N.West.dll
+protect /Binaries/ThirdParty/Mono/Mac/lib/mono/4.5/IBM.Data.DB2.dll
+protect /Binaries/ThirdParty/Mono/Mac/lib/mono/4.5/ICSharpCode.SharpZipLib.dll
+protect /Binaries/ThirdParty/Mono/Mac/lib/mono/4.5/Microsoft.Build.dll
+protect /Binaries/ThirdParty/Mono/Mac/lib/mono/4.5/Microsoft.Build.Engine.dll
+protect /Binaries/ThirdParty/Mono/Mac/lib/mono/4.5/Microsoft.Build.Framework.dll
+protect /Binaries/ThirdParty/Mono/Mac/lib/mono/4.5/Microsoft.Build.Tasks.v4.0.dll
+protect /Binaries/ThirdParty/Mono/Mac/lib/mono/4.5/Microsoft.Build.Utilities.v4.0.dll
+protect /Binaries/ThirdParty/Mono/Mac/lib/mono/4.5/Microsoft.CSharp.dll
+protect /Binaries/ThirdParty/Mono/Mac/lib/mono/4.5/Microsoft.VisualC.dll
+protect /Binaries/ThirdParty/Mono/Mac/lib/mono/4.5/Microsoft.Web.Infrastructure.dll
+protect /Binaries/ThirdParty/Mono/Mac/lib/mono/4.5/Mono.C5.dll
+protect /Binaries/ThirdParty/Mono/Mac/lib/mono/4.5/Mono.Cairo.dll
+protect /Binaries/ThirdParty/Mono/Mac/lib/mono/4.5/Mono.CodeContracts.dll
+protect /Binaries/ThirdParty/Mono/Mac/lib/mono/4.5/Mono.CompilerServices.SymbolWriter.dll
+protect /Binaries/ThirdParty/Mono/Mac/lib/mono/4.5/Mono.CSharp.dll
+protect /Binaries/ThirdParty/Mono/Mac/lib/mono/4.5/Mono.Data.Sqlite.dll
+protect /Binaries/ThirdParty/Mono/Mac/lib/mono/4.5/Mono.Data.Tds.dll
+protect /Binaries/ThirdParty/Mono/Mac/lib/mono/4.5/Mono.Debugger.Soft.dll
+protect /Binaries/ThirdParty/Mono/Mac/lib/mono/4.5/Mono.Http.dll
+protect /Binaries/ThirdParty/Mono/Mac/lib/mono/4.5/Mono.Management.dll
+protect /Binaries/ThirdParty/Mono/Mac/lib/mono/4.5/Mono.Messaging.dll
+protect /Binaries/ThirdParty/Mono/Mac/lib/mono/4.5/Mono.Messaging.RabbitMQ.dll
+protect /Binaries/ThirdParty/Mono/Mac/lib/mono/4.5/Mono.Parallel.dll
+protect /Binaries/ThirdParty/Mono/Mac/lib/mono/4.5/Mono.Posix.dll
+protect /Binaries/ThirdParty/Mono/Mac/lib/mono/4.5/Mono.Security.dll
+protect /Binaries/ThirdParty/Mono/Mac/lib/mono/4.5/Mono.Security.Win32.dll
+protect /Binaries/ThirdParty/Mono/Mac/lib/mono/4.5/Mono.Simd.dll
+protect /Binaries/ThirdParty/Mono/Mac/lib/mono/4.5/Mono.Tasklets.dll
+protect /Binaries/ThirdParty/Mono/Mac/lib/mono/4.5/Mono.WebBrowser.dll
+protect /Binaries/ThirdParty/Mono/Mac/lib/mono/4.5/Mono.XBuild.Tasks.dll
+protect /Binaries/ThirdParty/Mono/Mac/lib/mono/4.5/Novell.Directory.Ldap.dll
+protect /Binaries/ThirdParty/Mono/Mac/lib/mono/4.5/nunit-console-runner.dll
+protect /Binaries/ThirdParty/Mono/Mac/lib/mono/4.5/nunit.core.dll
+protect /Binaries/ThirdParty/Mono/Mac/lib/mono/4.5/nunit.core.extensions.dll
+protect /Binaries/ThirdParty/Mono/Mac/lib/mono/4.5/nunit.core.interfaces.dll
+protect /Binaries/ThirdParty/Mono/Mac/lib/mono/4.5/nunit.framework.dll
+protect /Binaries/ThirdParty/Mono/Mac/lib/mono/4.5/nunit.framework.extensions.dll
+protect /Binaries/ThirdParty/Mono/Mac/lib/mono/4.5/nunit.mocks.dll
+protect /Binaries/ThirdParty/Mono/Mac/lib/mono/4.5/nunit.util.dll
+protect /Binaries/ThirdParty/Mono/Mac/lib/mono/4.5/PEAPI.dll
+protect /Binaries/ThirdParty/Mono/Mac/lib/mono/4.5/RabbitMQ.Client.dll
+protect /Binaries/ThirdParty/Mono/Mac/lib/mono/4.5/System.ComponentModel.Composition.dll
+protect /Binaries/ThirdParty/Mono/Mac/lib/mono/4.5/System.ComponentModel.DataAnnotations.dll
+protect /Binaries/ThirdParty/Mono/Mac/lib/mono/4.5/System.Configuration.dll
+protect /Binaries/ThirdParty/Mono/Mac/lib/mono/4.5/System.Configuration.Install.dll
+protect /Binaries/ThirdParty/Mono/Mac/lib/mono/4.5/System.Core.dll
+protect /Binaries/ThirdParty/Mono/Mac/lib/mono/4.5/System.Data.DataSetExtensions.dll
+protect /Binaries/ThirdParty/Mono/Mac/lib/mono/4.5/System.Data.dll
+protect /Binaries/ThirdParty/Mono/Mac/lib/mono/4.5/System.Data.Entity.dll
+protect /Binaries/ThirdParty/Mono/Mac/lib/mono/4.5/System.Data.Linq.dll
+protect /Binaries/ThirdParty/Mono/Mac/lib/mono/4.5/System.Data.OracleClient.dll
+protect /Binaries/ThirdParty/Mono/Mac/lib/mono/4.5/System.Data.Services.Client.dll
+protect /Binaries/ThirdParty/Mono/Mac/lib/mono/4.5/System.Data.Services.dll
+protect /Binaries/ThirdParty/Mono/Mac/lib/mono/4.5/System.Design.dll
+protect /Binaries/ThirdParty/Mono/Mac/lib/mono/4.5/System.DirectoryServices.dll
+protect /Binaries/ThirdParty/Mono/Mac/lib/mono/4.5/System.DirectoryServices.Protocols.dll
+protect /Binaries/ThirdParty/Mono/Mac/lib/mono/4.5/System.dll
+protect /Binaries/ThirdParty/Mono/Mac/lib/mono/4.5/System.Drawing.Design.dll
+protect /Binaries/ThirdParty/Mono/Mac/lib/mono/4.5/System.Drawing.dll
+protect /Binaries/ThirdParty/Mono/Mac/lib/mono/4.5/System.Dynamic.dll
+protect /Binaries/ThirdParty/Mono/Mac/lib/mono/4.5/System.EnterpriseServices.dll
+protect /Binaries/ThirdParty/Mono/Mac/lib/mono/4.5/System.IdentityModel.dll
+protect /Binaries/ThirdParty/Mono/Mac/lib/mono/4.5/System.IdentityModel.Selectors.dll
+protect /Binaries/ThirdParty/Mono/Mac/lib/mono/4.5/System.IO.Compression.dll
+protect /Binaries/ThirdParty/Mono/Mac/lib/mono/4.5/System.IO.Compression.FileSystem.dll
+protect /Binaries/ThirdParty/Mono/Mac/lib/mono/4.5/System.Json.dll
+protect /Binaries/ThirdParty/Mono/Mac/lib/mono/4.5/System.Json.Microsoft.dll
+protect /Binaries/ThirdParty/Mono/Mac/lib/mono/4.5/System.Management.dll
+protect /Binaries/ThirdParty/Mono/Mac/lib/mono/4.5/System.Messaging.dll
+protect /Binaries/ThirdParty/Mono/Mac/lib/mono/4.5/System.Net.dll
+protect /Binaries/ThirdParty/Mono/Mac/lib/mono/4.5/System.Net.Http.dll
+protect /Binaries/ThirdParty/Mono/Mac/lib/mono/4.5/System.Net.Http.Formatting.dll
+protect /Binaries/ThirdParty/Mono/Mac/lib/mono/4.5/System.Net.Http.WebRequest.dll
+protect /Binaries/ThirdParty/Mono/Mac/lib/mono/4.5/System.Numerics.dll
+protect /Binaries/ThirdParty/Mono/Mac/lib/mono/4.5/System.Reactive.Core.dll
+protect /Binaries/ThirdParty/Mono/Mac/lib/mono/4.5/System.Reactive.Debugger.dll
+protect /Binaries/ThirdParty/Mono/Mac/lib/mono/4.5/System.Reactive.Experimental.dll
+protect /Binaries/ThirdParty/Mono/Mac/lib/mono/4.5/System.Reactive.Interfaces.dll
+protect /Binaries/ThirdParty/Mono/Mac/lib/mono/4.5/System.Reactive.Linq.dll
+protect /Binaries/ThirdParty/Mono/Mac/lib/mono/4.5/System.Reactive.Observable.Aliases.dll
+protect /Binaries/ThirdParty/Mono/Mac/lib/mono/4.5/System.Reactive.PlatformServices.dll
+protect /Binaries/ThirdParty/Mono/Mac/lib/mono/4.5/System.Reactive.Providers.dll
+protect /Binaries/ThirdParty/Mono/Mac/lib/mono/4.5/System.Reactive.Runtime.Remoting.dll
+protect /Binaries/ThirdParty/Mono/Mac/lib/mono/4.5/System.Reactive.Windows.Forms.dll
+protect /Binaries/ThirdParty/Mono/Mac/lib/mono/4.5/System.Reactive.Windows.Threading.dll
+protect /Binaries/ThirdParty/Mono/Mac/lib/mono/4.5/System.Runtime.Caching.dll
+protect /Binaries/ThirdParty/Mono/Mac/lib/mono/4.5/System.Runtime.DurableInstancing.dll
+protect /Binaries/ThirdParty/Mono/Mac/lib/mono/4.5/System.Runtime.Remoting.dll
+protect /Binaries/ThirdParty/Mono/Mac/lib/mono/4.5/System.Runtime.Serialization.dll
+protect /Binaries/ThirdParty/Mono/Mac/lib/mono/4.5/System.Runtime.Serialization.Formatters.Soap.dll
+protect /Binaries/ThirdParty/Mono/Mac/lib/mono/4.5/System.Security.dll
+protect /Binaries/ThirdParty/Mono/Mac/lib/mono/4.5/System.ServiceModel.Activation.dll
+protect /Binaries/ThirdParty/Mono/Mac/lib/mono/4.5/System.ServiceModel.Discovery.dll
+protect /Binaries/ThirdParty/Mono/Mac/lib/mono/4.5/System.ServiceModel.dll
+protect /Binaries/ThirdParty/Mono/Mac/lib/mono/4.5/System.ServiceModel.Routing.dll
+protect /Binaries/ThirdParty/Mono/Mac/lib/mono/4.5/System.ServiceModel.Web.dll
+protect /Binaries/ThirdParty/Mono/Mac/lib/mono/4.5/System.ServiceProcess.dll
+protect /Binaries/ThirdParty/Mono/Mac/lib/mono/4.5/System.Threading.Tasks.Dataflow.dll
+protect /Binaries/ThirdParty/Mono/Mac/lib/mono/4.5/System.Transactions.dll
+protect /Binaries/ThirdParty/Mono/Mac/lib/mono/4.5/System.Web.Abstractions.dll
+protect /Binaries/ThirdParty/Mono/Mac/lib/mono/4.5/System.Web.ApplicationServices.dll
+protect /Binaries/ThirdParty/Mono/Mac/lib/mono/4.5/System.Web.dll
+protect /Binaries/ThirdParty/Mono/Mac/lib/mono/4.5/System.Web.DynamicData.dll
+protect /Binaries/ThirdParty/Mono/Mac/lib/mono/4.5/System.Web.Extensions.Design.dll
+protect /Binaries/ThirdParty/Mono/Mac/lib/mono/4.5/System.Web.Extensions.dll
+protect /Binaries/ThirdParty/Mono/Mac/lib/mono/4.5/System.Web.Http.dll
+protect /Binaries/ThirdParty/Mono/Mac/lib/mono/4.5/System.Web.Http.SelfHost.dll
+protect /Binaries/ThirdParty/Mono/Mac/lib/mono/4.5/System.Web.Http.WebHost.dll
+protect /Binaries/ThirdParty/Mono/Mac/lib/mono/4.5/System.Web.Mvc.dll
+protect /Binaries/ThirdParty/Mono/Mac/lib/mono/4.5/System.Web.Razor.dll
+protect /Binaries/ThirdParty/Mono/Mac/lib/mono/4.5/System.Web.Routing.dll
+protect /Binaries/ThirdParty/Mono/Mac/lib/mono/4.5/System.Web.Services.dll
+protect /Binaries/ThirdParty/Mono/Mac/lib/mono/4.5/System.Web.WebPages.Deployment.dll
+protect /Binaries/ThirdParty/Mono/Mac/lib/mono/4.5/System.Web.WebPages.dll
+protect /Binaries/ThirdParty/Mono/Mac/lib/mono/4.5/System.Web.WebPages.Razor.dll
+protect /Binaries/ThirdParty/Mono/Mac/lib/mono/4.5/System.Windows.dll
+protect /Binaries/ThirdParty/Mono/Mac/lib/mono/4.5/System.Windows.Forms.DataVisualization.dll
+protect /Binaries/ThirdParty/Mono/Mac/lib/mono/4.5/System.Windows.Forms.dll
+protect /Binaries/ThirdParty/Mono/Mac/lib/mono/4.5/System.Xaml.dll
+protect /Binaries/ThirdParty/Mono/Mac/lib/mono/4.5/System.Xml.dll
+protect /Binaries/ThirdParty/Mono/Mac/lib/mono/4.5/System.Xml.Linq.dll
+protect /Binaries/ThirdParty/Mono/Mac/lib/mono/4.5/System.Xml.Serialization.dll
+protect /Binaries/ThirdParty/Mono/Mac/lib/mono/4.5/WebMatrix.Data.dll
+protect /Binaries/ThirdParty/Mono/Mac/lib/mono/4.5/WindowsBase.dll
+protect /Binaries/ThirdParty/Mono/Mac/lib/mono/xbuild/12.0/bin/Microsoft.Build.dll
+protect /Binaries/ThirdParty/Mono/Mac/lib/mono/xbuild/12.0/bin/Microsoft.Build.Engine.dll
+protect /Binaries/ThirdParty/Mono/Mac/lib/mono/xbuild/12.0/bin/Microsoft.Build.Framework.dll
+protect /Binaries/ThirdParty/Mono/Mac/lib/mono/xbuild/12.0/bin/Microsoft.Build.Tasks.v12.0.dll
+protect /Binaries/ThirdParty/Mono/Mac/lib/mono/xbuild/12.0/bin/Microsoft.Build.Utilities.v12.0.dll
+protect /Binaries/ThirdParty/Mono/Mac/lib/mono/xbuild/12.0/bin/Mono.XBuild.Tasks.dll
+protect /Binaries/ThirdParty/Mono/Mac/lib/mono/xbuild/14.0/bin/Microsoft.Build.dll
+protect /Binaries/ThirdParty/Mono/Mac/lib/mono/xbuild/14.0/bin/Microsoft.Build.Engine.dll
+protect /Binaries/ThirdParty/Mono/Mac/lib/mono/xbuild/14.0/bin/Microsoft.Build.Framework.dll
+protect /Binaries/ThirdParty/Mono/Mac/lib/mono/xbuild/14.0/bin/Microsoft.Build.Tasks.Core.dll
+protect /Binaries/ThirdParty/Mono/Mac/lib/mono/xbuild/14.0/bin/Microsoft.Build.Utilities.Core.dll
+protect /Binaries/ThirdParty/Mono/Mac/lib/mono/xbuild/14.0/bin/Mono.XBuild.Tasks.dll
+protect /Binaries/ThirdParty/Mono/Mac/lib/mono/monodoc/monodoc.dll
+
+# Exclude anything for platforms we don't care about
+- /**/Win32/
+- /**/Win64/
+- /**/Android/
+- /**/Linux/
+- /**/PS4/
+- /**/Switch/
+- /**/XboxOne/
+
+# Also exclude any C# intermediate directories
+- /Source/Programs/**/obj/
+
+# Add all the paths to sync
++ /
++ /Binaries/
++ /Binaries/DotNET/
++ /Binaries/DotNET/UnrealBuildTool.exe
++ /Binaries/DotNET/UnrealBuildTool.pdb
++ /Binaries/DotNET/DotNETUtilities.dll
++ /Binaries/DotNET/DotNETUtilities.pdb
++ /Binaries/DotNET/RPCUtility.exe
++ /Binaries/DotNET/RPCUtility.pdb
++ /Binaries/DotNET/Ionic.Zip.Reduced.dll
++ /Binaries/DotNET/IOS/
++ /Binaries/DotNET/IOS/**
++ /Binaries/ThirdParty/
++ /Binaries/ThirdParty/Ionic/
++ /Binaries/ThirdParty/Ionic/Ionic.Zip.Reduced.dll
++ /Binaries/ThirdParty/Ionic/Ionic.Zip.Reduced.Core.dll
++ /Binaries/ThirdParty/Ionic/Ionic.Zip.Reduced.Core.pdb
++ /Binaries/ThirdParty/Mono/
++ /Binaries/ThirdParty/Mono/Mac/
++ /Binaries/ThirdParty/Mono/Mac/**
++ /Build/
++ /Build/*
++ /Build/BatchFiles/
++ /Build/BatchFiles/**
++ /Build/IOS/
++ /Build/IOS/**
++ /Config/
++ /Config/**
++ /Intermediate/
++ /Intermediate/Build/
++ /Intermediate/Build/XmlConfigCache.bin
++ /Intermediate/Build/BuildRules/
++ /Intermediate/Build/BuildRules/*.dll
++ /Programs/
++ /Programs/UnrealHeaderTool/
++ /Programs/UnrealHeaderTool/Config/
++ /Programs/UnrealHeaderTool/Config/*.ini
++ /Plugins/
++ /Plugins/**/
++ /Plugins/**.uplugin
++ /Plugins/**/Source/**.h
++ /Plugins/**/Source/**.cpp
++ /Plugins/**/Source/**.mm
++ /Plugins/**/Source/**.m
++ /Plugins/**/Source/**.cs
++ /Plugins/**/Source/**.inl
++ /Source/
++ /Source/*
++ /Source/*.cs
++ /Source/**/
++ /Source/Runtime/**.h
++ /Source/Runtime/**.cpp
++ /Source/Runtime/**.mm
++ /Source/Runtime/**.m
++ /Source/Runtime/**.cs
++ /Source/Runtime/**.inl
++ /Source/Developer/**.h
++ /Source/Developer/**.cpp
++ /Source/Developer/**.mm
++ /Source/Developer/**.m
++ /Source/Developer/**.cs
++ /Source/Developer/**.inl
++ /Source/Editor/**.h
++ /Source/Editor/**.cpp
++ /Source/Editor/**.mm
++ /Source/Editor/**.m
++ /Source/Editor/**.cs
++ /Source/Editor/**.inl
++ /Source/Programs/**.h
++ /Source/Programs/**.cpp
++ /Source/Programs/**.mm
++ /Source/Programs/**.m
++ /Source/Programs/**.cs
++ /Source/Programs/**.inl
++ /Source/Programs/**.def
++ /Source/Programs/**.csproj
++ /Source/ThirdParty/
++ /Source/ThirdParty/**/
++ /Source/ThirdParty/**.cs
++ /Source/ThirdParty/**.h
++ /Source/ThirdParty/**.hh
++ /Source/ThirdParty/**.hpp
++ /Source/ThirdParty/**.mm
++ /Source/ThirdParty/**.inl
++ /Source/ThirdParty/**.a
++ /Source/ThirdParty/**.embeddedframework.zip
++ /Source/ThirdParty/libJPG/*.cpp
diff --git a/Engine/Build/Rsync/RsyncEngineScripts.txt b/Engine/Build/Rsync/RsyncEngineScripts.txt
new file mode 100644
--- /dev/null
+++ b/Engine/Build/Rsync/RsyncEngineScripts.txt
@@ -0,0 +1,10 @@
+; This file contains a list of shell scripts to be synced to the remote. These files will be locally converted to use Unix line endings, and will be
+; marked as executable on the remote. Wildcards are not supported.
+
+Binaries/ThirdParty/Mono/Mac/bin/xbuild
+Build/BatchFiles/Mac/Build.sh
+Build/BatchFiles/Mac/FixDependencyFiles.sh
+Build/BatchFiles/Mac/RunMono.sh
+Build/BatchFiles/Mac/FixMonoFiles.sh
+Build/BatchFiles/Mac/GenerateProjectFiles.sh
+Build/BatchFiles/Mac/SetupMono.sh
diff --git a/Engine/Build/Rsync/RsyncProject.txt b/Engine/Build/Rsync/RsyncProject.txt
new file mode 100644
--- /dev/null
+++ b/Engine/Build/Rsync/RsyncProject.txt
@@ -0,0 +1,28 @@
+# This file contains rsync rules that will be applied to the project directory when uploading to the remote. It is merged into the filter list
+# after the engine filter, but before the <ProjectDir>/Build/rsync-filter-project.txt list, if it exists.
+#
+# See rsync-filter-engine.txt for a list of caveats to be aware of.
+
++ /Intermediate/
++ /Intermediate/Build/
++ /Intermediate/Build/BuildRules/
++ /Intermediate/Build/BuildRules/*.dll
+
+protect /Intermediate/**
+
++ /*.uproject
++ /Build/
++ /Build/IOS/
++ /Build/IOS/**
++ /Build/TVOS/
++ /Build/TVOS/**
++ /Config/
++ /Config/**
++ /Plugins/
++ /Plugins/**/
++ /Plugins/**/*.uplugin
++ /Plugins/**/Config/**
++ /Plugins/**/Source/**
++ /Source/
++ /Source/**
+
diff --git a/Engine/Source/Programs/AutomationTool/IOS/IOSPlatform.Automation.cs b/Engine/Source/Programs/AutomationTool/IOS/IOSPlatform.Automation.cs
--- a/Engine/Source/Programs/AutomationTool/IOS/IOSPlatform.Automation.cs
+++ b/Engine/Source/Programs/AutomationTool/IOS/IOSPlatform.Automation.cs
@@ -168,65 +168,6 @@ public class IOSPlatform : Platform
 		SDKName = (TargetPlatform == UnrealTargetPlatform.TVOS) ? "appletvos" : "iphoneos";
 	}
 
-	// Run the integrated IPP code
-	// @todo ipp: How to log the output in a nice UAT way?
-	protected static int RunIPP(string IPPArguments)
-	{
-		List<string> Args = new List<string>();
-
-		StringBuilder Token = null;
-		int Index = 0;
-		bool bInQuote = false;
-		while (Index < IPPArguments.Length)
-		{
-			if (IPPArguments[Index] == '\"')
-			{
-				bInQuote = !bInQuote;
-			}
-			else if (IPPArguments[Index] == ' ' && !bInQuote)
-			{
-				if (Token != null)
-				{
-					Args.Add(Token.ToString());
-					Token = null;
-				}
-			}
-			else
-			{
-				if (Token == null)
-				{
-					Token = new StringBuilder();
-				}
-				Token.Append(IPPArguments[Index]);
-			}
-			Index++;
-		}
-
-		if (Token != null)
-		{
-			Args.Add(Token.ToString());
-		}
-	
-		return RunIPP(Args.ToArray());
-	}
-
-	protected static int RunIPP(string[] Args)
-	{
-		return 5;//@TODO: Reintegrate IPP to IOS.Automation iPhonePackager.Program.RealMain(Args);
-	}
-
-	public override int RunCommand(string Command)
-	{
-		// run generic IPP commands through the interface
-		if (Command.StartsWith("IPP:"))
-		{
-			return RunIPP(Command.Substring(4));
-		}
-
-		// non-zero error code
-		return 4;
-	}
-
 	public virtual bool PrepForUATPackageOrDeploy(UnrealTargetConfiguration Config, FileReference ProjectFile, string InProjectName, DirectoryReference InProjectDirectory, string InExecutablePath, DirectoryReference InEngineDir, bool bForDistribution, string CookFlavor, bool bIsDataDeploy, bool bCreateStubIPA, bool bIsUE4Game)
 	{
 		string TargetName = Path.GetFileNameWithoutExtension(InExecutablePath).Split("-".ToCharArray())[0];
@@ -384,7 +325,7 @@ public class IOSPlatform : Platform
 			}
 		}
 
-		IOSExports.GenerateAssetCatalog(Params.RawProjectPath, FullExePath, CombinePaths(Params.BaseStageDirectory, (Platform == UnrealTargetPlatform.IOS ? "IOS" : "TVOS")), Platform);
+		IOSExports.GenerateAssetCatalog(Params.RawProjectPath, new FileReference(FullExePath), new DirectoryReference(CombinePaths(Params.BaseStageDirectory, (Platform == UnrealTargetPlatform.IOS ? "IOS" : "TVOS"))), Platform);
 
         bCreatedIPA = false;
 		bool bNeedsIPA = false;
@@ -461,137 +402,63 @@ public class IOSPlatform : Platform
 
 				bCreatedIPA = true;
 
-				if (IOSExports.UseRPCUtil())
-				{
-					string IPPArguments = "RepackageFromStage \"" + (Params.IsCodeBasedProject ? Params.RawProjectPath.FullName : "Engine") + "\"";
-					IPPArguments += " -config " + TargetConfiguration.ToString();
-					IPPArguments += " -schemename " + SchemeName + " -schemeconfig \"" + SchemeConfiguration + "\"";
-
-					if (TargetConfiguration == UnrealTargetConfiguration.Shipping)
-					{
-						IPPArguments += " -compress=best";
-					}
-
-					// Determine if we should sign
-					bool bNeedToSign = GetCodeSignDesirability(Params);
+				string IPPArguments = "RepackageFromStage \"" + (Params.IsCodeBasedProject ? Params.RawProjectPath.FullName : "Engine") + "\"";
+				IPPArguments += " -config " + TargetConfiguration.ToString();
+				IPPArguments += " -schemename " + SchemeName + " -schemeconfig \"" + SchemeConfiguration + "\"";
 
-					if (!String.IsNullOrEmpty(Params.BundleName))
-					{
-						// Have to sign when a bundle name is specified
-						bNeedToSign = true;
-						IPPArguments += " -bundlename " + Params.BundleName;
-					}
-
-					if (bNeedToSign)
-					{
-						IPPArguments += " -sign";
-						if (Params.Distribution)
-						{
-							IPPArguments += " -distribution";
-						}
-						if (Params.IsCodeBasedProject)
-						{
-							IPPArguments += (" -codebased");
-						}
-					}
-
-					IPPArguments += (cookonthefly ? " -cookonthefly" : "");
-					IPPArguments += " -stagedir \"" + CombinePaths(Params.BaseStageDirectory, PlatformName) + "\"";
-					IPPArguments += " -project \"" + Params.RawProjectPath + "\"";
-					if (Params.IterativeDeploy)
-					{
-						IPPArguments += " -iterate";
-					}
-					if (!string.IsNullOrEmpty(Params.Provision))
-					{
-						IPPArguments += " -provision \"" + Params.Provision + "\""; 
-					}
-					if (!string.IsNullOrEmpty(Params.Certificate))
-					{
-						IPPArguments += " -certificate \"" + Params.Certificate + "\"";
-					}
-					if (PlatformName == "TVOS")
-					{
-						IPPArguments += " -tvos";
-					}
-					RunAndLog(CmdEnv, IPPExe, IPPArguments);
-
-                    if (IPPProjectIPA.Length > 0)
-                    {
-                        CopyFile(IPPProjectIPA, ProjectIPA);
-                        DeleteFile(IPPProjectIPA);
-                    }
-                }
-                else
+				if (TargetConfiguration == UnrealTargetConfiguration.Shipping)
 				{
-					List<string> IPPArguments = new List<string>();
-					IPPArguments.Add("RepackageFromStage");
-					IPPArguments.Add(Params.IsCodeBasedProject ? Params.RawProjectPath.FullName : "Engine");
-					IPPArguments.Add("-config");
-					IPPArguments.Add(TargetConfiguration.ToString());
-					IPPArguments.Add("-schemename");
-					IPPArguments.Add(SchemeName);
-					IPPArguments.Add("-schemeconfig");
-					IPPArguments.Add("\"" + SchemeConfiguration + "\"");
-					if (Params.IsCodeBasedProject)
-					{
-						IPPArguments.Add("-codebased");
-					}
-
-					if (TargetConfiguration == UnrealTargetConfiguration.Shipping)
-					{
-						IPPArguments.Add("-compress=best");
-					}
+					IPPArguments += " -compress=best";
+				}
 
-					// Determine if we should sign
-					bool bNeedToSign = GetCodeSignDesirability(Params);
+				// Determine if we should sign
+				bool bNeedToSign = GetCodeSignDesirability(Params);
 
-					if (!String.IsNullOrEmpty(Params.BundleName))
-					{
-						// Have to sign when a bundle name is specified
-						bNeedToSign = true;
-						IPPArguments.Add("-bundlename");
-						IPPArguments.Add(Params.BundleName);
-					}
-
-					if (bNeedToSign)
-					{
-						IPPArguments.Add("-sign");
-					}
+				if (!String.IsNullOrEmpty(Params.BundleName))
+				{
+					// Have to sign when a bundle name is specified
+					bNeedToSign = true;
+					IPPArguments += " -bundlename " + Params.BundleName;
+				}
 
-					if (cookonthefly)
-					{
-						IPPArguments.Add(" -cookonthefly");
-					}
-					IPPArguments.Add(" -stagedir");
-					IPPArguments.Add(CombinePaths(Params.BaseStageDirectory, "IOS"));
-					IPPArguments.Add(" -project");
-					IPPArguments.Add(Params.RawProjectPath.FullName);
-					if (Params.IterativeDeploy)
-					{
-						IPPArguments.Add(" -iterate");
-					}
-					if (!string.IsNullOrEmpty(Params.Provision))
+				if (bNeedToSign)
+				{
+					IPPArguments += " -sign";
+					if (Params.Distribution)
 					{
-						IPPArguments.Add(" -provision");
-						IPPArguments.Add(Params.Provision);
+						IPPArguments += " -distribution";
 					}
-					if (!string.IsNullOrEmpty(Params.Certificate))
+					if (Params.IsCodeBasedProject)
 					{
-						IPPArguments.Add(" -certificate");
-						IPPArguments.Add(Params.Certificate);
+						IPPArguments += (" -codebased");
 					}
+				}
 
-					if (RunIPP(IPPArguments.ToArray()) != 0)
-					{
-						throw new AutomationException("IPP Failed");
-					}
+				IPPArguments += (cookonthefly ? " -cookonthefly" : "");
+				IPPArguments += " -stagedir \"" + CombinePaths(Params.BaseStageDirectory, PlatformName) + "\"";
+				IPPArguments += " -project \"" + Params.RawProjectPath + "\"";
+				if (Params.IterativeDeploy)
+				{
+					IPPArguments += " -iterate";
+				}
+				if (!string.IsNullOrEmpty(Params.Provision))
+				{
+					IPPArguments += " -provision \"" + Params.Provision + "\""; 
+				}
+				if (!string.IsNullOrEmpty(Params.Certificate))
+				{
+					IPPArguments += " -certificate \"" + Params.Certificate + "\"";
+				}
+				if (PlatformName == "TVOS")
+				{
+					IPPArguments += " -tvos";
+				}
+				RunAndLog(CmdEnv, IPPExe, IPPArguments);
 
-                    if (IPPProjectIPA.Length > 0)
-                    {
-                        CopyFile(IPPProjectIPA, ProjectIPA);
-                        DeleteFile(IPPProjectIPA);
-                    }
+                if (IPPProjectIPA.Length > 0)
+                {
+                    CopyFile(IPPProjectIPA, ProjectIPA);
+                    DeleteFile(IPPProjectIPA);
                 }
             }
 
diff --git a/Engine/Source/Programs/IOS/iPhonePackager/CompileTime.cs b/Engine/Source/Programs/IOS/iPhonePackager/CompileTime.cs
--- a/Engine/Source/Programs/IOS/iPhonePackager/CompileTime.cs
+++ b/Engine/Source/Programs/IOS/iPhonePackager/CompileTime.cs
@@ -177,6 +177,61 @@ namespace iPhonePackager
 			}
 		}
 
+		static bool FindMobileProvision(string BundleIdentifier, out string OutFileName)
+		{
+			bool bNameMatch;
+			string ProvisionWithPrefix = MobileProvision.FindCompatibleProvision(BundleIdentifier, out bNameMatch, true, true, false);
+			if (!File.Exists(ProvisionWithPrefix))
+			{
+				ProvisionWithPrefix = FileOperations.FindPrefixedFile(Config.BuildDirectory, Program.GameName + ".mobileprovision");
+				if (!File.Exists(ProvisionWithPrefix))
+				{
+					ProvisionWithPrefix = FileOperations.FindPrefixedFile(Config.BuildDirectory + "/NotForLicensees/", Program.GameName + ".mobileprovision");
+					if (!File.Exists(ProvisionWithPrefix))
+					{
+						ProvisionWithPrefix = FileOperations.FindPrefixedFile(Config.EngineBuildDirectory, "UE4Game.mobileprovision");
+						if (!File.Exists(ProvisionWithPrefix))
+						{
+							ProvisionWithPrefix = FileOperations.FindPrefixedFile(Config.EngineBuildDirectory + "/NotForLicensees/", "UE4Game.mobileprovision");
+							if(!File.Exists(ProvisionWithPrefix))
+							{
+								OutFileName = null;
+								return false;
+							}
+						}
+					}
+				}
+			}
+
+			OutFileName = ProvisionWithPrefix;
+			return true;
+		}
+
+		/// <summary>
+		/// Export the certificate to a file
+		/// </summary>
+		static public void ExportCertificate()
+		{
+			string ProvisionWithPrefix;
+			if(!FindMobileProvision("", out ProvisionWithPrefix))
+			{
+				Program.Error("Missing provision");
+				return;
+			}
+
+			if(Config.Certificate == null)
+			{
+				Program.Error("Missing -Certificate=... argument");
+				return;
+			}
+
+			// export the signing certificate to a file
+			MobileProvision Provision = MobileProvisionParser.ParseFile(ProvisionWithPrefix);
+			var Certificate = CodeSignatureBuilder.FindCertificate(Provision);
+			byte[] Data = Certificate.Export(System.Security.Cryptography.X509Certificates.X509ContentType.Pkcs12, "A");
+			File.WriteAllBytes(Config.Certificate, Data);
+		}
+
 		/// <summary>
 		/// Copy the files always needed (even in a stub IPA)
 		/// </summary>
@@ -243,23 +298,12 @@ namespace iPhonePackager
 				// Copy the mobile provision file over
 				string CFBundleIdentifier = null;
 				Info.GetString("CFBundleIdentifier", out CFBundleIdentifier);
-				bool bNameMatch;
-				string ProvisionWithPrefix = MobileProvision.FindCompatibleProvision(CFBundleIdentifier, out bNameMatch, true, true, false);
-				if (!File.Exists(ProvisionWithPrefix))
+
+				string ProvisionWithPrefix;
+				if(!FindMobileProvision(CFBundleIdentifier, out ProvisionWithPrefix))
 				{
-					ProvisionWithPrefix = FileOperations.FindPrefixedFile(Config.BuildDirectory, Program.GameName + ".mobileprovision");
-					if (!File.Exists(ProvisionWithPrefix))
-					{
-						ProvisionWithPrefix = FileOperations.FindPrefixedFile(Config.BuildDirectory + "/NotForLicensees/", Program.GameName + ".mobileprovision");
-						if (!File.Exists(ProvisionWithPrefix))
-						{
-							ProvisionWithPrefix = FileOperations.FindPrefixedFile(Config.EngineBuildDirectory, "UE4Game.mobileprovision");
-							if (!File.Exists(ProvisionWithPrefix))
-							{
-								ProvisionWithPrefix = FileOperations.FindPrefixedFile(Config.EngineBuildDirectory + "/NotForLicensees/", "UE4Game.mobileprovision");
-							}
-						}
-					}
+					Program.Error("Unable to find mobileprovision");
+					return;
 				}
 				FinalMobileProvisionFilename = Path.Combine(Config.PCXcodeStagingDir, MacMobileProvisionFilename);
 				FileOperations.CopyRequiredFile(ProvisionWithPrefix, FinalMobileProvisionFilename);
@@ -688,6 +732,10 @@ namespace iPhonePackager
 					FileOperations.BatchUploadFolder(MacName, Config.PCXcodeStagingDir, MacXcodeStagingDir, false);
 					break;
 
+				case "exportcertificate":
+					CompileTime.ExportCertificate();
+					break;
+
 				default:
 					return false;
 			}
diff --git a/Engine/Source/Programs/IOS/iPhonePackager/Config.cs b/Engine/Source/Programs/IOS/iPhonePackager/Config.cs
--- a/Engine/Source/Programs/IOS/iPhonePackager/Config.cs
+++ b/Engine/Source/Programs/IOS/iPhonePackager/Config.cs
@@ -1,4 +1,4 @@
-﻿/*
+/*
  * Copyright 1998-2018 Epic Games, Inc. All Rights Reserved.
  */
 
@@ -479,7 +479,7 @@ namespace iPhonePackager
 				Program.Architecture);
 
 			// make a directory for the shared XcodeSupportFiles directory
-			Config.PCXcodeStagingDir = Config.PCStagingRootDir + @"..\XcodeSupportFiles";
+			Config.PCXcodeStagingDir = Path.Combine(Config.PCStagingRootDir, @"..\XcodeSupportFiles");
 
 			// Code signing identity
 			// Rules:
diff --git a/Engine/Source/Programs/IOS/iPhonePackager/iPhonePackager.cs b/Engine/Source/Programs/IOS/iPhonePackager/iPhonePackager.cs
--- a/Engine/Source/Programs/IOS/iPhonePackager/iPhonePackager.cs
+++ b/Engine/Source/Programs/IOS/iPhonePackager/iPhonePackager.cs
@@ -1,4 +1,4 @@
-﻿/**
+/**
  * Copyright 1998-2018 Epic Games, Inc. All Rights Reserved.
  */
 
@@ -929,7 +929,7 @@ namespace iPhonePackager
                             if (CodeSignatureBuilder.FindRequiredFiles(out Provision, out Cert, out bHasOverrideFile, out bNameMatch) && Cert != null)
                             {
                                 // print out the provision and cert name
-                                Program.LogVerbose("CERTIFICATE-{0},PROVISION-{1}", Cert.FriendlyName, Path.GetFileName(Provision.FileName));
+                                Program.LogVerbose("CERTIFICATE-{0},PROVISION-{1}", Cert.FriendlyName, Provision.FileName);
                             }
                             else
                             {
diff --git a/Engine/Source/Programs/UnrealBuildTool/Configuration/TargetDescriptor.cs b/Engine/Source/Programs/UnrealBuildTool/Configuration/TargetDescriptor.cs
--- a/Engine/Source/Programs/UnrealBuildTool/Configuration/TargetDescriptor.cs
+++ b/Engine/Source/Programs/UnrealBuildTool/Configuration/TargetDescriptor.cs
@@ -22,6 +22,7 @@ namespace UnrealBuildTool
 		public List<OnlyModule> OnlyModules;
 		public FileReference ForeignPlugin;
 		public string ForceReceiptFileName;
+		public string[] AdditionalArguments;
 
 		/// <summary>
 		/// Constructor
@@ -31,13 +32,15 @@ namespace UnrealBuildTool
 		/// <param name="Platform">Platform to build for</param>
 		/// <param name="Configuration">Configuration to build</param>
 		/// <param name="Architecture">Architecture to build for</param>
-		public TargetDescriptor(FileReference ProjectFile, string TargetName, UnrealTargetPlatform Platform, UnrealTargetConfiguration Configuration, string Architecture)
+		/// <param name="Arguments">Additional arguments for the target</param>
+		public TargetDescriptor(FileReference ProjectFile, string TargetName, UnrealTargetPlatform Platform, UnrealTargetConfiguration Configuration, string Architecture, string[] Arguments)
 		{
 			this.ProjectFile = ProjectFile;
 			this.Name = TargetName;
 			this.Platform = Platform;
 			this.Configuration = Configuration;
 			this.Architecture = Architecture;
+			this.AdditionalArguments = Arguments;
 		}
 
 		/// <summary>
@@ -45,9 +48,10 @@ namespace UnrealBuildTool
 		/// </summary>
 		/// <param name="Arguments">Command-line arguments</param>
 		/// <param name="bUsePrecompiled">Whether to use a precompiled engine distribution</param>
+		/// <param name="bSkipRulesCompile">Whether to skip compiling rules assemblies</param>
 		/// <param name="ProjectFile">The project file, if already set. May be updated if not.</param>
 		/// <returns>List of target descriptors</returns>
-		public static List<TargetDescriptor> ParseCommandLine(string[] Arguments, bool bUsePrecompiled, ref FileReference ProjectFile)
+		public static List<TargetDescriptor> ParseCommandLine(string[] Arguments, bool bUsePrecompiled, bool bSkipRulesCompile, ref FileReference ProjectFile)
 		{
 			UnrealTargetPlatform Platform = UnrealTargetPlatform.Unknown;
 			UnrealTargetConfiguration Configuration = UnrealTargetConfiguration.Unknown;
@@ -57,6 +61,7 @@ namespace UnrealBuildTool
 			List<OnlyModule> OnlyModules = new List<OnlyModule>();
 			FileReference ForeignPlugin = null;
 			string ForceReceiptFileName = null;
+			List<string> AdditionalArguments = new List<string>();
 
 			// Settings for creating/using static libraries for the engine
 			for (int ArgumentIndex = 0; ArgumentIndex < Arguments.Length; ArgumentIndex++)
@@ -141,6 +146,9 @@ namespace UnrealBuildTool
 								throw new BuildException("'-Plugin <Path>' syntax is no longer supported on the command line. Use '-Plugin=<Path>' instead.");
 							case "-RECEIPT":
 								throw new BuildException("'-Receipt <Path>' syntax is no longer supported on the command line. Use '-Receipt=<Path>' instead.");
+							default:
+								AdditionalArguments.Add(Arguments[ArgumentIndex]);
+								break;
 						}
 					}
 				}
@@ -169,7 +177,7 @@ namespace UnrealBuildTool
 				}
 				else
 				{
-					TargetNames.Add(RulesCompiler.CreateProjectRulesAssembly(ProjectFile, bUsePrecompiled).GetTargetNameByType(Type, Platform, Configuration, Architecture, ProjectFile, new ReadOnlyBuildVersion(BuildVersion.ReadDefault())));
+					TargetNames.Add(RulesCompiler.CreateProjectRulesAssembly(ProjectFile, bUsePrecompiled, bSkipRulesCompile).GetTargetNameByType(Type, Platform, Configuration, Architecture, ProjectFile, new ReadOnlyBuildVersion(BuildVersion.ReadDefault())));
 				}
 			}
 
@@ -183,7 +191,7 @@ namespace UnrealBuildTool
 					Log.TraceVerbose("Found project file for {0} - {1}", TargetName, ProjectFile);
 				}
 
-				TargetDescriptor Target = new TargetDescriptor(ProjectFile, TargetName, Platform, Configuration, Architecture);
+				TargetDescriptor Target = new TargetDescriptor(ProjectFile, TargetName, Platform, Configuration, Architecture, AdditionalArguments.ToArray());
 				Target.OnlyModules = OnlyModules;
 				Target.ForeignPlugin = ForeignPlugin;
 				Target.ForceReceiptFileName = ForceReceiptFileName;
diff --git a/Engine/Source/Programs/UnrealBuildTool/Configuration/TargetRules.cs b/Engine/Source/Programs/UnrealBuildTool/Configuration/TargetRules.cs
--- a/Engine/Source/Programs/UnrealBuildTool/Configuration/TargetRules.cs
+++ b/Engine/Source/Programs/UnrealBuildTool/Configuration/TargetRules.cs
@@ -874,6 +874,7 @@ namespace UnrealBuildTool
 		/// <summary>
 		/// If true, then a stub IPA will be generated when compiling is done (minimal files needed for a valid IPA).
 		/// </summary>
+		[CommandLine("-CreateStub", Value = "true")]
 		[CommandLine("-NoCreateStub", Value = "false")]
 		public bool bCreateStubIPA = true;
 
@@ -2109,9 +2110,9 @@ namespace UnrealBuildTool
 			get { return Inner.bGenerateProjectFiles; }
 		}
 
-#if !__MonoCS__
-#pragma warning restore C1591
-#endif
+		#if !__MonoCS__
+		#pragma warning restore C1591
+		#endif
 		#endregion
 
 		/// <summary>
diff --git a/Engine/Source/Programs/UnrealBuildTool/Configuration/UEBuildPlatform.cs b/Engine/Source/Programs/UnrealBuildTool/Configuration/UEBuildPlatform.cs
--- a/Engine/Source/Programs/UnrealBuildTool/Configuration/UEBuildPlatform.cs
+++ b/Engine/Source/Programs/UnrealBuildTool/Configuration/UEBuildPlatform.cs
@@ -303,14 +303,6 @@ namespace UnrealBuildTool
 		{
 		}
 
-		/// <summary>
-		/// Called immediately after UnrealHeaderTool is executed to generated code for all UObjects modules.  Only is called if UnrealHeaderTool was actually run in this session.
-		/// </summary>
-		/// <param name="Manifest">List of UObject modules we generated code for.</param>
-		public virtual void PostCodeGeneration(UHTManifest Manifest)
-		{
-		}
-
 		/// <summary>
 		/// Converts the passed in path from UBT host to compiler native format.
 		/// </summary>
@@ -626,18 +618,6 @@ namespace UnrealBuildTool
 			return false;
 		}
 
-		/// <summary>
-		/// Whether this platform requires the use of absolute paths in Unity files. The compiler will try to combine paths in
-		/// each #include directive with the standard include paths, and unity files in intermediate directories can result in the
-		/// maximum path length being exceeded on Windows. On the other hand, remote compilation requires relative paths so
-		/// dependency checking works correctly on the local machine as well as on the remote machine.
-		/// </summary>
-		/// <returns>bool true if it is required, false if not</returns>
-		public virtual bool UseAbsolutePathsInUnityFiles()
-		{
-			return true;
-		}
-
 		/// <summary>
 		/// Whether this platform should build a monolithic binary
 		/// </summary>
diff --git a/Engine/Source/Programs/UnrealBuildTool/Configuration/UEBuildTarget.cs b/Engine/Source/Programs/UnrealBuildTool/Configuration/UEBuildTarget.cs
--- a/Engine/Source/Programs/UnrealBuildTool/Configuration/UEBuildTarget.cs
+++ b/Engine/Source/Programs/UnrealBuildTool/Configuration/UEBuildTarget.cs
@@ -379,27 +379,28 @@ namespace UnrealBuildTool
 		/// </summary>
 		/// <param name="Desc">Information about the target</param>
 		/// <param name="Arguments">Command line arguments</param>
+		/// <param name="bSkipRulesCompile">Whether to skip compiling any rules assemblies</param>
 		/// <param name="bCompilingSingleFile">Whether we're compiling a single file</param>
 		/// <param name="bUsePrecompiled">Whether to use a precompiled engine/enterprise build</param>
 		/// <param name="Version">The current build version</param>
 		/// <returns>The build target object for the specified build rules source file</returns>
-		public static UEBuildTarget CreateTarget(TargetDescriptor Desc, string[] Arguments, bool bCompilingSingleFile, bool bUsePrecompiled, ReadOnlyBuildVersion Version)
+		public static UEBuildTarget CreateTarget(TargetDescriptor Desc, string[] Arguments, bool bSkipRulesCompile, bool bCompilingSingleFile, bool bUsePrecompiled, ReadOnlyBuildVersion Version)
 		{
 			DateTime CreateTargetStartTime = DateTime.UtcNow;
 
 			RulesAssembly RulesAssembly;
 			if (Desc.ProjectFile != null)
 			{
-				RulesAssembly = RulesCompiler.CreateProjectRulesAssembly(Desc.ProjectFile, bUsePrecompiled);
+				RulesAssembly = RulesCompiler.CreateProjectRulesAssembly(Desc.ProjectFile, bUsePrecompiled, bSkipRulesCompile);
 			}
 			else
 			{
-				RulesAssembly = RulesCompiler.CreateEngineRulesAssembly(bUsePrecompiled);
+				RulesAssembly = RulesCompiler.CreateEngineRulesAssembly(bUsePrecompiled, bSkipRulesCompile);
 
 				if (RulesAssembly.GetTargetFileName(Desc.Name) == null)
 				{
 					// Target isn't part of the engine assembly, try the enterprise assembly
-					RulesAssembly EnterpriseRulesAssembly = RulesCompiler.CreateEnterpriseRulesAssembly(bUsePrecompiled);
+					RulesAssembly EnterpriseRulesAssembly = RulesCompiler.CreateEnterpriseRulesAssembly(bUsePrecompiled, bSkipRulesCompile);
 
 					if (EnterpriseRulesAssembly != null)
 					{
@@ -409,7 +410,7 @@ namespace UnrealBuildTool
 			}
 			if (Desc.ForeignPlugin != null)
 			{
-				RulesAssembly = RulesCompiler.CreatePluginRulesAssembly(Desc.ForeignPlugin, RulesAssembly, true);
+				RulesAssembly = RulesCompiler.CreatePluginRulesAssembly(Desc.ForeignPlugin, bSkipRulesCompile, RulesAssembly, true);
 			}
 
 			FileReference TargetFileName;
@@ -518,7 +519,7 @@ namespace UnrealBuildTool
 				FileReference PluginFile = FileReference.Combine(RulesObject.ProjectFile.Directory, "Intermediate", "Plugins", "NativizedAssets", PlatformName, ProjectTargetType, "NativizedAssets.uplugin");
 				if (FileReference.Exists(PluginFile))
 				{
-					RulesAssembly = RulesCompiler.CreatePluginRulesAssembly(PluginFile, RulesAssembly, false);
+					RulesAssembly = RulesCompiler.CreatePluginRulesAssembly(PluginFile, bSkipRulesCompile, RulesAssembly, false);
 				}
 				else
 				{
@@ -3183,15 +3184,15 @@ namespace UnrealBuildTool
 				{
 					if(Plugin.EnabledByDefault && !ReferencedNames.Contains(Plugin.Name))
 					{
-						ReferencedNames.Add(Plugin.Name);
+							ReferencedNames.Add(Plugin.Name);
 
-						PluginReferenceDescriptor PluginReference = new PluginReferenceDescriptor(Plugin.Name, null, true);
-						PluginReference.bOptional = true;
+							PluginReferenceDescriptor PluginReference = new PluginReferenceDescriptor(Plugin.Name, null, true);
+							PluginReference.bOptional = true;
 
-						AddPlugin(PluginReference, "default plugins", ExcludeFolders, NameToInstance, NameToInfo);
+							AddPlugin(PluginReference, "default plugins", ExcludeFolders, NameToInstance, NameToInfo);
+						}
 					}
 				}
-			}
 
 			// If this is a program, synthesize references for plugins which are enabled via the config file
 			if(TargetType == TargetType.Program)
diff --git a/Engine/Source/Programs/UnrealBuildTool/Executors/LocalExecutor.cs b/Engine/Source/Programs/UnrealBuildTool/Executors/LocalExecutor.cs
--- a/Engine/Source/Programs/UnrealBuildTool/Executors/LocalExecutor.cs
+++ b/Engine/Source/Programs/UnrealBuildTool/Executors/LocalExecutor.cs
@@ -170,25 +170,14 @@ namespace UnrealBuildTool
 					{
 						ActionProcess = new Process();
 						ActionProcess.StartInfo = ActionStartInfo;
-						bool bShouldRedirectOuput = Action.bPrintDebugInfo;
-						if (bShouldRedirectOuput)
-						{
-							ActionStartInfo.RedirectStandardOutput = true;
-							ActionStartInfo.RedirectStandardError = true;
-							ActionProcess.EnableRaisingEvents = true;
-
-							if (Action.bPrintDebugInfo)
-							{
-								ActionProcess.OutputDataReceived += new DataReceivedEventHandler(ActionDebugOutput);
-								ActionProcess.ErrorDataReceived += new DataReceivedEventHandler(ActionDebugOutput);
-							}
-						}
+						ActionStartInfo.RedirectStandardOutput = true;
+						ActionStartInfo.RedirectStandardError = true;
+						ActionProcess.OutputDataReceived += new DataReceivedEventHandler(ActionDebugOutput);
+						ActionProcess.ErrorDataReceived += new DataReceivedEventHandler(ActionDebugOutput);
 						ActionProcess.Start();
-						if (bShouldRedirectOuput)
-						{
-							ActionProcess.BeginOutputReadLine();
-							ActionProcess.BeginErrorReadLine();
-						}
+
+						ActionProcess.BeginOutputReadLine();
+						ActionProcess.BeginErrorReadLine();
 					}
 					catch (Exception ex)
 					{
@@ -255,12 +244,6 @@ namespace UnrealBuildTool
 			// track how long it took
 			Action.EndTime = DateTimeOffset.Now;
 
-			if (!Utils.IsRunningOnMono)
-			{
-				// let RPCUtilHelper clean up anything thread related
-				RPCUtilHelper.OnThreadComplete();
-			}
-
 			// we are done!!
 			bComplete = true;
 		}
diff --git a/Engine/Source/Programs/UnrealBuildTool/Platform/IOS/IOSExports.cs b/Engine/Source/Programs/UnrealBuildTool/Platform/IOS/IOSExports.cs
--- a/Engine/Source/Programs/UnrealBuildTool/Platform/IOS/IOSExports.cs
+++ b/Engine/Source/Programs/UnrealBuildTool/Platform/IOS/IOSExports.cs
@@ -2,6 +2,7 @@
 
 using System;
 using System.Collections.Generic;
+using System.Diagnostics;
 using System.Linq;
 using System.Text;
 using System.Threading.Tasks;
@@ -14,16 +15,6 @@ namespace UnrealBuildTool
 	/// </summary>
 	public static class IOSExports
 	{
-		/// <summary>
-		/// 
-		/// </summary>
-		/// <returns></returns>
-		public static bool UseRPCUtil()
-		{
-			XmlConfig.ReadConfigFiles();
-			return RemoteToolChain.bUseRPCUtil;
-		}
-
 		/// <summary>
 		/// 
 		/// </summary>
@@ -132,16 +123,13 @@ namespace UnrealBuildTool
 		/// <param name="Executable"></param>
 		/// <param name="StageDirectory"></param>
 		/// <param name="PlatformType"></param>
-		public static void GenerateAssetCatalog(FileReference ProjectFile, string Executable, string StageDirectory, UnrealTargetPlatform PlatformType)
+		public static void GenerateAssetCatalog(FileReference ProjectFile, FileReference Executable, DirectoryReference StageDirectory, UnrealTargetPlatform PlatformType)
 		{
-			// Initialize the toolchain.
-			IOSProjectSettings ProjectSettings = ((IOSPlatform)UEBuildPlatform.GetBuildPlatform(PlatformType)).ReadProjectSettings(null);
-			IOSToolChain ToolChain = new IOSToolChain(ProjectFile, ProjectSettings);
+			CppPlatform Platform = PlatformType == UnrealTargetPlatform.IOS ? CppPlatform.IOS : CppPlatform.TVOS;
 
 			// Determine whether the user has modified icons that require a remote Mac to build.
-			CppPlatform Platform = PlatformType == UnrealTargetPlatform.IOS ? CppPlatform.IOS : CppPlatform.TVOS;
 			bool bUserImagesExist = false;
-			ToolChain.GenerateAssetCatalog(Platform, ref bUserImagesExist);
+			DirectoryReference ResourcesDir = IOSToolChain.GenerateAssetCatalog(ProjectFile, Platform, ref bUserImagesExist);
 
 			// Don't attempt to do anything remotely if the user is using the default UE4 images.
 			if (!bUserImagesExist)
@@ -155,111 +143,36 @@ namespace UnrealBuildTool
                 return;
             }
 
-			// Save off the current bUseRPCUtil setting to restore at the end of this function.
-			// At this time, iPhonePackager needs to be called with bUseRPCUtil == true.
-			bool bSaveUseRPCUtil = RemoteToolChain.bUseRPCUtil;
-
-			// Initialize the remote calling environment, taking into account the user's SSH setting.
-			ToolChain.SetUpGlobalEnvironment(false);
-
-			// Build the asset catalog ActionGraph.
-			ActionGraph ActionGraph = new ActionGraph();
-			List<FileItem> OutputFiles = new List<FileItem>();
-			ToolChain.CompileAssetCatalog(FileItem.GetItemByPath(Executable), Platform, ActionGraph, OutputFiles);
-
-			ActionGraph.FinalizeActionGraph();
-
-			// I'm not sure how to derive the UE4Game and Development arguments programmatically.
-			string[] Arguments = new string[] { "UE4Game", (PlatformType == UnrealTargetPlatform.IOS ? "IOS" : "TVOS"), "Development", "-UniqueBuildEnvironment" };
-
-			// Perform all of the setup necessary to actually execute the ActionGraph instance.
-			ReadOnlyBuildVersion Version = new ReadOnlyBuildVersion(BuildVersion.ReadDefault());
-			List<string[]> TargetSettings = new List<string[]>();
-			TargetSettings.Add(Arguments);
-			var Targets = new List<UEBuildTarget>();
-			Dictionary<UEBuildTarget, CPPHeaders> TargetToHeaders = new Dictionary<UEBuildTarget, CPPHeaders>();
-			List<TargetDescriptor> TargetDescs = new List<TargetDescriptor>();
-			foreach (string[] TargetSetting in TargetSettings)
-			{
-				TargetDescs.AddRange(TargetDescriptor.ParseCommandLine(TargetSetting, false, ref ProjectFile));
-			}
-			foreach (TargetDescriptor TargetDesc in TargetDescs)
+			// Compile the asset catalog immediately
+			if(BuildHostPlatform.Current.Platform != UnrealTargetPlatform.Mac)
 			{
-				UEBuildTarget Target = UEBuildTarget.CreateTarget(TargetDesc, Arguments, false, false, Version);
-				if (Target == null)
-				{
-					continue;
-				}
-				Targets.Add(Target);
-				TargetToHeaders.Add(Target, null);
-			}
+				FileReference OutputFile = FileReference.Combine(StageDirectory, "Assets.car");
 
-			bool bIsRemoteCompile = BuildHostPlatform.Current.Platform != UnrealTargetPlatform.Mac;
-
-			// Create the build configuration object, and read the settings
-			BuildConfiguration BuildConfiguration = new BuildConfiguration();
-			XmlConfig.ApplyTo(BuildConfiguration);
-			CommandLine.ParseArguments(Arguments, BuildConfiguration);
-			BuildConfiguration.bUseUBTMakefiles = false;
-
-			Action[] PrerequisiteActions;
-			{
-				HashSet<Action> PrerequisiteActionsSet = new HashSet<Action>();
-				foreach (FileItem OutputFile in OutputFiles)
-				{
-					ActionGraph.GatherPrerequisiteActions(OutputFile, ref PrerequisiteActionsSet);
-				}
-				PrerequisiteActions = PrerequisiteActionsSet.ToArray();
+				RemoteMac Remote = new RemoteMac(ProjectFile);
+				Remote.RunAssetCatalogTool(Platform, ResourcesDir, OutputFile);
 			}
-
-			// Copy any asset catalog files to the remote Mac, if necessary.
-			foreach (UEBuildTarget Target in Targets)
+			else
 			{
-				UEBuildPlatform.GetBuildPlatform(Target.Platform).PreBuildSync();
-			}
+				// Get the output file
+				FileReference OutputFile = IOSToolChain.GetAssetCatalogFile(Platform, Executable);
 
-			// Begin execution of the ActionGraph.
-			Dictionary<UEBuildTarget, List<FileItem>> TargetToOutdatedPrerequisitesMap;
-			List<Action> ActionsToExecute = ActionGraph.GetActionsToExecute(BuildConfiguration, PrerequisiteActions, Targets, TargetToHeaders, true, true, out TargetToOutdatedPrerequisitesMap);
-			string ExecutorName = "Unknown";
-			bool bSuccess = ActionGraph.ExecuteActions(BuildConfiguration, ActionsToExecute, bIsRemoteCompile, out ExecutorName, "", EHotReload.Disabled);
-			if (bSuccess)
-			{
-				if (bIsRemoteCompile)
+				// Delete the Assets.car file to force the asset catalog to build every time, because
+				// removals of files or copies of icons (for instance) with a timestamp earlier than
+				// the last generated Assets.car will result in nothing built.
+				if (FileReference.Exists(OutputFile))
 				{
-					// Copy the remotely built AssetCatalog directory locally.
-					foreach (FileItem OutputFile in OutputFiles)
-					{
-						string RemoteDirectory = System.IO.Path.GetDirectoryName(OutputFile.AbsolutePath).Replace("\\", "/");
-						FileItem LocalExecutable = ToolChain.RemoteToLocalFileItem(FileItem.GetItemByPath(Executable));
-						string LocalDirectory = System.IO.Path.Combine(System.IO.Path.GetDirectoryName(LocalExecutable.AbsolutePath), "AssetCatalog");
-						LocalDirectory = StageDirectory;
-						RPCUtilHelper.CopyDirectory(RemoteDirectory, LocalDirectory, RPCUtilHelper.ECopyOptions.DoNotReplace);
-					}
+					FileReference.Delete(OutputFile);
 				}
-				else
+
+				// Run the process locally
+				using(Process Process = new Process())
 				{
-					// Copy the built AssetCatalog directory to the StageDirectory.
-					foreach (FileItem OutputFile in OutputFiles)
-					{
-						string SourceDirectory = System.IO.Path.GetDirectoryName(OutputFile.AbsolutePath).Replace("\\", "/");
-						System.IO.DirectoryInfo SourceDirectoryInfo = new System.IO.DirectoryInfo(SourceDirectory);
-						if (!System.IO.Directory.Exists(StageDirectory))
-						{
-							System.IO.Directory.CreateDirectory(StageDirectory);
-						}
-						System.IO.FileInfo[] SourceFiles = SourceDirectoryInfo.GetFiles();
-						foreach (System.IO.FileInfo SourceFile in SourceFiles)
-						{
-							string DestinationPath = System.IO.Path.Combine(StageDirectory, SourceFile.Name);
-							SourceFile.CopyTo(DestinationPath, true);
-						}
-					}
+					Process.StartInfo.FileName = "/usr/bin/xcrun";
+					Process.StartInfo.Arguments = IOSToolChain.GetAssetCatalogArgs(Platform, ResourcesDir.FullName, OutputFile.Directory.FullName);; 
+					Process.StartInfo.UseShellExecute = false;
+					Utils.RunLocalProcess(Process);
 				}
 			}
-
-			// Restore the former bUseRPCUtil setting.
-			RemoteToolChain.bUseRPCUtil = bSaveUseRPCUtil;
 		}
 
         /// <summary>
diff --git a/Engine/Source/Programs/UnrealBuildTool/Platform/IOS/IOSToolChain.cs b/Engine/Source/Programs/UnrealBuildTool/Platform/IOS/IOSToolChain.cs
--- a/Engine/Source/Programs/UnrealBuildTool/Platform/IOS/IOSToolChain.cs
+++ b/Engine/Source/Programs/UnrealBuildTool/Platform/IOS/IOSToolChain.cs
@@ -13,6 +13,7 @@ using System.Text;
 using Ionic.Zip;
 using Ionic.Zlib;
 using Tools.DotNETCommon;
+using System.Security.Cryptography.X509Certificates;
 
 namespace UnrealBuildTool
 {
@@ -79,7 +80,7 @@ namespace UnrealBuildTool
 		}
 
 		protected IOSToolChain(CppPlatform TargetPlatform, FileReference InProjectFile, IOSProjectSettings InProjectSettings, Func<IOSToolChainSettings> InCreateSettings)
-			: base(TargetPlatform, UnrealTargetPlatform.Mac, InProjectFile)
+			: base(TargetPlatform, InProjectFile)
 		{
 			ProjectSettings = InProjectSettings;
 			Settings = new Lazy<IOSToolChainSettings>(InCreateSettings);
@@ -483,28 +484,17 @@ namespace UnrealBuildTool
 
 		static void CleanIntermediateDirectory(string Path)
 		{
-			if (BuildHostPlatform.Current.Platform != UnrealTargetPlatform.Mac)
-			{
-				// Delete the intermediate directory on the mac
-				RPCUtilHelper.Command("/", String.Format("rm -rf \"{0}\"", Path), "", null);
+			string ResultsText;
 
-				// Create a fresh intermediate after we delete it
-				RPCUtilHelper.Command("/", String.Format("mkdir -p \"{0}\"", Path), "", null);
-			}
-			else
+			// Delete the local dest directory if it exists
+			if (Directory.Exists(Path))
 			{
-				string ResultsText;
-
-				// Delete the local dest directory if it exists
-				if (Directory.Exists(Path))
-				{
-					// this can deal with linked files
-					RunExecutableAndWait("rm", String.Format("-rf \"{0}\"", Path), out ResultsText);
-				}
-
-				// Create the intermediate local directory
-				RunExecutableAndWait("mkdir", String.Format("-p \"{0}\"", Path), out ResultsText);
+				// this can deal with linked files
+				RunExecutableAndWait("rm", String.Format("-rf \"{0}\"", Path), out ResultsText);
 			}
+
+			// Create the intermediate local directory
+			RunExecutableAndWait("mkdir", String.Format("-p \"{0}\"", Path), out ResultsText);
 		}
 
 		bool IsBitcodeCompilingEnabled(CppConfiguration Configuration)
@@ -770,10 +760,6 @@ namespace UnrealBuildTool
 				FileArguments += string.Format(" \"{0}\"", ConvertPath(SourceFile.AbsolutePath));
 
 				string CompilerPath = Settings.Value.ToolchainDir + IOSCompiler;
-				if (!Utils.IsRunningOnMono && BuildHostPlatform.Current.Platform != UnrealTargetPlatform.Mac)
-				{
-					CompileAction.ActionHandler = new Action.BlockingActionHandler(RPCUtilHelper.RPCActionHandler);
-				}
 
 				string AllArgs = FilePCHArguments + Arguments + FileArguments + CompileEnvironment.AdditionalArguments;
 /*				string SourceText = System.IO.File.ReadAllText(SourceFile.AbsolutePath);
@@ -856,11 +842,6 @@ namespace UnrealBuildTool
 			// Create an action that invokes the linker.
 			Action LinkAction = ActionGraph.Add(ActionType.Link);
 
-			if (!Utils.IsRunningOnMono && BuildHostPlatform.Current.Platform != UnrealTargetPlatform.Mac)
-			{
-				LinkAction.ActionHandler = new Action.BlockingActionHandler(RPCUtilHelper.RPCActionHandler);
-			}
-
 			// RPC utility parameters are in terms of the Mac side
 			LinkAction.WorkingDirectory = GetMacDevSrcRoot();
 
@@ -955,15 +936,7 @@ namespace UnrealBuildTool
 			{
 				bool bIsUE4Game = LinkEnvironment.OutputFilePath.FullName.Contains("UE4Game");
 				FileReference ResponsePath = FileReference.Combine(((!bIsUE4Game && ProjectFile != null) ? ProjectFile.Directory : UnrealBuildTool.EngineDirectory), "Intermediate", "Build", LinkEnvironment.Platform.ToString(), "LinkFileList_" + LinkEnvironment.OutputFilePath.GetFileNameWithoutExtension() + ".tmp");
-				if (!Utils.IsRunningOnMono && BuildHostPlatform.Current.Platform != UnrealTargetPlatform.Mac)
-				{
-					FileItem.CreateIntermediateTextFile(ResponsePath, InputFileNames);
-					RPCUtilHelper.CopyFile(ResponsePath.FullName, ConvertPath(ResponsePath.FullName), true);
-				}
-				else
-				{
-					FileItem.CreateIntermediateTextFile(new FileReference(ConvertPath(ResponsePath.FullName)), InputFileNames);
-				}
+				FileItem.CreateIntermediateTextFile(new FileReference(ConvertPath(ResponsePath.FullName)), InputFileNames);
 				LinkCommandArguments += string.Format(" @\"{0}\"", ConvertPath(ResponsePath.FullName));
 			}
 
@@ -1004,77 +977,65 @@ namespace UnrealBuildTool
 			return RemoteOutputFile;
 		}
 
-        public FileItem CompileAssetCatalogHelper(FileItem Executable, string EngineDir, string BuildDir, string IntermediateDir, ActionGraph ActionGraph, CppPlatform Platform)
-        {
-            // Make a file item for the source and destination files
-            FileItem LocalExecutable = RemoteToLocalFileItem(Executable);
-			int Index = Path.GetFileName(LocalExecutable.AbsolutePath).IndexOf("-");
-			if (Index > 0)
+		static string GetAppBundleName(FileReference Executable)
+		{
+			// Get the app bundle name
+			string AppBundleName = Executable.GetFileNameWithoutExtension();
+
+			// Strip off any platform suffix
+			int SuffixIdx = AppBundleName.IndexOf('-');
+			if(SuffixIdx != -1)
 			{
-				LocalExecutable = FileItem.GetItemByPath(Path.Combine(Path.GetDirectoryName(LocalExecutable.AbsolutePath), Path.GetFileName(LocalExecutable.AbsolutePath).Substring(0, Index)));
+				AppBundleName = AppBundleName.Substring(0, SuffixIdx);
 			}
-            string FullDestPathRoot = Path.Combine(Path.GetDirectoryName(LocalExecutable.AbsolutePath), "AssetCatalog", "Assets.car");
-
-            FileItem OutputFile;
-            if (BuildHostPlatform.Current.Platform == UnrealTargetPlatform.Mac && Platform == CppPlatform.IOS)
-            {
-                FullDestPathRoot = Path.Combine(Path.GetDirectoryName(LocalExecutable.AbsolutePath), "Payload", Path.GetFileNameWithoutExtension(LocalExecutable.AbsolutePath) + ".app", "Assets.car");
-            }
-
-            // Delete the Assets.car file to force the asset catalog to build every time, because
-            // removals of files or copies of icons (for instance) with a timestamp earlier than
-            // the last generated Assets.car will result in nothing built.
-            if (BuildHostPlatform.Current.Platform == UnrealTargetPlatform.Mac)
-            {
-                if (File.Exists(FullDestPathRoot))
-                {
-                    File.Delete(FullDestPathRoot);
-                }
-            }
-
-            OutputFile = FileItem.GetItemByPath(FullDestPathRoot);
 
-            if (BuildHostPlatform.Current.Platform != UnrealTargetPlatform.Mac)
-            {
-                string RemoteFilename = ConvertPath(OutputFile.AbsolutePath);
-                RPCUtilHelper.Command("/", String.Format("rm -f {0}", RemoteFilename), "", null);
-            }
+			// Append the .app suffix
+			return AppBundleName + ".app";
+		}
 
-            // Make the compile action
-            Action CompileAssetAction = ActionGraph.Add(ActionType.CreateAppBundle);
-            if (!Utils.IsRunningOnMono)
+		public static FileReference GetAssetCatalogFile(CppPlatform Platform, FileReference Executable)
+		{
+			// Get the output file
+            if (Platform == CppPlatform.IOS)
             {
-                CompileAssetAction.ActionHandler = new Action.BlockingActionHandler(RPCUtilHelper.RPCActionHandler);
+                return FileReference.Combine(Executable.Directory, "Payload", GetAppBundleName(Executable), "Assets.car");
             }
+			else
+			{
+				return FileReference.Combine(Executable.Directory, "AssetCatalog", "Assets.car");
+			}
+		}
 
-            CompileAssetAction.WorkingDirectory = GetMacDevSrcRoot();
-            CompileAssetAction.CommandPath = "/usr/bin/xcrun";
-
-            FileItem AssetCat = FileItem.GetItemByPath(Path.Combine(IntermediateDir, "Resources", "Assets.xcassets"));
-            FileItem DestFile = LocalToRemoteFileItem(OutputFile, false);
-            FileItem InputFile = LocalToRemoteFileItem(AssetCat, BuildHostPlatform.Current.Platform != UnrealTargetPlatform.Mac);
-
-            string Arguments = "";
-            Arguments += " actool --output-format human-readable-text --notices --warnings";
-            Arguments += " --output-partial-info-plist \"";
-            Arguments += Path.GetDirectoryName(InputFile.AbsolutePath).Replace("\\", "/") + "/assetcatalog_generated_info.plist\"";
-            Arguments += " --app-icon AppIcon";
-            Arguments += Platform == CppPlatform.TVOS ? " --launch-image LaunchImage" : "";
-//            Arguments += " --compress-pngs";
-            Arguments += " --enable-on-demand-resources YES";
-            Arguments += Platform == CppPlatform.TVOS ? " --filter-for-device-model AppleTV5,3 --filter-for-device-os-version 9.2" : "";
-            Arguments += Platform == CppPlatform.TVOS ? " --target-device tv --minimum-deployment-target 9.2 --platform appletvos" : " --target-device iphone --target-device ipad --minimum-deployment-target 9.0 --platform iphoneos --product-type com.apple.product-type.application";
-            Arguments += " --compile \"";
-            Arguments += Path.GetDirectoryName(DestFile.AbsolutePath).Replace("\\", "/") + "\" \"" + InputFile.AbsolutePath + "\"";
-
-            CompileAssetAction.CommandArguments = Arguments;
-            CompileAssetAction.PrerequisiteItems.Add(Executable);
-            CompileAssetAction.ProducedItems.Add(DestFile);
-            CompileAssetAction.StatusDescription = CompileAssetAction.CommandArguments;// string.Format("Generating debug info for {0}", Path.GetFileName(Executable.AbsolutePath));
-            CompileAssetAction.bCanExecuteRemotely = false;
-
-            return DestFile;
-        }
+		public static string GetAssetCatalogArgs(CppPlatform Platform, string InputDir, string OutputDir)
+		{
+			StringBuilder Arguments = new StringBuilder("actool");
+            Arguments.Append(" --output-format human-readable-text");
+			Arguments.Append(" --notices");
+			Arguments.Append(" --warnings");
+            Arguments.AppendFormat(" --output-partial-info-plist '{0}/assetcatalog_generated_info.plist'", InputDir);
+            Arguments.Append(" --app-icon AppIcon");
+			if(Platform == CppPlatform.TVOS)
+			{
+				Arguments.Append(" --launch-image LaunchImage");
+				Arguments.Append(" --filter-for-device-model AppleTV5,3");
+				Arguments.Append(" --filter-for-device-os-version 9.2");
+				Arguments.Append(" --target-device tv");
+				Arguments.Append(" --minimum-deployment-target 9.2");
+				Arguments.Append(" --platform appletvos");
+			}
+			else
+			{
+				Arguments.Append(" --product-type com.apple.product-type.application");
+				Arguments.Append(" --target-device iphone");
+				Arguments.Append(" --target-device ipad");
+				Arguments.Append(" --minimum-deployment-target 9.0");
+				Arguments.Append(" --platform iphoneos");
+			}
+            Arguments.Append(" --enable-on-demand-resources YES");
+            Arguments.AppendFormat(" --compile '{0}'", OutputDir);
+			Arguments.AppendFormat(" '{0}/Assets.xcassets'", InputDir);
+			return Arguments.ToString();
+		}
 
         /// <summary>
         /// Generates debug info for a given executable
@@ -1096,10 +1057,6 @@ namespace UnrealBuildTool
 
             // Make the compile action
             Action GenDebugAction = ActionGraph.Add(ActionType.GenerateDebugInfo);
-			if (!Utils.IsRunningOnMono)
-			{
-				GenDebugAction.ActionHandler = new Action.BlockingActionHandler(RPCUtilHelper.RPCActionHandler);
-			}
 
 			GenDebugAction.WorkingDirectory = GetMacDevSrcRoot();
 			GenDebugAction.CommandPath = "sh";
@@ -1154,14 +1111,6 @@ namespace UnrealBuildTool
             // Make the compile action
             Action GenDebugAction = ActionGraph.Add(ActionType.GenerateDebugInfo);
             GenDebugAction.WorkingDirectory = GetMacDevEngineRoot() + "/Binaries/Mac/";
-            if (!Utils.IsRunningOnMono)
-            {
-                GenDebugAction.ActionHandler = new Action.BlockingActionHandler(RPCUtilHelper.RPCActionHandler);
-                FileItem DsymExporter = FileItem.GetItemByPath(Path.GetFullPath(Path.Combine(BranchDirectory, "Engine/")) + "/Binaries/Mac/" + "DsymExporter");
-                QueueFileForBatchUpload(DsymExporter);
-                QueueDirectoryForBatchUpload(new DirectoryReference(Path.GetFullPath(Path.Combine(BranchDirectory, "Engine/")) + "Binaries/ThirdParty/ICU/icu4c-53_1/Mac/", DirectoryReference.Sanitize.None));
-                QueueDirectoryForBatchUpload(new DirectoryReference(Path.GetFullPath(Path.Combine(BranchDirectory, "Engine/")) + "Content/Internationalization/", DirectoryReference.Sanitize.None));
-            }
 
             GenDebugAction.CommandPath = "sh";
             GenDebugAction.CommandArguments = string.Format("-c 'rm -rf \"{1}\"; dwarfdump --uuid \"{3}\" | cut -d\" \" -f2; chmod 777 ./DsymExporter; ./DsymExporter -UUID=$(dwarfdump --uuid \"{3}\" | cut -d\" \" -f2) \"{0}\" \"{2}\"'",
@@ -1291,20 +1240,8 @@ namespace UnrealBuildTool
 			}
 		}
 
-		public static new void PreBuildSync()
+		public static void PreBuildSync()
 		{
-			if (BuildHostPlatform.Current.Platform != UnrealTargetPlatform.Mac)
-			{
-				BuiltBinaries = new List<FileReference>();
-				DirectoryReference BinaryIOS = new DirectoryReference(Path.GetFullPath(Path.Combine(BranchDirectory, "Engine/")) + "Binaries/IOS/", DirectoryReference.Sanitize.None);
-				if (DirectoryReference.Exists(BinaryIOS))
-				{
-					QueueDirectoryForBatchUpload(BinaryIOS);
-				}
-			}
-
-			RemoteToolChain.PreBuildSync();
-
 			// Unzip any third party frameworks that are stored as zips
 			foreach (UEBuildFramework Framework in RememberedAdditionalFrameworks)
 			{
@@ -1318,30 +1255,21 @@ namespace UnrealBuildTool
 				// Assume that there is another directory inside the zip with the same name as the zip
 				ZipDstPath = ZipDstPath.Substring(0, ZipDstPath.LastIndexOf('/'));
 
-				if (BuildHostPlatform.Current.Platform == UnrealTargetPlatform.Mac)
-				{
-					// If we're on the mac, just unzip using the shell
-					string ResultsText;
-					RunExecutableAndWait("unzip", String.Format("-o \"{0}\" -d \"{1}\"", ZipSrcPath, ZipDstPath), out ResultsText);
-					continue;
-				}
-				else
-				{
-					// Use RPC utility if the zip is on remote mac
-					Hashtable Result = RPCUtilHelper.Command("/", String.Format("unzip -o \"{0}\" -d \"{1}\"", ZipSrcPath, ZipDstPath), "", null);
-
-					foreach (DictionaryEntry Entry in Result)
-					{
-						Log.TraceInformation("{0}", Entry.Value);
-					}
-				}
+				// If we're on the mac, just unzip using the shell
+				string ResultsText;
+				RunExecutableAndWait("unzip", String.Format("-o \"{0}\" -d \"{1}\"", ZipSrcPath, ZipDstPath), out ResultsText);
 			}
         }
 
-        public void GenerateAssetCatalogHelper(string EngineDir, string BuildDir, string IntermediateDir, CppPlatform Platform, ref bool bUserImagesExist)
+        public static DirectoryReference GenerateAssetCatalog(FileReference ProjectFile, CppPlatform Platform, ref bool bUserImagesExist)
         {
+            string EngineDir = UnrealBuildTool.EngineDirectory.ToString();
+            string BuildDir = (((ProjectFile != null) ? ProjectFile.Directory.ToString() : (string.IsNullOrEmpty(UnrealBuildTool.GetRemoteIniPath()) ? UnrealBuildTool.EngineDirectory.ToString() : UnrealBuildTool.GetRemoteIniPath()))) + "/Build/" + (Platform == CppPlatform.IOS ? "IOS" : "TVOS");
+            string IntermediateDir = (((ProjectFile != null) ? ProjectFile.Directory.ToString() : UnrealBuildTool.EngineDirectory.ToString())) + "/Intermediate/" + (Platform == CppPlatform.IOS ? "IOS" : "TVOS");
+
             bUserImagesExist = false;
 
+			string ResourcesDir = Path.Combine(IntermediateDir, "Resources");
             if (Platform == CppPlatform.TVOS)
             {
                 string[] Directories = { "Assets.xcassets",
@@ -1405,7 +1333,7 @@ namespace UnrealBuildTool
                 // create asset catalog for images
                 for (int i = 0; i < Directories.Length; ++i)
                 {
-                    string Dir = Path.Combine(IntermediateDir, "Resources", Directories[i]);
+                    string Dir = Path.Combine(ResourcesDir, Directories[i]);
                     if (!Directory.Exists(Dir))
                     {
                         Directory.CreateDirectory(Dir);
@@ -1431,7 +1359,7 @@ namespace UnrealBuildTool
             else
             {
                 // copy the template asset catalog to the appropriate directory
-                string Dir = Path.Combine(IntermediateDir, "Resources", "Assets.xcassets");
+                string Dir = Path.Combine(ResourcesDir, "Assets.xcassets");
                 if (!Directory.Exists(Dir))
                 {
                     Directory.CreateDirectory(Dir);
@@ -1500,6 +1428,7 @@ namespace UnrealBuildTool
                     }
                 }
             }
+			return new DirectoryReference(ResourcesDir);
         }
 
         public override ICollection<FileItem> PostBuild(FileItem Executable, LinkEnvironment BinaryLinkEnvironment, ActionGraph ActionGraph)
@@ -1521,33 +1450,38 @@ namespace UnrealBuildTool
                 }
             }
 
-            // generate the asset catalog
-            bool bUserImagesExist = false;
-            GenerateAssetCatalog(BinaryLinkEnvironment.Platform, ref bUserImagesExist);
-            CompileAssetCatalog(Executable, BinaryLinkEnvironment.Platform, ActionGraph, OutputFiles);
+			if(ShouldCompileAssetCatalog())
+			{
+				// generate the asset catalog
+				bool bUserImagesExist = false;
+				DirectoryReference ResourcesDir = GenerateAssetCatalog(ProjectFile, BinaryLinkEnvironment.Platform, ref bUserImagesExist);
 
-			return OutputFiles;
-        }
+				// Get the output location for the asset catalog
+				FileItem AssetCatalogFile = FileItem.GetItemByFileReference(GetAssetCatalogFile(CppPlatform, Executable.Location));
 
-        public void GenerateAssetCatalog(CppPlatform Platform, ref bool bUserImagesExist)
-        {
-            string EngineDir = UnrealBuildTool.EngineDirectory.ToString();
-            string BuildDir = (((ProjectFile != null) ? ProjectFile.Directory.ToString() : (string.IsNullOrEmpty(UnrealBuildTool.GetRemoteIniPath()) ? UnrealBuildTool.EngineDirectory.ToString() : UnrealBuildTool.GetRemoteIniPath()))) + "/Build/" + (Platform == CppPlatform.IOS ? "IOS" : "TVOS");
-            string IntermediateDir = (((ProjectFile != null) ? ProjectFile.Directory.ToString() : UnrealBuildTool.EngineDirectory.ToString())) + "/Intermediate/" + (Platform == CppPlatform.IOS ? "IOS" : "TVOS");
-            GenerateAssetCatalogHelper(EngineDir, BuildDir, IntermediateDir, Platform, ref bUserImagesExist);
-        }
+				// Make the compile action
+				Action CompileAssetAction = ActionGraph.Add(ActionType.CreateAppBundle);
+				CompileAssetAction.WorkingDirectory = GetMacDevSrcRoot();
+				CompileAssetAction.CommandPath = "/usr/bin/xcrun";
+				CompileAssetAction.CommandArguments = GetAssetCatalogArgs(CppPlatform, ResourcesDir.FullName, Path.GetDirectoryName(AssetCatalogFile.AbsolutePath));
+				CompileAssetAction.PrerequisiteItems.Add(Executable);
+				CompileAssetAction.ProducedItems.Add(AssetCatalogFile);
+				CompileAssetAction.DeleteItems.Add(AssetCatalogFile);
+				CompileAssetAction.StatusDescription = CompileAssetAction.CommandArguments;// string.Format("Generating debug info for {0}", Path.GetFileName(Executable.AbsolutePath));
+				CompileAssetAction.bCanExecuteRemotely = false;
 
-        public void CompileAssetCatalog(FileItem Executable, CppPlatform Platform, ActionGraph ActionGraph, ICollection<FileItem> OutputFiles)
-        {
-            if (CppPlatform == CppPlatform.TVOS || (CppPlatform == CppPlatform.IOS && Settings.Value.IOSSDKVersionFloat >= 11.0f))
-            {
-                string EngineDir = UnrealBuildTool.EngineDirectory.ToString();
-                string BuildDir = (((ProjectFile != null) ? ProjectFile.Directory.ToString() : (string.IsNullOrEmpty(UnrealBuildTool.GetRemoteIniPath()) ? UnrealBuildTool.EngineDirectory.ToString() : UnrealBuildTool.GetRemoteIniPath()))) + "/Build/" + (Platform == CppPlatform.IOS ? "IOS" : "TVOS");
-                string IntermediateDir = (((ProjectFile != null) ? ProjectFile.Directory.ToString() : UnrealBuildTool.EngineDirectory.ToString())) + "/Intermediate/" + (Platform == CppPlatform.IOS ? "IOS" : "TVOS");
-                OutputFiles.Add(CompileAssetCatalogHelper(Executable, EngineDir, BuildDir, IntermediateDir, ActionGraph, Platform));
-            }
+				// Add it to the output files so it's always built
+				OutputFiles.Add(AssetCatalogFile);
+			}
+
+			return OutputFiles;
         }
 
+		public bool ShouldCompileAssetCatalog()
+		{
+            return CppPlatform == CppPlatform.TVOS || (CppPlatform == CppPlatform.IOS && Settings.Value.IOSSDKVersionFloat >= 11.0f);
+		}
+
 		public static string GetCodesignPlatformName(UnrealTargetPlatform Platform)
 		{
 			switch(Platform)
@@ -1561,6 +1495,45 @@ namespace UnrealBuildTool
 			}
 		}
 
+		class ProcessOutput
+		{
+			/// <summary>
+			/// Substrings that indicate a line contains an error
+			/// </summary>
+			protected static readonly string[] ErrorMessageTokens =
+			{
+				"ERROR ",
+				"** BUILD FAILED **",
+				"[BEROR]",
+				"IPP ERROR",
+				"System.Net.Sockets.SocketException"
+			};
+
+			/// <summary>
+			/// Helper function to sync source files to and from the local system and a remote Mac
+			/// </summary>
+			//This chunk looks to be required to pipe output to VS giving information on the status of a remote build.
+			public bool OutputReceivedDataEventHandlerEncounteredError = false;
+			public string OutputReceivedDataEventHandlerEncounteredErrorMessage = "";
+			public void OutputReceivedDataEventHandler(Object Sender, DataReceivedEventArgs Line)
+			{
+				if ((Line != null) && (Line.Data != null))
+				{
+					Log.TraceInformation(Line.Data);
+
+					foreach (string ErrorToken in ErrorMessageTokens)
+					{
+						if (Line.Data.Contains(ErrorToken))
+						{
+							OutputReceivedDataEventHandlerEncounteredError = true;
+							OutputReceivedDataEventHandlerEncounteredErrorMessage += Line.Data;
+							break;
+						}
+					}
+				}
+			}
+		}
+
 		private static void GenerateCrashlyticsData(string ExecutableDirectory, string ExecutableName, string ProjectDir, string ProjectName)
         {
 			Log.TraceInformation("Generating and uploading Crashlytics Data");
@@ -1574,15 +1547,16 @@ namespace UnrealBuildTool
 				FabricProcess.StartInfo.Arguments = string.Format("-c 'chmod 777 \"{0}/Fabric.framework/upload-symbols\"; \"{0}/Fabric.framework/upload-symbols\" -a 7a4cebd0324af21696e5e321802c5e26ba541cad -p ios {1}'",
                     FabricPath,
 					ExecutableDirectory + "/" + ExecutableName + ".dSYM.zip");
-                FabricProcess.OutputDataReceived += new DataReceivedEventHandler(OutputReceivedDataEventHandler);
-                FabricProcess.ErrorDataReceived += new DataReceivedEventHandler(OutputReceivedDataEventHandler);
 
-                OutputReceivedDataEventHandlerEncounteredError = false;
-                OutputReceivedDataEventHandlerEncounteredErrorMessage = "";
+				ProcessOutput Output = new ProcessOutput();
+
+                FabricProcess.OutputDataReceived += new DataReceivedEventHandler(Output.OutputReceivedDataEventHandler);
+                FabricProcess.ErrorDataReceived += new DataReceivedEventHandler(Output.OutputReceivedDataEventHandler);
+
                 Utils.RunLocalProcess(FabricProcess);
-                if (OutputReceivedDataEventHandlerEncounteredError)
+                if (Output.OutputReceivedDataEventHandlerEncounteredError)
                 {
-                    throw new Exception(OutputReceivedDataEventHandlerEncounteredErrorMessage);
+                    throw new Exception(Output.OutputReceivedDataEventHandlerEncounteredErrorMessage);
                 }
             }
         }
@@ -1603,438 +1577,246 @@ namespace UnrealBuildTool
 
 			string AppName = Target.TargetType == TargetType.Game ? Target.TargetName : Target.AppName;
 
-			if (BuildHostPlatform.Current.Platform == UnrealTargetPlatform.Mac)
-			{
-				string RemoteShadowDirectoryMac = Path.GetDirectoryName(Target.OutputPath.FullName);
-				string FinalRemoteExecutablePath = String.Format("{0}/Payload/{1}.app/{1}", RemoteShadowDirectoryMac, AppName);
-
-				// strip the debug info from the executable if needed
-				if (Target.Rules.bStripSymbolsOnIOS || (Target.Configuration == UnrealTargetConfiguration.Shipping))
-				{
-					Process StripProcess = new Process();
-					StripProcess.StartInfo.WorkingDirectory = RemoteShadowDirectoryMac;
-					StripProcess.StartInfo.FileName = new IOSToolChainSettings().ToolchainDir + "strip";
-					StripProcess.StartInfo.Arguments = "\"" + Target.OutputPath + "\"";
-					StripProcess.OutputDataReceived += new DataReceivedEventHandler(OutputReceivedDataEventHandler);
-					StripProcess.ErrorDataReceived += new DataReceivedEventHandler(OutputReceivedDataEventHandler);
-
-					OutputReceivedDataEventHandlerEncounteredError = false;
-					OutputReceivedDataEventHandlerEncounteredErrorMessage = "";
-					Utils.RunLocalProcess(StripProcess);
-					if (OutputReceivedDataEventHandlerEncounteredError)
-					{
-						throw new Exception(OutputReceivedDataEventHandlerEncounteredErrorMessage);
-					}
-				}
+			string RemoteShadowDirectoryMac = Path.GetDirectoryName(Target.OutputPath.FullName);
+			string FinalRemoteExecutablePath = String.Format("{0}/Payload/{1}.app/{1}", RemoteShadowDirectoryMac, AppName);
 
-                // ensure the plist, entitlements, and provision files are properly copied
-                UEDeployIOS DeployHandler = (Target.Platform == UnrealTargetPlatform.IOS ? new UEDeployIOS() : new UEDeployTVOS());
-                DeployHandler.PrepTargetForDeployment(new UEBuildDeployTarget(Target));
+			// strip the debug info from the executable if needed
+			if (Target.Rules.bStripSymbolsOnIOS || (Target.Configuration == UnrealTargetConfiguration.Shipping))
+			{
+				Process StripProcess = new Process();
+				StripProcess.StartInfo.WorkingDirectory = RemoteShadowDirectoryMac;
+				StripProcess.StartInfo.FileName = new IOSToolChainSettings().ToolchainDir + "strip";
+				StripProcess.StartInfo.Arguments = "\"" + Target.OutputPath + "\"";
 
-				// copy the executable
-				if (!File.Exists(FinalRemoteExecutablePath))
-				{
-					Directory.CreateDirectory(String.Format("{0}/Payload/{1}.app", RemoteShadowDirectoryMac, AppName));
-				}
-				File.Copy(Target.OutputPath.FullName, FinalRemoteExecutablePath, true);
+				ProcessOutput Output = new ProcessOutput();
+				StripProcess.OutputDataReceived += new DataReceivedEventHandler(Output.OutputReceivedDataEventHandler);
+				StripProcess.ErrorDataReceived += new DataReceivedEventHandler(Output.OutputReceivedDataEventHandler);
 
-				if (!Target.Rules.IOSPlatform.bSkipCrashlytics)
+				Output.OutputReceivedDataEventHandlerEncounteredError = false;
+				Output.OutputReceivedDataEventHandlerEncounteredErrorMessage = "";
+				Utils.RunLocalProcess(StripProcess);
+				if (Output.OutputReceivedDataEventHandlerEncounteredError)
 				{
-					GenerateCrashlyticsData(RemoteShadowDirectoryMac, Path.GetFileName(Target.OutputPath.FullName), Target.ProjectDirectory.FullName, AppName);
+					throw new Exception(Output.OutputReceivedDataEventHandlerEncounteredErrorMessage);
 				}
+			}
 
-				if (Target.Rules.bCreateStubIPA)
-				{
-					string Project = Target.ProjectDirectory + "/" + AppName + ".uproject";
-
-					string SchemeName = AppName;
-
-                    // generate the dummy project so signing works
-                    if (AppName == "UE4Game" || AppName == "UE4Client" || Utils.IsFileUnderDirectory(Target.ProjectDirectory + "/" + AppName + ".uproject", Path.GetFullPath("../..")))
-					{
-						UnrealBuildTool.GenerateProjectFiles(new XcodeProjectFileGenerator(Target.ProjectFile), new string[] { "-platforms=" + (Target.Platform == UnrealTargetPlatform.IOS ? "IOS" : "TVOS"), "-NoIntellIsense", (Target.Platform == UnrealTargetPlatform.IOS ? "-iosdeployonly" : "-tvosdeployonly"), "-ignorejunk" });
-						Project = Path.GetFullPath("../..") + "/UE4_" + (Target.Platform == UnrealTargetPlatform.IOS ? "IOS" : "TVOS") + ".xcworkspace";
-                        if (AppName == "UE4Game" || AppName == "UE4Client")
-                        {
-                            SchemeName = "UE4";
-                        }
-					}
-					else
-					{
-						UnrealBuildTool.GenerateProjectFiles(new XcodeProjectFileGenerator(Target.ProjectFile), new string[] { "-platforms" + (Target.Platform == UnrealTargetPlatform.IOS ? "IOS" : "TVOS"), "-NoIntellIsense", (Target.Platform == UnrealTargetPlatform.IOS ? "-iosdeployonly" : "-tvosdeployonly"), "-ignorejunk", "-project=\"" + Target.ProjectDirectory + "/" + AppName + ".uproject\"", "-game" });
-						Project = Target.ProjectDirectory + "/" + AppName + "_" + (Target.Platform == UnrealTargetPlatform.IOS ? "IOS" : "TVOS") + ".xcworkspace";
-					}
-
-					if (Directory.Exists(Project))
-					{
-                        // ensure the plist, entitlements, and provision files are properly copied
-                        DeployHandler = (Target.Platform == UnrealTargetPlatform.IOS ? new UEDeployIOS() : new UEDeployTVOS());
-                        DeployHandler.PrepTargetForDeployment(new UEBuildDeployTarget(Target));
-
-						string ConfigName = Target.Configuration.ToString();
-						if (Target.Rules.Type != TargetType.Game && Target.Rules.Type != TargetType.Program)
-						{
-							ConfigName += " " + Target.Rules.Type.ToString();
-						}
-
-						// code sign the project
-						IOSProvisioningData ProvisioningData = ((IOSPlatform)UEBuildPlatform.GetBuildPlatform(Target.Platform)).ReadProvisioningData(ProjectSettings);
-						string CmdLine = new IOSToolChainSettings().XcodeDeveloperDir + "usr/bin/xcodebuild" +
-										" -workspace \"" + Project + "\"" +
-										" -configuration \"" + ConfigName + "\"" +
-										" -scheme '" + SchemeName + "'" +
-										" -sdk " + GetCodesignPlatformName(Target.Platform) +
-										" -destination generic/platform=" + (Target.Platform == UnrealTargetPlatform.IOS ? "iOS" : "tvOS") + 
-						                (!string.IsNullOrEmpty(ProvisioningData.TeamUUID) ? " DEVELOPMENT_TEAM=" + ProvisioningData.TeamUUID : "");
-						if (!ProjectSettings.bAutomaticSigning)
-						{
-							CmdLine += " CODE_SIGN_IDENTITY=\"" + (!string.IsNullOrEmpty(ProvisioningData.SigningCertificate) ? ProvisioningData.SigningCertificate : "IPhoneDeveloper") + "\"" +
-							(!string.IsNullOrEmpty(ProvisioningData.MobileProvisionUUID) ? (" PROVISIONING_PROFILE_SPECIFIER=" + ProvisioningData.MobileProvisionUUID) : "");
-						}
-						else
-						{
-							CmdLine += " CODE_SIGN_IDENTITY=\"iPhone Developer\"";
-						}
-
-                        Log.TraceInformation("Code signing with command line: " + CmdLine);
-
-						Process SignProcess = new Process();
-						SignProcess.StartInfo.WorkingDirectory = RemoteShadowDirectoryMac;
-						SignProcess.StartInfo.FileName = "/usr/bin/xcrun";
-						SignProcess.StartInfo.Arguments = CmdLine;
-						SignProcess.OutputDataReceived += new DataReceivedEventHandler(OutputReceivedDataEventHandler);
-						SignProcess.ErrorDataReceived += new DataReceivedEventHandler(OutputReceivedDataEventHandler);
+            // ensure the plist, entitlements, and provision files are properly copied
+            UEDeployIOS DeployHandler = (Target.Platform == UnrealTargetPlatform.IOS ? new UEDeployIOS() : new UEDeployTVOS());
+            DeployHandler.PrepTargetForDeployment(new UEBuildDeployTarget(Target));
 
-						OutputReceivedDataEventHandlerEncounteredError = false;
-						OutputReceivedDataEventHandlerEncounteredErrorMessage = "";
-						Utils.RunLocalProcess(SignProcess);
+			// copy the executable
+			if (!File.Exists(FinalRemoteExecutablePath))
+			{
+				Directory.CreateDirectory(String.Format("{0}/Payload/{1}.app", RemoteShadowDirectoryMac, AppName));
+			}
+			File.Copy(Target.OutputPath.FullName, FinalRemoteExecutablePath, true);
 
-						// delete the temp project
-						if (Project.Contains("_" + (Target.Platform == UnrealTargetPlatform.IOS ? "IOS" : "TVOS") + ".xcodeproj"))
-						{
-							Directory.Delete(Project, true);
-						}
+			if (!Target.Rules.IOSPlatform.bSkipCrashlytics)
+			{
+				GenerateCrashlyticsData(RemoteShadowDirectoryMac, Path.GetFileName(Target.OutputPath.FullName), Target.ProjectDirectory.FullName, AppName);
+			}
 
-						if (OutputReceivedDataEventHandlerEncounteredError)
-						{
-							throw new Exception(OutputReceivedDataEventHandlerEncounteredErrorMessage);
-						}
+			if (Target.Rules.bCreateStubIPA)
+			{
+				string Project = Target.ProjectDirectory + "/" + AppName + ".uproject";
 
-						// Package the stub
-						PackageStub(RemoteShadowDirectoryMac, AppName, Target.OutputPath.GetFileNameWithoutExtension());
-					}
-				}
+				string SchemeName = AppName;
 
+				// generate the dummy project so signing works
+				if (AppName == "UE4Game" || AppName == "UE4Client" || Utils.IsFileUnderDirectory(Target.ProjectDirectory + "/" + AppName + ".uproject", Path.GetFullPath("../..")))
 				{
-					// Copy bundled assets from additional frameworks to the intermediate assets directory (so they can get picked up during staging)
-					String LocalFrameworkAssets = Path.GetFullPath(Target.ProjectDirectory + "/Intermediate/" + (Target.Platform == UnrealTargetPlatform.IOS ? "IOS" : "TVOS") + "/FrameworkAssets");
-
-					// Clean the local dest directory if it exists
-					CleanIntermediateDirectory(LocalFrameworkAssets);
-
-					foreach (UEBuildFramework Framework in RememberedAdditionalFrameworks)
+					UnrealBuildTool.GenerateProjectFiles(new XcodeProjectFileGenerator(Target.ProjectFile), new string[] { "-platforms=" + (Target.Platform == UnrealTargetPlatform.IOS ? "IOS" : "TVOS"), "-NoIntellIsense", (Target.Platform == UnrealTargetPlatform.IOS ? "-iosdeployonly" : "-tvosdeployonly"), "-ignorejunk" });
+					Project = Path.GetFullPath("../..") + "/UE4_" + (Target.Platform == UnrealTargetPlatform.IOS ? "IOS" : "TVOS") + ".xcworkspace";
+					if (AppName == "UE4Game" || AppName == "UE4Client")
 					{
-						if (Framework.OwningModule == null || Framework.CopyBundledAssets == null || Framework.CopyBundledAssets == "")
-						{
-							continue;		// Only care if we need to copy bundle assets
-						}
-
-						string UnpackedZipPath = GetRemoteIntermediateFrameworkZipPath(Framework);
-
-						// For now, this is hard coded, but we need to loop over all modules, and copy bundled assets that need it
-						string LocalSource = UnpackedZipPath + "/" + Framework.CopyBundledAssets;
-						string BundleName = Framework.CopyBundledAssets.Substring(Framework.CopyBundledAssets.LastIndexOf('/') + 1);
-						string LocalDest = LocalFrameworkAssets + "/" + BundleName;
-
-						Log.TraceInformation("Copying bundled asset... LocalSource: {0}, LocalDest: {1}", LocalSource, LocalDest);
-
-						string ResultsText;
-                        RunExecutableAndWait("cp", String.Format("-R -L \"{0}\" \"{1}\"", LocalSource, LocalDest), out ResultsText);
-                    }
-                }
-			}
-			else
-			{
-				// store off the binaries
-				foreach (UEBuildBinary Binary in Target.Binaries)
-				{
-					BuiltBinaries.Add(Binary.OutputFilePath);
+						SchemeName = "UE4";
+					}
 				}
-
-				// check to see if the DangerouslyFast mode is valid (in other words, a build has gone through since a Rebuild/Clean operation)
-				FileReference DangerouslyFastValidFile = FileReference.Combine(Target.ProjectIntermediateDirectory, "DangerouslyFastIsNotDangerous");
-				bool bUseDangerouslyFastModeWasRequested = bUseDangerouslyFastMode;
-				if (bUseDangerouslyFastMode)
+				else
 				{
-					if (!FileReference.Exists(DangerouslyFastValidFile))
-					{
-						Log.TraceInformation("Dangeroulsy Fast mode was requested, but a slow mode hasn't completed. Performing slow now...");
-						bUseDangerouslyFastMode = false;
-					}
+					UnrealBuildTool.GenerateProjectFiles(new XcodeProjectFileGenerator(Target.ProjectFile), new string[] { "-platforms" + (Target.Platform == UnrealTargetPlatform.IOS ? "IOS" : "TVOS"), "-NoIntellIsense", (Target.Platform == UnrealTargetPlatform.IOS ? "-iosdeployonly" : "-tvosdeployonly"), "-ignorejunk", "-project=\"" + Target.ProjectDirectory + "/" + AppName + ".uproject\"", "-game" });
+					Project = Target.ProjectDirectory + "/" + AppName + "_" + (Target.Platform == UnrealTargetPlatform.IOS ? "IOS" : "TVOS") + ".xcworkspace";
 				}
 
-				foreach (FileReference FilePath in BuiltBinaries)
+				if (Directory.Exists(Project))
 				{
-					string RemoteExecutablePath = ConvertPath(FilePath.FullName);
+					// ensure the plist, entitlements, and provision files are properly copied
+					DeployHandler = (Target.Platform == UnrealTargetPlatform.IOS ? new UEDeployIOS() : new UEDeployTVOS());
+					DeployHandler.PrepTargetForDeployment(new UEBuildDeployTarget(Target));
 
-					// when going super fast, just copy the executable to the final resting spot
-					if (bUseDangerouslyFastMode)
+					string ConfigName = Target.Configuration.ToString();
+					if (Target.Rules.Type != TargetType.Game && Target.Rules.Type != TargetType.Program)
 					{
-						Log.TraceInformation("==============================================================================");
-						Log.TraceInformation("USING DANGEROUSLY FAST MODE! IF YOU HAVE ANY PROBLEMS, TRY A REBUILD/CLEAN!!!!");
-						Log.TraceInformation("==============================================================================");
-
-						// copy the executable
-						string RemoteShadowDirectoryMac = ConvertPath(Path.GetDirectoryName(Target.OutputPath.FullName));
-						string FinalRemoteExecutablePath = String.Format("{0}/Payload/{1}.app/{1}", RemoteShadowDirectoryMac, Target.TargetName);
-						RPCUtilHelper.Command("/", String.Format("cp -f {0} {1}", RemoteExecutablePath, FinalRemoteExecutablePath), "", null);
+						ConfigName += " " + Target.Rules.Type.ToString();
 					}
-					else if (!Utils.IsRunningOnMono && BuildHostPlatform.Current.Platform != UnrealTargetPlatform.Mac)
-					{
-						RPCUtilHelper.CopyFile(RemoteExecutablePath, FilePath.FullName, false);
 
-						if ((ProjectSettings.bGeneratedSYMFile == true || ProjectSettings.bGeneratedSYMBundle == true) && Path.GetExtension(FilePath.FullName) != ".a")
-						{
-							string DSYMExt;
-							if(ProjectSettings.bGeneratedSYMBundle)
-							{
-								DSYMExt = ".dSYM.zip";
-							}
-							else
-							{
-								DSYMExt = ".dSYM";
-							}
-							RPCUtilHelper.CopyFile(RemoteExecutablePath + DSYMExt, FilePath.FullName + DSYMExt, false);
-                            if (ProjectSettings.bGenerateCrashReportSymbols || Target.Rules.bUseMallocProfiler)
-                            {
-                                RPCUtilHelper.CopyFile(RemoteExecutablePath + ".udebugsymbols", FilePath.FullName + ".udebugsymbols", false);
-                            }
-
-                        }
-                    }
-				}
-
-				// Generate the stub
-				if (Target.Rules.bCreateStubIPA || bUseDangerouslyFastMode)
-				{
-                    // ensure the plist, entitlements, and provision files are properly copied
-                    UEDeployIOS DeployHandler = (Target.Platform == UnrealTargetPlatform.IOS ? new UEDeployIOS() : new UEDeployTVOS());
-                    DeployHandler.PrepTargetForDeployment(new UEBuildDeployTarget(Target));
-
-					if (!bUseDangerouslyFastMode)
+					FileReference SignProjectScript = FileReference.Combine(Target.ProjectIntermediateDirectory, "SignProject.sh");
+					using(StreamWriter Writer = new StreamWriter(SignProjectScript.FullName))
 					{
-						// generate the dummy project so signing works
-						UnrealBuildTool.GenerateProjectFiles(new XcodeProjectFileGenerator(Target.ProjectFile), new string[] { "-NoIntellisense", (Target.Platform == UnrealTargetPlatform.IOS ? "-iosdeployonly" : "-tvosdeployonly"), ((Target.ProjectFile != null) ? "-game" : "") });
-					}
+						// Boilerplate
+						Writer.WriteLine("#!/bin/sh");
+						Writer.WriteLine("set -e");
+						Writer.WriteLine("set -x");
 
-					// now that 
-					Process StubGenerateProcess = new Process();
-					StubGenerateProcess.StartInfo.WorkingDirectory = Path.GetFullPath("..\\Binaries\\DotNET\\IOS");
-					StubGenerateProcess.StartInfo.FileName = Path.Combine(StubGenerateProcess.StartInfo.WorkingDirectory, "iPhonePackager.exe");
+						// Copy the mobile provision into the system store
+						if(Target.Rules.IOSPlatform.ImportProvision != null)
+						{
+							Writer.WriteLine("cp -f {0} ~/Library/MobileDevice/Provisioning\\ Profiles/", Utils.EscapeShellArgument(Target.Rules.IOSPlatform.ImportProvision));
+						}
 
-					string Arguments = "";
-					string PathToApp = Target.RulesAssembly.GetTargetFileName(AppName).FullName;
-					string SchemeName = AppName;
+						// Path to the temporary keychain. When -ImportCertificate is specified, we will temporarily add this to the list of keychains to search, and remove it later.
+						FileReference TempKeychain = null;
 
-					// right now, no programs have a Source subdirectory, so assume the PathToApp is directly in the root
-					if (Path.GetDirectoryName(PathToApp).Contains(@"\Engine\Source\Programs"))
-					{
-						PathToApp = Path.GetDirectoryName(PathToApp);
-					}
-					else
-					{
-						Int32 SourceIndex = PathToApp.LastIndexOf(@"\Intermediate\Source");
-						if (SourceIndex != -1)
+						// Get the signing certificate to use
+						string SigningCertificate;
+						if(Target.Rules.IOSPlatform.ImportCertificate == null)
 						{
-							PathToApp = PathToApp.Substring(0, SourceIndex);
+							// Take it from the standard settings
+							IOSProvisioningData ProvisioningData = IOSPlatform.ReadProvisioningData(Target.ProjectFile);
+							SigningCertificate = ProvisioningData.SigningCertificate;
+
+							// Set the identity on the command line
+							if(!ProjectSettings.bAutomaticSigning)
+							{
+								Writer.WriteLine("CODE_SIGN_IDENTITY='{0}'", String.IsNullOrEmpty(SigningCertificate)? "IPhoneDeveloper" : SigningCertificate);
+							}
 						}
 						else
 						{
-							SourceIndex = PathToApp.LastIndexOf(@"\Source");
-							if (SourceIndex != -1)
+							// Read the name from the certificate
+							X509Certificate2 Certificate;
+							try
 							{
-								PathToApp = PathToApp.Substring(0, SourceIndex);
+								Certificate = new X509Certificate2(Target.Rules.IOSPlatform.ImportCertificate, Target.Rules.IOSPlatform.ImportCertificatePassword ?? "");
 							}
-							else
+							catch(Exception Ex)
 							{
-								throw new BuildException("The target was not in a /Source subdirectory");
+								throw new BuildException(Ex, "Unable to read certificate '{0}': {1}", Target.Rules.IOSPlatform.ImportCertificate, Ex.Message);
 							}
+							SigningCertificate = Certificate.GetNameInfo(X509NameType.SimpleName, false);
+
+							// Set the path to the temporary keychain
+							TempKeychain = FileReference.Combine(Target.ProjectIntermediateDirectory, "TempKeychain.keychain");//(DirectoryReference.GetSpecialFolder(Environment.SpecialFolder.UserProfile), "Library", "Keychains/UE4TempKeychain.keychain";
+
+							// Install a certificate given on the command line to a temporary keychain
+							Writer.WriteLine("security delete-keychain \"{0}\" || true", TempKeychain);
+							Writer.WriteLine("security create-keychain -p \"A\" \"{0}\"", TempKeychain);
+							Writer.WriteLine("security list-keychains -s \"{0}\"", TempKeychain);
+							Writer.WriteLine("security list-keychains");
+							Writer.WriteLine("security set-keychain-settings -t 3600 -l  \"{0}\"", TempKeychain);
+							Writer.WriteLine("security -v unlock-keychain -p \"A\" \"{0}\"", TempKeychain);
+							Writer.WriteLine("security import {0} -P {1} -k \"{2}\" -T /usr/bin/codesign -T /usr/bin/security -t agg", Utils.EscapeShellArgument(Target.Rules.IOSPlatform.ImportCertificate), Utils.EscapeShellArgument(Target.Rules.IOSPlatform.ImportCertificatePassword), TempKeychain);
+							Writer.WriteLine("security set-key-partition-list -S apple-tool:,apple:,codesign: -s -k \"A\" -D '{0}' -t private {1}", SigningCertificate, TempKeychain);
+
+							// Set parameters to make sure it uses the correct identity and keychain
+							Writer.WriteLine("CERT_IDENTITY='{0}'", SigningCertificate);
+							Writer.WriteLine("CODE_SIGN_IDENTITY='{0}'", SigningCertificate);
+							Writer.WriteLine("CODE_SIGN_KEYCHAIN='{0}'", TempKeychain);
 						}
-						if (AppName != "UE4Game" && AppName != "UE4Client")
+
+						FileReference MobileProvisionFile;
+						string MobileProvisionUUID;
+						string TeamUUID;
+						if(Target.Rules.IOSPlatform.ImportProvision == null)
 						{
-							PathToApp += "\\" + AppName + ".uproject";
+							IOSProvisioningData ProvisioningData = ((IOSPlatform)UEBuildPlatform.GetBuildPlatform(Target.Platform)).ReadProvisioningData(ProjectSettings);
+							MobileProvisionFile = ProvisioningData.MobileProvisionFile;
+							MobileProvisionUUID = ProvisioningData.MobileProvisionUUID;
+							TeamUUID = ProvisioningData.TeamUUID;
 						}
 						else
 						{
-							SchemeName = "UE4";
+							MobileProvisionFile = new FileReference(Target.Rules.IOSPlatform.ImportProvision);
+
+							MobileProvisionContents MobileProvision = MobileProvisionContents.Read(MobileProvisionFile);
+							MobileProvisionUUID = MobileProvision.GetUniqueId();
+
+							TeamUUID = null;
 						}
-					}
 
-					string SchemeConfiguration = Target.Configuration.ToString();
-					if (Target.Rules.Type != TargetType.Game && Target.Rules.Type != TargetType.Program)
-					{
-						SchemeConfiguration += " " + Target.Rules.Type.ToString();
-					}
+						// code sign the project
+						Console.WriteLine("Provisioning: {0}, {1}, {2}", MobileProvisionFile, MobileProvisionFile.GetFileName(), MobileProvisionUUID);
+						string CmdLine = new IOSToolChainSettings().XcodeDeveloperDir + "usr/bin/xcodebuild" +
+										" -workspace \"" + Project + "\"" +
+										" -configuration \"" + ConfigName + "\"" +
+										" -scheme '" + SchemeName + "'" +
+										" -sdk " + GetCodesignPlatformName(Target.Platform) +
+										" -destination generic/platform=" + (Target.Platform == UnrealTargetPlatform.IOS ? "iOS" : "tvOS") + 
+										(!string.IsNullOrEmpty(TeamUUID) ? " DEVELOPMENT_TEAM=" + TeamUUID : "");
+						CmdLine += String.Format(" CODE_SIGN_IDENTITY='{0}'", SigningCertificate);
+						if (!ProjectSettings.bAutomaticSigning)
+						{
+							CmdLine += (!string.IsNullOrEmpty(MobileProvisionUUID) ? (" PROVISIONING_PROFILE_SPECIFIER=" + MobileProvisionUUID) : "");
+						}
+						Writer.WriteLine("/usr/bin/xcrun {0}", CmdLine);
 
-					if (bUseDangerouslyFastMode)
-					{
-						// the quickness!
-						Arguments += "DangerouslyFast " + PathToApp;
-					}
-					else
-					{
-						Arguments += "PackageIPA \"" + PathToApp + "\" -createstub";
-                        IOSProvisioningData ProvisioningData = ((IOSPlatform)UEBuildPlatform.GetBuildPlatform(Target.Platform)).ReadProvisioningData(ProjectSettings);
-                        if (ProvisioningData != null && ProvisioningData.MobileProvisionUUID != "")
-                        {
-                            Arguments += " -provisioningUUID \"" + ProvisioningData.MobileProvisionUUID + "\"";
-                        }
-						// if we are making the dsym, then we can strip the debug info from the executable
-						if (Target.Rules.bStripSymbolsOnIOS || (Target.Configuration == UnrealTargetConfiguration.Shipping))
+						// Remove the temporary keychain from the search list
+						if(TempKeychain != null)
 						{
-							Arguments += " -strip";
+							Writer.WriteLine("security delete-keychain \"{0}\" || true", TempKeychain);
 						}
 					}
-					Arguments += " -config " + Target.Configuration + " -mac " + RemoteServerName + " -schemename " + SchemeName + " -schemeconfig \"" + SchemeConfiguration + "\"";
 
-                    string Architecture = Target.Architecture;
-                    if (Architecture != "")
-					{
-						// pass along the architecture if we need, skipping the initial -, so we have "-architecture simulator"
-						Arguments += " -architecture " + Architecture.Substring(1);
-					}
+					Log.TraceInformation("Executing {0}", SignProjectScript);
 
-					if (Target.Platform == UnrealTargetPlatform.TVOS)
-					{
-						Arguments += " -tvos";
-					}
+					Process SignProcess = new Process();
+					SignProcess.StartInfo.WorkingDirectory = RemoteShadowDirectoryMac;
+					SignProcess.StartInfo.FileName = "/bin/sh";
+					SignProcess.StartInfo.Arguments = SignProjectScript.FullName;
 
-					if (!bUseRPCUtil)
-					{
-						Arguments += " -usessh";
-						Arguments += " -sshexe \"" + ResolvedSSHExe + "\"";
-						Arguments += " -sshauth \"" + ResolvedSSHAuthentication + "\"";
-						Arguments += " -sshuser \"" + ResolvedRSyncUsername + "\"";
-						Arguments += " -rsyncexe \"" + ResolvedRSyncExe + "\"";
-						Arguments += " -overridedevroot \"" + UserDevRootMac + "\"";
-					}
+					ProcessOutput Output = new ProcessOutput();
 
-					// programmers that use Xcode packaging mode should use the following commandline instead, as it will package for Xcode on each compile
-					//				Arguments = "PackageApp " + GameName + " " + Configuration;
-					StubGenerateProcess.StartInfo.Arguments = Arguments;
+					SignProcess.OutputDataReceived += new DataReceivedEventHandler(Output.OutputReceivedDataEventHandler);
+					SignProcess.ErrorDataReceived += new DataReceivedEventHandler(Output.OutputReceivedDataEventHandler);
 
-					StubGenerateProcess.OutputDataReceived += new DataReceivedEventHandler(OutputReceivedDataEventHandler);
-					StubGenerateProcess.ErrorDataReceived += new DataReceivedEventHandler(OutputReceivedDataEventHandler);
+					Output.OutputReceivedDataEventHandlerEncounteredError = false;
+					Output.OutputReceivedDataEventHandlerEncounteredErrorMessage = "";
+					Utils.RunLocalProcess(SignProcess);
 
-					OutputReceivedDataEventHandlerEncounteredError = false;
-					OutputReceivedDataEventHandlerEncounteredErrorMessage = "";
-					int ExitCode = Utils.RunLocalProcess(StubGenerateProcess);
-					if (OutputReceivedDataEventHandlerEncounteredError)
+					// delete the temp project
+					if (Project.Contains("_" + (Target.Platform == UnrealTargetPlatform.IOS ? "IOS" : "TVOS") + ".xcodeproj"))
 					{
-						UnrealBuildTool.ExtendedErrorCode = ExitCode;
-						throw new Exception(OutputReceivedDataEventHandlerEncounteredErrorMessage);
+						Directory.Delete(Project, true);
 					}
 
-					// now that a slow mode sync has finished, we can now do DangerouslyFast mode again (if requested)
-					if (bUseDangerouslyFastModeWasRequested)
+					if (Output.OutputReceivedDataEventHandlerEncounteredError)
 					{
-						File.Create(DangerouslyFastValidFile.FullName);
-					}
-					else
-					{
-						// if we didn't want dangerously fast, then delete the file so that setting/unsetting the flag will do the right thing without a Rebuild
-						File.Delete(DangerouslyFastValidFile.FullName);
+						throw new Exception(Output.OutputReceivedDataEventHandlerEncounteredErrorMessage);
 					}
+
+					// Package the stub
+					PackageStub(RemoteShadowDirectoryMac, AppName, Target.OutputPath.GetFileNameWithoutExtension());
 				}
+			}
 
-				{
-					// Copy bundled assets from additional frameworks to the intermediate assets directory (so they can get picked up during staging)
-					String LocalFrameworkAssets = Path.GetFullPath(Target.ProjectDirectory + "/Intermediate/" + (Target.Platform == UnrealTargetPlatform.IOS ? "IOS" : "TVOS") + "/FrameworkAssets");
-					String RemoteFrameworkAssets = ConvertPath(LocalFrameworkAssets);
+			{
+				// Copy bundled assets from additional frameworks to the intermediate assets directory (so they can get picked up during staging)
+				String LocalFrameworkAssets = Path.GetFullPath(Target.ProjectDirectory + "/Intermediate/" + (Target.Platform == UnrealTargetPlatform.IOS ? "IOS" : "TVOS") + "/FrameworkAssets");
 
-					CleanIntermediateDirectory(RemoteFrameworkAssets);
+				// Clean the local dest directory if it exists
+				CleanIntermediateDirectory(LocalFrameworkAssets);
 
-					// Delete the local dest directory if it exists
-					if (Directory.Exists(LocalFrameworkAssets))
+				foreach (UEBuildFramework Framework in RememberedAdditionalFrameworks)
+				{
+					if (Framework.OwningModule == null || Framework.CopyBundledAssets == null || Framework.CopyBundledAssets == "")
 					{
-						Directory.Delete(LocalFrameworkAssets, true);
+						continue;		// Only care if we need to copy bundle assets
 					}
 
-					foreach (UEBuildFramework Framework in RememberedAdditionalFrameworks)
-					{
-						if (Framework.OwningModule == null || Framework.CopyBundledAssets == null || Framework.CopyBundledAssets == "")
-						{
-							continue;		// Only care if we need to copy bundle assets
-						}
-
-						string RemoteZipPath = GetRemoteIntermediateFrameworkZipPath(Framework);
-
-						RemoteZipPath = RemoteZipPath.Replace(".zip", "");
-
-						// For now, this is hard coded, but we need to loop over all modules, and copy bundled assets that need it
-						string RemoteSource = RemoteZipPath + "/" + Framework.CopyBundledAssets;
-						string BundleName = Framework.CopyBundledAssets.Substring(Framework.CopyBundledAssets.LastIndexOf('/') + 1);
-
-						String RemoteDest = RemoteFrameworkAssets + "/" + BundleName;
-						String LocalDest = LocalFrameworkAssets + "\\" + BundleName;
-
-						Log.TraceInformation("Copying bundled asset... RemoteSource: {0}, RemoteDest: {1}, LocalDest: {2}", RemoteSource, RemoteDest, LocalDest);
-
-                        Hashtable Results = RPCUtilHelper.Command("/", String.Format("cp -R -L \"{0}\" \"{1}\"", RemoteSource, RemoteDest), "", null);
+					string UnpackedZipPath = GetRemoteIntermediateFrameworkZipPath(Framework);
 
-                        foreach (DictionaryEntry Entry in Results)
-						{
-							Log.TraceInformation("{0}", Entry.Value);
-						}
+					// For now, this is hard coded, but we need to loop over all modules, and copy bundled assets that need it
+					string LocalSource = UnpackedZipPath + "/" + Framework.CopyBundledAssets;
+					string BundleName = Framework.CopyBundledAssets.Substring(Framework.CopyBundledAssets.LastIndexOf('/') + 1);
+					string LocalDest = LocalFrameworkAssets + "/" + BundleName;
 
-						// Copy the bundled resource from the remote mac to the local dest
-						RPCUtilHelper.CopyDirectory(RemoteDest, LocalDest, RPCUtilHelper.ECopyOptions.None);
-					}
-				}
+					Log.TraceInformation("Copying bundled asset... LocalSource: {0}, LocalDest: {1}", LocalSource, LocalDest);
 
-                if (Target.Platform == UnrealTargetPlatform.TVOS)
-                {
-                    // copy back the built asset catalog
-                    string CatPath = Path.GetDirectoryName(Target.OutputPath.FullName) + "\\AssetCatalog\\Assets.car";
-                    RPCUtilHelper.CopyFile(ConvertPath(CatPath), CatPath, false);
+					string ResultsText;
+                    RunExecutableAndWait("cp", String.Format("-R -L \"{0}\" \"{1}\"", LocalSource, LocalDest), out ResultsText);
                 }
-
-                // If it is requested, send the app bundle back to the platform executing these commands.
-                if (Target.Rules.bCopyAppBundleBackToDevice)
-				{
-					Log.TraceInformation("Copying binaries back to this device...");
-
-					try
-					{
-						string BinaryDir = Path.GetDirectoryName(Target.OutputPath.FullName) + "\\";
-						if (BinaryDir.EndsWith(Target.AppName + "\\Binaries\\" + (Target.Platform == UnrealTargetPlatform.IOS ? "IOS" : "TVOS") + "\\") && Target.TargetType != TargetType.Game)
-						{
-							BinaryDir = BinaryDir.Replace(Target.TargetType.ToString(), "Game");
-						}
-
-						// Get the app bundle's name
-						string AppFullName = Target.AppName;
-						if (Target.Configuration != UnrealTargetConfiguration.Development)
-						{
-							AppFullName += "-" + Target.Platform.ToString();
-							AppFullName += "-" + Target.Configuration.ToString();
-						}
-
-						foreach (string BinaryPath in BuiltBinaries.Select(x => x.FullName))
-						{
-							if (!BinaryPath.Contains("Dummy"))
-							{
-								RPCUtilHelper.CopyFile(ConvertPath(BinaryPath), BinaryPath, false);
-							}
-						}
-						Log.TraceInformation("Copied binaries successfully.");
-					}
-					catch (Exception)
-					{
-						Log.TraceInformation("Copying binaries back to this device failed.");
-					}
-				}
-			}
+            }
 		}
 
 		public static int RunExecutableAndWait(string ExeName, string ArgumentList, out string StdOutResults)
diff --git a/Engine/Source/Programs/UnrealBuildTool/Platform/IOS/MobileProvisionContents.cs b/Engine/Source/Programs/UnrealBuildTool/Platform/IOS/MobileProvisionContents.cs
new file mode 100644
--- /dev/null
+++ b/Engine/Source/Programs/UnrealBuildTool/Platform/IOS/MobileProvisionContents.cs
@@ -0,0 +1,121 @@
+using System;
+using System.Collections.Generic;
+using System.IO;
+using System.Linq;
+using System.Text;
+using System.Threading.Tasks;
+using System.Xml;
+using Tools.DotNETCommon;
+
+namespace UnrealBuildTool
+{
+	/// <summary>
+	/// Allows interrogating the contents of a .mobileprovision
+	/// </summary>
+	class MobileProvisionContents
+	{
+		/// <summary>
+		/// The contents of the provision
+		/// </summary>
+		XmlDocument Document;
+
+		/// <summary>
+		/// Map of key names to XML elements holding their values
+		/// </summary>
+		Dictionary<string, XmlElement> NameToValue = new Dictionary<string, XmlElement>();
+
+		/// <summary>
+		/// Constructor
+		/// </summary>
+		/// <param name="Document">XML file to create the mobile provision from. Call Read() to read from a signed file on disk.</param>
+		public MobileProvisionContents(XmlDocument Document)
+		{
+			this.Document = Document;
+
+			foreach(XmlElement KeyElement in Document.SelectNodes("/plist/dict/key"))
+			{
+				XmlNode ValueNode = KeyElement.NextSibling;
+				while(ValueNode != null)
+				{
+					XmlElement ValueElement = ValueNode as XmlElement;
+					if(ValueElement != null)
+					{
+						NameToValue[KeyElement.InnerText] = ValueElement;
+						break;
+					}
+				}
+			}
+		}
+
+		/// <summary>
+		/// Gets the unique id for this mobileprovision
+		/// </summary>
+		/// <returns>UUID for the provision</returns>
+		public string GetUniqueId()
+		{
+			XmlElement UniqueIdElement;
+			if(!NameToValue.TryGetValue("UUID", out UniqueIdElement))
+			{
+				throw new BuildException("Missing UUID in MobileProvision");
+			}
+			return UniqueIdElement.InnerText;
+		}
+
+		/// <summary>
+		/// Reads a mobileprovision from a file on disk
+		/// </summary>
+		/// <param name="Location">Path to the file</param>
+		/// <returns>New mobile provision instance</returns>
+		public static MobileProvisionContents Read(FileReference Location)
+		{
+			XmlDocument Document = ReadXml(Location);
+			return new MobileProvisionContents(Document);
+		}
+
+		/// <summary>
+		/// Reads the plist file inside a mobileprovision
+		/// </summary>
+		/// <param name="Location">Path to the file</param>
+		/// <returns>XML plist extracted from the mobile provision</returns>
+		public static XmlDocument ReadXml(FileReference Location)
+		{
+			// Provision data is stored as PKCS7-signed file in ASN.1 BER format
+			using(BinaryReader Reader = new BinaryReader(File.Open(Location.FullName, FileMode.Open)))
+			{
+				long Length = Reader.BaseStream.Length;
+				while(Reader.BaseStream.Position < Length)
+				{
+					Asn.FieldInfo Field = Asn.ReadField(Reader);
+					if(Field.Tag == Asn.FieldTag.OBJECT_IDENTIFIER)
+					{
+						int[] Identifier = Asn.ReadObjectIdentifier(Reader, Field.Length);
+						if(Enumerable.SequenceEqual(Identifier, Asn.ObjectIdentifier.Pkcs7_Data))
+						{
+							while(Reader.BaseStream.Position < Length)
+							{
+								Asn.FieldInfo NextField = Asn.ReadField(Reader);
+								if(NextField.Tag == Asn.FieldTag.OCTET_STRING)
+								{
+									byte[] Data = Reader.ReadBytes(NextField.Length);
+
+									XmlDocument Document = new XmlDocument();
+									Document.Load(new MemoryStream(Data));
+									return Document;
+								}
+								else
+								{
+									Asn.SkipValue(Reader, NextField);
+								}
+							}
+						}
+					}
+					else
+					{
+						Asn.SkipValue(Reader, Field);
+					}
+				}
+				throw new BuildException("No PKCS7-Data section found in {0}", Location);
+			}
+		}
+	}
+}
diff --git a/Engine/Source/Programs/UnrealBuildTool/Platform/IOS/UEBuildIOS.cs b/Engine/Source/Programs/UnrealBuildTool/Platform/IOS/UEBuildIOS.cs
--- a/Engine/Source/Programs/UnrealBuildTool/Platform/IOS/UEBuildIOS.cs
+++ b/Engine/Source/Programs/UnrealBuildTool/Platform/IOS/UEBuildIOS.cs
@@ -21,6 +21,24 @@ namespace UnrealBuildTool
 		/// </summary>
 		[CommandLine("-skipcrashlytics")]
 		public bool bSkipCrashlytics = false;
+
+		/// <summary>
+		/// Manual override for the provision to use. Should be a full path.
+		/// </summary>
+		[CommandLine("-ImportProvision=")]
+		public string ImportProvision = null;
+
+		/// <summary>
+		/// Imports the given certificate (inc private key) into a temporary keychain before signing.
+		/// </summary>
+		[CommandLine("-ImportCertificate=")]
+		public string ImportCertificate = null;
+
+		/// <summary>
+		/// Password for the imported certificate
+		/// </summary>
+		[CommandLine("-ImportCertificatePassword=")]
+		public string ImportCertificatePassword = null;
 	}
 
 	/// <summary>
@@ -54,6 +72,21 @@ namespace UnrealBuildTool
 			get { return Inner.bSkipCrashlytics; }
 		}
 
+		public string ImportProvision
+		{
+			get { return Inner.ImportProvision; }
+		}
+
+		public string ImportCertificate
+		{
+			get { return Inner.ImportCertificate; }
+		}
+
+		public string ImportCertificatePassword
+		{
+			get { return Inner.ImportCertificatePassword; }
+		}
+
 #if !__MonoCS__
 #pragma warning restore CS1591
 #endif
@@ -318,12 +351,17 @@ namespace UnrealBuildTool
     class IOSProvisioningData
     {
 		public string SigningCertificate;
-		public string MobileProvision;
+		public FileReference MobileProvisionFile;
         public string MobileProvisionUUID;
         public string MobileProvisionName;
         public string TeamUUID;
 		public bool bHaveCertificate = false;
 
+		public string MobileProvision
+		{
+			get { return (MobileProvisionFile == null)? null : MobileProvisionFile.GetFileName(); }
+		}
+
 		public IOSProvisioningData(IOSProjectSettings ProjectSettings, bool bForDistribution)
 			: this(ProjectSettings, false, bForDistribution)
 		{
@@ -332,7 +370,7 @@ namespace UnrealBuildTool
 		protected IOSProvisioningData(IOSProjectSettings ProjectSettings, bool bIsTVOS, bool bForDistribtion)
 		{
             SigningCertificate = ProjectSettings.SigningCertificate;
-            MobileProvision = ProjectSettings.MobileProvision;
+            string MobileProvision = ProjectSettings.MobileProvision;
 
 			FileReference ProjectFile = ProjectSettings.ProjectFile;
             if (!string.IsNullOrEmpty(SigningCertificate))
@@ -365,13 +403,31 @@ namespace UnrealBuildTool
                 bHaveCertificate = true;
             }
 
-            if (string.IsNullOrEmpty(MobileProvision) // no provision specified
-                || !File.Exists((BuildHostPlatform.Current.Platform == UnrealTargetPlatform.Mac ? (Environment.GetEnvironmentVariable("HOME") + "/Library/MobileDevice/Provisioning Profiles/") : (Environment.GetFolderPath(Environment.SpecialFolder.LocalApplicationData) + "/Apple Computer/MobileDevice/Provisioning Profiles/")) + MobileProvision) // file doesn't exist
-                || !bHaveCertificate) // certificate doesn't exist
+			if(!string.IsNullOrEmpty(MobileProvision))
+			{
+				DirectoryReference MobileProvisionDir;
+				if(BuildHostPlatform.Current.Platform == UnrealTargetPlatform.Mac)
+				{
+					MobileProvisionDir = DirectoryReference.Combine(new DirectoryReference(Environment.GetEnvironmentVariable("HOME")), "Library", "MobileDevice", "Provisioning Profiles");
+				}
+				else
+				{
+					MobileProvisionDir = DirectoryReference.Combine(DirectoryReference.GetSpecialFolder(Environment.SpecialFolder.LocalApplicationData), "Apple Computer", "MobileDevice", "Provisioning Profiles");
+				}
+
+				FileReference PossibleMobileProvisionFile = FileReference.Combine(MobileProvisionDir, MobileProvision);
+				if(FileReference.Exists(PossibleMobileProvisionFile))
+				{
+					MobileProvisionFile = PossibleMobileProvisionFile;
+				}
+			}
+
+            if (MobileProvisionFile == null || !bHaveCertificate)
             {
 
                 SigningCertificate = "";
                 MobileProvision = "";
+				MobileProvisionFile = null;
                 Log.TraceLog("Provision not specified or not found for " + ((ProjectFile != null) ? ProjectFile.GetFileNameWithoutAnyExtensions() : "UE4Game") + ", searching for compatible match...");
                 Process IPPProcess = new Process();
                 if (BuildHostPlatform.Current.Platform == UnrealTargetPlatform.Mac)
@@ -393,8 +449,12 @@ namespace UnrealBuildTool
                     IPPProcess.ErrorDataReceived += new DataReceivedEventHandler(IPPDataReceivedHandler);
                 }
                 Utils.RunLocalProcess(IPPProcess);
-                Log.TraceLog("Provision found for " + ((ProjectFile != null) ? ProjectFile.GetFileNameWithoutAnyExtensions() : "UE4Game") + ", Provision: " + MobileProvision + " Certificate: " + SigningCertificate);
+				if(MobileProvisionFile != null)
+				{
+					Log.TraceLog("Provision found for " + ((ProjectFile != null) ? ProjectFile.GetFileNameWithoutAnyExtensions() : "UE4Game") + ", Provision: " + MobileProvisionFile + " Certificate: " + SigningCertificate);
+				}
             }
+
             // add to the dictionary
             SigningCertificate = SigningCertificate.Replace("\"", "");
 
@@ -483,6 +543,7 @@ namespace UnrealBuildTool
         {
             if ((Line != null) && (Line.Data != null))
             {
+				Log.TraceLog("{0}", Line.Data);
                 if (!string.IsNullOrEmpty(SigningCertificate))
                 {
                     if (Line.Data.Contains("CERTIFICATE-") && Line.Data.Contains(SigningCertificate))
@@ -499,7 +560,10 @@ namespace UnrealBuildTool
                         cindex += "CERTIFICATE-".Length;
                         SigningCertificate = Line.Data.Substring(cindex, pindex - cindex - 1);
                         pindex += "PROVISION-".Length;
-                        MobileProvision = Line.Data.Substring(pindex);
+						if(pindex < Line.Data.Length)
+						{
+							MobileProvisionFile = new FileReference(Line.Data.Substring(pindex));
+						}
                     }
                 }
             }
@@ -624,6 +688,14 @@ namespace UnrealBuildTool
 			return new IOSProjectSettings(ProjectFile);
 		}
 
+		public static IOSProvisioningData ReadProvisioningData(FileReference ProjectFile, bool bForDistribution = false)
+		{
+			IOSPlatform Platform = (IOSPlatform)UEBuildPlatform.GetBuildPlatform(UnrealTargetPlatform.IOS);
+
+			IOSProjectSettings ProjectSettings = Platform.ReadProjectSettings(ProjectFile);
+			return Platform.ReadProvisioningData(ProjectSettings, bForDistribution);
+		}
+
 		public IOSProvisioningData ReadProvisioningData(IOSProjectSettings ProjectSettings, bool bForDistribution = false)
         {
 			string ProvisionKey = ProjectSettings.BundleIdentifier + " " + bForDistribution.ToString();
@@ -678,11 +750,6 @@ namespace UnrealBuildTool
 			IOSToolChain.PostBuildSync(Target);
 		}
 
-		public override void PostCodeGeneration(UHTManifest Manifest)
-		{
-			IOSToolChain.PostCodeGeneration(Manifest);
-		}
-
 		public bool HasCustomIcons(DirectoryReference ProjectDirectoryName)
 		{
 			string IconDir = Path.Combine(ProjectDirectoryName.FullName, "Build", "IOS", "Resources", "Graphics");
@@ -755,11 +822,6 @@ namespace UnrealBuildTool
 			return true;
 		}
 
-		public override bool UseAbsolutePathsInUnityFiles()
-		{
-			return false;
-		}
-
 		/// <summary>
 		/// Modify the rules for a newly created module, where the target is a different host platform.
 		/// This is not required - but allows for hiding details of a particular platform.
@@ -811,16 +873,6 @@ namespace UnrealBuildTool
 			}
 		}
 
-		/// <summary>
-		/// Converts the passed in path from UBT host to compiler native format.
-		/// </summary>
-		/// <param name="OriginalPath">The path to convert</param>
-		/// <returns>The path in native format for the toolchain</returns>
-		public override string ConvertPath(string OriginalPath)
-		{
-			return IOSToolChain.ConvertPath(OriginalPath);
-		}
-
 		/// <summary>
 		/// Whether this platform should create debug information or not
 		/// </summary>
diff --git a/Engine/Source/Programs/UnrealBuildTool/Platform/Mac/MacToolChain.cs b/Engine/Source/Programs/UnrealBuildTool/Platform/Mac/MacToolChain.cs
--- a/Engine/Source/Programs/UnrealBuildTool/Platform/Mac/MacToolChain.cs
+++ b/Engine/Source/Programs/UnrealBuildTool/Platform/Mac/MacToolChain.cs
@@ -110,7 +110,7 @@ namespace UnrealBuildTool
 		MacToolChainOptions Options;
 
 		public MacToolChain(FileReference InProjectFile, MacToolChainOptions InOptions)
-			: base(CppPlatform.Mac, UnrealTargetPlatform.Mac, InProjectFile)
+			: base(CppPlatform.Mac, InProjectFile)
 		{
 			this.Options = InOptions;
 		}
@@ -562,11 +562,6 @@ namespace UnrealBuildTool
 				// Add the source file path to the command-line.
 				FileArguments += string.Format(" \"{0}\"", ConvertPath(SourceFile.AbsolutePath));
 
-				if (BuildHostPlatform.Current.Platform != UnrealTargetPlatform.Mac)
-				{
-					CompileAction.ActionHandler = new Action.BlockingActionHandler(RPCUtilHelper.RPCActionHandler);
-				}
-				
 				string AllArgs = Arguments + FileArguments + CompileEnvironment.AdditionalArguments;
 				string CompilerPath = Settings.ToolchainDir + MacCompiler;
 				
@@ -698,11 +693,6 @@ namespace UnrealBuildTool
 			// Create an action that invokes the linker.
 			Action LinkAction = ActionGraph.Add(ActionType.Link);
 
-			if (BuildHostPlatform.Current.Platform != UnrealTargetPlatform.Mac)
-			{
-				LinkAction.ActionHandler = new Action.BlockingActionHandler(RPCUtilHelper.RPCActionHandler);
-			}
-
 			LinkAction.WorkingDirectory = GetMacDevSrcRoot();
 			LinkAction.CommandPath = "/bin/sh";
 			LinkAction.CommandDescription = "Link";
@@ -1250,11 +1240,6 @@ namespace UnrealBuildTool
 
 			LinkAction.CommandArguments = "-c 'chmod +x \"" + RemoteFixDylibDepsScript.AbsolutePath + "\"; \"" + RemoteFixDylibDepsScript.AbsolutePath + "\"; if [[ $? -ne 0 ]]; then exit 1; fi; ";
 
-			if (BuildHostPlatform.Current.Platform != UnrealTargetPlatform.Mac)
-			{
-				LinkAction.ActionHandler = new Action.BlockingActionHandler(RPCUtilHelper.RPCActionHandler);
-			}
-
 			// Make sure this action is executed after all the dylibs and the main executable are created
 
 			foreach (FileItem Dependency in BundleDependencies)
@@ -1278,20 +1263,6 @@ namespace UnrealBuildTool
 			return RemoteOutputFile;
 		}
 
-		private static Dictionary<Action, string> DebugOutputMap = new Dictionary<Action, string>();
-		static public void RPCDebugInfoActionHandler(Action Action, out int ExitCode, out string Output)
-		{
-			RPCUtilHelper.RPCActionHandler(Action, out ExitCode, out Output);
-			if (DebugOutputMap.ContainsKey(Action))
-			{
-				if (ExitCode == 0)
-				{
-					RPCUtilHelper.CopyDirectory(Action.ProducedItems[0].AbsolutePath, DebugOutputMap[Action], RPCUtilHelper.ECopyOptions.None);
-				}
-				DebugOutputMap.Remove(Action);
-			}
-		}
-
 		/// <summary>
 		/// Generates debug info for a given executable
 		/// </summary>
@@ -1326,11 +1297,6 @@ namespace UnrealBuildTool
 
 			// Make the compile action
 			Action GenDebugAction = ActionGraph.Add(ActionType.GenerateDebugInfo);
-			if (BuildHostPlatform.Current.Platform != UnrealTargetPlatform.Mac)
-			{
-				DebugOutputMap.Add(GenDebugAction, OutputFile.AbsolutePath);
-				GenDebugAction.ActionHandler = new Action.BlockingActionHandler(MacToolChain.RPCDebugInfoActionHandler);
-			}
 			GenDebugAction.WorkingDirectory = GetMacDevSrcRoot();
 			GenDebugAction.CommandPath = "sh";
 
@@ -1378,11 +1344,6 @@ namespace UnrealBuildTool
 			FileItem BundleScript = FileItem.GetItemByFileReference(FileReference.Combine(LinkEnvironment.IntermediateDirectory, "FinalizeAppBundle.sh"));
 			FileItem RemoteBundleScript = LocalToRemoteFileItem(BundleScript, true);
 
-			if (BuildHostPlatform.Current.Platform != UnrealTargetPlatform.Mac)
-			{
-				FinalizeAppBundleAction.ActionHandler = new Action.BlockingActionHandler(RPCUtilHelper.RPCActionHandler);
-			}
-
 			FinalizeAppBundleAction.CommandArguments = "\"" + RemoteBundleScript.AbsolutePath + "\"";
 			FinalizeAppBundleAction.PrerequisiteItems.Add(FixDylibOutputFile);
 			FinalizeAppBundleAction.ProducedItems.Add(RemoteDestFile);
@@ -1403,15 +1364,7 @@ namespace UnrealBuildTool
 			string SourcePath = Path.Combine(Path.GetFullPath("."), Resource.ResourcePath);
 			string TargetPath = Path.Combine(BundlePath, "Contents", Resource.BundleContentsSubdir, Path.GetFileName(Resource.ResourcePath));
 
-			FileItem TargetItem;
-			if (BuildHostPlatform.Current.Platform == UnrealTargetPlatform.Mac)
-			{
-				TargetItem = FileItem.GetItemByPath(TargetPath);
-			}
-			else
-			{
-				TargetItem = FileItem.GetRemoteItemByPath(TargetPath, RemoteToolChainPlatform);
-			}
+			FileItem TargetItem = FileItem.GetItemByPath(TargetPath);
 
 			CopyAction.CommandArguments = string.Format("-c 'cp -f -R \"{0}\" \"{1}\"; touch -c \"{2}\"'", ConvertPath(SourcePath), Path.GetDirectoryName(TargetPath).Replace('\\', '/') + "/", TargetPath.Replace('\\', '/'));
 			CopyAction.PrerequisiteItems.Add(Executable);
@@ -1420,11 +1373,6 @@ namespace UnrealBuildTool
 			CopyAction.StatusDescription = string.Format("Copying {0} to app bundle", Path.GetFileName(Resource.ResourcePath));
 			CopyAction.bCanExecuteRemotely = false;
 
-			if (BuildHostPlatform.Current.Platform != UnrealTargetPlatform.Mac)
-			{
-				CopyAction.ActionHandler = new Action.BlockingActionHandler(RPCUtilHelper.RPCActionHandler);
-			}
-
 			if (Directory.Exists(Resource.ResourcePath))
 			{
 				foreach (string ResourceFile in Directory.GetFiles(Resource.ResourcePath, "*", SearchOption.AllDirectories))
@@ -1597,174 +1545,6 @@ namespace UnrealBuildTool
 			AllBuildProducts.Add(Binary, new Dictionary<FileReference, BuildProductType>(BuildProducts));
 		}
 
-		public static void PostBuildSync(UEBuildTarget InTarget)
-		{
-			if (BuildHostPlatform.Current.Platform != UnrealTargetPlatform.Mac)
-			{
-				List<string> BuiltBinaries = new List<string>();
-				foreach (UEBuildBinary Binary in InTarget.Binaries)
-				{
-					BuiltBinaries.Add(Path.GetFullPath(Binary.ToString()));
-
-					string[] DebugExtensions = UEBuildPlatform.GetBuildPlatform(InTarget.Platform).GetDebugInfoExtensions(InTarget.Rules, Binary.Type);
-					if (Array.Exists(DebugExtensions, element => element == ".dSYM"))
-					{
-						Dictionary<FileReference, BuildProductType> BuildProducts = AllBuildProducts[Binary];
-						foreach (KeyValuePair<FileReference, BuildProductType> BuildProductPair in BuildProducts)
-						{
-							if (BuildProductPair.Value == BuildProductType.SymbolFile)
-							{
-								BuiltBinaries.Add(BuildProductPair.Key.FullName);
-							}
-						}
-					}
-				}
-
-				string IntermediateDirectory = InTarget.EngineIntermediateDirectory.FullName;
-				if (!Directory.Exists(IntermediateDirectory))
-				{
-					IntermediateDirectory = Path.GetFullPath("../../" + InTarget.AppName + "/Intermediate/Build/Mac/" + InTarget.AppName + "/" + InTarget.Configuration);
-					if (!Directory.Exists(IntermediateDirectory))
-					{
-						IntermediateDirectory = Path.GetFullPath("../../Engine/Intermediate/Build/Mac/" + InTarget.AppName + "/" + InTarget.Configuration);
-					}
-				}
-
-				string FixDylibDepsScript = Path.Combine(IntermediateDirectory, "FixDylibDependencies.sh");
-				string FinalizeAppBundleScript = Path.Combine(IntermediateDirectory, "FinalizeAppBundle.sh");
-
-				string RemoteWorkingDir = "";
-
-				bool bIsStaticLibrary = InTarget.OutputPath.HasExtension(".a");
-
-				if (!bIsStaticLibrary)
-				{
-					// Copy the command scripts to the intermediate on the target Mac.
-					string RemoteFixDylibDepsScript = ConvertPath(Path.GetFullPath(FixDylibDepsScript));
-					RemoteFixDylibDepsScript = RemoteFixDylibDepsScript.Replace("../../../../", "../../");
-					RPCUtilHelper.CopyFile(Path.GetFullPath(FixDylibDepsScript), RemoteFixDylibDepsScript, true);
-
-					if (!InTarget.Rules.bIsBuildingConsoleApplication)
-					{
-						string RemoteFinalizeAppBundleScript = ConvertPath(Path.GetFullPath(FinalizeAppBundleScript));
-						RemoteFinalizeAppBundleScript = RemoteFinalizeAppBundleScript.Replace("../../../../", "../../");
-						RPCUtilHelper.CopyFile(Path.GetFullPath(FinalizeAppBundleScript), RemoteFinalizeAppBundleScript, true);
-					}
-
-
-					// run it remotely
-					RemoteWorkingDir = ConvertPath(Path.GetDirectoryName(Path.GetFullPath(FixDylibDepsScript)));
-
-					Log.TraceInformation("Running FixDylibDependencies.sh...");
-					Hashtable Results = RPCUtilHelper.Command(RemoteWorkingDir, "/bin/sh", "FixDylibDependencies.sh", null);
-					if (Results != null)
-					{
-						string Result = (string)Results["CommandOutput"];
-						if (Result != null)
-						{
-							Log.TraceInformation(Result);
-						}
-					}
-
-					if (!InTarget.Rules.bIsBuildingConsoleApplication)
-					{
-						Log.TraceInformation("Running FinalizeAppBundle.sh...");
-						Results = RPCUtilHelper.Command(RemoteWorkingDir, "/bin/sh", "FinalizeAppBundle.sh", null);
-						if (Results != null)
-						{
-							string Result = (string)Results["CommandOutput"];
-							if (Result != null)
-							{
-								Log.TraceInformation(Result);
-							}
-						}
-					}
-				}
-
-
-				// If it is requested, send the app bundle back to the platform executing these commands.
-				if (InTarget.Rules.bCopyAppBundleBackToDevice)
-				{
-					Log.TraceInformation("Copying binaries back to this device...");
-
-					try
-					{
-						string BinaryDir = InTarget.OutputPath.Directory + "\\";
-						if (BinaryDir.EndsWith(InTarget.AppName + "\\Binaries\\Mac\\") && InTarget.TargetType != TargetType.Game)
-						{
-							BinaryDir = BinaryDir.Replace(InTarget.TargetType.ToString(), "Game");
-						}
-
-						string RemoteBinariesDir = ConvertPath(BinaryDir);
-						string LocalBinariesDir = BinaryDir;
-
-						// Get the app bundle's name
-						string AppFullName = InTarget.AppName;
-						if (InTarget.Configuration != InTarget.Rules.UndecoratedConfiguration)
-						{
-							AppFullName += "-" + InTarget.Platform.ToString();
-							AppFullName += "-" + InTarget.Configuration.ToString();
-						}
-
-						if (!InTarget.Rules.bIsBuildingConsoleApplication)
-						{
-							AppFullName += ".app";
-						}
-
-						List<string> NotBundledBinaries = new List<string>();
-						foreach (string BinaryPath in BuiltBinaries)
-						{
-							if (InTarget.Rules.bIsBuildingConsoleApplication || bIsStaticLibrary || !BinaryPath.StartsWith(LocalBinariesDir + AppFullName))
-							{
-								NotBundledBinaries.Add(BinaryPath);
-							}
-						}
-
-						// Zip the app bundle for transferring.
-						if (!InTarget.Rules.bIsBuildingConsoleApplication && !bIsStaticLibrary)
-						{
-							string ZipCommand = "zip -0 -r -y -T \"" + AppFullName + ".zip\" \"" + AppFullName + "\"";
-							RPCUtilHelper.Command(RemoteBinariesDir, ZipCommand, "", null);
-
-							// Copy the AppBundle back to the source machine
-							string LocalZipFileLocation = LocalBinariesDir + AppFullName + ".zip ";
-							string RemoteZipFileLocation = RemoteBinariesDir + AppFullName + ".zip";
-
-							RPCUtilHelper.CopyFile(RemoteZipFileLocation, LocalZipFileLocation, false);
-
-							// Extract the copied app bundle (in zip format) to the local binaries directory
-							using (ZipFile AppBundleZip = ZipFile.Read(LocalZipFileLocation))
-							{
-								foreach (ZipEntry Entry in AppBundleZip)
-								{
-									Entry.Extract(LocalBinariesDir, ExtractExistingFileAction.OverwriteSilently);
-								}
-							}
-
-							// Delete the zip as we no longer need/want it.
-							File.Delete(LocalZipFileLocation);
-							RPCUtilHelper.Command(RemoteBinariesDir, "rm -f \"" + AppFullName + ".zip\"", "", null);
-						}
-
-						if (NotBundledBinaries.Count > 0)
-						{
-
-							foreach (string BinaryPath in NotBundledBinaries)
-							{
-								RPCUtilHelper.CopyFile(ConvertPath(BinaryPath), BinaryPath, false);
-							}
-						}
-
-						Log.TraceInformation("Copied binaries successfully.");
-					}
-					catch (Exception)
-					{
-						Log.TraceInformation("Copying binaries back to this device failed.");
-					}
-				}
-			}
-		}
-
 		public override ICollection<FileItem> PostBuild(FileItem Executable, LinkEnvironment BinaryLinkEnvironment, ActionGraph ActionGraph)
 		{
 			ICollection<FileItem> OutputFiles = base.PostBuild(Executable, BinaryLinkEnvironment, ActionGraph);
diff --git a/Engine/Source/Programs/UnrealBuildTool/Platform/Mac/UEBuildMac.cs b/Engine/Source/Programs/UnrealBuildTool/Platform/Mac/UEBuildMac.cs
--- a/Engine/Source/Programs/UnrealBuildTool/Platform/Mac/UEBuildMac.cs
+++ b/Engine/Source/Programs/UnrealBuildTool/Platform/Mac/UEBuildMac.cs
@@ -240,16 +240,6 @@ namespace UnrealBuildTool
 		{
 		}
 
-		/// <summary>
-		/// Don't use absolute paths in unity files; we may be remote compiling.
-		/// </summary>
-		/// <returns>bool true if it is required, false if not</returns>
-		public override bool UseAbsolutePathsInUnityFiles()
-		{
-			return false;
-		}
-
-
 		/// <summary>
 		/// Return whether we wish to have this platform's binaries in our builds
 		/// </summary>
@@ -266,31 +256,6 @@ namespace UnrealBuildTool
 			return false;
 		}
 
-		/// <summary>
-		/// Converts the passed in path from UBT host to compiler native format.
-		/// </summary>
-		/// <param name="OriginalPath">The path to convert</param>
-		/// <returns>The path in native format for the toolchain</returns>
-		public override string ConvertPath(string OriginalPath)
-		{
-			return MacToolChain.ConvertPath(OriginalPath);
-		}
-
-		public override void PreBuildSync()
-		{
-			MacToolChain.PreBuildSync();
-		}
-
-		public override void PostBuildSync(UEBuildTarget Target)
-		{
-			MacToolChain.PostBuildSync(Target);
-		}
-
-		public override void PostCodeGeneration(UHTManifest Manifest)
-		{
-			MacToolChain.PostCodeGeneration(Manifest);
-		}
-
 		public override DirectoryReference GetBundleDirectory(ReadOnlyTargetRules Rules, List<FileReference> OutputFiles)
 		{
 			if(Rules.bIsBuildingConsoleApplication)
diff --git a/Engine/Source/Programs/UnrealBuildTool/ProjectFiles/Eddie/EddieProject.cs b/Engine/Source/Programs/UnrealBuildTool/ProjectFiles/Eddie/EddieProject.cs
--- a/Engine/Source/Programs/UnrealBuildTool/ProjectFiles/Eddie/EddieProject.cs
+++ b/Engine/Source/Programs/UnrealBuildTool/ProjectFiles/Eddie/EddieProject.cs
@@ -114,28 +114,7 @@ namespace UnrealBuildTool
 		
 		private static string ConvertPath(string InPath)
 		{
-			if (BuildHostPlatform.Current.Platform != UnrealTargetPlatform.Mac)
-			{
-				if (InPath[1] != ':')
-				{
-					throw new BuildException("Can only convert full paths ({0})", InPath);
-				}
-
-				string MacPath = string.Format("{0}/{1}/{2}/{3}",
-					RemoteToolChain.UserDevRootMac,
-					Environment.MachineName,
-					InPath[0].ToString().ToUpper(),
-					InPath.Substring(3));
-
-				// clean the path
-				MacPath = MacPath.Replace("\\", "/");
-
-				return MacPath;
-			}
-			else
-			{
-				return InPath.Replace("\\", "/");
-			}
+			return InPath.Replace("\\", "/");
 		}
 		
 		private void ParseSourceFilesIntoGroups()
diff --git a/Engine/Source/Programs/UnrealBuildTool/ProjectFiles/ProjectFileGenerator.cs b/Engine/Source/Programs/UnrealBuildTool/ProjectFiles/ProjectFileGenerator.cs
--- a/Engine/Source/Programs/UnrealBuildTool/ProjectFiles/ProjectFileGenerator.cs
+++ b/Engine/Source/Programs/UnrealBuildTool/ProjectFiles/ProjectFileGenerator.cs
@@ -2094,16 +2094,16 @@ namespace UnrealBuildTool
 					{
 						if(TargetFilePath.IsUnderDirectory(UnrealBuildTool.EnterpriseDirectory))
 						{
-							RulesAssembly = RulesCompiler.CreateEnterpriseRulesAssembly(false);
+							RulesAssembly = RulesCompiler.CreateEnterpriseRulesAssembly(false, false);
 						}
 						else
 						{
-							RulesAssembly = RulesCompiler.CreateEngineRulesAssembly(false);
+							RulesAssembly = RulesCompiler.CreateEngineRulesAssembly(false, false);
 						}
 					}
 					else
 					{
-						RulesAssembly = RulesCompiler.CreateProjectRulesAssembly(CheckProjectFile, false);
+						RulesAssembly = RulesCompiler.CreateProjectRulesAssembly(CheckProjectFile, false, false);
 					}
 
 					// Create target rules for all of the platforms and configuration combinations that we want to enable support for.
@@ -2301,17 +2301,19 @@ namespace UnrealBuildTool
 		{
 			// Setup a project file entry for this module's project.  Remember, some projects may host multiple modules!
 			DirectoryReference ShadersDirectory = DirectoryReference.Combine( UnrealBuildTool.EngineDirectory, "Shaders" );
-
-			List<string> SubdirectoryNamesToExclude = new List<string>();
+			if(DirectoryReference.Exists(ShadersDirectory))
 			{
-                // Don't include binary shaders in the project file.
-                SubdirectoryNamesToExclude.Add( "Binaries" );
-				// We never want shader intermediate files in our project file
-				SubdirectoryNamesToExclude.Add( "PDBDump" );
-				SubdirectoryNamesToExclude.Add( "WorkingDirectory" );
-			}
+				List<string> SubdirectoryNamesToExclude = new List<string>();
+				{
+					// Don't include binary shaders in the project file.
+					SubdirectoryNamesToExclude.Add( "Binaries" );
+					// We never want shader intermediate files in our project file
+					SubdirectoryNamesToExclude.Add( "PDBDump" );
+					SubdirectoryNamesToExclude.Add( "WorkingDirectory" );
+				}
 
-			EngineProject.AddFilesToProject( SourceFileSearch.FindFiles( ShadersDirectory, SubdirectoryNamesToExclude ), UnrealBuildTool.EngineDirectory );
+				EngineProject.AddFilesToProject( SourceFileSearch.FindFiles( ShadersDirectory, SubdirectoryNamesToExclude ), UnrealBuildTool.EngineDirectory );
+			}
 		}
 
 
diff --git a/Engine/Source/Programs/UnrealBuildTool/ProjectFiles/Xcode/XcodeProject.cs b/Engine/Source/Programs/UnrealBuildTool/ProjectFiles/Xcode/XcodeProject.cs
--- a/Engine/Source/Programs/UnrealBuildTool/ProjectFiles/Xcode/XcodeProject.cs
+++ b/Engine/Source/Programs/UnrealBuildTool/ProjectFiles/Xcode/XcodeProject.cs
@@ -223,28 +223,7 @@ namespace UnrealBuildTool
 		/// <returns>The normalized path</returns>
 		private static string ConvertPath(string InPath)
 		{
-			if (BuildHostPlatform.Current.Platform != UnrealTargetPlatform.Mac)
-			{
-				if (InPath[1] != ':')
-				{
-					throw new BuildException("Can only convert full paths ({0})", InPath);
-				}
-
-				string MacPath = string.Format("{0}/{1}/{2}/{3}",
-					RemoteToolChain.UserDevRootMac,
-					Environment.MachineName,
-					InPath[0].ToString().ToUpper(),
-					InPath.Substring(3));
-
-				// clean the path
-				MacPath = MacPath.Replace("\\", "/");
-
-				return MacPath;
-			}
-			else
-			{
-				return InPath.Replace("\\", "/");
-			}
+			return InPath.Replace("\\", "/");
 		}
 
 		/// <summary>
diff --git a/Engine/Source/Programs/UnrealBuildTool/System/ActionGraph.cs b/Engine/Source/Programs/UnrealBuildTool/System/ActionGraph.cs
--- a/Engine/Source/Programs/UnrealBuildTool/System/ActionGraph.cs
+++ b/Engine/Source/Programs/UnrealBuildTool/System/ActionGraph.cs
@@ -427,7 +427,7 @@ namespace UnrealBuildTool
 		/// <summary>
 		/// Executes a list of actions.
 		/// </summary>
-		public static bool ExecuteActions(BuildConfiguration BuildConfiguration, List<Action> ActionsToExecute, bool bIsRemoteCompile, out string ExecutorName, string TargetInfoForTelemetry, EHotReload HotReload)
+		public static bool ExecuteActions(BuildConfiguration BuildConfiguration, List<Action> ActionsToExecute, out string ExecutorName, string TargetInfoForTelemetry, EHotReload HotReload)
 		{
 			bool Result = true;
 			ExecutorName = "";
@@ -436,11 +436,7 @@ namespace UnrealBuildTool
 				DateTime StartTime = DateTime.UtcNow;
 
 				ActionExecutor Executor;
-				if(bIsRemoteCompile)
-				{
-					Executor = new RemoteExecutor();
-				}
-				else if(ActionsToExecute.Any(x => x.ActionHandler != null))
+				if(ActionsToExecute.Any(x => x.ActionHandler != null))
 				{
 					Executor = new LocalExecutor();
 				}
@@ -477,18 +473,7 @@ namespace UnrealBuildTool
 						{
 							foreach (FileItem Item in BuildAction.ProducedItems)
 							{
-								bool bExists;
-								if (Item.bIsRemoteFile)
-								{
-									DateTime UnusedTime;
-									long UnusedLength;
-									bExists = RPCUtilHelper.GetRemoteFileInfo(Item.AbsolutePath, out UnusedTime, out UnusedLength);
-								}
-								else
-								{
-									// allow output to be a directory
-									bExists = File.Exists(Item.AbsolutePath) || Directory.Exists(Item.AbsolutePath);
-								}
+								bool bExists = File.Exists(Item.AbsolutePath) || Directory.Exists(Item.AbsolutePath);
 
 								if (HotReload != EHotReload.Disabled)
 								{
@@ -905,7 +890,7 @@ namespace UnrealBuildTool
 					// If the produced file doesn't exist or has zero size, consider it outdated.  The zero size check is to detect cases
 					// where aborting an earlier compile produced invalid zero-sized obj files, but that may cause actions where that's
 					// legitimate output to always be considered outdated.
-					if (ProducedItem.bExists && (ProducedItem.bIsRemoteFile || ProducedItem.Length > 0 || ProducedItem.IsDirectory))
+					if (ProducedItem.bExists && (ProducedItem.Length > 0 || ProducedItem.IsDirectory))
 					{
 						// Use the oldest produced item's time as the last execution time.
 						if (ProducedItem.LastWriteTime < LastExecutionTime)
@@ -1167,29 +1152,11 @@ namespace UnrealBuildTool
 				{
 					foreach (FileItem ProducedItem in OutdatedActionInfo.Key.ProducedItems)
 					{
-						if (ProducedItem.bIsRemoteFile)
-						{
-							// we don't need to do this in the SSH mode, the action will have an output file, and it will use that to make the directory while executing the command
-							if (RemoteToolChain.bUseRPCUtil)
-							{
-								try
-								{
-									RPCUtilHelper.MakeDirectory(Path.GetDirectoryName(ProducedItem.AbsolutePath).Replace("\\", "/"));
-								}
-								catch (System.Exception Ex)
-								{
-									throw new BuildException(Ex, "Error while creating remote directory for '{0}'.  (Exception: {1})", ProducedItem.AbsolutePath, Ex.Message);
-								}
-							}
-						}
-						else
+						string DirectoryPath = Path.GetDirectoryName(ProducedItem.AbsolutePath);
+						if (!Directory.Exists(DirectoryPath))
 						{
-							string DirectoryPath = Path.GetDirectoryName(ProducedItem.AbsolutePath);
-							if (!Directory.Exists(DirectoryPath))
-							{
-								Log.TraceVerbose("Creating directory for produced item: {0}", DirectoryPath);
-								Directory.CreateDirectory(DirectoryPath);
-							}
+							Log.TraceVerbose("Creating directory for produced item: {0}", DirectoryPath);
+							Directory.CreateDirectory(DirectoryPath);
 						}
 					}
 				}
diff --git a/Engine/Source/Programs/UnrealBuildTool/System/Asn.cs b/Engine/Source/Programs/UnrealBuildTool/System/Asn.cs
new file mode 100644
--- /dev/null
+++ b/Engine/Source/Programs/UnrealBuildTool/System/Asn.cs
@@ -0,0 +1,187 @@
+using System;
+using System.Collections.Generic;
+using System.IO;
+using System.Linq;
+using System.Text;
+using System.Threading.Tasks;
+
+namespace UnrealBuildTool
+{
+	/// <summary>
+	/// Parser for ASN.1 BER formats (eg. iOS mobile provisions)
+	/// </summary>
+	static class Asn
+	{
+		public enum FieldTag
+		{
+			EOC = 0,
+			BOOLEAN = 1,
+			INTEGER = 2,
+			OCTET_STRING = 4,
+			NULL = 5,
+			OBJECT_IDENTIFIER = 6,
+			SEQUENCE = 16,
+			SET = 17,
+			PRINTABLE_STRING = 19,
+		}
+
+		public struct FieldInfo
+		{
+			public int TagClass;
+			public bool bPrimitive;
+			public FieldTag Tag;
+			public int Length;
+		}
+
+		public static class ObjectIdentifier
+		{
+			public static readonly int[] CountryName = new int[]{ 2, 5, 4, 6 };
+			public static readonly int[] OrganizationName = new int[]{ 2, 5, 4, 10 };
+			public static readonly int[] Pkcs7_SignedData = new int[]{ 1, 2, 840, 113549, 1, 7, 2 };
+			public static readonly int[] Pkcs7_Data = new int[]{ 1, 2, 840, 113549, 1, 7, 1 };
+		}
+
+		static readonly List<KeyValuePair<int[], string>> KnownObjectIdentifiers = new List<KeyValuePair<int[], string>>
+		{
+			new KeyValuePair<int[], string>(ObjectIdentifier.CountryName, "CountryName"),
+			new KeyValuePair<int[], string>(ObjectIdentifier.OrganizationName, "OrganizationName"),
+			new KeyValuePair<int[], string>(ObjectIdentifier.Pkcs7_SignedData, "Pkcs7-SignedData"),
+			new KeyValuePair<int[], string>(ObjectIdentifier.Pkcs7_Data, "Pkcs7-Data"),
+		};
+
+		public static FieldInfo ReadField(BinaryReader Reader)
+		{
+			// Read the type and length
+			int Type = Reader.ReadByte();
+			int Length = ReadLength(Reader);
+
+			// Unpack the type
+			FieldInfo Field = new FieldInfo();
+			Field.TagClass = (Type >> 6);
+			Field.bPrimitive = (Type & 0x20) == 0;
+			Field.Tag = (FieldTag)(Type & 0x1f);
+			Field.Length = Length;
+			return Field;
+		}
+
+		public static void SkipValue(BinaryReader Reader, FieldInfo Field)
+		{
+			if(Field.bPrimitive)
+			{
+				Reader.BaseStream.Seek(Field.Length, SeekOrigin.Current);
+			}
+		}
+
+		public static int ReadInteger(BinaryReader Reader, int Length)
+		{
+			int Value = 0;
+			for(int Idx = 0; Idx < Length; Idx++)
+			{
+				Value = (Value << 8) | (int)Reader.ReadByte();
+			}
+			return Value;
+		}
+
+		public static int ReadLength(BinaryReader Reader)
+		{
+			int Count = (int)Reader.ReadByte();
+			if(Count <= 0x7f)
+			{
+				return Count;
+			}
+			else
+			{
+				return ReadInteger(Reader, Count & 0x7f);
+			}
+		}
+
+		public static int[] ReadObjectIdentifier(BinaryReader Reader, int Length)
+		{
+			byte[] Data = Reader.ReadBytes(Length);
+
+			List<int> Values = new List<int>();
+			Values.Add((int)Data[0] / 40);
+			Values.Add((int)Data[0] % 40);
+			for(int Idx = 1; Idx < Data.Length; Idx++)
+			{
+				int Value = (int)Data[Idx] & 0x7f;
+				while(((int)Data[Idx] & 0x80) != 0)
+				{
+					Value = (Value << 7) | ((int)Data[++Idx] & 0x7f);
+				}
+				Values.Add(Value);
+			}
+			return Values.ToArray();
+		}
+
+		public static string GetObjectIdentifierName(int[] Values)
+		{
+			string Description;
+			if(!TryGetObjectIdentifierName(Values, out Description))
+			{
+				Description = String.Format("Unknown ({0})", String.Join(", ", Values.Select(x => x.ToString())));
+			}
+			return Description;
+		}
+
+		public static bool TryGetObjectIdentifierName(int[] Values, out string Name)
+		{
+			foreach(KeyValuePair<int[], string> KnownObjectIdentifier in KnownObjectIdentifiers)
+			{
+				if(Enumerable.SequenceEqual(Values, KnownObjectIdentifier.Key))
+				{
+					Name = KnownObjectIdentifier.Value;
+					return true;
+				}
+			}
+
+			Name = null;
+			return false;
+		}
+
+		public static void Dump(BinaryReader Reader, string Indent)
+		{
+			FieldInfo Field = ReadField(Reader);
+
+			// If it's a primitive type, unpack the value
+			string DescriptionSuffix = "";
+			if(Field.bPrimitive)
+			{
+				if(Field.Length > 0)
+				{
+					string Description;
+					switch(Field.Tag)
+					{
+						case FieldTag.INTEGER:
+							Description = String.Format("{0}", ReadInteger(Reader, Field.Length));
+							break;
+						case FieldTag.OBJECT_IDENTIFIER:
+							Description = GetObjectIdentifierName(ReadObjectIdentifier(Reader, Field.Length));
+							break;
+						case FieldTag.PRINTABLE_STRING:
+							Description = Encoding.ASCII.GetString(Reader.ReadBytes(Field.Length));
+							break;
+						default:
+							Reader.ReadBytes(Field.Length);
+							Description = "unknown";
+							break;
+					}
+					DescriptionSuffix = String.Format(" = {0}", Description);
+				}
+			}
+
+			// Print the contents/name of this element
+			Console.WriteLine("{0}{1}{2}", Indent, Field.Tag, DescriptionSuffix);
+
+			// If it's not a primitive, print the contents
+			if(!Field.bPrimitive)
+			{
+				long EndOffset = Reader.BaseStream.Position + Field.Length;
+				while(Reader.BaseStream.Position < EndOffset)
+				{
+					Dump(Reader, Indent + "  ");
+				}
+			}
+		}
+	}
+}
diff --git a/Engine/Source/Programs/UnrealBuildTool/System/DependencyCache.cs b/Engine/Source/Programs/UnrealBuildTool/System/DependencyCache.cs
--- a/Engine/Source/Programs/UnrealBuildTool/System/DependencyCache.cs
+++ b/Engine/Source/Programs/UnrealBuildTool/System/DependencyCache.cs
@@ -34,14 +34,6 @@ namespace UnrealBuildTool
 		/// <param name="InIncludeName"></param>
 		public DependencyInclude(string InIncludeName)
 		{
-            if (BuildHostPlatform.Current.Platform != UnrealTargetPlatform.Mac)
-            {
-                if (InIncludeName.StartsWith("/Users") || InIncludeName.StartsWith("\\Users"))
-                {
-                    // convert the path back to Windows
-                    InIncludeName = RemoteExports.UnconvertPath(InIncludeName);
-                }
-            }
 			IncludeName = InIncludeName;
 		}
 	}
diff --git a/Engine/Source/Programs/UnrealBuildTool/System/DynamicCompilation.cs b/Engine/Source/Programs/UnrealBuildTool/System/DynamicCompilation.cs
--- a/Engine/Source/Programs/UnrealBuildTool/System/DynamicCompilation.cs
+++ b/Engine/Source/Programs/UnrealBuildTool/System/DynamicCompilation.cs
@@ -1,4 +1,4 @@
-﻿// Copyright 1998-2018 Epic Games, Inc. All Rights Reserved.
+// Copyright 1998-2018 Epic Games, Inc. All Rights Reserved.
 
 using System;
 using System.Collections.Generic;
@@ -360,7 +360,7 @@ namespace UnrealBuildTool
 			}
 			catch (Exception Ex)
 			{
-				throw new BuildException(Ex, "Failed to launch compiler to compile assembly from source files '{0}' (Exception: {1})", SourceFileNames.ToString(), Ex.Message);
+				throw new BuildException(Ex, "Failed to launch compiler to compile assembly from source files:\n  {0}\n(Exception: {1})", String.Join("\n  ", SourceFileNames), Ex.ToString());
 			}
 
 			// Display compilation warnings and errors
@@ -432,6 +432,10 @@ namespace UnrealBuildTool
 					Log.TraceInformation(String.Format("Compiled assembly file '{0}' appears to be for a newer CLR version or is otherwise invalid.  Unreal Build Tool will try to recompile this assembly now.  (Exception: {1})", OutputAssemblyPath, Ex.Message));
 					bNeedsCompilation = true;
 				}
+				catch (FileNotFoundException)
+			    {
+				    throw new BuildException("Precompiled rules assembly '{0}' does not exist.", OutputAssemblyPath);
+			    }
 				catch (Exception Ex)
 				{
 					throw new BuildException(Ex, "Error while loading previously-compiled assembly file '{0}'.  (Exception: {1})", OutputAssemblyPath, Ex.Message);
diff --git a/Engine/Source/Programs/UnrealBuildTool/System/ExternalExecution.cs b/Engine/Source/Programs/UnrealBuildTool/System/ExternalExecution.cs
--- a/Engine/Source/Programs/UnrealBuildTool/System/ExternalExecution.cs
+++ b/Engine/Source/Programs/UnrealBuildTool/System/ExternalExecution.cs
@@ -1276,16 +1276,6 @@ namespace UnrealBuildTool
 
 				Progress.Write(2, 3);
 
-				// There will never be generated code if we're building UHT, so this should never be called.
-				if (!bIsBuildingUHT)
-				{
-					// Allow generated code to be sync'd to remote machines if needed. This needs to be done even if UHT did not run because
-					// generated headers include other generated headers using absolute paths which in case of building remotely are already
-					// the remote machine absolute paths. Because of that parsing headers will not result in finding all includes properly.
-					// @todo ubtmake: Need to figure out what this does in the assembler case, and whether we need to run it
-					UEBuildPlatform.GetBuildPlatform(Target.Platform).PostCodeGeneration(Manifest);
-				}
-
 				// touch the directories
 				UpdateDirectoryTimestamps(UObjectModules);
 
diff --git a/Engine/Source/Programs/UnrealBuildTool/System/FileItem.cs b/Engine/Source/Programs/UnrealBuildTool/System/FileItem.cs
--- a/Engine/Source/Programs/UnrealBuildTool/System/FileItem.cs
+++ b/Engine/Source/Programs/UnrealBuildTool/System/FileItem.cs
@@ -37,11 +37,6 @@ namespace UnrealBuildTool
 		/// </summary>
 		public bool bNeedsHotReloadNumbersDLLCleanUp = false;
 
-		/// <summary>
-		/// Whether or not this is a remote file, in which case we can't access it directly
-		/// </summary>
-		public bool bIsRemoteFile = false;
-
 		/// <summary>
 		/// Accessor for the absolute path to the file
 		/// </summary>
@@ -98,14 +93,7 @@ namespace UnrealBuildTool
 		public DateTimeOffset _LastWriteTime;
 		public DateTimeOffset LastWriteTime
 		{
-			get
-			{
-				if (bIsRemoteFile)
-				{
-					LookupOutstandingFiles();
-				}
-				return _LastWriteTime;
-			}
+			get { return _LastWriteTime; }
 			set { _LastWriteTime = value; }
 		}
 
@@ -115,14 +103,7 @@ namespace UnrealBuildTool
 		public bool _bExists = false;
 		public bool bExists
 		{
-			get
-			{
-				if (bIsRemoteFile)
-				{
-					LookupOutstandingFiles();
-				}
-				return _bExists;
-			}
+			get { return _bExists; }
 			set { _bExists = value; }
 		}
 
@@ -132,14 +113,7 @@ namespace UnrealBuildTool
 		public long _Length = -1;
 		public long Length
 		{
-			get
-			{
-				if (bIsRemoteFile)
-				{
-					LookupOutstandingFiles();
-				}
-				return _Length;
-			}
+			get { return _Length; }
 			set { _Length = value; }
 		}
 
@@ -159,18 +133,12 @@ namespace UnrealBuildTool
 		/// </summary>
 		static Dictionary<FileReference, FileItem> UniqueSourceFileMap = new Dictionary<FileReference, FileItem>();
 
-		/// <summary>
-		/// A list of remote file items that have been created but haven't needed the remote info yet, so we can gang up many into one request
-		/// </summary>
-		static List<FileItem> DelayedRemoteLookupFiles = new List<FileItem>();
-
 		/// <summary>
 		/// Clears the FileItem caches.
 		/// </summary>
 		public static void ClearCaches()
 		{
 			UniqueSourceFileMap.Clear();
-			DelayedRemoteLookupFiles.Clear();
 		}
 
 		/// <summary>
@@ -184,32 +152,6 @@ namespace UnrealBuildTool
 			}
 		}
 
-		/// <summary>
-		/// Resolve any outstanding remote file info lookups
-		/// </summary>
-		private void LookupOutstandingFiles()
-		{
-			// for remote files, look up any outstanding files
-			if (bIsRemoteFile)
-			{
-				FileItem[] Files = null;
-				lock (DelayedRemoteLookupFiles)
-				{
-					if (DelayedRemoteLookupFiles.Count > 0)
-					{
-						// make an array so we can clear the original array, just in case BatchFileInfo does something that uses
-						// DelayedRemoteLookupFiles, so we don't deadlock
-						Files = DelayedRemoteLookupFiles.ToArray();
-						DelayedRemoteLookupFiles.Clear();
-					}
-				}
-				if (Files != null)
-				{
-					RPCUtilHelper.BatchFileInfo(Files);
-				}
-			}
-		}
-
 		/// <returns>The FileItem that represents the given file path.</returns>
 		public static FileItem GetItemByPath(string FilePath)
 		{
@@ -334,7 +276,6 @@ namespace UnrealBuildTool
 		public void Delete()
 		{
 			Debug.Assert(_bExists);
-			Debug.Assert(!bIsRemoteFile);
 
 			int MaxRetryCount = 3;
 			int DeleteTryCount = 0;
@@ -403,7 +344,6 @@ namespace UnrealBuildTool
 		{
 			ProducingAction = (Action)SerializationInfo.GetValue("pa", typeof(Action));
 			Location = (FileReference)SerializationInfo.GetValue("fi", typeof(FileReference));
-			bIsRemoteFile = SerializationInfo.GetBoolean("rf");
 			bNeedsHotReloadNumbersDLLCleanUp = SerializationInfo.GetBoolean("hr");
 			CachedIncludePaths = (CppIncludePaths)SerializationInfo.GetValue("ci", typeof(CppIncludePaths));
 
@@ -418,17 +358,7 @@ namespace UnrealBuildTool
 					// Log.TraceInformation( "Missing: " + FileAbsolutePath );
 				}
 
-				if (bIsRemoteFile)
-				{
-					lock (DelayedRemoteLookupFiles)
-					{
-						DelayedRemoteLookupFiles.Add(this);
-					}
-				}
-				else
-				{
-					UniqueSourceFileMap[Location] = this;
-				}
+				UniqueSourceFileMap[Location] = this;
 			}
 		}
 
@@ -440,7 +370,6 @@ namespace UnrealBuildTool
 		{
 			SerializationInfo.AddValue("pa", ProducingAction);
 			SerializationInfo.AddValue("fi", Location);
-			SerializationInfo.AddValue("rf", bIsRemoteFile);
 			SerializationInfo.AddValue("hr", bNeedsHotReloadNumbersDLLCleanUp);
 			SerializationInfo.AddValue("ci", CachedIncludePaths);
 		}
@@ -489,41 +418,22 @@ namespace UnrealBuildTool
 		/// </summary>
 		protected FileItem(FileReference InReference, bool InIsRemoteFile, UnrealTargetPlatform Platform)
 		{
-			bIsRemoteFile = InIsRemoteFile;
 			Location = InReference;
 
-			// @todo iosmerge: This doesn't handle remote directories (may be needed for compiling Mac from Windows)
-			if (bIsRemoteFile)
+			FileInfo Info = new FileInfo(AbsolutePath);
+
+			_bExists = Info.Exists;
+			if (_bExists)
 			{
-				if (Platform == UnrealTargetPlatform.IOS || Platform == UnrealTargetPlatform.Mac)
-				{
-					lock (DelayedRemoteLookupFiles)
-					{
-						DelayedRemoteLookupFiles.Add(this);
-					}
-				}
-				else
-				{
-					throw new BuildException("Only IPhone and Mac support remote FileItems");
-				}
+				_LastWriteTime = Info.LastWriteTimeUtc;
+				_Length = Info.Length;
 			}
-			else
-			{
-				FileInfo Info = new FileInfo(AbsolutePath);
 
-				_bExists = Info.Exists;
-				if (_bExists)
-				{
-					_LastWriteTime = Info.LastWriteTimeUtc;
-					_Length = Info.Length;
-				}
-
-				++TotalFileItemCount;
-				if (!_bExists)
-				{
-					++MissingFileItemCount;
-					// Log.TraceInformation( "Missing: " + FileAbsolutePath );
-				}
+			++TotalFileItemCount;
+			if (!_bExists)
+			{
+				++MissingFileItemCount;
+				// Log.TraceInformation( "Missing: " + FileAbsolutePath );
 			}
 
 			// @todo iosmerge: This was in UE3, why commented out now?
diff --git a/Engine/Source/Programs/UnrealBuildTool/System/PlatformExports.cs b/Engine/Source/Programs/UnrealBuildTool/System/PlatformExports.cs
--- a/Engine/Source/Programs/UnrealBuildTool/System/PlatformExports.cs
+++ b/Engine/Source/Programs/UnrealBuildTool/System/PlatformExports.cs
@@ -146,7 +146,7 @@ namespace UnrealBuildTool
 			UnrealBuildTool.SetIsEngineInstalled(bIsEngineInstalled);
 
 			// Read the XML configuration files
-			if(!XmlConfig.ReadConfigFiles())
+			if(!XmlConfig.ReadConfigFiles(false))
 			{
 				return false;
 			}
diff --git a/Engine/Source/Programs/UnrealBuildTool/System/RPCUtilHelper.cs b/Engine/Source/Programs/UnrealBuildTool/System/RPCUtilHelper.cs
deleted file mode 100644
diff --git a/Engine/Source/Programs/UnrealBuildTool/System/RemoteExecutor.cs b/Engine/Source/Programs/UnrealBuildTool/System/RemoteExecutor.cs
deleted file mode 100644
diff --git a/Engine/Source/Programs/UnrealBuildTool/System/RulesAssembly.cs b/Engine/Source/Programs/UnrealBuildTool/System/RulesAssembly.cs
--- a/Engine/Source/Programs/UnrealBuildTool/System/RulesAssembly.cs
+++ b/Engine/Source/Programs/UnrealBuildTool/System/RulesAssembly.cs
@@ -85,8 +85,9 @@ namespace UnrealBuildTool
 		/// <param name="bContainsEngineModules">Whether this assembly contains engine modules. Used to initialize the default value for ModuleRules.bTreatAsEngineModule.</param>
 		/// <param name="bUseBackwardsCompatibleDefaults">Whether modules in this assembly should use backwards-compatible defaults.</param>
 		/// <param name="bReadOnly">Whether the modules and targets in this assembly are installed, and should be created with the bUsePrecompiled flag set</param> 
+		/// <param name="bSkipCompile">Whether to skip compiling this assembly</param>
 		/// <param name="Parent">The parent rules assembly</param>
-		public RulesAssembly(DirectoryReference BaseDir, IReadOnlyList<PluginInfo> Plugins, List<FileReference> ModuleFiles, List<FileReference> TargetFiles, Dictionary<FileReference, PluginInfo> ModuleFileToPluginInfo, FileReference AssemblyFileName, bool bContainsEngineModules, bool bUseBackwardsCompatibleDefaults, bool bReadOnly, RulesAssembly Parent)
+		public RulesAssembly(DirectoryReference BaseDir, IReadOnlyList<PluginInfo> Plugins, List<FileReference> ModuleFiles, List<FileReference> TargetFiles, Dictionary<FileReference, PluginInfo> ModuleFileToPluginInfo, FileReference AssemblyFileName, bool bContainsEngineModules, bool bUseBackwardsCompatibleDefaults, bool bReadOnly, bool bSkipCompile, RulesAssembly Parent)
 		{
 			this.BaseDir = BaseDir;
 			this.Plugins = Plugins;
@@ -105,7 +106,7 @@ namespace UnrealBuildTool
 			if (AssemblySourceFiles.Count > 0)
 			{
 				List<string> PreprocessorDefines = GetPreprocessorDefinitions();
-				CompiledAssembly = DynamicCompilation.CompileAndLoadAssembly(AssemblyFileName, AssemblySourceFiles, PreprocessorDefines: PreprocessorDefines);
+				CompiledAssembly = DynamicCompilation.CompileAndLoadAssembly(AssemblyFileName, AssemblySourceFiles, PreprocessorDefines: PreprocessorDefines, DoNotCompile: bSkipCompile);
 			}
 
 			// Setup the module map
diff --git a/Engine/Source/Programs/UnrealBuildTool/System/RulesCompiler.cs b/Engine/Source/Programs/UnrealBuildTool/System/RulesCompiler.cs
--- a/Engine/Source/Programs/UnrealBuildTool/System/RulesCompiler.cs
+++ b/Engine/Source/Programs/UnrealBuildTool/System/RulesCompiler.cs
@@ -1,4 +1,4 @@
-// Copyright 1998-2018 Epic Games, Inc. All Rights Reserved.
+﻿// Copyright 1998-2018 Epic Games, Inc. All Rights Reserved.
 
 using System;
 using System.IO;
@@ -259,13 +259,14 @@ namespace UnrealBuildTool
 		/// Creates the engine rules assembly
 		/// </summary>
 		/// <param name="bUsePrecompiled">Whether to use a precompiled engine</param>
+		/// <param name="bSkipCompile">Whether to skip compilation for this assembly</param>
 		/// <returns>New rules assembly</returns>
-		public static RulesAssembly CreateEngineRulesAssembly(bool bUsePrecompiled)
+		public static RulesAssembly CreateEngineRulesAssembly(bool bUsePrecompiled, bool bSkipCompile)
 		{
 			if (EngineRulesAssembly == null)
 			{
 				IReadOnlyList<PluginInfo> IncludedPlugins = Plugins.ReadEnginePlugins(UnrealBuildTool.EngineDirectory);
-				EngineRulesAssembly = CreateEngineOrEnterpriseRulesAssembly(UnrealBuildTool.EngineDirectory, ProjectFileGenerator.EngineProjectFileNameBase, IncludedPlugins, UnrealBuildTool.IsEngineInstalled() || bUsePrecompiled, null);
+				EngineRulesAssembly = CreateEngineOrEnterpriseRulesAssembly(UnrealBuildTool.EngineDirectory, ProjectFileGenerator.EngineProjectFileNameBase, IncludedPlugins, UnrealBuildTool.IsEngineInstalled() || bUsePrecompiled, bSkipCompile, null);
 			}
 			return EngineRulesAssembly;
 		}
@@ -274,13 +275,14 @@ namespace UnrealBuildTool
 		/// Creates the enterprise rules assembly
 		/// </summary>
 		/// <param name="bUsePrecompiled">Whether to use a precompiled enterprise and engine folder</param>
+		/// <param name="bSkipCompile">Whether to skip compilation for this assembly</param>
 		/// <returns>New rules assembly. Returns null if the enterprise directory is unavailable.</returns>
-		public static RulesAssembly CreateEnterpriseRulesAssembly(bool bUsePrecompiled)
+		public static RulesAssembly CreateEnterpriseRulesAssembly(bool bUsePrecompiled, bool bSkipCompile)
 		{
 			if (EnterpriseRulesAssembly == null && DirectoryReference.Exists(UnrealBuildTool.EnterpriseDirectory))
 			{
 				IReadOnlyList<PluginInfo> IncludedPlugins = Plugins.ReadEnterprisePlugins(UnrealBuildTool.EnterpriseDirectory);
-				EnterpriseRulesAssembly = CreateEngineOrEnterpriseRulesAssembly(UnrealBuildTool.EnterpriseDirectory, ProjectFileGenerator.EnterpriseProjectFileNameBase, IncludedPlugins, UnrealBuildTool.IsEnterpriseInstalled() || bUsePrecompiled, CreateEngineRulesAssembly(bUsePrecompiled));
+				EnterpriseRulesAssembly = CreateEngineOrEnterpriseRulesAssembly(UnrealBuildTool.EnterpriseDirectory, ProjectFileGenerator.EnterpriseProjectFileNameBase, IncludedPlugins, UnrealBuildTool.IsEnterpriseInstalled() || bUsePrecompiled, bSkipCompile, CreateEngineRulesAssembly(bUsePrecompiled, bSkipCompile));
 			}
 			return EnterpriseRulesAssembly;
 		}
@@ -292,9 +294,10 @@ namespace UnrealBuildTool
 		/// <param name="AssemblyPrefix">A prefix for the assembly file name</param>
 		/// <param name="Plugins">List of plugins to include in this assembly</param>
 		/// <param name="bReadOnly">Whether the assembly should be marked as installed</param>
+		/// <param name="bSkipCompile">Whether to skip compilation for this assembly</param>
 		/// <param name="Parent">The parent rules assembly</param>
 		/// <returns>New rules assembly</returns>
-		private static RulesAssembly CreateEngineOrEnterpriseRulesAssembly(DirectoryReference RootDirectory, string AssemblyPrefix, IReadOnlyList<PluginInfo> Plugins, bool bReadOnly, RulesAssembly Parent)
+		private static RulesAssembly CreateEngineOrEnterpriseRulesAssembly(DirectoryReference RootDirectory, string AssemblyPrefix, IReadOnlyList<PluginInfo> Plugins, bool bReadOnly, bool bSkipCompile, RulesAssembly Parent)
 		{
 			DirectoryReference SourceDirectory = DirectoryReference.Combine(RootDirectory, "Source");
 			DirectoryReference ProgramsDirectory = DirectoryReference.Combine(SourceDirectory, "Programs");
@@ -315,7 +318,7 @@ namespace UnrealBuildTool
 
 			// Create the assembly
 			FileReference EngineAssemblyFileName = FileReference.Combine(RootDirectory, "Intermediate", "Build", "BuildRules", AssemblyPrefix + "Rules" + FrameworkAssemblyExtension);
-			RulesAssembly EngineAssembly = new RulesAssembly(RootDirectory, Plugins, EngineModuleFiles, new List<FileReference>(), ModuleFileToPluginInfo, EngineAssemblyFileName, bContainsEngineModules: true, bUseBackwardsCompatibleDefaults: false, bReadOnly: bReadOnly, Parent: Parent);
+			RulesAssembly EngineAssembly = new RulesAssembly(RootDirectory, Plugins, EngineModuleFiles, new List<FileReference>(), ModuleFileToPluginInfo, EngineAssemblyFileName, bContainsEngineModules: true, bUseBackwardsCompatibleDefaults: false, bReadOnly: bReadOnly, bSkipCompile: bSkipCompile, Parent: Parent);
 
 			// Find all the rules files
 			List<FileReference> ProgramModuleFiles = new List<FileReference>(FindAllRulesFiles(ProgramsDirectory, RulesFileType.Module));
@@ -323,7 +326,7 @@ namespace UnrealBuildTool
 
 			// Create a path to the assembly that we'll either load or compile
 			FileReference ProgramAssemblyFileName = FileReference.Combine(RootDirectory, "Intermediate", "Build", "BuildRules", AssemblyPrefix + "ProgramRules" + FrameworkAssemblyExtension);
-			RulesAssembly ProgramAssembly = new RulesAssembly(RootDirectory, new List<PluginInfo>().AsReadOnly(), ProgramModuleFiles, ProgramTargetFiles, new Dictionary<FileReference, PluginInfo>(), ProgramAssemblyFileName, bContainsEngineModules: false, bUseBackwardsCompatibleDefaults: false, bReadOnly: bReadOnly, Parent: EngineAssembly);
+			RulesAssembly ProgramAssembly = new RulesAssembly(RootDirectory, new List<PluginInfo>().AsReadOnly(), ProgramModuleFiles, ProgramTargetFiles, new Dictionary<FileReference, PluginInfo>(), ProgramAssemblyFileName, bContainsEngineModules: false, bUseBackwardsCompatibleDefaults: false, bReadOnly: bReadOnly, bSkipCompile: bSkipCompile, Parent: EngineAssembly);
 
 			// Return the combined assembly
 			return ProgramAssembly;
@@ -334,8 +337,9 @@ namespace UnrealBuildTool
 		/// </summary>
 		/// <param name="ProjectFileName">The project file to create rules for. Null for the engine.</param>
 		/// <param name="bUsePrecompiled">Whether to use a precompiled engine</param>
+		/// <param name="bSkipCompile">Whether to skip compilation for this assembly</param>
 		/// <returns>New rules assembly</returns>
-		public static RulesAssembly CreateProjectRulesAssembly(FileReference ProjectFileName, bool bUsePrecompiled)
+		public static RulesAssembly CreateProjectRulesAssembly(FileReference ProjectFileName, bool bUsePrecompiled, bool bSkipCompile)
 		{
 			// Check if there's an existing assembly for this project
 			RulesAssembly ProjectRulesAssembly;
@@ -347,11 +351,11 @@ namespace UnrealBuildTool
 				RulesAssembly Parent;
 				if (Project.IsEnterpriseProject)
 				{
-					Parent = CreateEnterpriseRulesAssembly(bUsePrecompiled);
+					Parent = CreateEnterpriseRulesAssembly(bUsePrecompiled, bSkipCompile);
 				}
 				else
 				{
-					Parent = CreateEngineRulesAssembly(bUsePrecompiled);
+					Parent = CreateEngineRulesAssembly(bUsePrecompiled, bSkipCompile);
 				}
 
 				// Find all the rules under the project source directory
@@ -386,20 +390,21 @@ namespace UnrealBuildTool
 
 				// Compile the assembly
 				FileReference AssemblyFileName = FileReference.Combine(ProjectDirectory, "Intermediate", "Build", "BuildRules", ProjectFileName.GetFileNameWithoutExtension() + "ModuleRules" + FrameworkAssemblyExtension);
-				ProjectRulesAssembly = new RulesAssembly(ProjectDirectory, ProjectPlugins, ModuleFiles, TargetFiles, ModuleFileToPluginInfo, AssemblyFileName, bContainsEngineModules: false, bUseBackwardsCompatibleDefaults: true, bReadOnly: UnrealBuildTool.IsProjectInstalled(), Parent: Parent);
+				ProjectRulesAssembly = new RulesAssembly(ProjectDirectory, ProjectPlugins, ModuleFiles, TargetFiles, ModuleFileToPluginInfo, AssemblyFileName, bContainsEngineModules: false, bUseBackwardsCompatibleDefaults: true, bReadOnly: UnrealBuildTool.IsProjectInstalled(), bSkipCompile: bSkipCompile, Parent: Parent);
 				LoadedAssemblyMap.Add(ProjectFileName, ProjectRulesAssembly);
 			}
 			return ProjectRulesAssembly;
 		}
 
-        /// <summary>
-        /// Creates a rules assembly with the given parameters.
-        /// </summary>
-        /// <param name="PluginFileName">The plugin file to create rules for</param>
-        /// <param name="Parent">The parent rules assembly</param>
+		/// <summary>
+		/// Creates a rules assembly with the given parameters.
+		/// </summary>
+		/// <param name="PluginFileName">The plugin file to create rules for</param>
+		/// <param name="bSkipCompile">Whether to skip compilation for this assembly</param>
+		/// <param name="Parent">The parent rules assembly</param>
         /// <param name="bContainsEngineModules">Whether the plugin contains engine modules. Used to initialize the default value for ModuleRules.bTreatAsEngineModule.</param>
-        /// <returns>The new rules assembly</returns>
-        public static RulesAssembly CreatePluginRulesAssembly(FileReference PluginFileName, RulesAssembly Parent, bool bContainsEngineModules)
+		/// <returns>The new rules assembly</returns>
+        public static RulesAssembly CreatePluginRulesAssembly(FileReference PluginFileName, bool bSkipCompile, RulesAssembly Parent, bool bContainsEngineModules)
 		{
 			// Check if there's an existing assembly for this project
 			RulesAssembly PluginRulesAssembly;
@@ -422,7 +427,7 @@ namespace UnrealBuildTool
 
 				// Compile the assembly
 				FileReference AssemblyFileName = FileReference.Combine(PluginFileName.Directory, "Intermediate", "Build", "BuildRules", Path.GetFileNameWithoutExtension(PluginFileName.FullName) + "ModuleRules" + FrameworkAssemblyExtension);
-				PluginRulesAssembly = new RulesAssembly(PluginFileName.Directory, ForeignPlugins, ModuleFiles, TargetFiles, ModuleFileToPluginInfo, AssemblyFileName, bContainsEngineModules, bUseBackwardsCompatibleDefaults: false, bReadOnly: false, Parent: Parent);
+				PluginRulesAssembly = new RulesAssembly(PluginFileName.Directory, ForeignPlugins, ModuleFiles, TargetFiles, ModuleFileToPluginInfo, AssemblyFileName, bContainsEngineModules, bUseBackwardsCompatibleDefaults: false, bReadOnly: false, bSkipCompile: bSkipCompile, Parent: Parent);
 				LoadedAssemblyMap.Add(PluginFileName, PluginRulesAssembly);
 			}
 			return PluginRulesAssembly;
diff --git a/Engine/Source/Programs/UnrealBuildTool/System/Unity.cs b/Engine/Source/Programs/UnrealBuildTool/System/Unity.cs
--- a/Engine/Source/Programs/UnrealBuildTool/System/Unity.cs
+++ b/Engine/Source/Programs/UnrealBuildTool/System/Unity.cs
@@ -331,18 +331,7 @@ namespace UnrealBuildTool
 				// Add source files to the unity file
 				foreach (FileItem CPPFile in UnityFile.Files)
 				{
-					string IncludePath;
-					if(BuildPlatform.UseAbsolutePathsInUnityFiles())
-					{
-						IncludePath = CPPFile.AbsolutePath;
-					}
-					else
-					{
-						// @todo: MakeRelativeTo does not work with code projects on a different drive than the engine. reverting to old version until we can come
-						// up with a better solution
-						IncludePath = RemoteExports.ConvertPath(CPPFile.AbsolutePath).Replace('\\', '/');
-					}
-					OutputUnityCPPWriter.WriteLine("#include \"{0}\"", IncludePath);
+					OutputUnityCPPWriter.WriteLine("#include \"{0}\"", CPPFile.AbsolutePath);
 				}
 
 				// Determine unity file path name
diff --git a/Engine/Source/Programs/UnrealBuildTool/System/XmlConfig.cs b/Engine/Source/Programs/UnrealBuildTool/System/XmlConfig.cs
--- a/Engine/Source/Programs/UnrealBuildTool/System/XmlConfig.cs
+++ b/Engine/Source/Programs/UnrealBuildTool/System/XmlConfig.cs
@@ -48,14 +48,12 @@ namespace UnrealBuildTool
 		/// <summary>
 		/// Initialize the config system with the given types
 		/// </summary>
-		public static bool ReadConfigFiles()
+		/// <param name="bForceLoadCache">Force use of the cached XML config without checking if it's valid (useful for remote builds)</param>
+		public static bool ReadConfigFiles(bool bForceLoadCache)
 		{
 			// Find all the configurable types
 			List<Type> ConfigTypes = FindConfigurableTypes();
 
-			// Find all the input files
-			FileReference[] InputFiles = FindInputFiles().Select(x => x.Location).ToArray();
-
 			// Get the path to the cache file
 			FileReference CacheFile = FileReference.Combine(UnrealBuildTool.EngineDirectory, "Intermediate", "Build", "XmlConfigCache.bin");
 			if(UnrealBuildTool.IsEngineInstalled())
@@ -67,50 +65,65 @@ namespace UnrealBuildTool
 				}
 			}
 
-			// Get the path to the schema
-			FileReference SchemaFile = GetSchemaLocation();
-
-			// Try to read the existing cache from disk
-			XmlConfigData CachedValues;
-			if(IsCacheUpToDate(CacheFile, InputFiles) && FileReference.Exists(SchemaFile))
+			// Update the cache if necessary
+			if(bForceLoadCache)
 			{
-				if(XmlConfigData.TryRead(CacheFile, ConfigTypes, out CachedValues) && Enumerable.SequenceEqual(InputFiles, CachedValues.InputFiles))
+				// Never rebuild the cache; just try to load it.
+				if(!XmlConfigData.TryRead(CacheFile, ConfigTypes, out Values))
 				{
-					Values = CachedValues;
+					throw new BuildException("Unable to load XML config cache ({0})", CacheFile);
 				}
 			}
-
-			// If that failed, regenerate it
-			if(Values == null)
+			else
 			{
-				// Find all the configurable fields from the given types
-				Dictionary<string, Dictionary<string, FieldInfo>> CategoryToFields = new Dictionary<string, Dictionary<string, FieldInfo>>();
-				FindConfigurableFields(ConfigTypes, CategoryToFields);
+				// Find all the input files
+				FileReference[] InputFiles = FindInputFiles().Select(x => x.Location).ToArray();
 
-				// Create a schema for the config files
-				XmlSchema Schema = CreateSchema(CategoryToFields);
-				if(!UnrealBuildTool.IsEngineInstalled())
+				// Get the path to the schema
+				FileReference SchemaFile = GetSchemaLocation();
+
+				// Try to read the existing cache from disk
+				XmlConfigData CachedValues;
+				if(IsCacheUpToDate(CacheFile, InputFiles) && FileReference.Exists(SchemaFile))
 				{
-					WriteSchema(Schema, SchemaFile);
+					if(XmlConfigData.TryRead(CacheFile, ConfigTypes, out CachedValues) && Enumerable.SequenceEqual(InputFiles, CachedValues.InputFiles))
+					{
+						Values = CachedValues;
+					}
 				}
 
-				// Read all the XML files and validate them against the schema
-				Dictionary<Type, Dictionary<FieldInfo, object>> TypeToValues = new Dictionary<Type, Dictionary<FieldInfo, object>>();
-				foreach(FileReference InputFile in InputFiles)
+				// If that failed, regenerate it
+				if(Values == null)
 				{
-					if(!TryReadFile(InputFile, CategoryToFields, TypeToValues, Schema))
+					// Find all the configurable fields from the given types
+					Dictionary<string, Dictionary<string, FieldInfo>> CategoryToFields = new Dictionary<string, Dictionary<string, FieldInfo>>();
+					FindConfigurableFields(ConfigTypes, CategoryToFields);
+
+					// Create a schema for the config files
+					XmlSchema Schema = CreateSchema(CategoryToFields);
+					if(!UnrealBuildTool.IsEngineInstalled())
 					{
-						Log.TraceError("Failed to properly read XML file : {0}", InputFile.FullName);
-						return false;
+						WriteSchema(Schema, SchemaFile);
 					}
-				}
 
-				// Make sure the cache directory exists
-				DirectoryReference.CreateDirectory(CacheFile.Directory);
+					// Read all the XML files and validate them against the schema
+					Dictionary<Type, Dictionary<FieldInfo, object>> TypeToValues = new Dictionary<Type, Dictionary<FieldInfo, object>>();
+					foreach(FileReference InputFile in InputFiles)
+					{
+						if(!TryReadFile(InputFile, CategoryToFields, TypeToValues, Schema))
+						{
+							Log.TraceError("Failed to properly read XML file : {0}", InputFile.FullName);
+							return false;
+						}
+					}
+
+					// Make sure the cache directory exists
+					DirectoryReference.CreateDirectory(CacheFile.Directory);
 
-				// Create the new cache
-				Values = new XmlConfigData(InputFiles, TypeToValues.ToDictionary(x => x.Key, x => x.Value.ToArray()));
-				Values.Write(CacheFile);
+					// Create the new cache
+					Values = new XmlConfigData(InputFiles, TypeToValues.ToDictionary(x => x.Key, x => x.Value.ToArray()));
+					Values.Write(CacheFile);
+				}
 			}
 
 			// Apply all the static field values
diff --git a/Engine/Source/Programs/UnrealBuildTool/ToolChain/AppleToolChain.cs b/Engine/Source/Programs/UnrealBuildTool/ToolChain/AppleToolChain.cs
--- a/Engine/Source/Programs/UnrealBuildTool/ToolChain/AppleToolChain.cs
+++ b/Engine/Source/Programs/UnrealBuildTool/ToolChain/AppleToolChain.cs
@@ -1,4 +1,4 @@
-﻿// Copyright 1998-2018 Epic Games, Inc. All Rights Reserved.
+// Copyright 1998-2018 Epic Games, Inc. All Rights Reserved.
 
 using System;
 using System.Collections;
@@ -30,25 +30,13 @@ namespace UnrealBuildTool
 			{
 				Reason = "xcode-select";
 
-				if (Utils.IsRunningOnMono)
-				{
-					// on the Mac, run xcode-select directly
-					DeveloperDir = Utils.RunLocalProcessAndReturnStdOut("xcode-select", "--print-path");
+				// on the Mac, run xcode-select directly
+				DeveloperDir = Utils.RunLocalProcessAndReturnStdOut("xcode-select", "--print-path");
 
-					// make sure we get a full path
-					if (Directory.Exists(DeveloperDir) == false)
-					{
-						throw new BuildException("Selected Xcode ('{0}') doesn't exist, cannot continue.", DeveloperDir);
-					}
-				}
-				else
+				// make sure we get a full path
+				if (Directory.Exists(DeveloperDir) == false)
 				{
-					Hashtable Results = RPCUtilHelper.Command("/", "xcode-select", "--print-path", null);
-					if (Results != null)
-					{
-						DeveloperDir = (string)Results["CommandOutput"];
-						DeveloperDir = DeveloperDir.TrimEnd();
-					}
+					throw new BuildException("Selected Xcode ('{0}') doesn't exist, cannot continue.", DeveloperDir);
 				}
 
 				if (DeveloperDir.EndsWith("/") == false)
@@ -72,21 +60,8 @@ namespace UnrealBuildTool
 				PlatformSDKVersion = "";
 				try
 				{
-					string[] SubDirs = null;
-					if (Utils.IsRunningOnMono)
-					{
-						// on the Mac, we can just get the directory name
-						SubDirs = System.IO.Directory.GetDirectories(BaseSDKDir);
-					}
-					else
-					{
-						Hashtable Results = RPCUtilHelper.Command("/", "ls", BaseSDKDir, null);
-						if (Results != null)
-						{
-							string Result = (string)Results["CommandOutput"];
-							SubDirs = Result.Split("\r\n".ToCharArray(), StringSplitOptions.RemoveEmptyEntries);
-						}
-					}
+					// on the Mac, we can just get the directory name
+					string[] SubDirs = System.IO.Directory.GetDirectories(BaseSDKDir);
 
 					// loop over the subdirs and parse out the version
 					int MaxSDKVersionMajor = 0;
@@ -150,25 +125,50 @@ namespace UnrealBuildTool
 
 			if (bVerbose && !ProjectFileGenerator.bGenerateProjectFiles)
 			{
-				if (BuildHostPlatform.Current.Platform != UnrealTargetPlatform.Mac)
-				{
-					Log.TraceInformation("Compiling with {0} SDK {1} on Mac {2}", OSPrefix, PlatformSDKVersion, RemoteToolChain.RemoteServerName);
-				}
-				else
-				{
-					Log.TraceInformation("Compiling with {0} SDK {1}", OSPrefix, PlatformSDKVersion);
-				}
+				Log.TraceInformation("Compiling with {0} SDK {1}", OSPrefix, PlatformSDKVersion);
 			}
 		}
 	}
 
-	abstract class AppleToolChain : RemoteToolChain
+	abstract class AppleToolChain : UEToolChain
 	{
-		public AppleToolChain(CppPlatform InCppPlatform, UnrealTargetPlatform InRemoteToolChainPlatform, FileReference InProjectFile)
-			: base(InCppPlatform, InRemoteToolChainPlatform, InProjectFile)
+		protected FileReference ProjectFile;
+
+		public AppleToolChain(CppPlatform InCppPlatform, FileReference InProjectFile)
+			: base(InCppPlatform)
 		{
+			ProjectFile = InProjectFile;
 		}
 
+		protected static string ConvertPath(string Path)
+		{
+			return Path;
+		}
+
+		protected static FileItem RemoteToLocalFileItem(FileItem Executable)
+		{
+			return Executable;
+		}
+
+		protected static FileItem LocalToRemoteFileItem(FileItem Item, bool bSomething = false)
+		{
+			return Item;
+		}
+
+		protected static void QueueFileForBatchUpload(FileItem Item)
+		{
+		}
+
+		protected string GetMacDevSrcRoot()
+		{
+			return UnrealBuildTool.EngineSourceDirectory.FullName;
+		}
+
+        protected string GetMacDevEngineRoot()
+        {
+            return UnrealBuildTool.EngineDirectory.FullName;
+        }
+
 		protected void StripSymbolsWithXcode(FileReference SourceFile, FileReference TargetFile, string ToolchainDir)
 		{
 			if (SourceFile != TargetFile)
diff --git a/Engine/Source/Programs/UnrealBuildTool/ToolChain/RemoteExports.cs b/Engine/Source/Programs/UnrealBuildTool/ToolChain/RemoteExports.cs
deleted file mode 100644
diff --git a/Engine/Source/Programs/UnrealBuildTool/ToolChain/RemoteMac.cs b/Engine/Source/Programs/UnrealBuildTool/ToolChain/RemoteMac.cs
new file mode 100644
--- /dev/null
+++ b/Engine/Source/Programs/UnrealBuildTool/ToolChain/RemoteMac.cs
@@ -0,0 +1,935 @@
+﻿// Copyright 1998-2018 Epic Games, Inc. All Rights Reserved.
+
+using System;
+using System.Collections;
+using System.Collections.Generic;
+using System.IO;
+using System.Diagnostics;
+using System.Linq;
+using System.Security.AccessControl;
+using System.Xml;
+using System.Text;
+using System.Text.RegularExpressions;
+using Tools.DotNETCommon;
+using System.Security.Cryptography.X509Certificates;
+
+namespace UnrealBuildTool
+{
+	/// <summary>
+	/// Stores information about how a local directory maps to a remote directory
+	/// </summary>
+	[DebuggerDisplay("{LocalDirectory}")]
+	class RemoteMapping
+	{
+		public DirectoryReference LocalDirectory;
+		public string RemoteDirectory;
+
+		public RemoteMapping(DirectoryReference LocalDirectory, string RemoteDirectory)
+		{
+			this.LocalDirectory = LocalDirectory;
+			this.RemoteDirectory = RemoteDirectory;
+		}
+	}
+
+	/// <summary>
+	/// Handles uploading and building on a remote Mac
+	/// </summary>
+	class RemoteMac
+	{
+		/// <summary>
+		/// These two variables will be loaded from XML config file in XmlConfigLoader.Init()
+		/// </summary>
+		[XmlConfigFile]
+		private readonly string ServerName;
+
+		/// <summary>
+		/// The remote username
+		/// </summary>
+		[XmlConfigFile]
+		private readonly string UserName;
+
+		/// <summary>
+		/// Instead of looking for RemoteToolChainPrivate.key in the usual places (Documents/Unreal Engine/UnrealBuildTool/SSHKeys, Engine/Build/SSHKeys), this private key will be used if set
+		/// </summary>
+		[XmlConfigFile]
+		private FileReference SshPrivateKey;
+
+		/// <summary>
+		/// The authentication used for Rsync (for the -e rsync flag)
+		/// </summary>
+		[XmlConfigFile]
+		private string RsyncAuthentication = "ssh -i '${CYGWIN_SSH_PRIVATE_KEY}'";
+
+		/// <summary>
+		/// The authentication used for SSH (probably similar to RsyncAuthentication)
+		/// </summary>
+		[XmlConfigFile]
+		private string SshAuthentication = "-i '${CYGWIN_SSH_PRIVATE_KEY}'";
+
+		/// <summary>
+		/// Save the specified port so that RemoteServerName is the machine address only
+		/// </summary>
+		private readonly int ServerPort = 22;	// Default ssh port
+
+		/// <summary>
+		/// Path to Rsync
+		/// </summary>
+		private readonly FileReference RsyncExe;
+
+		/// <summary>
+		/// Path to SSH
+		/// </summary>
+		private readonly FileReference SshExe;
+
+		/// <summary>
+		/// The project being built. Settings will be read from config files in this project.
+		/// </summary>
+		private readonly FileReference ProjectFile;
+
+		/// <summary>
+		/// The directory containing the project being built.
+		/// </summary>
+		private readonly DirectoryReference ProjectDirectory;
+
+		/// <summary>
+		/// The base directory on the remote machine
+		/// </summary>
+		private string RemoteBaseDir;
+
+		/// <summary>
+		/// Mappings from local directories to remote directories
+		/// </summary>
+		private List<RemoteMapping> Mappings;
+
+		/// <summary>
+		/// Arguments that are used by every Ssh call
+		/// </summary>
+		private List<string> CommonSshArguments;
+
+		/// <summary>
+		/// Arguments that are used by every Rsync call
+		/// </summary>
+		private List<string> CommonRsyncArguments;
+
+		/// <summary>
+		/// Constructor
+		/// </summary>
+		/// <param name="ProjectFile">Project to read settings from</param>
+		public RemoteMac(FileReference ProjectFile)
+		{
+			this.RsyncExe = FileReference.Combine(UnrealBuildTool.EngineDirectory, "Extras", "ThirdPartyNotUE", "DeltaCopy", "Binaries", "Rsync.exe");
+			this.SshExe = FileReference.Combine(UnrealBuildTool.EngineDirectory, "Extras", "ThirdPartyNotUE", "DeltaCopy", "Binaries", "Ssh.exe");
+			this.ProjectFile = ProjectFile;
+			this.ProjectDirectory = DirectoryReference.FromFile(ProjectFile);
+
+			// Apply settings from the XML file
+			XmlConfig.ApplyTo(this);
+
+			// Get the project config file path
+			DirectoryReference EngineIniPath = ProjectFile != null ? ProjectFile.Directory : null;
+			if (EngineIniPath == null && UnrealBuildTool.GetRemoteIniPath() != null)
+			{
+				EngineIniPath = new DirectoryReference(UnrealBuildTool.GetRemoteIniPath());
+			}
+			ConfigHierarchy Ini = ConfigCache.ReadHierarchy(ConfigHierarchyType.Engine, EngineIniPath, UnrealTargetPlatform.IOS);
+
+			// Read the project settings if we don't have anything in the build configuration settings
+			if(String.IsNullOrEmpty(ServerName))
+			{
+				// Read the server name
+				string IniServerName;
+				if (Ini.GetString("/Script/IOSRuntimeSettings.IOSRuntimeSettings", "RemoteServerName", out IniServerName) && !String.IsNullOrEmpty(IniServerName))
+				{
+					this.ServerName = IniServerName;
+				}
+				else
+				{
+					throw new BuildException("Remote compiling requires a server name. Use the editor (Project Settings > IOS) to set up your remote compilation settings.");
+				}
+
+				// Parse the username
+				string IniUserName;
+				if (Ini.GetString("/Script/IOSRuntimeSettings.IOSRuntimeSettings", "RSyncUsername", out IniUserName) && !String.IsNullOrEmpty(IniUserName))
+				{
+					this.UserName = IniUserName;
+				}
+			}
+
+			// Split port out from the server name
+			int PortIdx = ServerName.LastIndexOf(':');
+			if(PortIdx != -1)
+			{
+				string Port = ServerName.Substring(PortIdx + 1);
+				if(!int.TryParse(Port, out ServerPort))
+				{
+					throw new BuildException("Unable to parse port number from '{0}'", ServerName);
+				}
+				ServerName = ServerName.Substring(0, PortIdx);
+			}
+
+			// If a user name is not set, use the current user
+			if (String.IsNullOrEmpty(UserName))
+			{
+				UserName = Environment.UserName;
+			}
+
+			// Print out the server info
+			Log.TraceInformation("[Remote] Using remote server '{0}' on port {1} (user '{2}')", ServerName, ServerPort, UserName);
+
+			// Get the path to the SSH private key
+			string OverrideSshPrivateKeyPath;
+			if (Ini.GetString("/Script/IOSRuntimeSettings.IOSRuntimeSettings", "SSHPrivateKeyOverridePath", out OverrideSshPrivateKeyPath) && !String.IsNullOrEmpty(OverrideSshPrivateKeyPath))
+			{
+				SshPrivateKey = new FileReference(OverrideSshPrivateKeyPath);
+				if (!FileReference.Exists(SshPrivateKey))
+				{
+					throw new BuildException("SSH private key specified in config file ({0}) does not exist.", SshPrivateKey);
+				}
+			}
+
+			// If it's not set, look in the standard locations. If that fails, spawn the batch file to generate one.
+			if (SshPrivateKey == null && !TryGetSshPrivateKey(out SshPrivateKey))
+			{
+				Log.TraceWarning("No SSH private key found for {0}@{1}. Launching SSH to generate one.", UserName, ServerName);
+
+				StringBuilder CommandLine = new StringBuilder();
+				CommandLine.AppendFormat("/C \"\"{0}\"", FileReference.Combine(UnrealBuildTool.EngineDirectory, "Build", "BatchFiles", "MakeAndInstallSSHKey.bat"));
+				CommandLine.AppendFormat(" \"{0}\"", SshExe);
+				CommandLine.AppendFormat(" \"{0}\"", ServerPort);
+				CommandLine.AppendFormat(" \"{0}\"", RsyncExe);
+				CommandLine.AppendFormat(" \"{0}\"", UserName);
+				CommandLine.AppendFormat(" \"{0}\"", ServerName);
+				CommandLine.AppendFormat(" \"{0}\"", DirectoryReference.GetSpecialFolder(Environment.SpecialFolder.MyDocuments));
+				CommandLine.AppendFormat(" \"{0}\"", GetLocalCygwinPath(DirectoryReference.GetSpecialFolder(Environment.SpecialFolder.MyDocuments)));
+				CommandLine.AppendFormat(" \"{0}\"", UnrealBuildTool.EngineDirectory);
+				CommandLine.Append("\"");
+
+				using(Process ChildProcess = Process.Start("C:\\Windows\\System32\\Cmd.exe", CommandLine.ToString()))
+				{
+					ChildProcess.WaitForExit();
+				}
+
+				if(!TryGetSshPrivateKey(out SshPrivateKey))
+				{
+					throw new BuildException("Failed to generate SSH private key for {0}@{1}.", UserName, ServerName);
+				}
+			}
+
+			// resolve the rest of the strings
+			RsyncAuthentication = ExpandVariables(RsyncAuthentication);
+			SshAuthentication = ExpandVariables(SshAuthentication);
+
+			// Get the remote base directory
+			RemoteBaseDir = String.Format("[REDACTED_USER_HOME]/UE4/Builds/{1}", UserName, Environment.MachineName);
+
+			// Build the list of directory mappings between the local and remote machines
+			Mappings = new List<RemoteMapping>();
+			Mappings.Add(new RemoteMapping(UnrealBuildTool.EngineDirectory, GetRemotePath(UnrealBuildTool.EngineDirectory)));
+			if(ProjectFile != null && !ProjectFile.IsUnderDirectory(UnrealBuildTool.EngineDirectory))
+			{
+				Mappings.Add(new RemoteMapping(ProjectFile.Directory, GetRemotePath(ProjectFile.Directory)));
+			}
+
+			// Build a list of arguments for SSH
+			CommonSshArguments = new List<string>();
+			CommonSshArguments.Add("-o BatchMode=yes");
+			CommonSshArguments.Add(SshAuthentication);
+			CommonSshArguments.Add(String.Format("-p {0}", ServerPort));
+			CommonSshArguments.Add(String.Format("\"{0}@{1}\"", UserName, ServerName));
+
+			// Build a list of arguments for Rsync
+			CommonRsyncArguments = new List<string>();
+			CommonRsyncArguments.Add("--compress");
+			CommonRsyncArguments.Add("--recursive");
+			CommonRsyncArguments.Add("--delete"); // Delete anything not in the source directory
+			CommonRsyncArguments.Add("--delete-excluded"); // Delete anything not in the source directory
+			CommonRsyncArguments.Add("--times"); // Preserve modification times
+			CommonRsyncArguments.Add("--verbose");
+			CommonRsyncArguments.Add("-m");
+			CommonRsyncArguments.Add("--chmod=ug=rwX,o=rxX");
+			CommonRsyncArguments.Add(String.Format("--rsh=\"{0} -p {1}\"", RsyncAuthentication, ServerPort));
+		}
+
+		/// <summary>
+		/// Attempts to get the SSH private key from the standard locations
+		/// </summary>
+		/// <param name="OutPrivateKey">If successful, receives the location of the private key that was found</param>
+		/// <returns>True if a private key was found, false otherwise</returns>
+		private bool TryGetSshPrivateKey(out FileReference OutPrivateKey)
+		{
+			// Build a list of all the places to look for a private key
+			List<DirectoryReference> Locations = new List<DirectoryReference>();
+			Locations.Add(DirectoryReference.Combine(DirectoryReference.GetSpecialFolder(Environment.SpecialFolder.ApplicationData), "Unreal Engine", "UnrealBuildTool"));
+			Locations.Add(DirectoryReference.Combine(DirectoryReference.GetSpecialFolder(Environment.SpecialFolder.Personal), "Unreal Engine", "UnrealBuildTool"));
+			if (ProjectFile != null)
+			{
+				Locations.Add(DirectoryReference.Combine(ProjectFile.Directory, "Build", "NotForLicensees"));
+				Locations.Add(DirectoryReference.Combine(ProjectFile.Directory, "Build", "NoRedist"));
+				Locations.Add(DirectoryReference.Combine(ProjectFile.Directory, "Build"));
+			}
+			Locations.Add(DirectoryReference.Combine(UnrealBuildTool.EngineDirectory, "Build", "NotForLicensees"));
+			Locations.Add(DirectoryReference.Combine(UnrealBuildTool.EngineDirectory, "Build", "NoRedist"));
+			Locations.Add(DirectoryReference.Combine(UnrealBuildTool.EngineDirectory, "Build"));
+
+			// Find the first that exists
+			foreach (DirectoryReference Location in Locations)
+			{
+				FileReference KeyFile = FileReference.Combine(Location, "SSHKeys", ServerName, UserName, "RemoteToolChainPrivate.key");
+				if (FileReference.Exists(KeyFile))
+				{
+					OutPrivateKey = KeyFile;
+					return true;
+				}
+			}
+
+			// Nothing found
+			OutPrivateKey = null;
+			return false;
+		}
+
+		/// <summary>
+		/// Expand all the variables in the given string
+		/// </summary>
+		/// <param name="Input">The input string</param>
+		/// <returns>String with any variables expanded</returns>
+		private string ExpandVariables(string Input)
+		{
+			string Result = Input;
+			Result = Result.Replace("${SSH_PRIVATE_KEY}", SshPrivateKey.FullName);
+			Result = Result.Replace("${CYGWIN_SSH_PRIVATE_KEY}", GetLocalCygwinPath(SshPrivateKey));
+			return Result;
+		}
+
+		/// <summary>
+		/// Flush the remote machine, removing all existing files
+		/// </summary>
+		public void FlushRemote()
+		{
+			Log.TraceInformation("[Remote] Deleting all files under {0}...", RemoteBaseDir);
+			Execute("/", String.Format("rm -rf \"{0}\"", RemoteBaseDir));
+		}
+
+		/// <summary>
+		/// Build a target remotely
+		/// </summary>
+		/// <param name="TargetDesc">Descriptor for the target to build</param>
+		/// <param name="RemoteLogFile">Path to store the remote log file</param>
+		/// <returns>True if the build succeeded, false otherwise</returns>
+		public bool Build(TargetDescriptor TargetDesc, FileReference RemoteLogFile)
+		{
+			// Get the directory for working files
+			DirectoryReference BaseDir = DirectoryReference.FromFile(TargetDesc.ProjectFile) ?? UnrealBuildTool.EngineDirectory;
+			DirectoryReference TempDir = DirectoryReference.Combine(BaseDir, "Intermediate", "Remote", TargetDesc.Name, TargetDesc.Platform.ToString(), TargetDesc.Configuration.ToString());
+			DirectoryReference.CreateDirectory(TempDir);
+
+			// Path to the local manifest file. This has to be translated from the remote format after the build is complete.
+			List<FileReference> LocalManifestFiles = new List<FileReference>();
+
+			// Path to the remote manifest file
+			FileReference RemoteManifestFile = FileReference.Combine(TempDir, "Manifest.xml");
+
+			// Prepare the arguments we will pass to the remote build
+			List<string> RemoteArguments = new List<string>();
+			RemoteArguments.Add(TargetDesc.Name);
+			RemoteArguments.Add(TargetDesc.Platform.ToString());
+			RemoteArguments.Add(TargetDesc.Configuration.ToString());
+			RemoteArguments.Add("-SkipRulesCompile"); // Use the rules assembly built locally
+			RemoteArguments.Add("-ForceXmlConfigCache"); // Use the XML config cache built locally, since the remote won't have it
+			RemoteArguments.Add(String.Format("-Log={0}", GetRemotePath(RemoteLogFile)));
+			RemoteArguments.Add(String.Format("-Manifest={0}", GetRemotePath(RemoteManifestFile)));
+
+			foreach(string LocalArgument in TargetDesc.AdditionalArguments)
+			{
+				int EqualsIdx = LocalArgument.IndexOf('=');
+				if(EqualsIdx == -1)
+				{
+					RemoteArguments.Add(LocalArgument);
+					continue;
+				}
+
+				string Key = LocalArgument.Substring(0, EqualsIdx);
+				string Value = LocalArgument.Substring(EqualsIdx + 1);
+
+				if(Key.Equals("-Log", StringComparison.InvariantCultureIgnoreCase))
+				{
+					// We are already writing to the local log file. The remote will produce a different log (RemoteLogFile)
+					continue;
+				}
+				if(Key.Equals("-Manifest", StringComparison.InvariantCultureIgnoreCase))
+				{
+					LocalManifestFiles.Add(new FileReference(Value));
+					continue;
+				}
+
+				string RemoteArgument = LocalArgument;
+				foreach(RemoteMapping Mapping in Mappings)
+				{
+					if(Value.StartsWith(Mapping.LocalDirectory.FullName, StringComparison.InvariantCultureIgnoreCase))
+					{
+						RemoteArgument = String.Format("{0}={1}", Key, GetRemotePath(Value));
+						break;
+					}
+				}
+				RemoteArguments.Add(RemoteArgument);
+			}
+
+			// Handle any per-platform setup that is required
+			if(TargetDesc.Platform == UnrealTargetPlatform.IOS || TargetDesc.Platform == UnrealTargetPlatform.TVOS)
+			{
+				// Always generate a .stub
+				RemoteArguments.Add("-CreateStub");
+
+				// Get the provisioning data for this project
+				IOSProvisioningData ProvisioningData = IOSPlatform.ReadProvisioningData(TargetDesc.ProjectFile);
+				if(ProvisioningData == null || ProvisioningData.MobileProvisionFile == null)
+				{
+					throw new BuildException("Unable to find mobile provision for {0}. See log for more information.", TargetDesc.Name);
+				}
+
+				// Create a local copy of the provision
+				FileReference MobileProvisionFile = FileReference.Combine(TempDir, ProvisioningData.MobileProvisionFile.GetFileName());
+				if(FileReference.Exists(MobileProvisionFile))
+				{
+					FileReference.SetAttributes(MobileProvisionFile, FileAttributes.Normal);
+				}
+				FileReference.Copy(ProvisioningData.MobileProvisionFile, MobileProvisionFile, true);
+				Log.TraceInformation("[Remote] Uploading {0}", MobileProvisionFile);
+				UploadFile(MobileProvisionFile);
+
+				// Extract the certificate for the project
+				FileReference CertificateFile = FileReference.Combine(TempDir, "Certificate.p12");
+				if(!FileReference.Exists(CertificateFile))
+				{
+					StringBuilder Arguments = new StringBuilder("ExportCertificate");
+					if(TargetDesc.ProjectFile == null)
+					{
+						Arguments.AppendFormat(" \"{0}\"", UnrealBuildTool.EngineSourceDirectory);
+					}
+					else
+					{
+						Arguments.AppendFormat(" \"{0}\"", TargetDesc.ProjectFile.Directory);
+					}
+					Arguments.AppendFormat(" -certificate \"{0}\"", CertificateFile);
+
+					ProcessStartInfo StartInfo = new ProcessStartInfo();
+					StartInfo.FileName = FileReference.Combine(UnrealBuildTool.EngineDirectory, "Binaries", "DotNET", "IOS", "IPhonePackager.exe").FullName;
+					StartInfo.Arguments = Arguments.ToString();
+					if(Utils.RunLocalProcessAndLogOutput(StartInfo) != 0)
+					{
+						throw new BuildException("IphonePackager failed.");
+					}
+				}
+
+				// Upload the certificate to the remote
+				Log.TraceInformation("[Remote] Uploading {0}", CertificateFile);
+				UploadFile(CertificateFile);
+
+				// Tell the remote UBT instance to use them
+				RemoteArguments.Add(String.Format("-ImportProvision={0}", GetRemotePath(MobileProvisionFile)));
+				RemoteArguments.Add(String.Format("-ImportCertificate={0}", GetRemotePath(CertificateFile)));
+				RemoteArguments.Add(String.Format("-ImportCertificatePassword=A"));
+			}
+
+			// Upload the workspace files
+			Log.TraceInformation("[Remote] Uploading workspace files");
+			UploadWorkspace(TempDir);
+
+			// Fixup permissions on any shell scripts
+			Execute(RemoteBaseDir, String.Format("chmod +x {0}/Build/BatchFiles/Mac/*.sh", EscapeShellArgument(GetRemotePath(UnrealBuildTool.EngineDirectory))));
+
+			// Execute the compile
+			Log.TraceInformation("[Remote] Executing build");
+
+			StringBuilder BuildCommandLine = new StringBuilder("Engine/Build/BatchFiles/Mac/Build.sh");
+			foreach(string RemoteArgument in RemoteArguments)
+			{
+				BuildCommandLine.AppendFormat(" {0}", EscapeShellArgument(RemoteArgument));
+			}
+
+			int Result = Execute(GetRemotePath(UnrealBuildTool.RootDirectory), BuildCommandLine.ToString());
+			if(Result != 0)
+			{
+				if(RemoteLogFile != null)
+				{
+					Log.TraceInformation("[Remote] Downloading {0}", RemoteLogFile);
+					DownloadFile(RemoteLogFile);
+				}
+				return false;
+			}
+
+			// Download the manifest
+			Log.TraceInformation("[Remote] Downloading {0}", RemoteManifestFile);
+			DownloadFile(RemoteManifestFile);
+
+			// Convert the manifest to local form
+			BuildManifest Manifest = Utils.ReadClass<BuildManifest>(RemoteManifestFile.FullName);
+			for(int Idx = 0; Idx < Manifest.BuildProducts.Count; Idx++)
+			{
+				Manifest.BuildProducts[Idx] = GetLocalPath(Manifest.BuildProducts[Idx]).FullName;
+			}
+
+			// Download the files from the remote
+			if(TargetDesc.AdditionalArguments.Any(x => x.Equals("-GenerateManifest", StringComparison.InvariantCultureIgnoreCase)))
+			{
+				LocalManifestFiles.Add(FileReference.Combine(UnrealBuildTool.EngineDirectory, "Intermediate", "Build", "Manifest.xml"));
+			}
+			else
+			{
+				Log.TraceInformation("[Remote] Downloading build products");
+
+				List<FileReference> FilesToDownload = new List<FileReference>();
+				FilesToDownload.Add(RemoteLogFile);
+				FilesToDownload.AddRange(Manifest.BuildProducts.Select(x => new FileReference(x)));
+				DownloadFiles(FilesToDownload);
+			}
+
+			// Write out all the local manifests
+			foreach(FileReference LocalManifestFile in LocalManifestFiles)
+			{
+				Log.TraceInformation("[Remote] Writing {0}", LocalManifestFile);
+				Utils.WriteClass<BuildManifest>(Manifest, LocalManifestFile.FullName, "");
+			}
+			return true;
+		}
+
+		/// <summary>
+		/// Runs the actool utility on a directory to create an Assets.car file
+		/// </summary>
+		/// <param name="Platform">The target platform</param>
+		/// <param name="InputDir">Input directory containing assets</param>
+		/// <param name="OutputFile">Path to the Assets.car file to produce</param>
+		public void RunAssetCatalogTool(CppPlatform Platform, DirectoryReference InputDir, FileReference OutputFile)
+		{
+			Log.TraceInformation("Running asset catalog tool for {0}: {1} -> {2}", Platform, InputDir, OutputFile);
+
+			string RemoteInputDir = GetRemotePath(InputDir);
+			UploadDirectory(InputDir);
+
+			string RemoteOutputFile = GetRemotePath(OutputFile);
+			Execute(RemoteBaseDir, String.Format("rm -f {0}", EscapeShellArgument(RemoteOutputFile)));
+
+			string RemoteOutputDir = Path.GetDirectoryName(RemoteOutputFile).Replace(Path.DirectorySeparatorChar, '/');
+			Execute(RemoteBaseDir, String.Format("mkdir -p {0}", EscapeShellArgument(RemoteOutputDir)));
+
+			string RemoteArguments = IOSToolChain.GetAssetCatalogArgs(Platform, RemoteInputDir, RemoteOutputDir); 
+			if(Execute(RemoteBaseDir, String.Format("/usr/bin/xcrun {0}", RemoteArguments)) != 0)
+			{
+				throw new BuildException("Failed to run actool.");
+			}
+			DownloadFile(OutputFile);
+		}
+
+		/// <summary>
+		/// Convers a remote path into local form
+		/// </summary>
+		/// <param name="RemotePath">The remote filename</param>
+		/// <returns>Local filename corresponding to the remote path</returns>
+		private FileReference GetLocalPath(string RemotePath)
+		{
+			foreach(RemoteMapping Mapping in Mappings)
+			{
+				if(RemotePath.StartsWith(Mapping.RemoteDirectory, StringComparison.InvariantCultureIgnoreCase) && RemotePath.Length > Mapping.RemoteDirectory.Length && RemotePath[Mapping.RemoteDirectory.Length] == '/')
+				{
+					return FileReference.Combine(Mapping.LocalDirectory, RemotePath.Substring(Mapping.RemoteDirectory.Length + 1));
+				}
+			}
+			throw new BuildException("Unable to map remote path '{0}' to local path", RemotePath);
+		}
+
+		/// <summary>
+		/// Converts a local path into a remote one
+		/// </summary>
+		/// <param name="LocalPath">The local path to convert</param>
+		/// <returns>Equivalent remote path</returns>
+		private string GetRemotePath(FileSystemReference LocalPath)
+		{
+			return GetRemotePath(LocalPath.FullName);
+		}
+
+		/// <summary>
+		/// Converts a local path into a remote one
+		/// </summary>
+		/// <param name="LocalPath">The local path to convert</param>
+		/// <returns>Equivalent remote path</returns>
+		private string GetRemotePath(string LocalPath)
+		{
+			return String.Format("{0}/{1}", RemoteBaseDir, LocalPath.Replace(":", "").Replace("\\", "/").Replace(" ", "_"));
+		}
+
+		/// <summary>
+		/// Gets the local path in Cygwin format (eg. /cygdrive/C/...)
+		/// </summary>
+		/// <param name="InPath">Local path</param>
+		/// <returns>Path in cygwin format</returns>
+		private static string GetLocalCygwinPath(FileSystemReference InPath)
+		{
+			if(InPath.FullName.Length < 2 || InPath.FullName[1] != ':')
+			{
+				throw new BuildException("Invalid local path for converting to cygwin format ({0}).", InPath);
+			}
+			return String.Format("/cygdrive/{0}{1}", InPath.FullName.Substring(0, 1), InPath.FullName.Substring(2).Replace('\\', '/'));
+		}
+
+		/// <summary>
+		/// Escapes spaces in a shell command argument
+		/// </summary>
+		/// <param name="Argument">The argument to escape</param>
+		/// <returns>The escaped argument</returns>
+		private static string EscapeShellArgument(string Argument)
+		{
+			return Argument.Replace(" ", "\\ ");
+		}
+
+		/// <summary>
+		/// Upload a single file to the remote
+		/// </summary>
+		/// <param name="LocalFile">The file to upload</param>
+		void UploadFile(FileReference LocalFile)
+		{
+			string RemoteFile = GetRemotePath(LocalFile);
+			string RemoteDirectory = GetRemotePath(LocalFile.Directory);
+
+			List<string> Arguments = new List<string>(CommonRsyncArguments);
+			Arguments.Add(String.Format("--rsync-path=\"mkdir -p {0} && rsync\"", RemoteDirectory));
+			Arguments.Add(String.Format("\"{0}\"", GetLocalCygwinPath(LocalFile)));
+			Arguments.Add(String.Format("\"{0}@{1}\":'{2}/'", UserName, ServerName, RemoteDirectory));
+			Arguments.Add("-q");
+
+			int Result = Rsync(String.Join(" ", Arguments));
+			if(Result != 0)
+			{
+				throw new BuildException("Error while running Rsync (exit code {0})", Result);
+			}
+		}
+
+		/// <summary>
+		/// Upload a single directory to the remote
+		/// </summary>
+		/// <param name="LocalDirectory">The local directory to upload</param>
+		void UploadDirectory(DirectoryReference LocalDirectory)
+		{
+			string RemoteDirectory = GetRemotePath(LocalDirectory);
+
+			List<string> Arguments = new List<string>(CommonRsyncArguments);
+			Arguments.Add(String.Format("--rsync-path=\"mkdir -p {0} && rsync\"", RemoteDirectory));
+			Arguments.Add(String.Format("\"{0}/\"", GetLocalCygwinPath(LocalDirectory)));
+			Arguments.Add(String.Format("\"{0}@{1}\":'{2}/'", UserName, ServerName, RemoteDirectory));
+			Arguments.Add("-q");
+
+			int Result = Rsync(String.Join(" ", Arguments));
+			if(Result != 0)
+			{
+				throw new BuildException("Error while running Rsync (exit code {0})", Result);
+			}
+		}
+
+		/// <summary>
+		/// Uploads a directory to the remote using a specific filter list
+		/// </summary>
+		/// <param name="LocalDirectory">The local directory to copy from</param>
+		/// <param name="RemoteDirectory">The remote directory to copy to</param>
+		/// <param name="FilterLocations">List of paths to filter</param>
+		void UploadDirectory(DirectoryReference LocalDirectory, string RemoteDirectory, List<FileReference> FilterLocations)
+		{
+			List<string> Arguments = new List<string>(CommonRsyncArguments);
+			Arguments.Add(String.Format("--rsync-path=\"mkdir -p {0} && rsync\"", RemoteDirectory));
+			foreach(FileReference FilterLocation in FilterLocations)
+			{
+				Arguments.Add(String.Format("--filter=\"merge {0}\"", GetLocalCygwinPath(FilterLocation)));
+			}
+			Arguments.Add("--exclude='*'");
+			Arguments.Add(String.Format("\"{0}/\"", GetLocalCygwinPath(LocalDirectory)));
+			Arguments.Add(String.Format("\"{0}@{1}\":'{2}/'", UserName, ServerName, RemoteDirectory));
+
+			int Result = Rsync(String.Join(" ", Arguments));
+			if(Result != 0)
+			{
+				throw new BuildException("Error while running Rsync (exit code {0})", Result);
+			}
+		}
+
+		/// <summary>
+		/// Upload all the files in the workspace for the current project
+		/// </summary>
+		void UploadWorkspace(DirectoryReference TempDir)
+		{
+			// Path to the scripts to be uploaded
+			FileReference ScriptPathsFileName = FileReference.Combine(UnrealBuildTool.EngineDirectory, "Build", "Rsync", "RsyncEngineScripts.txt");
+
+			// Read the list of scripts to be uploaded
+			List<string> ScriptPaths = new List<string>();
+			foreach(string Line in FileReference.ReadAllLines(ScriptPathsFileName))
+			{
+				string FileToUpload = Line.Trim();
+				if(FileToUpload.Length > 0 && FileToUpload[0] != ';')
+				{
+					ScriptPaths.Add(FileToUpload);
+				}
+			}
+
+			// Fixup the line endings
+			List<FileReference> TargetFiles = new List<FileReference>();
+			foreach(string ScriptPath in ScriptPaths)
+			{
+				FileReference SourceFile = FileReference.Combine(UnrealBuildTool.EngineDirectory, ScriptPath.TrimStart('/'));
+				if(!FileReference.Exists(SourceFile))
+				{
+					throw new BuildException("Missing script required for remote upload: {0}", SourceFile);
+				}
+
+				FileReference TargetFile = FileReference.Combine(TempDir, SourceFile.MakeRelativeTo(UnrealBuildTool.EngineDirectory));
+				if(!FileReference.Exists(TargetFile) || FileReference.GetLastWriteTimeUtc(TargetFile) < FileReference.GetLastWriteTimeUtc(SourceFile))
+				{
+					DirectoryReference.CreateDirectory(TargetFile.Directory);
+					string ScriptText = FileReference.ReadAllText(SourceFile);
+					FileReference.WriteAllText(TargetFile, ScriptText.Replace("\r\n", "\n"));
+				}
+				TargetFiles.Add(TargetFile);
+			}
+
+			// Write a file that protects all the scripts from being overridden by the standard engine filters
+			FileReference ScriptUploadList = FileReference.Combine(TempDir, "RsyncEngineScripts-Upload.txt");
+			using(StreamWriter Writer = new StreamWriter(ScriptUploadList.FullName))
+			{
+				foreach(string ScriptPath in ScriptPaths)
+				{
+					for(int SlashIdx = ScriptPath.IndexOf('/', 1); SlashIdx != -1; SlashIdx = ScriptPath.IndexOf('/', SlashIdx + 1))
+					{
+						Writer.WriteLine("+ {0}", ScriptPath.Substring(0, SlashIdx));
+					}
+					Writer.WriteLine("+ {0}", ScriptPath);
+				}
+				Writer.WriteLine("protect *");
+			}
+
+			// Write a file that protects all the scripts from being overridden by the standard engine filters
+			FileReference ScriptProtectList = FileReference.Combine(TempDir, "RsyncEngineScripts-Protect.txt");
+			using(StreamWriter Writer = new StreamWriter(ScriptProtectList.FullName))
+			{
+				foreach(string ScriptPath in ScriptPaths)
+				{
+					Writer.WriteLine("protect {0}", ScriptPath);
+				}
+			}
+
+			// Upload these files to the remote
+			List<FileReference> FilterLocations = new List<FileReference>();
+			FilterLocations.Add(ScriptUploadList);
+			UploadDirectory(TempDir, GetRemotePath(UnrealBuildTool.EngineDirectory), FilterLocations);
+
+			// Upload the engine files
+			List<FileReference> EngineFilters = new List<FileReference>();
+			EngineFilters.Add(ScriptProtectList);
+			EngineFilters.Add(FileReference.Combine(UnrealBuildTool.EngineDirectory, "Build", "Rsync", "RsyncEngine.txt"));
+			UploadDirectory(UnrealBuildTool.EngineDirectory, GetRemotePath(UnrealBuildTool.EngineDirectory), EngineFilters);
+
+			// Upload the project files
+			if(ProjectFile != null && !ProjectFile.IsUnderDirectory(UnrealBuildTool.EngineDirectory))
+			{
+				List<FileReference> ProjectFilters = new List<FileReference>();
+
+				FileReference CustomProjectFilter = FileReference.Combine(ProjectFile.Directory, "Build", "Rsync", "RsyncProject.txt");
+				if(FileReference.Exists(CustomProjectFilter))
+				{
+					ProjectFilters.Add(CustomProjectFilter);
+				}
+				ProjectFilters.Add(FileReference.Combine(UnrealBuildTool.EngineDirectory, "Build", "Rsync", "RsyncProject.txt"));
+
+				UploadDirectory(ProjectFile.Directory, GetRemotePath(ProjectFile.Directory), ProjectFilters);
+			}
+		}
+
+		/// <summary>
+		/// Downloads a single file from the remote
+		/// </summary>
+		/// <param name="LocalFile">The file to download</param>
+		void DownloadFile(FileReference LocalFile)
+		{
+			RemoteMapping Mapping = Mappings.FirstOrDefault(x => LocalFile.IsUnderDirectory(x.LocalDirectory));
+			if(Mapping == null)
+			{
+				throw new BuildException("File for download '{0}' is not under any mapped directory.", LocalFile);
+			}
+
+			List<string> Arguments = new List<string>(CommonRsyncArguments);
+			Arguments.Add(String.Format("\"{0}@{1}\":'{2}/{3}'", UserName, ServerName, Mapping.RemoteDirectory, LocalFile.MakeRelativeTo(Mapping.LocalDirectory).Replace('\\', '/')));
+			Arguments.Add(String.Format("\"{0}/\"", GetLocalCygwinPath(LocalFile.Directory)));
+			Arguments.Add("-q");
+
+			int Result = Rsync(String.Join(" ", Arguments));
+			if(Result != 0)
+			{
+				throw new BuildException("Unable to download '{0}' from the remote Mac (exit code {1}).", LocalFile, Result);
+			}
+		}
+
+		/// <summary>
+		/// Download multiple files from the remote Mac
+		/// </summary>
+		/// <param name="Files">List of local files to download</param>
+		void DownloadFiles(IEnumerable<FileReference> Files)
+		{
+			List<FileReference>[] FileGroups = new List<FileReference>[Mappings.Count];
+			for(int Idx = 0; Idx < Mappings.Count; Idx++)
+			{
+				FileGroups[Idx] = new List<FileReference>();
+			}
+			foreach(FileReference File in Files)
+			{
+				int MappingIdx = Mappings.FindIndex(x => File.IsUnderDirectory(x.LocalDirectory));
+				if(MappingIdx == -1)
+				{
+					throw new BuildException("File for download '{0}' is not under the engine or project directory.", File);
+				}
+				FileGroups[MappingIdx].Add(File);
+			}
+			for(int Idx = 0; Idx < Mappings.Count; Idx++)
+			{
+				if(FileGroups[Idx].Count > 0)
+				{
+					FileReference DownloadListLocation = FileReference.Combine(UnrealBuildTool.EngineDirectory, "Intermediate", "Rsync", "Download.txt");
+					DirectoryReference.CreateDirectory(DownloadListLocation.Directory);
+					FileReference.WriteAllLines(DownloadListLocation, FileGroups[Idx].Select(x => x.MakeRelativeTo(Mappings[Idx].LocalDirectory).Replace('\\', '/')));
+
+					List<string> Arguments = new List<string>(CommonRsyncArguments);
+					Arguments.Add(String.Format("--files-from=\"{0}\"", GetLocalCygwinPath(DownloadListLocation)));
+					Arguments.Add(String.Format("\"{0}@{1}\":'{2}/'", UserName, ServerName, Mappings[Idx].RemoteDirectory));
+					Arguments.Add(String.Format("\"{0}/\"", GetLocalCygwinPath(Mappings[Idx].LocalDirectory)));
+
+					int Result = Rsync(String.Join(" ", Arguments));
+					if(Result != 0)
+					{
+						throw new BuildException("Unable to download files from remote Mac (exit code {0})", Result);
+					}
+				}
+			}
+		}
+
+		/// <summary>
+		/// Execute Rsync
+		/// </summary>
+		/// <param name="Arguments">Arguments for the Rsync command</param>
+		/// <returns>Exit code from Rsync</returns>
+		private int Rsync(string Arguments)
+		{
+			using(Process RsyncProcess = new Process())
+			{
+				RsyncProcess.StartInfo.FileName = RsyncExe.FullName;
+				RsyncProcess.StartInfo.Arguments = Arguments;
+				RsyncProcess.StartInfo.WorkingDirectory = SshExe.Directory.FullName;
+				RsyncProcess.OutputDataReceived += RsyncOutput;
+				RsyncProcess.ErrorDataReceived += RsyncOutput;
+
+				Log.TraceLog("[Rsync] {0} {1}", Utils.MakePathSafeToUseWithCommandLine(RsyncProcess.StartInfo.FileName), RsyncProcess.StartInfo.Arguments);
+				return Utils.RunLocalProcess(RsyncProcess);
+			}
+		}
+
+		/// <summary>
+		/// Handles data output by rsync
+		/// </summary>
+		/// <param name="Sender">The object sending the messag</param>
+		/// <param name="Args">The received data</param>e
+		private void RsyncOutput(object Sender, DataReceivedEventArgs Args)
+		{
+			if(Args.Data != null)
+			{
+				Log.TraceInformation("  {0}", Args.Data);
+			}
+		}
+
+		/// <summary>
+		/// Execute a command on the remote in the remote equivalent of a local directory
+		/// </summary>
+		/// <param name="WorkingDir"></param>
+		/// <param name="Command"></param>
+		/// <returns></returns>
+		public int Execute(DirectoryReference WorkingDir, string Command)
+		{
+			return Execute(GetRemotePath(WorkingDir), Command);
+		}
+
+		/// <summary>
+		/// Execute a remote command, capturing the output text
+		/// </summary>
+		/// <param name="WorkingDirectory">The remote working directory</param>
+		/// <param name="Command">Command to be executed</param>
+		/// <returns></returns>
+		protected int Execute(string WorkingDirectory, string Command)
+		{
+			string FullCommand = String.Format("cd {0} && {1}", EscapeShellArgument(WorkingDirectory), Command);
+			using(Process SSHProcess = new Process())
+			{
+				DataReceivedEventHandler OutputHandler = (E, Args) => { SshOutput(Args); };
+
+				SSHProcess.StartInfo.FileName = SshExe.FullName;
+				SSHProcess.StartInfo.WorkingDirectory = SshExe.Directory.FullName;
+				SSHProcess.StartInfo.Arguments = String.Format("{0} {1}", String.Join(" ", CommonSshArguments), FullCommand.Replace("\"", "\\\""));
+				SSHProcess.OutputDataReceived += OutputHandler;
+				SSHProcess.ErrorDataReceived += OutputHandler;
+
+				Log.TraceLog("[SSH] {0} {1}", Utils.MakePathSafeToUseWithCommandLine(SSHProcess.StartInfo.FileName), SSHProcess.StartInfo.Arguments);
+				return Utils.RunLocalProcess(SSHProcess);
+			}
+		}
+
+		/// <summary>
+		/// Handler for output from running remote SSH commands
+		/// </summary>
+		/// <param name="Args"></param>
+		private void SshOutput(DataReceivedEventArgs Args)
+		{
+			if(Args.Data != null)
+			{
+				string FormattedOutput = ConvertRemotePathsToLocal(Args.Data);
+				Log.TraceInformation("  {0}", FormattedOutput);
+			}
+		}
+
+		/// <summary>
+		/// Converts any remote paths within the given string to local format
+		/// </summary>
+		/// <param name="Text">The text containing strings to convert</param>
+		/// <returns>The string with paths converted to local format</returns>
+		private string ConvertRemotePathsToLocal(string Text)
+		{
+			// Try to match any source file with the remote base directory in front of it
+			string Pattern = String.Format("(?<![a-zA-Z=]){0}[^:]*\\.(?:cpp|inl|h|hpp|hh|txt)(?![a-zA-Z])", Regex.Escape(RemoteBaseDir));
+
+			// Find the matches, and early out if there are none
+			MatchCollection Matches = Regex.Matches(Text, Pattern, RegexOptions.IgnoreCase);
+			if(Matches.Count == 0)
+			{
+				return Text;
+			}
+
+			// Replace any remote paths with local ones
+			StringBuilder Result = new StringBuilder();
+			int StartIdx = 0;
+			foreach(Match Match in Matches)
+			{
+				// Append the text leading up to this path
+				Result.Append(Text, StartIdx, Match.Index - StartIdx);
+
+				// Try to convert the path
+				string Path = Match.Value;
+				foreach(RemoteMapping Mapping in Mappings)
+				{
+					if(Path.StartsWith(Mapping.RemoteDirectory))
+					{
+						Path = Mapping.LocalDirectory + Path.Substring(Mapping.RemoteDirectory.Length).Replace('/', '\\');
+						break;
+					}
+				}
+
+				// Append the path to the output string
+				Result.Append(Path);
+
+				// Move past this match
+				StartIdx = Match.Index + Match.Length;
+			}
+			Result.Append(Text, StartIdx, Text.Length - StartIdx);
+			return Result.ToString();
+		}
+	}
+}
diff --git a/Engine/Source/Programs/UnrealBuildTool/ToolChain/RemoteToolChain.cs b/Engine/Source/Programs/UnrealBuildTool/ToolChain/RemoteToolChain.cs
deleted file mode 100644
diff --git a/Engine/Source/Programs/UnrealBuildTool/UnrealBuildTool.cs b/Engine/Source/Programs/UnrealBuildTool/UnrealBuildTool.cs
--- a/Engine/Source/Programs/UnrealBuildTool/UnrealBuildTool.cs
+++ b/Engine/Source/Programs/UnrealBuildTool/UnrealBuildTool.cs
@@ -497,7 +497,8 @@ namespace UnrealBuildTool
 					}
 
 					// Read the XML configuration files
-					if (!XmlConfig.ReadConfigFiles())
+					bool bForceXmlConfigCache = Arguments.Any(x => x.Equals("-ForceXmlConfigCache", StringComparison.InvariantCultureIgnoreCase));
+					if (!XmlConfig.ReadConfigFiles(bForceXmlConfigCache))
 					{
 						return 1;
 					}
@@ -1110,13 +1111,15 @@ namespace UnrealBuildTool
 					Log.TraceInformation("RunUBT initialization took " + RunUBTInitTime + "s");
 				}
 
+				bool bSkipRulesCompile = Arguments.Any(x => x.Equals("-skiprulescompile", StringComparison.InvariantCultureIgnoreCase));
+
 				List<TargetDescriptor> TargetDescs = new List<TargetDescriptor>();
 				{
 					DateTime TargetDescstStartTime = DateTime.UtcNow;
 
 					foreach (string[] TargetSetting in TargetSettings)
 					{
-						TargetDescs.AddRange(TargetDescriptor.ParseCommandLine(TargetSetting, BuildConfiguration.bUsePrecompiled, ref ProjectFile));
+						TargetDescs.AddRange(TargetDescriptor.ParseCommandLine(TargetSetting, BuildConfiguration.bUsePrecompiled, bSkipRulesCompile, ref ProjectFile));
 					}
 
 					if (UnrealBuildTool.bPrintPerformanceInfo)
@@ -1126,6 +1129,35 @@ namespace UnrealBuildTool
 					}
 				}
 
+				// Handle remote builds
+				if(BuildHostPlatform.Current.Platform == UnrealTargetPlatform.Win64)
+				{
+					for(int Idx = 0; Idx < TargetDescs.Count; Idx++)
+					{
+						TargetDescriptor TargetDesc = TargetDescs[Idx];
+						if(TargetDesc.Platform == UnrealTargetPlatform.Mac || TargetDesc.Platform == UnrealTargetPlatform.IOS || TargetDesc.Platform == UnrealTargetPlatform.TVOS)
+						{
+							FileReference LogFile = null;
+							if(!String.IsNullOrEmpty(BuildConfiguration.LogFileName))
+							{
+								LogFile = new FileReference(Path.Combine(Path.GetDirectoryName(BuildConfiguration.LogFileName), Path.GetFileNameWithoutExtension(BuildConfiguration.LogFileName) + "_Remote.txt"));
+							}
+
+							RemoteMac RemoteMac = new RemoteMac(TargetDesc.ProjectFile);
+							if(!RemoteMac.Build(TargetDesc, LogFile))
+							{
+								return ECompilationResult.Unknown;
+							}
+
+							TargetDescs.RemoveAt(Idx--);
+						}
+					}
+					if(TargetDescs.Count == 0)
+					{
+						return ECompilationResult.Succeeded;
+					}
+				}
+
 				if (Arguments.Any(x => x.Equals("-InvalidateMakefilesOnly", StringComparison.InvariantCultureIgnoreCase)))
 				{
 					Log.TraceInformation("Invalidating makefiles only in this run.");
@@ -1285,7 +1317,7 @@ namespace UnrealBuildTool
 						// Make sure the gather phase is executed if we're not actually building anything
 						if (ProjectFileGenerator.bGenerateProjectFiles || BuildConfiguration.bGenerateManifest || BuildConfiguration.bCleanProject || BuildConfiguration.bXGEExport || GeneratingActionGraph)
 						{
-							bIsGatheringBuild = true;
+								bIsGatheringBuild = true;
 						}
 
 						// Were we asked to run in 'assembler only' mode?  If so, let's check to see if that's even possible by seeing if
@@ -1377,7 +1409,7 @@ namespace UnrealBuildTool
 						Targets = new List<UEBuildTarget>();
 						foreach (TargetDescriptor TargetDesc in TargetDescs)
 						{
-							UEBuildTarget Target = UEBuildTarget.CreateTarget(TargetDesc, Arguments, BuildConfiguration.SingleFileToCompile != null, BuildConfiguration.bUsePrecompiled, Version);
+							UEBuildTarget Target = UEBuildTarget.CreateTarget(TargetDesc, Arguments, bSkipRulesCompile, BuildConfiguration.SingleFileToCompile != null, BuildConfiguration.bUsePrecompiled, Version);
 							if ((Target == null) && (BuildConfiguration.bCleanProject))
 							{
 								continue;
@@ -1626,8 +1658,7 @@ namespace UnrealBuildTool
 								}
 								else
 								{
-									bool bIsRemoteCompile = BuildHostPlatform.Current.Platform != UnrealTargetPlatform.Mac && TargetDescs.Any(x => x.Platform == UnrealTargetPlatform.Mac || x.Platform == UnrealTargetPlatform.IOS);
-									bSuccess = ActionGraph.ExecuteActions(BuildConfiguration, ActionsToExecute, bIsRemoteCompile, out ExecutorName, TargetInfoForTelemetry, HotReload);
+									bSuccess = ActionGraph.ExecuteActions(BuildConfiguration, ActionsToExecute, out ExecutorName, TargetInfoForTelemetry, HotReload);
 								}
 								TotalExecutorTime += (DateTime.UtcNow - ExecutorStartTime).TotalSeconds;
 
@@ -2111,7 +2142,6 @@ namespace UnrealBuildTool
 							// Copy the other important settings from the original file item
 							NewPrerequisiteItem.bNeedsHotReloadNumbersDLLCleanUp = OriginalPrerequisiteItem.bNeedsHotReloadNumbersDLLCleanUp;
 							NewPrerequisiteItem.ProducingAction = OriginalPrerequisiteItem.ProducingAction;
-							NewPrerequisiteItem.bIsRemoteFile = OriginalPrerequisiteItem.bIsRemoteFile;
 
 							// Keep track of it so we can fix up dependencies in a second pass afterwards
 							AffectedOriginalFileItemAndNewFileItemMap.Add(OriginalPrerequisiteItem, NewPrerequisiteItem);
@@ -2153,7 +2183,6 @@ namespace UnrealBuildTool
 						// Copy the other important settings from the original file item
 						NewProducedItem.bNeedsHotReloadNumbersDLLCleanUp = OriginalProducedItem.bNeedsHotReloadNumbersDLLCleanUp;
 						NewProducedItem.ProducingAction = OriginalProducedItem.ProducingAction;
-						NewProducedItem.bIsRemoteFile = OriginalProducedItem.bIsRemoteFile;
 
 						// Keep track of it so we can fix up dependencies in a second pass afterwards
 						AffectedOriginalFileItemAndNewFileItemMap.Add(OriginalProducedItem, NewProducedItem);
diff --git a/Engine/Source/Programs/UnrealBuildTool/UnrealBuildTool.csproj b/Engine/Source/Programs/UnrealBuildTool/UnrealBuildTool.csproj
--- a/Engine/Source/Programs/UnrealBuildTool/UnrealBuildTool.csproj
+++ b/Engine/Source/Programs/UnrealBuildTool/UnrealBuildTool.csproj
@@ -153,6 +153,7 @@
     <Compile Include="Executors\ParallelExecutor.cs" />
     <Compile Include="Executors\SNDBS.cs" />
     <Compile Include="Executors\XGE.cs" />
+    <Compile Include="Platform\IOS\MobileProvisionContents.cs" />
     <Compile Include="Platform\Lumin\UEDeployLumin.cs" />
     <Compile Include="ProjectFiles\CLion\CLionGenerator.cs" />
     <Compile Include="ProjectFiles\CMake\CMakefileGenerator.cs" />
@@ -171,6 +172,7 @@
     <Compile Include="ProjectFiles\VisualStudio\VCSolutionOptions.cs" />
     <Compile Include="ProjectFiles\Xcode\XcodeProject.cs" />
     <Compile Include="ProjectFiles\Xcode\XcodeProjectFileGenerator.cs" />
+    <Compile Include="System\Asn.cs" />
     <Compile Include="System\BuildVersion.cs" />
     <Compile Include="System\CppCompileEnvironment.cs" />
     <Compile Include="System\CppIncludeParser.cs" />
@@ -199,7 +201,6 @@
     <Compile Include="System\LocalizationDescriptor.cs" />
     <Compile Include="System\ObsoleteOverrideAttribute.cs" />
     <Compile Include="System\ProgressWriter.cs" />
-    <Compile Include="System\RemoteExecutor.cs" />
     <Compile Include="System\RulesAssembly.cs" />
     <Compile Include="System\ScopedTimer.cs" />
     <Compile Include="System\SourceFileWorkingSet.cs" />
@@ -247,7 +248,7 @@
     </When>
   </Choose>
   <ItemGroup>
-    <Compile Include="ToolChain\RemoteExports.cs" />
+    <Compile Include="ToolChain\RemoteMac.cs" />
     <Compile Include="UnrealBuildTool.cs" />
     <Compile Include="Platform\Windows\PVSToolChain.cs" />
     <Compile Include="Platform\Windows\WindowsDeploy.cs" />
@@ -333,14 +334,12 @@
     <Compile Include="System\JunkDeleter.cs" />
     <Compile Include="System\LinkEnvironment.cs" />
     <Compile Include="System\Plugins.cs" />
-    <Compile Include="System\RPCUtilHelper.cs" />
     <Compile Include="System\RulesCompiler.cs" />
     <Compile Include="System\SourceFileSearch.cs" />
     <Compile Include="System\UEPlatformProjectGenerator.cs" />
     <Compile Include="System\Unity.cs" />
     <Compile Include="System\UProjectInfo.cs" />
     <Compile Include="System\Utils.cs" />
-    <Compile Include="ToolChain\RemoteToolChain.cs" />
     <Compile Include="ToolChain\UEToolChain.cs" />
     <Compile Include="Platform\Windows\UEBuildWindows.cs" />
     <Compile Include="Platform\Windows\VCEnvironment.cs" />
