PR #4921: Honor ModuleRules' bEnableExceptions flag when creating precompiled h. (Contributed by [REDACTED_PERSON])

diff --git a/Engine/Source/Programs/UnrealBuildTool/Configuration/UEBuildModuleCPP.cs b/Engine/Source/Programs/UnrealBuildTool/Configuration/UEBuildModuleCPP.cs
--- a/Engine/Source/Programs/UnrealBuildTool/Configuration/UEBuildModuleCPP.cs
+++ b/Engine/Source/Programs/UnrealBuildTool/Configuration/UEBuildModuleCPP.cs
@@ -449,7 +449,7 @@ namespace UnrealBuildTool
 					PrecompiledHeaderTemplate Template = SharedPCHs.FirstOrDefault(x => ReferencedModules.Contains(x.Module));
 					if(Template != null && Template.IsValidFor(CompileEnvironment))
 					{
-						PrecompiledHeaderInstance Instance = FindOrCreateSharedPCH(ToolChain, Template, ModuleCompileEnvironment.bOptimizeCode, ModuleCompileEnvironment.bUseRTTI, ActionGraph);
+						PrecompiledHeaderInstance Instance = FindOrCreateSharedPCH(ToolChain, Template, ModuleCompileEnvironment.bOptimizeCode, ModuleCompileEnvironment.bUseRTTI, ModuleCompileEnvironment.bEnableExceptions, ActionGraph);
 
 						FileReference PrivateDefinitionsFile = FileReference.Combine(IntermediateDirectory, String.Format("Definitions.{0}.h", Name));
 
@@ -636,7 +636,7 @@ namespace UnrealBuildTool
 
 			// Create the action to compile the PCH file.
 			CPPOutput Output = ToolChain.CompileCPPFiles(CompileEnvironment, new List<FileItem>() { WrapperFile }, IntermediateDirectory, Name, ActionGraph);
-			return new PrecompiledHeaderInstance(WrapperFile, CompileEnvironment.bOptimizeCode, CompileEnvironment.bUseRTTI, Output);
+			return new PrecompiledHeaderInstance(WrapperFile, CompileEnvironment.bOptimizeCode, CompileEnvironment.bUseRTTI, CompileEnvironment.bEnableExceptions, Output);
 		}
 
 		/// <summary>
@@ -646,11 +646,12 @@ namespace UnrealBuildTool
 		/// <param name="Template">The PCH template</param>
 		/// <param name="bOptimizeCode">Whether optimization should be enabled for this PCH</param>
 		/// <param name="bUseRTTI">Whether to enable RTTI for this PCH</param>
+		/// <param name="bEnableExceptions">Whether to enable exceptions for this PCH</param>
 		/// <param name="ActionGraph">Graph containing build actions</param>
 		/// <returns>Instance of a PCH</returns>
-		public PrecompiledHeaderInstance FindOrCreateSharedPCH(UEToolChain ToolChain, PrecompiledHeaderTemplate Template, bool bOptimizeCode, bool bUseRTTI, ActionGraph ActionGraph)
+		public PrecompiledHeaderInstance FindOrCreateSharedPCH(UEToolChain ToolChain, PrecompiledHeaderTemplate Template, bool bOptimizeCode, bool bUseRTTI, bool bEnableExceptions, ActionGraph ActionGraph)
 		{
-			PrecompiledHeaderInstance Instance = Template.Instances.Find(x => x.bOptimizeCode == bOptimizeCode && x.bUseRTTI == bUseRTTI);
+			PrecompiledHeaderInstance Instance = Template.Instances.Find(x => x.bOptimizeCode == bOptimizeCode && x.bUseRTTI == bUseRTTI && x.bEnableExceptions == bEnableExceptions);
 			if(Instance == null)
 			{
 				// Create a suffix to distinguish this shared PCH variant from any others. Currently only optimized and non-optimized shared PCHs are supported.
@@ -677,6 +678,17 @@ namespace UnrealBuildTool
 						Variant += ".NonRTTI";
 					}
 				}
+				if (bEnableExceptions != Template.BaseCompileEnvironment.bEnableExceptions)
+				{
+					if (bEnableExceptions)
+					{
+						Variant += ".Exceptions";
+					}
+					else
+					{
+						Variant += ".NoExceptions";
+					}
+				}
 
 				// Create the wrapper file, which sets all the definitions needed to compile it
 				FileReference WrapperLocation = FileReference.Combine(Template.OutputDir, String.Format("SharedPCH.{0}{1}.h", Template.Module.Name, Variant));
@@ -689,10 +701,11 @@ namespace UnrealBuildTool
 				CompileEnvironment.PrecompiledHeaderIncludeFilename = WrapperFile.Location;
 				CompileEnvironment.bOptimizeCode = bOptimizeCode;
 				CompileEnvironment.bUseRTTI = bUseRTTI;
+				CompileEnvironment.bEnableExceptions = bEnableExceptions;
 
 				// Create the PCH
 				CPPOutput Output = ToolChain.CompileCPPFiles(CompileEnvironment, new List<FileItem>() { WrapperFile }, Template.OutputDir, "Shared", ActionGraph);
-				Instance = new PrecompiledHeaderInstance(WrapperFile, bOptimizeCode, bUseRTTI, Output);
+				Instance = new PrecompiledHeaderInstance(WrapperFile, bOptimizeCode, bUseRTTI, bEnableExceptions, Output);
 				Template.Instances.Add(Instance);
 			}
 			return Instance;
diff --git a/Engine/Source/Programs/UnrealBuildTool/System/PrecompiledHeaderInstance.cs b/Engine/Source/Programs/UnrealBuildTool/System/PrecompiledHeaderInstance.cs
--- a/Engine/Source/Programs/UnrealBuildTool/System/PrecompiledHeaderInstance.cs
+++ b/Engine/Source/Programs/UnrealBuildTool/System/PrecompiledHeaderInstance.cs
@@ -28,6 +28,11 @@ namespace UnrealBuildTool
 		/// </summary>
 		public bool bUseRTTI;
 
+		/// <summary>
+		/// Whether to enable exceptions
+		/// </summary>
+		public bool bEnableExceptions;
+
 		/// <summary>
 		/// The output files for the shared PCH
 		/// </summary>
@@ -36,11 +41,12 @@ namespace UnrealBuildTool
 		/// <summary>
 		/// Constructor
 		/// </summary>
-		public PrecompiledHeaderInstance(FileItem HeaderFile, bool bOptimizeCode, bool bUseRTTI, CPPOutput Output)
+		public PrecompiledHeaderInstance(FileItem HeaderFile, bool bOptimizeCode, bool bUseRTTI, bool bEnableExceptions, CPPOutput Output)
 		{
 			this.HeaderFile = HeaderFile;
 			this.bOptimizeCode = bOptimizeCode;
 			this.bUseRTTI = bUseRTTI;
+			this.bEnableExceptions = bEnableExceptions;
 			this.Output = Output;
 		}
 
@@ -50,7 +56,7 @@ namespace UnrealBuildTool
 		/// <returns>String representation of the object</returns>
 		public override string ToString()
 		{
-			return String.Format("{0} (Optimized={1}, RTTI={2})", HeaderFile.Location.GetFileName(), bOptimizeCode, bUseRTTI);
+			return String.Format("{0} (Optimized={1}, RTTI={2}, Exceptions={3})", HeaderFile.Location.GetFileName(), bOptimizeCode, bUseRTTI, bEnableExceptions);
 		}
 	}
 }
