UBT: Add an error if a module referenced by a plugin descriptor doesn't exist.

diff --git a/Engine/Source/Programs/UnrealBuildTool/Configuration/UEBuildTarget.cs b/Engine/Source/Programs/UnrealBuildTool/Configuration/UEBuildTarget.cs
--- a/Engine/Source/Programs/UnrealBuildTool/Configuration/UEBuildTarget.cs
+++ b/Engine/Source/Programs/UnrealBuildTool/Configuration/UEBuildTarget.cs
@@ -2855,6 +2855,10 @@ namespace UnrealBuildTool
 								if (ModuleDescriptor.IsCompiledInConfiguration(Platform, Configuration, TargetName, TargetType, bAllowDeveloperModules && Rules.bBuildDeveloperTools, Rules.bBuildEditor, Rules.bBuildRequiresCookedData))
 								{
 									FileReference ModuleFileName = RulesAssembly.GetModuleFileName(ModuleDescriptor.Name);
+									if(ModuleFileName == null)
+									{
+										throw new BuildException("Unable to find module '{0}' referenced by {1}", ModuleDescriptor.Name, Plugin.File);
+									}
 									if(!ModuleFileName.ContainsAnyNames(ExcludeFoldersArray, Plugin.Directory))
 									{
 										FilteredModuleNames.Add(ModuleDescriptor.Name);
