UBT: Fix project files not being generated for foreign projects when creating .stub files.

diff --git a/Engine/Source/Programs/UnrealBuildTool/Platform/IOS/IOSToolChain.cs b/Engine/Source/Programs/UnrealBuildTool/Platform/IOS/IOSToolChain.cs
--- a/Engine/Source/Programs/UnrealBuildTool/Platform/IOS/IOSToolChain.cs
+++ b/Engine/Source/Programs/UnrealBuildTool/Platform/IOS/IOSToolChain.cs
@@ -1635,7 +1635,7 @@ namespace UnrealBuildTool
 				}
 				else
 				{
-					UnrealBuildTool.GenerateProjectFiles(new XcodeProjectFileGenerator(Target.ProjectFile), new string[] { "-platforms" + (Target.Platform == UnrealTargetPlatform.IOS ? "IOS" : "TVOS"), "-NoIntellIsense", (Target.Platform == UnrealTargetPlatform.IOS ? "-iosdeployonly" : "-tvosdeployonly"), "-ignorejunk", "-project=\"" + Target.ProjectDirectory + "/" + AppName + ".uproject\"", "-game" });
+					UnrealBuildTool.GenerateProjectFiles(new XcodeProjectFileGenerator(Target.ProjectFile), new string[] { "-platforms=" + (Target.Platform == UnrealTargetPlatform.IOS ? "IOS" : "TVOS"), "-NoIntellIsense", (Target.Platform == UnrealTargetPlatform.IOS ? "-iosdeployonly" : "-tvosdeployonly"), "-ignorejunk", "-project=\"" + Target.ProjectDirectory + "/" + AppName + ".uproject\"", "-game" });
 					Project = Target.ProjectDirectory + "/" + AppName + "_" + (Target.Platform == UnrealTargetPlatform.IOS ? "IOS" : "TVOS") + ".xcworkspace";
 				}
 
