UBT: Clean up the logic around generating project files before creating a stub IPA, so that it fails loudly if project files do not exist, and can accept target names not matching project names.

diff --git a/Engine/Source/Programs/UnrealBuildTool/Platform/IOS/IOSToolChain.cs b/Engine/Source/Programs/UnrealBuildTool/Platform/IOS/IOSToolChain.cs
--- a/Engine/Source/Programs/UnrealBuildTool/Platform/IOS/IOSToolChain.cs
+++ b/Engine/Source/Programs/UnrealBuildTool/Platform/IOS/IOSToolChain.cs
@@ -1619,175 +1619,178 @@ namespace UnrealBuildTool
 
 			if (Target.Rules.bCreateStubIPA)
 			{
-				string Project = Target.ProjectDirectory + "/" + AppName + ".uproject";
-
-				string SchemeName = AppName;
-
 				// generate the dummy project so signing works
-				if (AppName == "UE4Game" || AppName == "UE4Client" || Utils.IsFileUnderDirectory(Target.ProjectDirectory + "/" + AppName + ".uproject", Path.GetFullPath("../..")))
+				DirectoryReference XcodeWorkspaceDir;
+				if (AppName == "UE4Game" || AppName == "UE4Client" || Target.ProjectFile == null || Target.ProjectFile.IsUnderDirectory(UnrealBuildTool.EngineDirectory))
 				{
 					UnrealBuildTool.GenerateProjectFiles(new XcodeProjectFileGenerator(Target.ProjectFile), new string[] { "-platforms=" + (Target.Platform == UnrealTargetPlatform.IOS ? "IOS" : "TVOS"), "-NoIntellIsense", (Target.Platform == UnrealTargetPlatform.IOS ? "-iosdeployonly" : "-tvosdeployonly"), "-ignorejunk" });
-					Project = Path.GetFullPath("../..") + "/UE4_" + (Target.Platform == UnrealTargetPlatform.IOS ? "IOS" : "TVOS") + ".xcworkspace";
-					if (AppName == "UE4Game" || AppName == "UE4Client")
-					{
-						SchemeName = "UE4";
-					}
+					XcodeWorkspaceDir = DirectoryReference.Combine(UnrealBuildTool.EngineDirectory, String.Format("UE4_{0}.xcworkspace", (Target.Platform == UnrealTargetPlatform.IOS ? "IOS" : "TVOS")));
 				}
 				else
 				{
-					UnrealBuildTool.GenerateProjectFiles(new XcodeProjectFileGenerator(Target.ProjectFile), new string[] { "-platforms=" + (Target.Platform == UnrealTargetPlatform.IOS ? "IOS" : "TVOS"), "-NoIntellIsense", (Target.Platform == UnrealTargetPlatform.IOS ? "-iosdeployonly" : "-tvosdeployonly"), "-ignorejunk", "-project=\"" + Target.ProjectDirectory + "/" + AppName + ".uproject\"", "-game" });
-					Project = Target.ProjectDirectory + "/" + AppName + "_" + (Target.Platform == UnrealTargetPlatform.IOS ? "IOS" : "TVOS") + ".xcworkspace";
+					UnrealBuildTool.GenerateProjectFiles(new XcodeProjectFileGenerator(Target.ProjectFile), new string[] { "-platforms=" + (Target.Platform == UnrealTargetPlatform.IOS ? "IOS" : "TVOS"), "-NoIntellIsense", (Target.Platform == UnrealTargetPlatform.IOS ? "-iosdeployonly" : "-tvosdeployonly"), "-ignorejunk", String.Format("-project={0}", Target.ProjectFile), "-game" });
+					XcodeWorkspaceDir = DirectoryReference.Combine(Target.ProjectDirectory, String.Format("{0}_{1}.xcworkspace", AppName, (Target.Platform == UnrealTargetPlatform.IOS ? "IOS" : "TVOS")));
 				}
 
-				if (Directory.Exists(Project))
+				// Make sure it exists
+				if (!DirectoryReference.Exists(XcodeWorkspaceDir))
 				{
-					// ensure the plist, entitlements, and provision files are properly copied
-					DeployHandler = (Target.Platform == UnrealTargetPlatform.IOS ? new UEDeployIOS() : new UEDeployTVOS());
-					DeployHandler.PrepTargetForDeployment(new UEBuildDeployTarget(Target));
+					throw new BuildException("Unable to create stub IPA; Xcode workspace not found at {0}", XcodeWorkspaceDir);
+				}
 
-					string ConfigName = Target.Configuration.ToString();
-					if (Target.Rules.Type != TargetType.Game && Target.Rules.Type != TargetType.Program)
+				// ensure the plist, entitlements, and provision files are properly copied
+				DeployHandler = (Target.Platform == UnrealTargetPlatform.IOS ? new UEDeployIOS() : new UEDeployTVOS());
+				DeployHandler.PrepTargetForDeployment(new UEBuildDeployTarget(Target));
+
+				FileReference SignProjectScript = FileReference.Combine(Target.ProjectIntermediateDirectory, "SignProject.sh");
+				using(StreamWriter Writer = new StreamWriter(SignProjectScript.FullName))
+				{
+					// Boilerplate
+					Writer.WriteLine("#!/bin/sh");
+					Writer.WriteLine("set -e");
+					Writer.WriteLine("set -x");
+
+					// Copy the mobile provision into the system store
+					if(Target.Rules.IOSPlatform.ImportProvision != null)
 					{
-						ConfigName += " " + Target.Rules.Type.ToString();
+						Writer.WriteLine("cp -f {0} ~/Library/MobileDevice/Provisioning\\ Profiles/", Utils.EscapeShellArgument(Target.Rules.IOSPlatform.ImportProvision));
 					}
 
-					FileReference SignProjectScript = FileReference.Combine(Target.ProjectIntermediateDirectory, "SignProject.sh");
-					using(StreamWriter Writer = new StreamWriter(SignProjectScript.FullName))
+					// Path to the temporary keychain. When -ImportCertificate is specified, we will temporarily add this to the list of keychains to search, and remove it later.
+					FileReference TempKeychain = null;
+
+					// Get the signing certificate to use
+					string SigningCertificate;
+					if(Target.Rules.IOSPlatform.ImportCertificate == null)
 					{
-						// Boilerplate
-						Writer.WriteLine("#!/bin/sh");
-						Writer.WriteLine("set -e");
-						Writer.WriteLine("set -x");
+						// Take it from the standard settings
+						IOSProvisioningData ProvisioningData = IOSPlatform.ReadProvisioningData(Target.ProjectFile);
+						SigningCertificate = ProvisioningData.SigningCertificate;
 
-						// Copy the mobile provision into the system store
-						if(Target.Rules.IOSPlatform.ImportProvision != null)
+						// Set the identity on the command line
+						if(!ProjectSettings.bAutomaticSigning)
 						{
-							Writer.WriteLine("cp -f {0} ~/Library/MobileDevice/Provisioning\\ Profiles/", Utils.EscapeShellArgument(Target.Rules.IOSPlatform.ImportProvision));
+							Writer.WriteLine("CODE_SIGN_IDENTITY='{0}'", String.IsNullOrEmpty(SigningCertificate)? "IPhoneDeveloper" : SigningCertificate);
 						}
-
-						// Path to the temporary keychain. When -ImportCertificate is specified, we will temporarily add this to the list of keychains to search, and remove it later.
-						FileReference TempKeychain = null;
-
-						// Get the signing certificate to use
-						string SigningCertificate;
-						if(Target.Rules.IOSPlatform.ImportCertificate == null)
+					}
+					else
+					{
+						// Read the name from the certificate
+						X509Certificate2 Certificate;
+						try
 						{
-							// Take it from the standard settings
-							IOSProvisioningData ProvisioningData = IOSPlatform.ReadProvisioningData(Target.ProjectFile);
-							SigningCertificate = ProvisioningData.SigningCertificate;
-
-							// Set the identity on the command line
-							if(!ProjectSettings.bAutomaticSigning)
-							{
-								Writer.WriteLine("CODE_SIGN_IDENTITY='{0}'", String.IsNullOrEmpty(SigningCertificate)? "IPhoneDeveloper" : SigningCertificate);
-							}
+							Certificate = new X509Certificate2(Target.Rules.IOSPlatform.ImportCertificate, Target.Rules.IOSPlatform.ImportCertificatePassword ?? "");
 						}
-						else
+						catch(Exception Ex)
 						{
-							// Read the name from the certificate
-							X509Certificate2 Certificate;
-							try
-							{
-								Certificate = new X509Certificate2(Target.Rules.IOSPlatform.ImportCertificate, Target.Rules.IOSPlatform.ImportCertificatePassword ?? "");
-							}
-							catch(Exception Ex)
-							{
-								throw new BuildException(Ex, "Unable to read certificate '{0}': {1}", Target.Rules.IOSPlatform.ImportCertificate, Ex.Message);
-							}
-							SigningCertificate = Certificate.GetNameInfo(X509NameType.SimpleName, false);
-
-							// Set the path to the temporary keychain
-							TempKeychain = FileReference.Combine(Target.ProjectIntermediateDirectory, "TempKeychain.keychain");//(DirectoryReference.GetSpecialFolder(Environment.SpecialFolder.UserProfile), "Library", "Keychains/UE4TempKeychain.keychain";
-
-							// Install a certificate given on the command line to a temporary keychain
-							Writer.WriteLine("security delete-keychain \"{0}\" || true", TempKeychain);
-							Writer.WriteLine("security create-keychain -p \"A\" \"{0}\"", TempKeychain);
-							Writer.WriteLine("security list-keychains -s \"{0}\"", TempKeychain);
-							Writer.WriteLine("security list-keychains");
-							Writer.WriteLine("security set-keychain-settings -t 3600 -l  \"{0}\"", TempKeychain);
-							Writer.WriteLine("security -v unlock-keychain -p \"A\" \"{0}\"", TempKeychain);
-							Writer.WriteLine("security import {0} -P {1} -k \"{2}\" -T /usr/bin/codesign -T /usr/bin/security -t agg", Utils.EscapeShellArgument(Target.Rules.IOSPlatform.ImportCertificate), Utils.EscapeShellArgument(Target.Rules.IOSPlatform.ImportCertificatePassword), TempKeychain);
-							Writer.WriteLine("security set-key-partition-list -S apple-tool:,apple:,codesign: -s -k \"A\" -D '{0}' -t private {1}", SigningCertificate, TempKeychain);
-
-							// Set parameters to make sure it uses the correct identity and keychain
-							Writer.WriteLine("CERT_IDENTITY='{0}'", SigningCertificate);
-							Writer.WriteLine("CODE_SIGN_IDENTITY='{0}'", SigningCertificate);
-							Writer.WriteLine("CODE_SIGN_KEYCHAIN='{0}'", TempKeychain);
+							throw new BuildException(Ex, "Unable to read certificate '{0}': {1}", Target.Rules.IOSPlatform.ImportCertificate, Ex.Message);
 						}
+						SigningCertificate = Certificate.GetNameInfo(X509NameType.SimpleName, false);
+
+						// Set the path to the temporary keychain
+						TempKeychain = FileReference.Combine(Target.ProjectIntermediateDirectory, "TempKeychain.keychain");//(DirectoryReference.GetSpecialFolder(Environment.SpecialFolder.UserProfile), "Library", "Keychains/UE4TempKeychain.keychain";
+
+						// Install a certificate given on the command line to a temporary keychain
+						Writer.WriteLine("security delete-keychain \"{0}\" || true", TempKeychain);
+						Writer.WriteLine("security create-keychain -p \"A\" \"{0}\"", TempKeychain);
+						Writer.WriteLine("security list-keychains -s \"{0}\"", TempKeychain);
+						Writer.WriteLine("security list-keychains");
+						Writer.WriteLine("security set-keychain-settings -t 3600 -l  \"{0}\"", TempKeychain);
+						Writer.WriteLine("security -v unlock-keychain -p \"A\" \"{0}\"", TempKeychain);
+						Writer.WriteLine("security import {0} -P {1} -k \"{2}\" -T /usr/bin/codesign -T /usr/bin/security -t agg", Utils.EscapeShellArgument(Target.Rules.IOSPlatform.ImportCertificate), Utils.EscapeShellArgument(Target.Rules.IOSPlatform.ImportCertificatePassword), TempKeychain);
+						Writer.WriteLine("security set-key-partition-list -S apple-tool:,apple:,codesign: -s -k \"A\" -D '{0}' -t private {1}", SigningCertificate, TempKeychain);
+
+						// Set parameters to make sure it uses the correct identity and keychain
+						Writer.WriteLine("CERT_IDENTITY='{0}'", SigningCertificate);
+						Writer.WriteLine("CODE_SIGN_IDENTITY='{0}'", SigningCertificate);
+						Writer.WriteLine("CODE_SIGN_KEYCHAIN='{0}'", TempKeychain);
+					}
 
-						FileReference MobileProvisionFile;
-						string MobileProvisionUUID;
-						string TeamUUID;
-						if(Target.Rules.IOSPlatform.ImportProvision == null)
-						{
-							IOSProvisioningData ProvisioningData = ((IOSPlatform)UEBuildPlatform.GetBuildPlatform(Target.Platform)).ReadProvisioningData(ProjectSettings);
-							MobileProvisionFile = ProvisioningData.MobileProvisionFile;
-							MobileProvisionUUID = ProvisioningData.MobileProvisionUUID;
-							TeamUUID = ProvisioningData.TeamUUID;
-						}
-						else
-						{
-							MobileProvisionFile = new FileReference(Target.Rules.IOSPlatform.ImportProvision);
+					FileReference MobileProvisionFile;
+					string MobileProvisionUUID;
+					string TeamUUID;
+					if(Target.Rules.IOSPlatform.ImportProvision == null)
+					{
+						IOSProvisioningData ProvisioningData = ((IOSPlatform)UEBuildPlatform.GetBuildPlatform(Target.Platform)).ReadProvisioningData(ProjectSettings);
+						MobileProvisionFile = ProvisioningData.MobileProvisionFile;
+						MobileProvisionUUID = ProvisioningData.MobileProvisionUUID;
+						TeamUUID = ProvisioningData.TeamUUID;
+					}
+					else
+					{
+						MobileProvisionFile = new FileReference(Target.Rules.IOSPlatform.ImportProvision);
 
-							MobileProvisionContents MobileProvision = MobileProvisionContents.Read(MobileProvisionFile);
-							MobileProvisionUUID = MobileProvision.GetUniqueId();
+						MobileProvisionContents MobileProvision = MobileProvisionContents.Read(MobileProvisionFile);
+						MobileProvisionUUID = MobileProvision.GetUniqueId();
 
-							TeamUUID = null;
-						}
+						TeamUUID = null;
+					}
 
-						// code sign the project
-						Console.WriteLine("Provisioning: {0}, {1}, {2}", MobileProvisionFile, MobileProvisionFile.GetFileName(), MobileProvisionUUID);
-						string CmdLine = new IOSToolChainSettings().XcodeDeveloperDir + "usr/bin/xcodebuild" +
-										" -workspace \"" + Project + "\"" +
-										" -configuration \"" + ConfigName + "\"" +
-										" -scheme '" + SchemeName + "'" +
-										" -sdk " + GetCodesignPlatformName(Target.Platform) +
-										" -destination generic/platform=" + (Target.Platform == UnrealTargetPlatform.IOS ? "iOS" : "tvOS") + 
-										(!string.IsNullOrEmpty(TeamUUID) ? " DEVELOPMENT_TEAM=" + TeamUUID : "");
-						CmdLine += String.Format(" CODE_SIGN_IDENTITY='{0}'", SigningCertificate);
-						if (!ProjectSettings.bAutomaticSigning)
-						{
-							CmdLine += (!string.IsNullOrEmpty(MobileProvisionUUID) ? (" PROVISIONING_PROFILE_SPECIFIER=" + MobileProvisionUUID) : "");
-						}
-						Writer.WriteLine("/usr/bin/xcrun {0}", CmdLine);
+					string ConfigName = Target.Configuration.ToString();
+					if (Target.Rules.Type != TargetType.Game && Target.Rules.Type != TargetType.Program)
+					{
+						ConfigName += " " + Target.Rules.Type.ToString();
+					}
 
-						// Remove the temporary keychain from the search list
-						if(TempKeychain != null)
-						{
-							Writer.WriteLine("security delete-keychain \"{0}\" || true", TempKeychain);
-						}
+					string SchemeName;
+					if (AppName == "UE4Game" || AppName == "UE4Client")
+					{
+						SchemeName = "UE4";
+					}
+					else
+					{
+						SchemeName = AppName;
+					}
+
+					// code sign the project
+					Console.WriteLine("Provisioning: {0}, {1}, {2}", MobileProvisionFile, MobileProvisionFile.GetFileName(), MobileProvisionUUID);
+					string CmdLine = new IOSToolChainSettings().XcodeDeveloperDir + "usr/bin/xcodebuild" +
+									" -workspace \"" + XcodeWorkspaceDir + "\"" +
+									" -configuration \"" + ConfigName + "\"" +
+									" -scheme '" + SchemeName + "'" +
+									" -sdk " + GetCodesignPlatformName(Target.Platform) +
+									" -destination generic/platform=" + (Target.Platform == UnrealTargetPlatform.IOS ? "iOS" : "tvOS") + 
+									(!string.IsNullOrEmpty(TeamUUID) ? " DEVELOPMENT_TEAM=" + TeamUUID : "");
+					CmdLine += String.Format(" CODE_SIGN_IDENTITY='{0}'", SigningCertificate);
+					if (!ProjectSettings.bAutomaticSigning)
+					{
+						CmdLine += (!string.IsNullOrEmpty(MobileProvisionUUID) ? (" PROVISIONING_PROFILE_SPECIFIER=" + MobileProvisionUUID) : "");
 					}
+					Writer.WriteLine("/usr/bin/xcrun {0}", CmdLine);
 
-					Log.TraceInformation("Executing {0}", SignProjectScript);
+					// Remove the temporary keychain from the search list
+					if(TempKeychain != null)
+					{
+						Writer.WriteLine("security delete-keychain \"{0}\" || true", TempKeychain);
+					}
+				}
 
-					Process SignProcess = new Process();
-					SignProcess.StartInfo.WorkingDirectory = RemoteShadowDirectoryMac;
-					SignProcess.StartInfo.FileName = "/bin/sh";
-					SignProcess.StartInfo.Arguments = SignProjectScript.FullName;
+				Log.TraceInformation("Executing {0}", SignProjectScript);
 
-					ProcessOutput Output = new ProcessOutput();
+				Process SignProcess = new Process();
+				SignProcess.StartInfo.WorkingDirectory = RemoteShadowDirectoryMac;
+				SignProcess.StartInfo.FileName = "/bin/sh";
+				SignProcess.StartInfo.Arguments = SignProjectScript.FullName;
 
-					SignProcess.OutputDataReceived += new DataReceivedEventHandler(Output.OutputReceivedDataEventHandler);
-					SignProcess.ErrorDataReceived += new DataReceivedEventHandler(Output.OutputReceivedDataEventHandler);
+				ProcessOutput Output = new ProcessOutput();
 
-					Output.OutputReceivedDataEventHandlerEncounteredError = false;
-					Output.OutputReceivedDataEventHandlerEncounteredErrorMessage = "";
-					Utils.RunLocalProcess(SignProcess);
+				SignProcess.OutputDataReceived += new DataReceivedEventHandler(Output.OutputReceivedDataEventHandler);
+				SignProcess.ErrorDataReceived += new DataReceivedEventHandler(Output.OutputReceivedDataEventHandler);
 
-					// delete the temp project
-					if (Project.Contains("_" + (Target.Platform == UnrealTargetPlatform.IOS ? "IOS" : "TVOS") + ".xcodeproj"))
-					{
-						Directory.Delete(Project, true);
-					}
+				Output.OutputReceivedDataEventHandlerEncounteredError = false;
+				Output.OutputReceivedDataEventHandlerEncounteredErrorMessage = "";
+				Utils.RunLocalProcess(SignProcess);
 
-					if (Output.OutputReceivedDataEventHandlerEncounteredError)
-					{
-						throw new Exception(Output.OutputReceivedDataEventHandlerEncounteredErrorMessage);
-					}
+				// delete the temp project
+				DirectoryReference.Delete(XcodeWorkspaceDir, true);
 
-					// Package the stub
-					PackageStub(RemoteShadowDirectoryMac, AppName, Target.OutputPath.GetFileNameWithoutExtension());
+				if (Output.OutputReceivedDataEventHandlerEncounteredError)
+				{
+					throw new Exception(Output.OutputReceivedDataEventHandlerEncounteredErrorMessage);
 				}
+
+				// Package the stub
+				PackageStub(RemoteShadowDirectoryMac, AppName, Target.OutputPath.GetFileNameWithoutExtension());
 			}
 
 			{
