UBT: Make it an error if a framework doesn't exist, rather than failing silently. Also remove some remote toolchain stuff that's no longer necessary.

diff --git a/Engine/Source/Programs/UnrealBuildTool/Platform/IOS/IOSToolChain.cs b/Engine/Source/Programs/UnrealBuildTool/Platform/IOS/IOSToolChain.cs
--- a/Engine/Source/Programs/UnrealBuildTool/Platform/IOS/IOSToolChain.cs
+++ b/Engine/Source/Programs/UnrealBuildTool/Platform/IOS/IOSToolChain.cs
@@ -369,7 +369,7 @@ namespace UnrealBuildTool
 			{
 				if (Framework.OwningModule != null && Framework.FrameworkZipPath != null && Framework.FrameworkZipPath != "")
 				{
-					Result += " -F\"" + GetRemoteIntermediateFrameworkZipPath(Framework) + "\"";
+					Result += " -F\"" + GetExtractedFrameworkDir(Framework) + "\"";
 				}
 			}
 
@@ -444,42 +444,24 @@ namespace UnrealBuildTool
 			return Result;
 		}
 
-		static string GetLocalFrameworkZipPath(UEBuildFramework Framework)
+		static FileReference GetFrameworkZip(UEBuildFramework Framework)
 		{
 			if (Framework.OwningModule == null)
 			{
 				throw new BuildException("GetLocalFrameworkZipPath: No owning module for framework {0}", Framework.FrameworkName);
 			}
 
-			return Path.GetFullPath(Framework.OwningModule.ModuleDirectory + "/" + Framework.FrameworkZipPath);
+			return FileReference.Combine(Framework.OwningModule.ModuleDirectory, Framework.FrameworkZipPath);
 		}
 
-		static string GetRemoteFrameworkZipPath(UEBuildFramework Framework)
-		{
-			if (BuildHostPlatform.Current.Platform != UnrealTargetPlatform.Mac)
-			{
-				return ConvertPath(GetLocalFrameworkZipPath(Framework));
-			}
-
-			return GetLocalFrameworkZipPath(Framework);
-		}
-
-		static string GetRemoteIntermediateFrameworkZipPath(UEBuildFramework Framework)
+		static DirectoryReference GetExtractedFrameworkDir(UEBuildFramework Framework)
 		{
 			if (Framework.OwningModule == null)
 			{
 				throw new BuildException("GetRemoteIntermediateFrameworkZipPath: No owning module for framework {0}", Framework.FrameworkName);
 			}
 
-			string IntermediatePath = UnrealBuildTool.EngineDirectory + "/Intermediate/UnzippedFrameworks/" + Framework.OwningModule.Name;
-			IntermediatePath = Path.GetFullPath((IntermediatePath + Framework.FrameworkZipPath).Replace(".zip", ""));
-
-			if (BuildHostPlatform.Current.Platform != UnrealTargetPlatform.Mac)
-			{
-				return ConvertPath(IntermediatePath);
-			}
-
-			return IntermediatePath;
+			return DirectoryReference.Combine(UnrealBuildTool.EngineDirectory, "Intermediate", "UnzippedFrameworks", Framework.OwningModule.Name, Framework.FrameworkZipPath.Replace(".zip", ""));
 		}
 
 		static void CleanIntermediateDirectory(string Path)
@@ -605,7 +587,7 @@ namespace UnrealBuildTool
 				if (Framework.OwningModule != null && Framework.FrameworkZipPath != null && Framework.FrameworkZipPath != "")
 				{
 					// If this framework has a zip specified, we'll need to setup the path as well
-					Result += " -F\"" + GetRemoteIntermediateFrameworkZipPath(Framework) + "\"";
+					Result += " -F\"" + GetExtractedFrameworkDir(Framework) + "\"";
 				}
 
 				Result += " -framework " + Framework.FrameworkName;
@@ -815,22 +797,6 @@ namespace UnrealBuildTool
 
 				// Remember any files we need to unzip
 				RememberedAdditionalFrameworks.Add(Framework);
-
-				// Copy them to remote mac if needed
-				if (BuildHostPlatform.Current.Platform != UnrealTargetPlatform.Mac)
-				{
-					FileItem ShadowFile = FileItem.GetExistingItemByPath(GetLocalFrameworkZipPath(Framework));
-
-					if (ShadowFile != null)
-					{
-						QueueFileForBatchUpload(ShadowFile);
-						DependentAction.PrerequisiteItems.Add(ShadowFile);
-					}
-					else
-					{
-						throw new BuildException("Couldn't find required additional file to shadow: {0}", Framework.FrameworkZipPath);
-					}
-				}
 			}
 		}
 
@@ -877,24 +843,6 @@ namespace UnrealBuildTool
 				}
 			}
 
-			if (BuildHostPlatform.Current.Platform != UnrealTargetPlatform.Mac)
-			{
-				// Add any additional files that we'll need in order to link the app
-				foreach (string AdditionalShadowFile in LinkEnvironment.AdditionalShadowFiles)
-				{
-					FileItem ShadowFile = FileItem.GetExistingItemByPath(AdditionalShadowFile);
-					if (ShadowFile != null)
-					{
-						QueueFileForBatchUpload(ShadowFile);
-						LinkAction.PrerequisiteItems.Add(ShadowFile);
-					}
-					else
-					{
-						throw new BuildException("Couldn't find required additional file to shadow: {0}", AdditionalShadowFile);
-					}
-				}
-			}
-
 			// Handle additional framework assets that might need to be shadowed
 			AddFrameworksToTrack(LinkEnvironment.AdditionalFrameworks, LinkAction);
 
@@ -1245,19 +1193,26 @@ namespace UnrealBuildTool
 			// Unzip any third party frameworks that are stored as zips
 			foreach (UEBuildFramework Framework in RememberedAdditionalFrameworks)
 			{
-				string ZipSrcPath = GetRemoteFrameworkZipPath(Framework);
-				string ZipDstPath = GetRemoteIntermediateFrameworkZipPath(Framework);
+				FileReference ZipSrcPath = GetFrameworkZip(Framework);
+				if(!FileReference.Exists(ZipSrcPath))
+				{
+					throw new BuildException("Unable to find framework '{0}'", ZipSrcPath);
+				}
 
-				Log.TraceInformation("Unzipping: {0} -> {1}", ZipSrcPath, ZipDstPath);
+				DirectoryReference FrameworkDstPath = GetExtractedFrameworkDir(Framework);
+				Log.TraceInformation("Unzipping: {0} -> {1}", ZipSrcPath, FrameworkDstPath);
 
-				CleanIntermediateDirectory(ZipDstPath);
+				CleanIntermediateDirectory(FrameworkDstPath.FullName);
 
 				// Assume that there is another directory inside the zip with the same name as the zip
-				ZipDstPath = ZipDstPath.Substring(0, ZipDstPath.LastIndexOf('/'));
+				DirectoryReference ZipDstPath = FrameworkDstPath.ParentDirectory;
 
 				// If we're on the mac, just unzip using the shell
 				string ResultsText;
-				RunExecutableAndWait("unzip", String.Format("-o \"{0}\" -d \"{1}\"", ZipSrcPath, ZipDstPath), out ResultsText);
+				if(RunExecutableAndWait("unzip", String.Format("-o \"{0}\" -d \"{1}\"", ZipSrcPath, ZipDstPath), out ResultsText) != 0)
+				{
+					throw new BuildException("Unable to extract {0}:\n{1}", ZipSrcPath, ResultsText);
+				}
 			}
         }
 
@@ -1807,7 +1762,7 @@ namespace UnrealBuildTool
 						continue;		// Only care if we need to copy bundle assets
 					}
 
-					string UnpackedZipPath = GetRemoteIntermediateFrameworkZipPath(Framework);
+					string UnpackedZipPath = GetExtractedFrameworkDir(Framework).FullName;
 
 					// For now, this is hard coded, but we need to loop over all modules, and copy bundled assets that need it
 					string LocalSource = UnpackedZipPath + "/" + Framework.CopyBundledAssets;
diff --git a/Engine/Source/Programs/UnrealBuildTool/Platform/Mac/MacToolChain.cs b/Engine/Source/Programs/UnrealBuildTool/Platform/Mac/MacToolChain.cs
--- a/Engine/Source/Programs/UnrealBuildTool/Platform/Mac/MacToolChain.cs
+++ b/Engine/Source/Programs/UnrealBuildTool/Platform/Mac/MacToolChain.cs
@@ -729,38 +729,6 @@ namespace UnrealBuildTool
 
 			List<string> ThirdPartyLibraries = new List<string>();
 
-			if (BuildHostPlatform.Current.Platform != UnrealTargetPlatform.Mac)
-			{
-				// Add any additional files that we'll need in order to link the app
-				foreach (string AdditionalShadowFile in LinkEnvironment.AdditionalShadowFiles)
-				{
-					FileItem ShadowFile = FileItem.GetExistingItemByPath(AdditionalShadowFile);
-					if (ShadowFile != null)
-					{
-						QueueFileForBatchUpload(ShadowFile);
-						LinkAction.PrerequisiteItems.Add(ShadowFile);
-					}
-					else
-					{
-						throw new BuildException("Couldn't find required additional file to shadow: {0}", AdditionalShadowFile);
-					}
-				}
-
-				// Add any frameworks to be shadowed to the remote
-				foreach (string FrameworkPath in LinkEnvironment.Frameworks)
-				{
-					if (FrameworkPath.EndsWith(".framework"))
-					{
-						foreach (string FrameworkFile in Directory.EnumerateFiles(FrameworkPath, "*", SearchOption.AllDirectories))
-						{
-							FileItem FrameworkFileItem = FileItem.GetExistingItemByPath(FrameworkFile);
-							QueueFileForBatchUpload(FrameworkFileItem);
-							LinkAction.PrerequisiteItems.Add(FrameworkFileItem);
-						}
-					}
-				}
-			}
-
 			bool bIsBuildingAppBundle = !LinkEnvironment.bIsBuildingDLL && !LinkEnvironment.bIsBuildingLibrary && !LinkEnvironment.bIsBuildingConsoleApplication;
 
 			List<string> RPaths = new List<string>();
@@ -812,16 +780,6 @@ namespace UnrealBuildTool
 						{
 							ThirdPartyLibraries.Add(AdditionalLibrary);
 						}
-
-						if (BuildHostPlatform.Current.Platform != UnrealTargetPlatform.Mac)
-						{
-							// copy over libs we may need
-							FileItem ShadowFile = FileItem.GetExistingItemByPath(AdditionalLibrary);
-							if (ShadowFile != null)
-							{
-								QueueFileForBatchUpload(ShadowFile);
-							}
-						}
 					}
 
 					AddLibraryPathToRPaths(AdditionalLibrary, AbsolutePath, ref RPaths, ref LinkCommand, bIsBuildingAppBundle);
@@ -838,16 +796,6 @@ namespace UnrealBuildTool
 					LinkCommand += string.Format(" -weak_library \"{0}\"", ConvertPath(Path.GetFullPath(AdditionalLibrary)));
 
 					AddLibraryPathToRPaths(AdditionalLibrary, AbsolutePath, ref RPaths, ref LinkCommand, bIsBuildingAppBundle);
-
-                    if (BuildHostPlatform.Current.Platform != UnrealTargetPlatform.Mac)
-                    {
-                        // copy over libs we may need
-                        FileItem ShadowFile = FileItem.GetExistingItemByPath(AdditionalLibrary);
-                        if (ShadowFile != null)
-                        {
-                            QueueFileForBatchUpload(ShadowFile);
-                        }
-                    }
                 }
 			}
 
@@ -1162,7 +1110,6 @@ namespace UnrealBuildTool
 
 						if (CustomIcon != DefaultIcon)
 						{
-							QueueFileForBatchUpload(FileItem.GetItemByFileReference(new FileReference(CustomIcon)));
 							CustomIcon = ConvertPath(CustomIcon);
 						}
 					}
@@ -1180,7 +1127,6 @@ namespace UnrealBuildTool
 					}
 					else
 					{
-						QueueFileForBatchUpload(FileItem.GetItemByFileReference(new FileReference(InfoPlistFile)));
 						InfoPlistFile = ConvertPath(InfoPlistFile);
 					}
 
@@ -1207,13 +1153,6 @@ namespace UnrealBuildTool
 					AppendMacLine(FinalizeAppBundleScript, "touch -c \"{0}.app\"", ExeName);
 
 					FinalizeAppBundleScript.Close();
-
-					// copy over some needed files
-					// @todo mac: Make a QueueDirectoryForBatchUpload
-					QueueFileForBatchUpload(FileItem.GetItemByFileReference(new FileReference("../../Engine/Source/Runtime/Launch/Resources/Mac/" + GameName + ".icns")));
-					QueueFileForBatchUpload(FileItem.GetItemByFileReference(new FileReference("../../Engine/Source/Runtime/Launch/Resources/Mac/UProject.icns")));
-					QueueFileForBatchUpload(FileItem.GetItemByFileReference(new FileReference("../../Engine/Source/Runtime/Launch/Resources/Mac/Info.plist")));
-					QueueFileForBatchUpload(FileItem.GetItemByFileReference(FileReference.Combine(LinkEnvironment.IntermediateDirectory, "DylibCopy.sh")));
 				}
 			}
 
@@ -1373,18 +1312,6 @@ namespace UnrealBuildTool
 			CopyAction.StatusDescription = string.Format("Copying {0} to app bundle", Path.GetFileName(Resource.ResourcePath));
 			CopyAction.bCanExecuteRemotely = false;
 
-			if (Directory.Exists(Resource.ResourcePath))
-			{
-				foreach (string ResourceFile in Directory.GetFiles(Resource.ResourcePath, "*", SearchOption.AllDirectories))
-				{
-					QueueFileForBatchUpload(FileItem.GetItemByFileReference(new FileReference(ResourceFile)));
-				}
-			}
-			else
-			{
-				QueueFileForBatchUpload(FileItem.GetItemByFileReference(new FileReference(SourcePath)));
-			}
-
 			return TargetItem;
 		}
 
diff --git a/Engine/Source/Programs/UnrealBuildTool/ToolChain/AppleToolChain.cs b/Engine/Source/Programs/UnrealBuildTool/ToolChain/AppleToolChain.cs
--- a/Engine/Source/Programs/UnrealBuildTool/ToolChain/AppleToolChain.cs
+++ b/Engine/Source/Programs/UnrealBuildTool/ToolChain/AppleToolChain.cs
@@ -155,10 +155,6 @@ namespace UnrealBuildTool
 			return Item;
 		}
 
-		protected static void QueueFileForBatchUpload(FileItem Item)
-		{
-		}
-
 		protected string GetMacDevSrcRoot()
 		{
 			return UnrealBuildTool.EngineSourceDirectory.FullName;
