UBT: Remove some leftover functions for handling the remote toolchain.

#rnx

diff --git a/Engine/Source/Programs/UnrealBuildTool/Platform/IOS/IOSToolChain.cs b/Engine/Source/Programs/UnrealBuildTool/Platform/IOS/IOSToolChain.cs
--- a/Engine/Source/Programs/UnrealBuildTool/Platform/IOS/IOSToolChain.cs
+++ b/Engine/Source/Programs/UnrealBuildTool/Platform/IOS/IOSToolChain.cs
@@ -147,7 +147,7 @@ namespace UnrealBuildTool
 					FileReference AssetFile = FileReference.Combine(Binary.OutputFilePath.Directory, "AssetCatalog", "Assets.car");
                     BuildProducts.Add(AssetFile, BuildProductType.RequiredResource);
                 }
-				else if (CppPlatform == CppPlatform.IOS && Settings.Value.IOSSDKVersionFloat >= 11.0f && BuildHostPlatform.Current.Platform == UnrealTargetPlatform.Mac)
+				else if (CppPlatform == CppPlatform.IOS && Settings.Value.IOSSDKVersionFloat >= 11.0f)
 				{
 					int Index = Binary.OutputFilePath.GetFileNameWithoutExtension().IndexOf("-");
 					string OutputFile = Binary.OutputFilePath.GetFileNameWithoutExtension().Substring(0, Index > 0 ? Index : Binary.OutputFilePath.GetFileNameWithoutExtension().Length);
@@ -541,10 +541,9 @@ namespace UnrealBuildTool
 			if(IsBitcodeCompilingEnabled(LinkEnvironment.Configuration))
 			{
 				FileItem OutputFile = FileItem.GetItemByFileReference(LinkEnvironment.OutputFilePath);
-				FileItem RemoteOutputFile = LocalToRemoteFileItem(OutputFile, false);
 
 				Result += " -fembed-bitcode -Xlinker -bitcode_verify -Xlinker -bitcode_hide_symbols -Xlinker -bitcode_symbol_map ";
-				Result += " -Xlinker " + Path.GetDirectoryName(RemoteOutputFile.AbsolutePath);
+				Result += " -Xlinker " + Path.GetDirectoryName(OutputFile.AbsolutePath);
 			}
 
 			Result += " -dead_strip";
@@ -616,12 +615,12 @@ namespace UnrealBuildTool
 
 			if (CompileEnvironment.PrecompiledHeaderAction == PrecompiledHeaderAction.Include)
 			{
-				PCHArguments += string.Format(" -include \"{0}\"", ConvertPath(CompileEnvironment.PrecompiledHeaderIncludeFilename.FullName));
+				PCHArguments += string.Format(" -include \"{0}\"", CompileEnvironment.PrecompiledHeaderIncludeFilename.FullName);
 			}
 
 			foreach(FileItem ForceIncludeFile in CompileEnvironment.ForceIncludeFiles)
 			{
-				Arguments += String.Format(" -include \"{0}\"", ConvertPath(ForceIncludeFile.Location.FullName));
+				Arguments += String.Format(" -include \"{0}\"", ForceIncludeFile.Location.FullName);
 			}
 
 			// Add include paths to the argument list.
@@ -629,28 +628,7 @@ namespace UnrealBuildTool
 			AllIncludes.UnionWith(CompileEnvironment.IncludePaths.SystemIncludePaths);
 			foreach (DirectoryReference IncludePath in AllIncludes)
 			{
-				Arguments += string.Format(" -I\"{0}\"", ConvertPath(IncludePath.FullName));
-
-				if (BuildHostPlatform.Current.Platform != UnrealTargetPlatform.Mac)
-				{
-					// sync any third party headers we may need
-					if (IncludePath.FullName.Contains("ThirdParty") && DirectoryReference.Exists(IncludePath))
-					{
-						string[] FileList = Directory.GetFiles(IncludePath.FullName, "*.h", SearchOption.AllDirectories);
-						foreach (string File in FileList)
-						{
-							FileItem ExternalDependency = FileItem.GetItemByPath(File);
-							LocalToRemoteFileItem(ExternalDependency, true);
-						}
-
-						FileList = Directory.GetFiles(IncludePath.FullName, "*.cpp", SearchOption.AllDirectories);
-						foreach (string File in FileList)
-						{
-							FileItem ExternalDependency = FileItem.GetItemByPath(File);
-							LocalToRemoteFileItem(ExternalDependency, true);
-						}
-					}
-				}
+				Arguments += string.Format(" -I\"{0}\"", IncludePath.FullName);
 			}
 
 			foreach (string Definition in CompileEnvironment.Definitions)
@@ -714,12 +692,11 @@ namespace UnrealBuildTool
 					// Add the precompiled header file to the produced item list.
 					FileItem PrecompiledHeaderFile = FileItem.GetItemByFileReference(FileReference.Combine(OutputDir, Path.GetFileName(SourceFile.AbsolutePath) + ".gch"));
 
-					FileItem RemotePrecompiledHeaderFile = LocalToRemoteFileItem(PrecompiledHeaderFile, false);
-					CompileAction.ProducedItems.Add(RemotePrecompiledHeaderFile);
-					Result.PrecompiledHeaderFile = RemotePrecompiledHeaderFile;
+					CompileAction.ProducedItems.Add(PrecompiledHeaderFile);
+					Result.PrecompiledHeaderFile = PrecompiledHeaderFile;
 
 					// Add the parameters needed to compile the precompiled header file to the command-line.
-					FileArguments += string.Format(" -o \"{0}\"", RemotePrecompiledHeaderFile.AbsolutePath);
+					FileArguments += string.Format(" -o \"{0}\"", PrecompiledHeaderFile.AbsolutePath);
 				}
 				else
 				{
@@ -731,15 +708,14 @@ namespace UnrealBuildTool
 					// Add the object file to the produced item list.
 					FileItem ObjectFile = FileItem.GetItemByFileReference(FileReference.Combine(OutputDir, Path.GetFileName(SourceFile.AbsolutePath) + ".o"));
 
-					FileItem RemoteObjectFile = LocalToRemoteFileItem(ObjectFile, false);
-					CompileAction.ProducedItems.Add(RemoteObjectFile);
-					Result.ObjectFiles.Add(RemoteObjectFile);
-					FileArguments += string.Format(" -o \"{0}\"", RemoteObjectFile.AbsolutePath);
-					OutputFilePath = RemoteObjectFile.AbsolutePath;
+					CompileAction.ProducedItems.Add(ObjectFile);
+					Result.ObjectFiles.Add(ObjectFile);
+					FileArguments += string.Format(" -o \"{0}\"", ObjectFile.AbsolutePath);
+					OutputFilePath = ObjectFile.AbsolutePath;
 				}
 
 				// Add the source file path to the command-line.
-				FileArguments += string.Format(" \"{0}\"", ConvertPath(SourceFile.AbsolutePath));
+				FileArguments += string.Format(" \"{0}\"", SourceFile.AbsolutePath);
 
 				string CompilerPath = Settings.Value.ToolchainDir + IOSCompiler;
 
@@ -819,7 +795,7 @@ namespace UnrealBuildTool
 				// Add the library paths to the argument list.
 				foreach (DirectoryReference LibraryPath in LinkEnvironment.LibraryPaths)
 				{
-					LinkCommandArguments += string.Format(" -L\"{0}\"", ConvertPath(LibraryPath.FullName));
+					LinkCommandArguments += string.Format(" -L\"{0}\"", LibraryPath.FullName);
 				}
 
 				// Add the additional libraries to the argument list.
@@ -830,11 +806,10 @@ namespace UnrealBuildTool
 					{
 						// add it to the prerequisites to make sure it's built first (this should be the case of non-system libraries)
 						FileItem LibFile = FileItem.GetItemByPath(AdditionalLibrary);
-						FileItem RemoteLibFile = LocalToRemoteFileItem(LibFile, true);
-						LinkAction.PrerequisiteItems.Add(RemoteLibFile);
+						LinkAction.PrerequisiteItems.Add(LibFile);
 
 						// and add to the commandline
-						LinkCommandArguments += string.Format(" \"{0}\"", ConvertPath(Path.GetFullPath(AdditionalLibrary)));
+						LinkCommandArguments += string.Format(" \"{0}\"", Path.GetFullPath(AdditionalLibrary));
 					}
 					else
 					{
@@ -848,8 +823,7 @@ namespace UnrealBuildTool
 
 			// Add the output file as a production of the link action.
 			FileItem OutputFile = FileItem.GetItemByFileReference(LinkEnvironment.OutputFilePath);
-			FileItem RemoteOutputFile = LocalToRemoteFileItem(OutputFile, false);
-			LinkAction.ProducedItems.Add(RemoteOutputFile);
+			LinkAction.ProducedItems.Add(OutputFile);
 
 			// Add arguments to generate a map file too
 			if (!LinkEnvironment.bIsBuildingLibrary && LinkEnvironment.bCreateMapFile)
@@ -863,7 +837,7 @@ namespace UnrealBuildTool
 			List<string> InputFileNames = new List<string>();
 			foreach (FileItem InputFile in LinkEnvironment.InputFiles)
 			{
-				string FileNamePath = string.Format("\"{0}\"", ConvertPath(InputFile.AbsolutePath));
+				string FileNamePath = string.Format("\"{0}\"", InputFile.AbsolutePath);
 				InputFileNames.Add(FileNamePath);
 				LinkAction.PrerequisiteItems.Add(InputFile);
 			}
@@ -884,12 +858,12 @@ namespace UnrealBuildTool
 			{
 				bool bIsUE4Game = LinkEnvironment.OutputFilePath.FullName.Contains("UE4Game");
 				FileReference ResponsePath = FileReference.Combine(((!bIsUE4Game && ProjectFile != null) ? ProjectFile.Directory : UnrealBuildTool.EngineDirectory), "Intermediate", "Build", LinkEnvironment.Platform.ToString(), "LinkFileList_" + LinkEnvironment.OutputFilePath.GetFileNameWithoutExtension() + ".tmp");
-				FileItem.CreateIntermediateTextFile(new FileReference(ConvertPath(ResponsePath.FullName)), InputFileNames);
-				LinkCommandArguments += string.Format(" @\"{0}\"", ConvertPath(ResponsePath.FullName));
+				FileItem.CreateIntermediateTextFile(ResponsePath, InputFileNames);
+				LinkCommandArguments += string.Format(" @\"{0}\"", ResponsePath.FullName);
 			}
 
 			// Add the output file to the command-line.
-			LinkCommandArguments += string.Format(" -o \"{0}\"", RemoteOutputFile.AbsolutePath);
+			LinkCommandArguments += string.Format(" -o \"{0}\"", OutputFile.AbsolutePath);
 
 			// Add the additional arguments specified by the environment.
 			LinkCommandArguments += LinkEnvironment.AdditionalArguments;
@@ -900,7 +874,7 @@ namespace UnrealBuildTool
 			LinkAction.StatusDescription = string.Format("{0}", OutputFile.AbsolutePath);
 
 			LinkAction.CommandPath = "sh";
-			if(LinkEnvironment.Configuration == CppConfiguration.Shipping && Path.GetExtension(RemoteOutputFile.AbsolutePath) != ".a")
+			if(LinkEnvironment.Configuration == CppConfiguration.Shipping && Path.GetExtension(OutputFile.AbsolutePath) != ".a")
 			{
 				// When building a shipping package, symbols are stripped from the exe as the last build step. This is a problem
 				// when re-packaging and no source files change because the linker skips symbol generation and dsymutil will 
@@ -908,8 +882,8 @@ namespace UnrealBuildTool
 				// the output file to force the linker to recreate it with symbols again.
 				string	linkCommandArguments = "-c '";
 
-				linkCommandArguments += string.Format("rm -f \"{0}\";", RemoteOutputFile.AbsolutePath);
-				linkCommandArguments += string.Format("rm -f \"{0}\\*.bcsymbolmap\";", Path.GetDirectoryName(RemoteOutputFile.AbsolutePath));
+				linkCommandArguments += string.Format("rm -f \"{0}\";", OutputFile.AbsolutePath);
+				linkCommandArguments += string.Format("rm -f \"{0}\\*.bcsymbolmap\";", Path.GetDirectoryName(OutputFile.AbsolutePath));
 				linkCommandArguments += LinkerPath + " " + LinkCommandArguments + ";";
 
 				linkCommandArguments += "'";
@@ -922,7 +896,7 @@ namespace UnrealBuildTool
 				LinkAction.CommandArguments = string.Format("-c '{0} {1}'", LinkerPath, LinkCommandArguments);
 			}
 
-			return RemoteOutputFile;
+			return OutputFile;
 		}
 
 		static string GetAppBundleName(FileReference Executable)
@@ -993,15 +967,10 @@ namespace UnrealBuildTool
         public FileItem GenerateDebugInfo(FileItem Executable, ActionGraph ActionGraph)
 		{
             // Make a file item for the source and destination files
-            FileItem LocalExecutable = RemoteToLocalFileItem(Executable);
-			string FullDestPathRoot = Path.Combine(Path.GetDirectoryName(LocalExecutable.AbsolutePath), Path.GetFileName(LocalExecutable.AbsolutePath) + ".dSYM");
+			string FullDestPathRoot = Path.Combine(Path.GetDirectoryName(Executable.AbsolutePath), Path.GetFileName(Executable.AbsolutePath) + ".dSYM");
 
-            FileItem OutputFile;
-            OutputFile = FileItem.GetItemByPath(FullDestPathRoot);
-            FileItem ZipOutputFile;
-            ZipOutputFile = FileItem.GetItemByPath(FullDestPathRoot + ".zip");
-            FileItem DestFile = LocalToRemoteFileItem(OutputFile, false);
-            FileItem ZipDestFile = LocalToRemoteFileItem(ZipOutputFile, false);
+            FileItem OutputFile = FileItem.GetItemByPath(FullDestPathRoot);
+            FileItem ZipOutputFile = FileItem.GetItemByPath(FullDestPathRoot + ".zip");
 
             // Make the compile action
             Action GenDebugAction = ActionGraph.Add(ActionType.GenerateDebugInfo);
@@ -1012,23 +981,23 @@ namespace UnrealBuildTool
 			{
 				GenDebugAction.CommandArguments = string.Format("-c 'rm -rf \"{1}\"; /usr/bin/dsymutil \"{0}\" -o \"{1}\"; cd \"{1}/..\"; zip -r -y -1 {2}.zip {2}'",
 					Executable.AbsolutePath,
-                    DestFile.AbsolutePath,
+                    OutputFile.AbsolutePath,
 					Path.GetFileName(FullDestPathRoot));
-                GenDebugAction.ProducedItems.Add(ZipDestFile);
-                Log.TraceInformation("Zip file: " + ZipDestFile.AbsolutePath);
+                GenDebugAction.ProducedItems.Add(ZipOutputFile);
+                Log.TraceInformation("Zip file: " + ZipOutputFile.AbsolutePath);
             }
             else
 			{
 				GenDebugAction.CommandArguments = string.Format("-c 'rm -rf \"{1}\"; /usr/bin/dsymutil \"{0}\" -f -o \"{1}\"'",
 						Executable.AbsolutePath,
-						DestFile.AbsolutePath);
+						OutputFile.AbsolutePath);
 			}
 			GenDebugAction.PrerequisiteItems.Add(Executable);
-            GenDebugAction.ProducedItems.Add(DestFile);
+            GenDebugAction.ProducedItems.Add(OutputFile);
             GenDebugAction.StatusDescription = GenDebugAction.CommandArguments;// string.Format("Generating debug info for {0}", Path.GetFileName(Executable.AbsolutePath));
 			GenDebugAction.bCanExecuteRemotely = false;
 
-			return (ProjectSettings.bGeneratedSYMBundle ? ZipDestFile : DestFile);
+			return (ProjectSettings.bGeneratedSYMBundle ? ZipOutputFile : OutputFile);
 		}
 
         /// <summary>
@@ -1039,22 +1008,15 @@ namespace UnrealBuildTool
         public FileItem GeneratePseudoPDB(FileItem Executable, ActionGraph ActionGraph)
         {
             // Make a file item for the source and destination files
-            FileItem LocalExecutable = RemoteToLocalFileItem(Executable);
-            string FullDestPathRoot = Path.Combine(Path.GetDirectoryName(LocalExecutable.AbsolutePath), Path.GetFileName(LocalExecutable.AbsolutePath) + ".udebugsymbols");
-            string FulldSYMPathRoot = Path.Combine(Path.GetDirectoryName(LocalExecutable.AbsolutePath), Path.GetFileName(LocalExecutable.AbsolutePath) + ".dSYM");
-            string PathToDWARF = Path.Combine(FulldSYMPathRoot, "Contents", "Resources", "DWARF", Path.GetFileName(LocalExecutable.AbsolutePath));
+            string FullDestPathRoot = Path.Combine(Path.GetDirectoryName(Executable.AbsolutePath), Path.GetFileName(Executable.AbsolutePath) + ".udebugsymbols");
+            string FulldSYMPathRoot = Path.Combine(Path.GetDirectoryName(Executable.AbsolutePath), Path.GetFileName(Executable.AbsolutePath) + ".dSYM");
+            string PathToDWARF = Path.Combine(FulldSYMPathRoot, "Contents", "Resources", "DWARF", Path.GetFileName(Executable.AbsolutePath));
 
-            FileItem dSYMFile;
-            dSYMFile = FileItem.GetItemByPath(FulldSYMPathRoot);
-            FileItem dSYMOutFile = LocalToRemoteFileItem(dSYMFile, false);
+            FileItem dSYMFile = FileItem.GetItemByPath(FulldSYMPathRoot);
 
-            FileItem DWARFFile;
-            DWARFFile = FileItem.GetItemByPath(PathToDWARF);
-            FileItem DWARFOutFile = LocalToRemoteFileItem(DWARFFile, false);
+            FileItem DWARFFile = FileItem.GetItemByPath(PathToDWARF);
 
-            FileItem OutputFile;
-            OutputFile = FileItem.GetItemByPath(FullDestPathRoot);
-            FileItem DestFile = LocalToRemoteFileItem(OutputFile, false);
+            FileItem OutputFile = FileItem.GetItemByPath(FullDestPathRoot);
 
             // Make the compile action
             Action GenDebugAction = ActionGraph.Add(ActionType.GenerateDebugInfo);
@@ -1062,16 +1024,16 @@ namespace UnrealBuildTool
 
             GenDebugAction.CommandPath = "sh";
             GenDebugAction.CommandArguments = string.Format("-c 'rm -rf \"{1}\"; dwarfdump --uuid \"{3}\" | cut -d\" \" -f2; chmod 777 ./DsymExporter; ./DsymExporter -UUID=$(dwarfdump --uuid \"{3}\" | cut -d\" \" -f2) \"{0}\" \"{2}\"'",
-                    DWARFOutFile.AbsolutePath,
-                    DestFile.AbsolutePath,
-                    Path.GetDirectoryName(DestFile.AbsolutePath),
-                    dSYMOutFile.AbsolutePath);
-            GenDebugAction.PrerequisiteItems.Add(dSYMOutFile);
-            GenDebugAction.ProducedItems.Add(DestFile);
+                    DWARFFile.AbsolutePath,
+                    OutputFile.AbsolutePath,
+                    Path.GetDirectoryName(OutputFile.AbsolutePath),
+                    dSYMFile.AbsolutePath);
+            GenDebugAction.PrerequisiteItems.Add(dSYMFile);
+            GenDebugAction.ProducedItems.Add(OutputFile);
             GenDebugAction.StatusDescription = GenDebugAction.CommandArguments;// string.Format("Generating debug info for {0}", Path.GetFileName(Executable.AbsolutePath));
             GenDebugAction.bCanExecuteRemotely = false;
 
-            return DestFile;
+            return OutputFile;
         }
 
         private static void PackageStub(string BinaryPath, string GameName, string ExeName)
@@ -1294,7 +1256,6 @@ namespace UnrealBuildTool
                         Directory.CreateDirectory(Dir);
                     }
                     File.WriteAllText(Path.Combine(Dir, "Contents.json"), Contents[i]);
-                    LocalToRemoteFileItem(FileItem.GetItemByPath(Path.Combine(Dir, "Contents.json")), BuildHostPlatform.Current.Platform != UnrealTargetPlatform.Mac);
                     if (Images[i] != null)
                     {
                         // This assumption might not be true, but we need the asset catalog process to fire anyway.
@@ -1304,7 +1265,6 @@ namespace UnrealBuildTool
                         if (File.Exists(Image))
                         {
                             File.Copy(Image, Path.Combine(Dir, Images[i]), true);
-                            LocalToRemoteFileItem(FileItem.GetItemByPath(Path.Combine(Dir, Images[i])), BuildHostPlatform.Current.Platform != UnrealTargetPlatform.Mac);
                             FileInfo DestFileInfo = new FileInfo(Path.Combine(Dir, Images[i]));
                             DestFileInfo.Attributes = DestFileInfo.Attributes & ~FileAttributes.ReadOnly;
                         }
@@ -1333,7 +1293,6 @@ namespace UnrealBuildTool
                 {
                     Dir = file.Replace(Path.Combine(EngineDir, "Build", "IOS"), IntermediateDir);
                     File.Copy(file, Dir, true);
-                    LocalToRemoteFileItem(FileItem.GetItemByPath(Dir), BuildHostPlatform.Current.Platform != UnrealTargetPlatform.Mac);
                     FileInfo DestFileInfo = new FileInfo(Dir);
                     DestFileInfo.Attributes = DestFileInfo.Attributes & ~FileAttributes.ReadOnly;
                 }
@@ -1377,7 +1336,6 @@ namespace UnrealBuildTool
                         bUserImagesExist |= Image.StartsWith(BuildResourcesGraphicsDir);
 
                         File.Copy(Image, Path.Combine(Dir, Images[Index][0]), true);
-                        LocalToRemoteFileItem(FileItem.GetItemByPath(Path.Combine(Dir, Images[Index][0])), BuildHostPlatform.Current.Platform != UnrealTargetPlatform.Mac);
                         FileInfo DestFileInfo = new FileInfo(Path.Combine(Dir, Images[Index][0]));
                         DestFileInfo.Attributes = DestFileInfo.Attributes & ~FileAttributes.ReadOnly;
                     }
diff --git a/Engine/Source/Programs/UnrealBuildTool/Platform/Mac/MacToolChain.cs b/Engine/Source/Programs/UnrealBuildTool/Platform/Mac/MacToolChain.cs
--- a/Engine/Source/Programs/UnrealBuildTool/Platform/Mac/MacToolChain.cs
+++ b/Engine/Source/Programs/UnrealBuildTool/Platform/Mac/MacToolChain.cs
@@ -167,11 +167,6 @@ namespace UnrealBuildTool
 			Result += " -fexceptions";
 			Result += " -fasm-blocks";
 
-			if (!Utils.IsRunningOnMono)
-			{
-				Result += " -fdiagnostics-format=msvc";
-			}
-
 			if(CompileEnvironment.bHideSymbolsByDefault)
 			{
 				Result += " -fvisibility=hidden";
@@ -359,7 +354,7 @@ namespace UnrealBuildTool
 			string Result = "";
 			if (FrameworkName.EndsWith(".framework"))
 			{
-				Result += " -F \"" + ConvertPath(Path.GetDirectoryName(Path.GetFullPath(FrameworkName))) + "\"";
+				Result += " -F \"" + Path.GetDirectoryName(Path.GetFullPath(FrameworkName)) + "\"";
 				FrameworkName = Path.GetFileNameWithoutExtension(FrameworkName);
 			}
 			Result += " " + Arg + " \"" + FrameworkName + "\"";
@@ -445,34 +440,7 @@ namespace UnrealBuildTool
 			foreach (DirectoryReference IncludePath in AllIncludes)
 			{
 				Arguments.Append(" -I\"");
-
-				if (BuildHostPlatform.Current.Platform != UnrealTargetPlatform.Mac)
-				{
-					Arguments.Append(ConvertPath(IncludePath.FullName));
-
-					// sync any third party headers we may need
-					if (IncludePath.FullName.Contains("ThirdParty"))
-					{
-						string[] FileList = Directory.GetFiles(IncludePath.FullName, "*.h", SearchOption.AllDirectories);
-						foreach (string File in FileList)
-						{
-							FileItem ExternalDependency = FileItem.GetItemByPath(File);
-							LocalToRemoteFileItem(ExternalDependency, true);
-						}
-
-						FileList = Directory.GetFiles(IncludePath.FullName, "*.cpp", SearchOption.AllDirectories);
-						foreach (string File in FileList)
-						{
-							FileItem ExternalDependency = FileItem.GetItemByPath(File);
-							LocalToRemoteFileItem(ExternalDependency, true);
-						}
-					}
-				}
-				else
-				{
-					Arguments.Append(IncludePath);
-				}
-
+				Arguments.Append(IncludePath);
 				Arguments.Append("\"");
 			}
 
@@ -534,13 +502,11 @@ namespace UnrealBuildTool
 				{
 					// Add the precompiled header file to the produced item list.
 					FileItem PrecompiledHeaderFile = FileItem.GetItemByFileReference(FileReference.Combine(OutputDir, Path.GetFileName(SourceFile.AbsolutePath) + ".gch"));
-
-					FileItem RemotePrecompiledHeaderFile = LocalToRemoteFileItem(PrecompiledHeaderFile, false);
-					CompileAction.ProducedItems.Add(RemotePrecompiledHeaderFile);
-					Result.PrecompiledHeaderFile = RemotePrecompiledHeaderFile;
+					CompileAction.ProducedItems.Add(PrecompiledHeaderFile);
+					Result.PrecompiledHeaderFile = PrecompiledHeaderFile;
 
 					// Add the parameters needed to compile the precompiled header file to the command-line.
-					FileArguments += string.Format(" -o \"{0}\"", RemotePrecompiledHeaderFile.AbsolutePath);
+					FileArguments += string.Format(" -o \"{0}\"", PrecompiledHeaderFile.AbsolutePath);
 				}
 				else
 				{
@@ -552,15 +518,14 @@ namespace UnrealBuildTool
 					// Add the object file to the produced item list.
 					FileItem ObjectFile = FileItem.GetItemByFileReference(FileReference.Combine(OutputDir, Path.GetFileName(SourceFile.AbsolutePath) + ".o"));
 
-					FileItem RemoteObjectFile = LocalToRemoteFileItem(ObjectFile, false);
-					CompileAction.ProducedItems.Add(RemoteObjectFile);
-					Result.ObjectFiles.Add(RemoteObjectFile);
-					FileArguments += string.Format(" -o \"{0}\"", RemoteObjectFile.AbsolutePath);
-					OutputFilePath = RemoteObjectFile.AbsolutePath;
+					CompileAction.ProducedItems.Add(ObjectFile);
+					Result.ObjectFiles.Add(ObjectFile);
+					FileArguments += string.Format(" -o \"{0}\"", ObjectFile.AbsolutePath);
+					OutputFilePath = ObjectFile.AbsolutePath;
 				}
 
 				// Add the source file path to the command-line.
-				FileArguments += string.Format(" \"{0}\"", ConvertPath(SourceFile.AbsolutePath));
+				FileArguments += string.Format(" \"{0}\"", SourceFile.AbsolutePath);
 
 				string AllArgs = Arguments + FileArguments + CompileEnvironment.AdditionalArguments;
 				string CompilerPath = Settings.ToolchainDir + MacCompiler;
@@ -712,8 +677,6 @@ namespace UnrealBuildTool
 			FileItem OutputFile = FileItem.GetItemByFileReference(LinkEnvironment.OutputFilePath);
 			OutputFile.bNeedsHotReloadNumbersDLLCleanUp = LinkEnvironment.bIsBuildingDLL;
 
-			FileItem RemoteOutputFile = LocalToRemoteFileItem(OutputFile, false);
-
 			// To solve the problem with cross dependencies, for now we create a broken dylib that does not link with other engine dylibs.
 			// This is fixed in later step, FixDylibDependencies. For this and to know what libraries to copy whilst creating an app bundle,
 			// we gather the list of engine dylibs.
@@ -749,17 +712,9 @@ namespace UnrealBuildTool
 							 Path.GetDirectoryName(AdditionalLibrary).Contains("Binaries\\Mac")))
 					{
 						// It's an engine or game dylib. Save it for later
-						EngineAndGameLibraries.Add(ConvertPath(Path.GetFullPath(AdditionalLibrary)));
+						EngineAndGameLibraries.Add(Path.GetFullPath(AdditionalLibrary));
 
-						if (!Utils.IsRunningOnMono)
-						{
-							FileItem EngineLibDependency = FileItem.GetItemByPath(AdditionalLibrary);
-							LinkAction.PrerequisiteItems.Add(EngineLibDependency);
-							FileItem RemoteEngineLibDependency = FileItem.GetRemoteItemByPath(ConvertPath(Path.GetFullPath(AdditionalLibrary)), UnrealTargetPlatform.Mac);
-							LinkAction.PrerequisiteItems.Add(RemoteEngineLibDependency);
-							//Log.TraceInformation("Adding {0} / {1} as a prereq to {2}", EngineLibDependency.AbsolutePath, RemoteEngineLibDependency.AbsolutePath, RemoteOutputFile.AbsolutePath);
-						}
-						else if (LinkEnvironment.bIsCrossReferenced == false)
+						if (LinkEnvironment.bIsCrossReferenced == false)
 						{
 							FileItem EngineLibDependency = FileItem.GetItemByPath(AdditionalLibrary);
 							LinkAction.PrerequisiteItems.Add(EngineLibDependency);
@@ -775,7 +730,7 @@ namespace UnrealBuildTool
 					}
 					else
 					{
-						LinkCommand += string.Format(" \"{0}\"", ConvertPath(Path.GetFullPath(AdditionalLibrary)));
+						LinkCommand += string.Format(" \"{0}\"", Path.GetFullPath(AdditionalLibrary));
 						if (Path.GetExtension(AdditionalLibrary) == ".dylib")
 						{
 							ThirdPartyLibraries.Add(AdditionalLibrary);
@@ -793,7 +748,7 @@ namespace UnrealBuildTool
 						continue;
 					}
 
-					LinkCommand += string.Format(" -weak_library \"{0}\"", ConvertPath(Path.GetFullPath(AdditionalLibrary)));
+					LinkCommand += string.Format(" -weak_library \"{0}\"", Path.GetFullPath(AdditionalLibrary));
 
 					AddLibraryPathToRPaths(AdditionalLibrary, AbsolutePath, ref RPaths, ref LinkCommand, bIsBuildingAppBundle);
                 }
@@ -914,7 +869,7 @@ namespace UnrealBuildTool
 			}
 
 			// Add the output file to the command-line.
-			LinkCommand += string.Format(" -o \"{0}\"", RemoteOutputFile.AbsolutePath);
+			LinkCommand += string.Format(" -o \"{0}\"", OutputFile.AbsolutePath);
 
 			// Add the additional arguments specified by the environment.
 			LinkCommand += LinkEnvironment.AdditionalArguments;
@@ -925,7 +880,7 @@ namespace UnrealBuildTool
 				foreach (string Library in ThirdPartyLibraries)
 				{
 					string LibraryFileName = Path.GetFileName(Library);
-					LinkCommand += "; " + Settings.ToolchainDir + "install_name_tool -change " + LibraryFileName + " " + DylibsPath + "/" + LibraryFileName + " \"" + ConvertPath(OutputFile.AbsolutePath) + "\"";
+					LinkCommand += "; " + Settings.ToolchainDir + "install_name_tool -change " + LibraryFileName + " " + DylibsPath + "/" + LibraryFileName + " \"" + OutputFile.AbsolutePath + "\"";
 				}
 			}
 
@@ -936,7 +891,7 @@ namespace UnrealBuildTool
 
 			LinkAction.StatusDescription = Path.GetFileName(OutputFile.AbsolutePath);
 
-			LinkAction.ProducedItems.Add(RemoteOutputFile);
+			LinkAction.ProducedItems.Add(OutputFile);
 
 			if (!DirectoryReference.Exists(LinkEnvironment.IntermediateDirectory))
 			{
@@ -970,10 +925,10 @@ namespace UnrealBuildTool
 					{
 						EngineAndGameLibrariesString += " \"" + Library + "\"";
 					}
-					string FixDylibLine = "pushd \"" + ConvertPath(Directory.GetCurrentDirectory()) + "\"  > /dev/null; ";
-					FixDylibLine += string.Format("TIMESTAMP=`stat -n -f \"%Sm\" -t \"%Y%m%d%H%M.%S\" \"{0}\"`; ", RemoteOutputFile.AbsolutePath);
+					string FixDylibLine = "pushd \"" + Directory.GetCurrentDirectory() + "\"  > /dev/null; ";
+					FixDylibLine += string.Format("TIMESTAMP=`stat -n -f \"%Sm\" -t \"%Y%m%d%H%M.%S\" \"{0}\"`; ", OutputFile.AbsolutePath);
 					FixDylibLine += LinkCommand.Replace("-undefined dynamic_lookup", EngineAndGameLibrariesString).Replace("$", "\\$");
-					FixDylibLine += string.Format("; touch -t $TIMESTAMP \"{0}\"; if [[ $? -ne 0 ]]; then exit 1; fi; ", RemoteOutputFile.AbsolutePath);
+					FixDylibLine += string.Format("; touch -t $TIMESTAMP \"{0}\"; if [[ $? -ne 0 ]]; then exit 1; fi; ", OutputFile.AbsolutePath);
 					FixDylibLine += "popd > /dev/null";
 					AppendMacLine(FixDylibDepsScript, FixDylibLine);
 				}
@@ -999,7 +954,7 @@ namespace UnrealBuildTool
 				StreamWriter DylibCopyScript = File.AppendText(DylibCopyScriptPath.FullName);
 				foreach (string Library in ThirdPartyLibraries)
 				{
-					string CopyCommandLineEntry = FormatCopyCommand(ConvertPath(Path.GetFullPath(Library)).Replace("$", "\\$"), "$1.app/Contents/MacOS");
+					string CopyCommandLineEntry = FormatCopyCommand(Path.GetFullPath(Library).Replace("$", "\\$"), "$1.app/Contents/MacOS");
 					if (!ExistingScript.Contains(CopyCommandLineEntry))
 					{
 						AppendMacLine(DylibCopyScript, CopyCommandLineEntry);
@@ -1015,7 +970,7 @@ namespace UnrealBuildTool
 					AppendMacLine(FinalizeAppBundleScript, "#!/bin/sh");
 					string BinariesPath = Path.GetDirectoryName(OutputFile.AbsolutePath);
 					BinariesPath = Path.GetDirectoryName(BinariesPath.Substring(0, BinariesPath.IndexOf(".app")));
-					AppendMacLine(FinalizeAppBundleScript, "cd \"{0}\"", ConvertPath(BinariesPath).Replace("$", "\\$"));
+					AppendMacLine(FinalizeAppBundleScript, "cd \"{0}\"", BinariesPath.Replace("$", "\\$"));
 
 					string BundleVersion = LinkEnvironment.BundleVersion;
 					if(BundleVersion == null)
@@ -1057,10 +1012,10 @@ namespace UnrealBuildTool
 					AppendMacLine(FinalizeAppBundleScript, "mkdir -p \"{0}.app/Contents/Resources\"", ExeName);
 
 					// Copy third party dylibs by calling additional script prepared earlier
-					AppendMacLine(FinalizeAppBundleScript, "sh \"{0}\" \"{1}\"", ConvertPath(DylibCopyScriptPath.FullName).Replace("$", "\\$"), ExeName);
+					AppendMacLine(FinalizeAppBundleScript, "sh \"{0}\" \"{1}\"", DylibCopyScriptPath.FullName.Replace("$", "\\$"), ExeName);
 
 					string IconName = "UE4";
-					string EngineSourcePath = ConvertPath(Directory.GetCurrentDirectory()).Replace("$", "\\$");
+					string EngineSourcePath = Directory.GetCurrentDirectory().Replace("$", "\\$");
 					string CustomResourcesPath = "";
 					string CustomBuildPath = "";
 					if (UProjectFilePath == null)
@@ -1107,11 +1062,6 @@ namespace UnrealBuildTool
 								CustomIcon = DefaultIcon;
 							}
 						}
-
-						if (CustomIcon != DefaultIcon)
-						{
-							CustomIcon = ConvertPath(CustomIcon);
-						}
 					}
 					AppendMacLine(FinalizeAppBundleScript, FormatCopyCommand(CustomIcon, String.Format("{0}.app/Contents/Resources/{1}.icns", ExeName, GameName)));
 
@@ -1125,10 +1075,6 @@ namespace UnrealBuildTool
 					{
 						InfoPlistFile = EngineSourcePath + "/Runtime/Launch/Resources/Mac/" + (bBuildingEditor ? "Info-Editor.plist" : "Info.plist");
 					}
-					else
-					{
-						InfoPlistFile = ConvertPath(InfoPlistFile);
-					}
 
 					string TempInfoPlist = "$TMPDIR/TempInfo.plist";
 					AppendMacLine(FinalizeAppBundleScript, FormatCopyCommand(InfoPlistFile, TempInfoPlist));
@@ -1156,7 +1102,7 @@ namespace UnrealBuildTool
 				}
 			}
 
-			return RemoteOutputFile;
+			return OutputFile;
 		}
 
 		static string FormatCopyCommand(string SourceFile, string TargetFile)
@@ -1175,9 +1121,8 @@ namespace UnrealBuildTool
 			// once all are already created, so the cross dependency problem no longer prevents linking.
 			// The script is deleted after it's executed so it's empty when we start appending link commands for the next executable.
 			FileItem FixDylibDepsScript = FileItem.GetItemByFileReference(FileReference.Combine(LinkEnvironment.LocalShadowDirectory, "FixDylibDependencies.sh"));
-			FileItem RemoteFixDylibDepsScript = LocalToRemoteFileItem(FixDylibDepsScript, true);
 
-			LinkAction.CommandArguments = "-c 'chmod +x \"" + RemoteFixDylibDepsScript.AbsolutePath + "\"; \"" + RemoteFixDylibDepsScript.AbsolutePath + "\"; if [[ $? -ne 0 ]]; then exit 1; fi; ";
+			LinkAction.CommandArguments = "-c 'chmod +x \"" + FixDylibDepsScript.AbsolutePath + "\"; \"" + FixDylibDepsScript.AbsolutePath + "\"; if [[ $? -ne 0 ]]; then exit 1; fi; ";
 
 			// Make sure this action is executed after all the dylibs and the main executable are created
 
@@ -1192,14 +1137,13 @@ namespace UnrealBuildTool
 			LinkAction.bCanExecuteRemotely = false;
 
 			FileItem OutputFile = FileItem.GetItemByFileReference(FileReference.Combine(LinkEnvironment.LocalShadowDirectory, Path.GetFileNameWithoutExtension(Executable.AbsolutePath) + ".link"));
-			FileItem RemoteOutputFile = LocalToRemoteFileItem(OutputFile, false);
 
-			LinkAction.CommandArguments += "echo \"Dummy\" >> \"" + RemoteOutputFile.AbsolutePath + "\"";
+			LinkAction.CommandArguments += "echo \"Dummy\" >> \"" + OutputFile.AbsolutePath + "\"";
 			LinkAction.CommandArguments += "'";
 
-			LinkAction.ProducedItems.Add(RemoteOutputFile);
+			LinkAction.ProducedItems.Add(OutputFile);
 
-			return RemoteOutputFile;
+			return OutputFile;
 		}
 
 		/// <summary>
@@ -1225,8 +1169,6 @@ namespace UnrealBuildTool
 			}
 
 			FileItem OutputFile = FileItem.GetItemByPath(BinaryPath);
-			FileItem DestFile = LocalToRemoteFileItem(OutputFile, false);
-			FileItem InputFile = LocalToRemoteFileItem(MachOBinary, false);
 
 			// Delete on the local machine
 			if (Directory.Exists(OutputFile.AbsolutePath))
@@ -1244,19 +1186,19 @@ namespace UnrealBuildTool
 			// Note that the source and dest are switched from a copy command
 			GenDebugAction.CommandArguments = string.Format("-c 'rm -rf \"{2}\"; for i in {{1..30}}; do if [ -f \"{1}\" ] ; then break; else echo\"Waiting for {1} before generating dSYM file.\"; sleep 1; fi; done; \"{0}\"dsymutil -f \"{1}\" -o \"{2}\"'",
 				Settings.ToolchainDir,
-				InputFile.AbsolutePath,
-				DestFile.AbsolutePath);
+				MachOBinary.AbsolutePath,
+				OutputFile.AbsolutePath);
 			if (LinkEnvironment.bIsCrossReferenced)
 			{
 				GenDebugAction.PrerequisiteItems.Add(FixDylibOutputFile);
 			}
-			GenDebugAction.PrerequisiteItems.Add(LocalToRemoteFileItem(InputFile, false));
-			GenDebugAction.ProducedItems.Add(DestFile);
+			GenDebugAction.PrerequisiteItems.Add(MachOBinary);
+			GenDebugAction.ProducedItems.Add(OutputFile);
 			GenDebugAction.CommandDescription = "";
 			GenDebugAction.StatusDescription = "Generating " + Path.GetFileName(BinaryPath);
 			GenDebugAction.bCanExecuteRemotely = false;
 
-			return DestFile;
+			return OutputFile;
 		}
 
 		/// <summary>
@@ -1271,7 +1213,6 @@ namespace UnrealBuildTool
 			// Make a file item for the source and destination files
 			string FullDestPath = Executable.AbsolutePath.Substring(0, Executable.AbsolutePath.IndexOf(".app") + 4);
 			FileItem DestFile = FileItem.GetItemByPath(FullDestPath);
-			FileItem RemoteDestFile = LocalToRemoteFileItem(DestFile, false);
 
 			// Make the compile action
 			Action FinalizeAppBundleAction = ActionGraph.Add(ActionType.CreateAppBundle);
@@ -1281,15 +1222,14 @@ namespace UnrealBuildTool
 
 			// make path to the script
 			FileItem BundleScript = FileItem.GetItemByFileReference(FileReference.Combine(LinkEnvironment.IntermediateDirectory, "FinalizeAppBundle.sh"));
-			FileItem RemoteBundleScript = LocalToRemoteFileItem(BundleScript, true);
 
-			FinalizeAppBundleAction.CommandArguments = "\"" + RemoteBundleScript.AbsolutePath + "\"";
+			FinalizeAppBundleAction.CommandArguments = "\"" + BundleScript.AbsolutePath + "\"";
 			FinalizeAppBundleAction.PrerequisiteItems.Add(FixDylibOutputFile);
-			FinalizeAppBundleAction.ProducedItems.Add(RemoteDestFile);
+			FinalizeAppBundleAction.ProducedItems.Add(DestFile);
 			FinalizeAppBundleAction.StatusDescription = string.Format("Finalizing app bundle: {0}.app", Path.GetFileName(Executable.AbsolutePath));
 			FinalizeAppBundleAction.bCanExecuteRemotely = false;
 
-			return RemoteDestFile;
+			return DestFile;
 		}
 
 		FileItem CopyBundleResource(UEBuildBundleResource Resource, FileItem Executable, DirectoryReference BundleDirectory, ActionGraph ActionGraph)
@@ -1305,7 +1245,7 @@ namespace UnrealBuildTool
 
 			FileItem TargetItem = FileItem.GetItemByPath(TargetPath);
 
-			CopyAction.CommandArguments = string.Format("-c 'cp -f -R \"{0}\" \"{1}\"; touch -c \"{2}\"'", ConvertPath(SourcePath), Path.GetDirectoryName(TargetPath).Replace('\\', '/') + "/", TargetPath.Replace('\\', '/'));
+			CopyAction.CommandArguments = string.Format("-c 'cp -f -R \"{0}\" \"{1}\"; touch -c \"{2}\"'", SourcePath, Path.GetDirectoryName(TargetPath).Replace('\\', '/') + "/", TargetPath.Replace('\\', '/'));
 			CopyAction.PrerequisiteItems.Add(Executable);
 			CopyAction.ProducedItems.Add(TargetItem);
 			CopyAction.bShouldOutputStatusDescription = Resource.bShouldLog;
@@ -1503,12 +1443,6 @@ namespace UnrealBuildTool
 				}
 			}
 
-			// If building for Mac on a Mac, use actions to finalize the builds (otherwise, we use Deploy)
-			if (BuildHostPlatform.Current.Platform != UnrealTargetPlatform.Mac)
-			{
-				return OutputFiles;
-			}
-
 			if (BinaryLinkEnvironment.bIsBuildingDLL)
 			{
 				return OutputFiles;
diff --git a/Engine/Source/Programs/UnrealBuildTool/System/FileItem.cs b/Engine/Source/Programs/UnrealBuildTool/System/FileItem.cs
--- a/Engine/Source/Programs/UnrealBuildTool/System/FileItem.cs
+++ b/Engine/Source/Programs/UnrealBuildTool/System/FileItem.cs
@@ -172,27 +172,6 @@ namespace UnrealBuildTool
 			}
 		}
 
-		/// <returns>The remote FileItem that represents the given file path.</returns>
-		public static FileItem GetRemoteItemByPath(string AbsoluteRemotePath, UnrealTargetPlatform Platform)
-		{
-			if (AbsoluteRemotePath.StartsWith("."))
-			{
-				throw new BuildException("GetRemoteItemByPath must be passed an absolute path, not a relative path '{0}'", AbsoluteRemotePath);
-			}
-
-			FileReference RemoteFileReference = FileReference.MakeRemote(AbsoluteRemotePath);
-
-			FileItem Result = null;
-			if (UniqueSourceFileMap.TryGetValue(RemoteFileReference, out Result))
-			{
-				return Result;
-			}
-			else
-			{
-				return new FileItem(RemoteFileReference, true, Platform);
-			}
-		}
-
 		/// <summary>
 		/// If the given file path identifies a file that already exists, returns the FileItem that represents it.
 		/// </summary>
diff --git a/Engine/Source/Programs/UnrealBuildTool/ToolChain/AppleToolChain.cs b/Engine/Source/Programs/UnrealBuildTool/ToolChain/AppleToolChain.cs
--- a/Engine/Source/Programs/UnrealBuildTool/ToolChain/AppleToolChain.cs
+++ b/Engine/Source/Programs/UnrealBuildTool/ToolChain/AppleToolChain.cs
@@ -140,21 +140,6 @@ namespace UnrealBuildTool
 			ProjectFile = InProjectFile;
 		}
 
-		protected static string ConvertPath(string Path)
-		{
-			return Path;
-		}
-
-		protected static FileItem RemoteToLocalFileItem(FileItem Executable)
-		{
-			return Executable;
-		}
-
-		protected static FileItem LocalToRemoteFileItem(FileItem Item, bool bSomething = false)
-		{
-			return Item;
-		}
-
 		protected string GetMacDevSrcRoot()
 		{
 			return UnrealBuildTool.EngineSourceDirectory.FullName;
