UBT: Speculative fix (and better error reporting) for IOS mobile provision not being found in CIS.

#rnx

diff --git a/Engine/Source/Programs/UnrealBuildTool/Platform/IOS/UEBuildIOS.cs b/Engine/Source/Programs/UnrealBuildTool/Platform/IOS/UEBuildIOS.cs
--- a/Engine/Source/Programs/UnrealBuildTool/Platform/IOS/UEBuildIOS.cs
+++ b/Engine/Source/Programs/UnrealBuildTool/Platform/IOS/UEBuildIOS.cs
@@ -459,10 +459,17 @@ namespace UnrealBuildTool
             SigningCertificate = SigningCertificate.Replace("\"", "");
 
             // read the provision to get the UUID
-            string filename = (BuildHostPlatform.Current.Platform == UnrealTargetPlatform.Mac ? (Environment.GetEnvironmentVariable("HOME") + "/Library/MobileDevice/Provisioning Profiles/") : (Environment.GetFolderPath(Environment.SpecialFolder.LocalApplicationData) + "/Apple Computer/MobileDevice/Provisioning Profiles/")) + MobileProvision;
-            if (File.Exists(filename))
+			if(MobileProvisionFile == null)
+			{
+                Log.TraceWarning("No matching provision file was discovered. Please ensure you have a compatible provision installed.");
+			}
+			else if(!FileReference.Exists(MobileProvisionFile))
+			{
+                Log.TraceWarning("Selected mobile provision ({0}) was not found. Please ensure you have a compatible provision installed.", MobileProvisionFile);
+			}
+			else
             {
-				byte[] AllBytes = File.ReadAllBytes(filename);
+				byte[] AllBytes = FileReference.ReadAllBytes(MobileProvisionFile);
 
 				uint StartIndex = (uint)AllBytes.Length;
 				uint EndIndex = (uint)AllBytes.Length;
@@ -532,10 +539,6 @@ namespace UnrealBuildTool
 					SigningCertificate = null;
 					Log.TraceLog("Failed to parse the mobile provisioning profile.");
 				}
-            }
-            else
-            {
-                Log.TraceLog("No matching provision file was discovered. Please ensure you have a compatible provision installed.");
             }
 		}
 
