UBT: Downgrade warning to a log message, since it appears when generating project files.

#rnx

diff --git a/Engine/Source/Programs/UnrealBuildTool/Platform/IOS/UEBuildIOS.cs b/Engine/Source/Programs/UnrealBuildTool/Platform/IOS/UEBuildIOS.cs
--- a/Engine/Source/Programs/UnrealBuildTool/Platform/IOS/UEBuildIOS.cs
+++ b/Engine/Source/Programs/UnrealBuildTool/Platform/IOS/UEBuildIOS.cs
@@ -461,11 +461,11 @@ namespace UnrealBuildTool
             // read the provision to get the UUID
 			if(MobileProvisionFile == null)
 			{
-                Log.TraceWarning("No matching provision file was discovered. Please ensure you have a compatible provision installed.");
+				Log.TraceLog("No matching provision file was discovered for {0}. Please ensure you have a compatible provision installed.", ProjectFile);
 			}
 			else if(!FileReference.Exists(MobileProvisionFile))
 			{
-                Log.TraceWarning("Selected mobile provision ({0}) was not found. Please ensure you have a compatible provision installed.", MobileProvisionFile);
+				Log.TraceLog("Selected mobile provision for {0} ({1}) was not found. Please ensure you have a compatible provision installed.", ProjectFile, MobileProvisionFile);
 			}
 			else
             {
