Fix compile errors due to missing include path when hot-reloading a module from the editor. There are not necessarily source files to compile when -modulewithsuffix is specified on the command line, which was results in GeneratedCodeWildcard not being set.

#rnx

diff --git a/Engine/Source/Programs/UnrealBuildTool/Configuration/UEBuildModuleCPP.cs b/Engine/Source/Programs/UnrealBuildTool/Configuration/UEBuildModuleCPP.cs
--- a/Engine/Source/Programs/UnrealBuildTool/Configuration/UEBuildModuleCPP.cs
+++ b/Engine/Source/Programs/UnrealBuildTool/Configuration/UEBuildModuleCPP.cs
@@ -86,6 +86,11 @@ namespace UnrealBuildTool
 		/// </summary>
 		public readonly DirectoryReference GeneratedCodeDirectory;
 
+		/// <summary>
+		/// Set for modules that have generated code
+		/// </summary>
+		public bool bAddGeneratedCodeIncludePath;
+
 		/// <summary>
 		/// Wildcard matching the *.gen.cpp files for this module.  If this is null then this module doesn't have any UHT-produced code.
 		/// </summary>
@@ -315,7 +320,7 @@ namespace UnrealBuildTool
 			bool bLegacyPublicIncludePaths
 			)
 		{
-			if(GeneratedCodeWildcard != null)
+			if(bAddGeneratedCodeIncludePath)
 			{
 				IncludePaths.Add(GeneratedCodeDirectory);
 			}
diff --git a/Engine/Source/Programs/UnrealBuildTool/System/ExternalExecution.cs b/Engine/Source/Programs/UnrealBuildTool/System/ExternalExecution.cs
--- a/Engine/Source/Programs/UnrealBuildTool/System/ExternalExecution.cs
+++ b/Engine/Source/Programs/UnrealBuildTool/System/ExternalExecution.cs
@@ -613,6 +613,9 @@ namespace UnrealBuildTool
 				UHTModuleInfo Info = ExternalExecution.CreateUHTModuleInfo(HeaderFiles, Module.Name, Module.RulesFile, Module.ModuleDirectory, ModuleToType[Module], GeneratedCodeVersion, Module.Rules.bUsePrecompiled);
 				if (Info.PublicUObjectClassesHeaders.Count > 0 || Info.PrivateUObjectHeaders.Count > 0 || Info.PublicUObjectHeaders.Count > 0)
 				{
+					// Set a flag indicating that we need to add the generated headers directory
+					Module.bAddGeneratedCodeIncludePath = true;
+
 					// If we've got this far and there are no source files then it's likely we're installed and ignoring
 					// engine files, so we don't need a .gen.cpp either
 					Info.GeneratedCPPFilenameBase = Path.Combine(Module.GeneratedCodeDirectory.FullName, Info.ModuleName) + ".gen";
