Fix incorrect executable name for DebugGame configurations in Xcode.

diff --git a/Engine/Source/Programs/UnrealBuildTool/ProjectFiles/Xcode/XcodeProject.cs b/Engine/Source/Programs/UnrealBuildTool/ProjectFiles/Xcode/XcodeProject.cs
--- a/Engine/Source/Programs/UnrealBuildTool/ProjectFiles/Xcode/XcodeProject.cs
+++ b/Engine/Source/Programs/UnrealBuildTool/ProjectFiles/Xcode/XcodeProject.cs
@@ -1109,14 +1109,14 @@ namespace UnrealBuildTool
 
                                             if (BuildPlatform.Platform == UnrealTargetPlatform.Mac)
                                             {
-                                                string MacExecutableName = MakeExecutableFileName(ExeName, UnrealTargetPlatform.Mac, (ExeName == "UE4Editor" && Configuration == UnrealTargetConfiguration.DebugGame) ? UnrealTargetConfiguration.Development : Configuration, ProjectTarget.TargetRules.Architecture, ProjectTarget.TargetRules.UndecoratedConfiguration);
+                                                string MacExecutableName = MakeExecutableFileName(ExeName, UnrealTargetPlatform.Mac, Configuration, ProjectTarget.TargetRules.Architecture, ProjectTarget.TargetRules.UndecoratedConfiguration);
                                                 string IOSExecutableName = MacExecutableName.Replace("-Mac-", "-IOS-");
                                                 string TVOSExecutableName = MacExecutableName.Replace("-Mac-", "-TVOS-");
                                                 BuildConfigs.Add(new XcodeBuildConfig(ConfigName, TargetName, FileReference.Combine(OutputDirectory, "Mac", MacExecutableName), FileReference.Combine(OutputDirectory, "IOS", IOSExecutableName), FileReference.Combine(OutputDirectory, "TVOS", TVOSExecutableName), ProjectTarget, Configuration));
                                             }
                                             else if (BuildPlatform.Platform == UnrealTargetPlatform.IOS)
                                             {
-                                                string IOSExecutableName = MakeExecutableFileName(ExeName, UnrealTargetPlatform.IOS, (ExeName == "UE4Editor" && Configuration == UnrealTargetConfiguration.DebugGame) ? UnrealTargetConfiguration.Development : Configuration, ProjectTarget.TargetRules.Architecture, ProjectTarget.TargetRules.UndecoratedConfiguration);
+                                                string IOSExecutableName = MakeExecutableFileName(ExeName, UnrealTargetPlatform.IOS, Configuration, ProjectTarget.TargetRules.Architecture, ProjectTarget.TargetRules.UndecoratedConfiguration);
                                                 string TVOSExecutableName = IOSExecutableName.Replace("-IOS-", "-TVOS-");
                                                 string MacExecutableName = IOSExecutableName.Replace("-IOS-", "-Mac-");
                                                 BuildConfigs.Add(new XcodeBuildConfig(ConfigName, TargetName, FileReference.Combine(OutputDirectory, "Mac", IOSExecutableName), FileReference.Combine(OutputDirectory, "IOS", IOSExecutableName), FileReference.Combine(OutputDirectory, "TVOS", TVOSExecutableName), ProjectTarget, Configuration));
