Fix incremental compile failures due to dependency checking for unity files. CachedIncludePaths was not correctly being set on file items, so dependencies were being ignored.

diff --git a/Engine/Source/Programs/UnrealBuildTool/System/Unity.cs b/Engine/Source/Programs/UnrealBuildTool/System/Unity.cs
--- a/Engine/Source/Programs/UnrealBuildTool/System/Unity.cs
+++ b/Engine/Source/Programs/UnrealBuildTool/System/Unity.cs
@@ -348,14 +348,9 @@ namespace UnrealBuildTool
 
 				// Write the unity file to the intermediate folder.
 				FileItem UnityCPPFile = FileItem.CreateIntermediateTextFile(UnityCPPFilePath, OutputUnityCPPWriter.ToString());
-
+				UnityCPPFile.CachedIncludePaths = CompileEnvironment.IncludePaths;
 				UnityCPPFile.RelativeCost = UnityFile.TotalLength;
 				NewCPPFiles.Add(UnityCPPFile);
-
-				// Cache information about the unity .cpp dependencies
-				// @todo ubtmake urgent: Fails when building remotely for Mac because unity .cpp has an include for a PCH on the REMOTE machine
-				FileItem FirstCppFile = UnityFile.Files.First();
-				UnityCPPFile.CachedIncludePaths = FirstCppFile.CachedIncludePaths;
 			}
 
 			return NewCPPFiles;
