- Fixed the CopyAction for runtime dependencies that need to be copied to different location, on non-XGE

diff --git a/Engine/Source/Programs/UnrealBuildTool/Configuration/UEBuildBinary.cs b/Engine/Source/Programs/UnrealBuildTool/Configuration/UEBuildBinary.cs
--- a/Engine/Source/Programs/UnrealBuildTool/Configuration/UEBuildBinary.cs
+++ b/Engine/Source/Programs/UnrealBuildTool/Configuration/UEBuildBinary.cs
@@ -316,6 +316,7 @@ namespace UnrealBuildTool
 				CopyAction.CommandPath = "/bin/sh";
 				CopyAction.CommandArguments = String.Format("cp -f {0} {1}", Utils.EscapeShellArgument(SourceFile.FullName), Utils.EscapeShellArgument(TargetFile.FullName));
 			}
+			CopyAction.WorkingDirectory = Environment.CurrentDirectory;
 			CopyAction.PrerequisiteItems.Add(SourceFileItem);
 			CopyAction.ProducedItems.Add(TargetFileItem);
 			CopyAction.DeleteItems.Add(TargetFileItem);
