Remove override for the __IPHONE_OS_VERSION_MIN_REQUIRED macro on TVOS.

This macro is already defined by system headers (in <AvailabilityInternal.h>). Now that we support PCHs on IOS and TVOS, manually defining this macro results in it being defined three times (once for the PCH, once by AvailabilityInternal.h, and once by the force-included list of definitions for the source file being built). The errors for redefining the macro in AvailabilityInternal.h are suppressed due to it being a system header, but the error for redefining it for the source file being compiled are not.

diff --git a/Engine/Source/Programs/UnrealBuildTool/Platform/TVOS/UEBuildTVOS.cs b/Engine/Source/Programs/UnrealBuildTool/Platform/TVOS/UEBuildTVOS.cs
--- a/Engine/Source/Programs/UnrealBuildTool/Platform/TVOS/UEBuildTVOS.cs
+++ b/Engine/Source/Programs/UnrealBuildTool/Platform/TVOS/UEBuildTVOS.cs
@@ -116,7 +116,6 @@ namespace UnrealBuildTool
 		{
 			base.SetUpEnvironment(Target, CompileEnvironment, LinkEnvironment);
 			CompileEnvironment.Definitions.Add("PLATFORM_TVOS=1");
-			CompileEnvironment.Definitions.Add("__IPHONE_OS_VERSION_MIN_REQUIRED=__APPLETV_OS_VERSION_MIN_REQUIRED");
 		}
 
 		/// <summary>
