Fixes incremental build failure when compile arguments for PCH have changed on IOS/TVOS. Compile action needs to have a dependency on PCH build action.

#rnx

diff --git a/Engine/Source/Programs/UnrealBuildTool/Platform/IOS/IOSToolChain.cs b/Engine/Source/Programs/UnrealBuildTool/Platform/IOS/IOSToolChain.cs
--- a/Engine/Source/Programs/UnrealBuildTool/Platform/IOS/IOSToolChain.cs
+++ b/Engine/Source/Programs/UnrealBuildTool/Platform/IOS/IOSToolChain.cs
@@ -680,6 +680,12 @@ namespace UnrealBuildTool
 				// Add the C++ source file and its included files to the prerequisite item list.
 				AddPrerequisiteSourceFile(CompileEnvironment, SourceFile, CompileAction.PrerequisiteItems);
 
+				// Add the precompiled header
+				if (CompileEnvironment.PrecompiledHeaderAction == PrecompiledHeaderAction.Include)
+				{
+					CompileAction.PrerequisiteItems.Add(CompileEnvironment.PrecompiledHeaderFile);
+				}
+
 				// Upload the force included files
 				foreach(FileItem ForceIncludeFile in CompileEnvironment.ForceIncludeFiles)
 				{
