Fix IOS provisioning data being using when remote compiling on TVOS.

diff --git a/Engine/Source/Programs/UnrealBuildTool/Platform/IOS/IOSToolChain.cs b/Engine/Source/Programs/UnrealBuildTool/Platform/IOS/IOSToolChain.cs
--- a/Engine/Source/Programs/UnrealBuildTool/Platform/IOS/IOSToolChain.cs
+++ b/Engine/Source/Programs/UnrealBuildTool/Platform/IOS/IOSToolChain.cs
@@ -1583,7 +1583,7 @@ namespace UnrealBuildTool
 					if(Target.Rules.IOSPlatform.ImportCertificate == null)
 					{
 						// Take it from the standard settings
-						IOSProvisioningData ProvisioningData = IOSPlatform.ReadProvisioningData(Target.ProjectFile);
+						IOSProvisioningData ProvisioningData = ((IOSPlatform)UEBuildPlatform.GetBuildPlatform(Target.Platform)).ReadProvisioningData(Target.ProjectFile);
 						SigningCertificate = ProvisioningData.SigningCertificate;
 
 						// Set the identity on the command line
diff --git a/Engine/Source/Programs/UnrealBuildTool/Platform/IOS/UEBuildIOS.cs b/Engine/Source/Programs/UnrealBuildTool/Platform/IOS/UEBuildIOS.cs
--- a/Engine/Source/Programs/UnrealBuildTool/Platform/IOS/UEBuildIOS.cs
+++ b/Engine/Source/Programs/UnrealBuildTool/Platform/IOS/UEBuildIOS.cs
@@ -691,12 +691,10 @@ namespace UnrealBuildTool
 			return new IOSProjectSettings(ProjectFile);
 		}
 
-		public static IOSProvisioningData ReadProvisioningData(FileReference ProjectFile, bool bForDistribution = false)
+		public IOSProvisioningData ReadProvisioningData(FileReference ProjectFile, bool bForDistribution = false)
 		{
-			IOSPlatform Platform = (IOSPlatform)UEBuildPlatform.GetBuildPlatform(UnrealTargetPlatform.IOS);
-
-			IOSProjectSettings ProjectSettings = Platform.ReadProjectSettings(ProjectFile);
-			return Platform.ReadProvisioningData(ProjectSettings, bForDistribution);
+			IOSProjectSettings ProjectSettings = ReadProjectSettings(ProjectFile);
+			return ReadProvisioningData(ProjectSettings, bForDistribution);
 		}
 
 		public IOSProvisioningData ReadProvisioningData(IOSProjectSettings ProjectSettings, bool bForDistribution = false)
diff --git a/Engine/Source/Programs/UnrealBuildTool/ToolChain/RemoteMac.cs b/Engine/Source/Programs/UnrealBuildTool/ToolChain/RemoteMac.cs
--- a/Engine/Source/Programs/UnrealBuildTool/ToolChain/RemoteMac.cs
+++ b/Engine/Source/Programs/UnrealBuildTool/ToolChain/RemoteMac.cs
@@ -380,7 +380,7 @@ namespace UnrealBuildTool
 				RemoteArguments.Add("-CreateStub");
 
 				// Get the provisioning data for this project
-				IOSProvisioningData ProvisioningData = IOSPlatform.ReadProvisioningData(TargetDesc.ProjectFile);
+				IOSProvisioningData ProvisioningData = ((IOSPlatform)UEBuildPlatform.GetBuildPlatform(TargetDesc.Platform)).ReadProvisioningData(TargetDesc.ProjectFile);
 				if(ProvisioningData == null || ProvisioningData.MobileProvisionFile == null)
 				{
 					throw new BuildException("Unable to find mobile provision for {0}. See log for more information.", TargetDesc.Name);
