Fix Slate dependencies not being enumerated, and rules assembly not being rebuilt when building remotely.

diff --git a/Engine/Build/Rsync/RsyncEngine.txt b/Engine/Build/Rsync/RsyncEngine.txt
--- a/Engine/Build/Rsync/RsyncEngine.txt
+++ b/Engine/Build/Rsync/RsyncEngine.txt
@@ -223,6 +223,11 @@ protect /Binaries/ThirdParty/Mono/Mac/lib/mono/monodoc/monodoc.dll
 + /Build/IOS/**
 + /Config/
 + /Config/**
++ /Content/
++ /Content/Slate/
++ /Content/Slate/**
++ /Content/SlateDebug/
++ /Content/SlateDebug/**
 + /Intermediate/
 + /Intermediate/Build/
 + /Intermediate/Build/XmlConfigCache.bin
diff --git a/Engine/Build/Rsync/RsyncProject.txt b/Engine/Build/Rsync/RsyncProject.txt
--- a/Engine/Build/Rsync/RsyncProject.txt
+++ b/Engine/Build/Rsync/RsyncProject.txt
@@ -18,6 +18,11 @@ protect /Intermediate/**
 + /Build/TVOS/**
 + /Config/
 + /Config/**
++ /Content/
++ /Content/Slate/
++ /Content/Slate/**
++ /Content/SlateDebug/
++ /Content/SlateDebug/**
 + /Plugins/
 + /Plugins/**/
 + /Plugins/**/*.uplugin
diff --git a/Engine/Source/Programs/UnrealBuildTool/Configuration/TargetDescriptor.cs b/Engine/Source/Programs/UnrealBuildTool/Configuration/TargetDescriptor.cs
--- a/Engine/Source/Programs/UnrealBuildTool/Configuration/TargetDescriptor.cs
+++ b/Engine/Source/Programs/UnrealBuildTool/Configuration/TargetDescriptor.cs
@@ -191,6 +191,7 @@ namespace UnrealBuildTool
 					Log.TraceVerbose("Found project file for {0} - {1}", TargetName, ProjectFile);
 				}
 
+				// Create the target descriptor
 				TargetDescriptor Target = new TargetDescriptor(ProjectFile, TargetName, Platform, Configuration, Architecture, AdditionalArguments.ToArray());
 				Target.OnlyModules = OnlyModules;
 				Target.ForeignPlugin = ForeignPlugin;
diff --git a/Engine/Source/Programs/UnrealBuildTool/Configuration/UEBuildTarget.cs b/Engine/Source/Programs/UnrealBuildTool/Configuration/UEBuildTarget.cs
--- a/Engine/Source/Programs/UnrealBuildTool/Configuration/UEBuildTarget.cs
+++ b/Engine/Source/Programs/UnrealBuildTool/Configuration/UEBuildTarget.cs
@@ -388,30 +388,7 @@ namespace UnrealBuildTool
 		{
 			DateTime CreateTargetStartTime = DateTime.UtcNow;
 
-			RulesAssembly RulesAssembly;
-			if (Desc.ProjectFile != null)
-			{
-				RulesAssembly = RulesCompiler.CreateProjectRulesAssembly(Desc.ProjectFile, bUsePrecompiled, bSkipRulesCompile);
-			}
-			else
-			{
-				RulesAssembly = RulesCompiler.CreateEngineRulesAssembly(bUsePrecompiled, bSkipRulesCompile);
-
-				if (RulesAssembly.GetTargetFileName(Desc.Name) == null)
-				{
-					// Target isn't part of the engine assembly, try the enterprise assembly
-					RulesAssembly EnterpriseRulesAssembly = RulesCompiler.CreateEnterpriseRulesAssembly(bUsePrecompiled, bSkipRulesCompile);
-
-					if (EnterpriseRulesAssembly != null)
-					{
-						RulesAssembly = EnterpriseRulesAssembly;
-					}
-				}
-			}
-			if (Desc.ForeignPlugin != null)
-			{
-				RulesAssembly = RulesCompiler.CreatePluginRulesAssembly(Desc.ForeignPlugin, bSkipRulesCompile, RulesAssembly, true);
-			}
+			RulesAssembly RulesAssembly = RulesCompiler.CreateTargetRulesAssembly(Desc.ProjectFile, Desc.Name, bSkipRulesCompile, bUsePrecompiled, Desc.ForeignPlugin);
 
 			FileReference TargetFileName;
 			TargetRules RulesObject = RulesAssembly.CreateTargetRules(Desc.Name, Desc.Platform, Desc.Configuration, Desc.Architecture, Desc.ProjectFile, Version, Arguments, out TargetFileName);
diff --git a/Engine/Source/Programs/UnrealBuildTool/System/RulesCompiler.cs b/Engine/Source/Programs/UnrealBuildTool/System/RulesCompiler.cs
--- a/Engine/Source/Programs/UnrealBuildTool/System/RulesCompiler.cs
+++ b/Engine/Source/Programs/UnrealBuildTool/System/RulesCompiler.cs
@@ -1,4 +1,4 @@
-﻿// Copyright 1998-2018 Epic Games, Inc. All Rights Reserved.
+// Copyright 1998-2018 Epic Games, Inc. All Rights Reserved.
 
 using System;
 using System.IO;
@@ -433,6 +433,43 @@ namespace UnrealBuildTool
 			return PluginRulesAssembly;
 		}
 
+		/// <summary>
+		/// Compile a rules assembly for the current target
+		/// </summary>
+		/// <param name="ProjectFile">The project file being compiled</param>
+		/// <param name="TargetName">The target being built</param>
+		/// <param name="bSkipRulesCompile">Whether to skip compiling any rules assemblies</param>
+		/// <param name="bUsePrecompiled">Whether to use a precompiled engine/enterprise build</param>
+		/// <param name="ForeignPlugin">Foreign plugin to be compiled</param>
+		/// <returns>The compiled rules assembly</returns>
+		public static RulesAssembly CreateTargetRulesAssembly(FileReference ProjectFile, string TargetName, bool bSkipRulesCompile, bool bUsePrecompiled, FileReference ForeignPlugin)
+		{
+			RulesAssembly RulesAssembly;
+			if (ProjectFile != null)
+			{
+				RulesAssembly = CreateProjectRulesAssembly(ProjectFile, bUsePrecompiled, bSkipRulesCompile);
+			}
+			else
+			{
+				RulesAssembly = CreateEngineRulesAssembly(bUsePrecompiled, bSkipRulesCompile);
+
+				if (RulesAssembly.GetTargetFileName(TargetName) == null)
+				{
+					// Target isn't part of the engine assembly, try the enterprise assembly
+					RulesAssembly EnterpriseRulesAssembly = CreateEnterpriseRulesAssembly(bUsePrecompiled, bSkipRulesCompile);
+					if (EnterpriseRulesAssembly != null)
+					{
+						RulesAssembly = EnterpriseRulesAssembly;
+					}
+				}
+			}
+			if (ForeignPlugin != null)
+			{
+				RulesAssembly = CreatePluginRulesAssembly(ForeignPlugin, bSkipRulesCompile, RulesAssembly, true);
+			}
+			return RulesAssembly;
+		}
+
 		/// <summary>
 		/// Finds all the module rules for plugins under the given directory.
 		/// </summary>
diff --git a/Engine/Source/Programs/UnrealBuildTool/ToolChain/RemoteMac.cs b/Engine/Source/Programs/UnrealBuildTool/ToolChain/RemoteMac.cs
--- a/Engine/Source/Programs/UnrealBuildTool/ToolChain/RemoteMac.cs
+++ b/Engine/Source/Programs/UnrealBuildTool/ToolChain/RemoteMac.cs
@@ -322,6 +322,9 @@ namespace UnrealBuildTool
 			DirectoryReference TempDir = DirectoryReference.Combine(BaseDir, "Intermediate", "Remote", TargetDesc.Name, TargetDesc.Platform.ToString(), TargetDesc.Configuration.ToString());
 			DirectoryReference.CreateDirectory(TempDir);
 
+			// Compile the rules assembly
+			RulesCompiler.CreateTargetRulesAssembly(TargetDesc.ProjectFile, TargetDesc.Name, false, false, TargetDesc.ForeignPlugin);
+
 			// Path to the local manifest file. This has to be translated from the remote format after the build is complete.
 			List<FileReference> LocalManifestFiles = new List<FileReference>();
 
