Always parse the team UUID out of the mobile provision when doing a remote compile. The provision installed on the remote Mac (and selected for signing) may be different.

diff --git a/Engine/Source/Programs/UnrealBuildTool/Platform/IOS/IOSToolChain.cs b/Engine/Source/Programs/UnrealBuildTool/Platform/IOS/IOSToolChain.cs
--- a/Engine/Source/Programs/UnrealBuildTool/Platform/IOS/IOSToolChain.cs
+++ b/Engine/Source/Programs/UnrealBuildTool/Platform/IOS/IOSToolChain.cs
@@ -1641,8 +1641,7 @@ namespace UnrealBuildTool
 
 						MobileProvisionContents MobileProvision = MobileProvisionContents.Read(MobileProvisionFile);
 						MobileProvisionUUID = MobileProvision.GetUniqueId();
-
-						TeamUUID = null;
+						MobileProvision.TryGetTeamUniqueId(out TeamUUID);
 					}
 
 					string ConfigName = Target.Configuration.ToString();
diff --git a/Engine/Source/Programs/UnrealBuildTool/Platform/IOS/MobileProvisionContents.cs b/Engine/Source/Programs/UnrealBuildTool/Platform/IOS/MobileProvisionContents.cs
--- a/Engine/Source/Programs/UnrealBuildTool/Platform/IOS/MobileProvisionContents.cs
+++ b/Engine/Source/Programs/UnrealBuildTool/Platform/IOS/MobileProvisionContents.cs
@@ -61,6 +61,31 @@ namespace UnrealBuildTool
 			return UniqueIdElement.InnerText;
 		}
 
+		/// <summary>
+		/// Gets the team unique id for this mobileprovision
+		/// </summary>
+		/// <param name="UniqueId">Receives the team unique id</param>
+		/// <returns>True if the team unique ID was found, false otherwise</returns>
+		public bool TryGetTeamUniqueId(out string UniqueId)
+		{
+			XmlElement UniqueIdElement;
+			if(!NameToValue.TryGetValue("TeamIdentifier", out UniqueIdElement) || UniqueIdElement.Name != "array")
+			{
+				UniqueId = null;
+				return false;
+			}
+
+			XmlElement ValueElement = UniqueIdElement.SelectSingleNode("string") as XmlElement;
+			if(ValueElement == null)
+			{
+				UniqueId = null;
+				return false;
+			}
+
+			UniqueId = ValueElement.InnerText;
+			return true;
+		}
+
 		/// <summary>
 		/// Reads a mobileprovision from a file on disk
 		/// </summary>
diff --git a/Engine/Source/Programs/UnrealBuildTool/ToolChain/RemoteMac.cs b/Engine/Source/Programs/UnrealBuildTool/ToolChain/RemoteMac.cs
--- a/Engine/Source/Programs/UnrealBuildTool/ToolChain/RemoteMac.cs
+++ b/Engine/Source/Programs/UnrealBuildTool/ToolChain/RemoteMac.cs
@@ -341,6 +341,11 @@ namespace UnrealBuildTool
 			RemoteArguments.Add(String.Format("-Log={0}", GetRemotePath(RemoteLogFile)));
 			RemoteArguments.Add(String.Format("-Manifest={0}", GetRemotePath(RemoteManifestFile)));
 
+			if(TargetDesc.ProjectFile != null)
+			{
+				RemoteArguments.Add(String.Format("-Project={0}", GetRemotePath(TargetDesc.ProjectFile)));
+			}
+
 			foreach(string LocalArgument in TargetDesc.AdditionalArguments)
 			{
 				int EqualsIdx = LocalArgument.IndexOf('=');
