Fix BuildVersion string not being passed through from <SetVersion> task. Also add a -BuildVersion command line argument to UBT to override it for a particular build.

diff --git a/Engine/Source/Programs/AutomationTool/AutomationUtils/UE4Build.cs b/Engine/Source/Programs/AutomationTool/AutomationUtils/UE4Build.cs
--- a/Engine/Source/Programs/AutomationTool/AutomationUtils/UE4Build.cs
+++ b/Engine/Source/Programs/AutomationTool/AutomationUtils/UE4Build.cs
@@ -485,6 +485,7 @@ namespace AutomationTool
 					CommandUtils.LogLog("  IsLicenseeVersion={0}", bIsLicenseeVersion? 1 : 0);
 					CommandUtils.LogLog("  IsPromotedBuild={0}", bIsPromotedBuild? 1 : 0);
 					CommandUtils.LogLog("  BranchName={0}", Branch);
+					CommandUtils.LogLog("  BuildVersion={0}", Build);
 
 					Version.Changelist = ChangelistNumber;
 					if(CompatibleChangelistNumber > 0)
@@ -494,6 +495,7 @@ namespace AutomationTool
 					Version.IsLicenseeVersion = bIsLicenseeVersion;
 					Version.IsPromotedBuild = bIsPromotedBuild;
 					Version.BranchName = Branch;
+					Version.BuildVersionString = Build;
 
 					VersionFileUpdater.MakeFileWriteable(BuildVersionFile.FullName);
 
diff --git a/Engine/Source/Programs/UnrealBuildTool/Configuration/TargetRules.cs b/Engine/Source/Programs/UnrealBuildTool/Configuration/TargetRules.cs
--- a/Engine/Source/Programs/UnrealBuildTool/Configuration/TargetRules.cs
+++ b/Engine/Source/Programs/UnrealBuildTool/Configuration/TargetRules.cs
@@ -948,6 +948,7 @@ namespace UnrealBuildTool
 		/// <summary>
 		/// The build version string
 		/// </summary>
+		[CommandLine("-BuildVersion")]
 		public string BuildVersion;
 
 		/// <summary>
@@ -1152,7 +1153,17 @@ namespace UnrealBuildTool
 			}
 
 			// Set the default build version
-			BuildVersion = String.Format("{0}-CL-{1}", Target.Version.BranchName, Target.Version.Changelist);
+			if(String.IsNullOrEmpty(BuildVersion))
+			{
+				if(String.IsNullOrEmpty(Target.Version.BuildVersionString))
+				{
+					BuildVersion = String.Format("{0}-CL-{1}", Target.Version.BranchName, Target.Version.Changelist);
+				}
+				else
+				{
+					BuildVersion = Target.Version.BuildVersionString;
+				}
+			}
 
 			// Setup macros for signing and encryption keys
 			EncryptionAndSigning.CryptoSettings CryptoSettings = EncryptionAndSigning.ParseCryptoSettings(DirectoryReference.FromFile(ProjectFile), Platform);
diff --git a/Engine/Source/Programs/UnrealBuildTool/System/BuildVersion.cs b/Engine/Source/Programs/UnrealBuildTool/System/BuildVersion.cs
--- a/Engine/Source/Programs/UnrealBuildTool/System/BuildVersion.cs
+++ b/Engine/Source/Programs/UnrealBuildTool/System/BuildVersion.cs
@@ -61,6 +61,11 @@ namespace UnrealBuildTool
 		/// </summary>
 		public string BuildId;
 
+		/// <summary>
+		/// The build version string
+		/// </summary>
+		public string BuildVersionString;
+
 		/// <summary>
 		/// Returns the value which can be used as the compatible changelist. Requires that the regular changelist is also set, and defaults to the 
 		/// regular changelist if a specific compatible changelist is not set.
@@ -169,6 +174,7 @@ namespace UnrealBuildTool
 
 			Object.TryGetStringField("BranchName", out NewVersion.BranchName);
 			Object.TryGetStringField("BuildId", out NewVersion.BuildId);
+			Object.TryGetStringField("BuildVersion", out NewVersion.BuildVersionString);
 
 			Version = NewVersion;
 			return true;
@@ -216,6 +222,7 @@ namespace UnrealBuildTool
 			Writer.WriteValue("IsPromotedBuild", IsPromotedBuild? 1 : 0);
 			Writer.WriteValue("BranchName", BranchName);
 			Writer.WriteValue("BuildId", BuildId);
+			Writer.WriteValue("BuildVersion", BuildVersionString);
 		}
 	}
 
@@ -288,6 +295,11 @@ namespace UnrealBuildTool
 			get { return Inner.BranchName; }
 		}
 
+		public string BuildVersionString
+		{
+			get { return Inner.BuildVersionString; }
+		}
+
 		#if !__MonoCS__
 		#pragma warning restore C1591
 		#endif
