Allow std::aligned_storage to build under VS [REDACTED_PHONE].

diff --git a/Engine/Source/Programs/UnrealBuildTool/Platform/Windows/VCToolChain.cs b/Engine/Source/Programs/UnrealBuildTool/Platform/Windows/VCToolChain.cs
--- a/Engine/Source/Programs/UnrealBuildTool/Platform/Windows/VCToolChain.cs
+++ b/Engine/Source/Programs/UnrealBuildTool/Platform/Windows/VCToolChain.cs
@@ -249,8 +249,11 @@ namespace UnrealBuildTool
 				AddDefinition(Arguments, "_WINDLL");
 			}
 
-            // Fix Incredibuild errors with helpers using heterogeneous character sets
-            if (Target.WindowsPlatform.Compiler >= WindowsCompiler.VisualStudio2015)
+			// Maintain the old std::aligned_storage behavior from VS from v15.8 onwards, in case of prebuilt third party libraries are reliant on it
+			AddDefinition(Arguments, "_DISABLE_EXTENDED_ALIGNED_STORAGE");
+
+			// Fix Incredibuild errors with helpers using heterogeneous character sets
+			if (Target.WindowsPlatform.Compiler >= WindowsCompiler.VisualStudio2015)
             {
                 Arguments.Add("/source-charset:utf-8 /execution-charset:utf-8");
             }
