Fix a crash when data encryption options are enabled but no key is set

diff --git a/Engine/Source/Programs/UnrealBuildTool/Configuration/TargetRules.cs b/Engine/Source/Programs/UnrealBuildTool/Configuration/TargetRules.cs
--- a/Engine/Source/Programs/UnrealBuildTool/Configuration/TargetRules.cs
+++ b/Engine/Source/Programs/UnrealBuildTool/Configuration/TargetRules.cs
@@ -1191,7 +1191,7 @@ namespace UnrealBuildTool
 				ProjectDefinitions.Add("IMPLEMENT_ENCRYPTION_KEY_REGISTRATION()=");
 			}
 
-			if (CryptoSettings.bEnablePakSigning)
+			if (CryptoSettings.IsPakSigningEnabled())
 			{
 				ProjectDefinitions.Add(String.Format("IMPLEMENT_SIGNING_KEY_REGISTRATION()=UE_REGISTER_SIGNING_KEY(UE_LIST_ARGUMENT({0}), UE_LIST_ARGUMENT({1}))", FormatHexBytes(CryptoSettings.SigningKey.PublicKey.Exponent), FormatHexBytes(CryptoSettings.SigningKey.PublicKey.Modulus)));
 			}
diff --git a/Engine/Source/Programs/UnrealBuildTool/System/EncryptionAndSigning.cs b/Engine/Source/Programs/UnrealBuildTool/System/EncryptionAndSigning.cs
--- a/Engine/Source/Programs/UnrealBuildTool/System/EncryptionAndSigning.cs
+++ b/Engine/Source/Programs/UnrealBuildTool/System/EncryptionAndSigning.cs
@@ -29,6 +29,14 @@ namespace UnrealBuildTool
 			/// Modulus
 			/// </summary>
 			public byte[] Modulus;
+
+			/// <summary>
+			/// Returns TRUE if this is a valid Signing key
+			/// </summary>
+			public bool IsValid()
+			{
+				return Exponent != null && Modulus != null;
+			}
 		}
 
 		/// <summary>
@@ -45,6 +53,14 @@ namespace UnrealBuildTool
 			/// Private key
 			/// </summary>
 			public SigningKey PrivateKey = new SigningKey();
+
+			/// <summary>
+			/// Returns TRUE if this is a valid signing key pair
+			/// </summary>
+			public bool IsValid()
+			{
+				return PublicKey != null && PrivateKey != null && PublicKey.IsValid() && PrivateKey.IsValid();
+			}
 		}
 
 		/// <summary>
@@ -64,6 +80,14 @@ namespace UnrealBuildTool
 			/// 128 bit AES key
 			/// </summary>
 			public byte[] Key;
+
+			/// <summary>
+			/// Returns TRUE if this is a valid encryption key
+			/// </summary>
+			public bool IsValid()
+			{
+				return Key != null && Guid != null;
+			}
 		}
 
 		/// <summary>
@@ -122,7 +146,15 @@ namespace UnrealBuildTool
 			/// </summary>
 			public bool IsAnyEncryptionEnabled()
 			{
-				return bEnablePakFullAssetEncryption || bEnablePakUAssetEncryption || bEnablePakIndexEncryption || bEnablePakIniEncryption;
+				return (EncryptionKey != null && EncryptionKey.IsValid()) && (bEnablePakFullAssetEncryption || bEnablePakUAssetEncryption || bEnablePakIndexEncryption || bEnablePakIniEncryption);
+			}
+
+			/// <summary>
+			/// 
+			/// </summary>
+			public bool IsPakSigningEnabled()
+			{
+				return (SigningKey != null && SigningKey.IsValid()) && bEnablePakSigning;
 			}
 
 			/// <summary>
