Fix for encryption keys not being embedded in executable

diff --git a/Engine/Source/Programs/UnrealBuildTool/System/EncryptionAndSigning.cs b/Engine/Source/Programs/UnrealBuildTool/System/EncryptionAndSigning.cs
--- a/Engine/Source/Programs/UnrealBuildTool/System/EncryptionAndSigning.cs
+++ b/Engine/Source/Programs/UnrealBuildTool/System/EncryptionAndSigning.cs
@@ -264,6 +264,8 @@ namespace UnrealBuildTool
 				{
 					Settings.EncryptionKey = new EncryptionKey();
 					Settings.EncryptionKey.Key = System.Convert.FromBase64String(EncryptionKeyString);
+					Settings.EncryptionKey.Guid = Guid.Empty.ToString();
+					Settings.EncryptionKey.Name = "Embedded";
 				}
 
 				// Parse secondary encryption keys
