PR #4476: Fix VSCode IntelliSense (Contributed by [REDACTED_PERSON])

Reapplying this fix after it was trashed by a merge from //UE4/Main

diff --git a/Engine/Source/Programs/UnrealBuildTool/ProjectFiles/VisualStudioCode/VSCodeProjectFileGenerator.cs b/Engine/Source/Programs/UnrealBuildTool/ProjectFiles/VisualStudioCode/VSCodeProjectFileGenerator.cs
--- a/Engine/Source/Programs/UnrealBuildTool/ProjectFiles/VisualStudioCode/VSCodeProjectFileGenerator.cs
+++ b/Engine/Source/Programs/UnrealBuildTool/ProjectFiles/VisualStudioCode/VSCodeProjectFileGenerator.cs
@@ -345,6 +345,7 @@ namespace UnrealBuildTool
 			public List<Project> CSharpProjects = new List<Project>();
 			public List<Project> AllProjects = new List<Project>();
 			public List<string> CombinedIncludePaths = new List<string>();
+			public List<string> CombinedPreprocessorDefinitions = new List<string>();
 		}
 
 		private ProjectData GatherProjectData(List<ProjectFile> InProjects)
@@ -435,6 +436,15 @@ namespace UnrealBuildTool
 							
 						}
 					}
+					
+					foreach (string Definition in Project.IntelliSensePreprocessorDefinitions)
+					{
+						string Processed = Definition.Replace("\"", "\\\"");
+						if (!ProjectData.CombinedPreprocessorDefinitions.Contains(Processed))
+						{
+							ProjectData.CombinedPreprocessorDefinitions.Add(Processed);
+						}
+					}
 
 					ProjectData.NativeProjects.Add(NewProject);
 				}
@@ -615,8 +625,10 @@ namespace UnrealBuildTool
 
 						OutFile.BeginArray("defines");
 						{
-							OutFile.AddUnnamedField("_DEBUG");
-							OutFile.AddUnnamedField("UNICODE");
+							foreach (string Definition in Projects.CombinedPreprocessorDefinitions)
+							{
+								OutFile.AddUnnamedField(Definition);
+							}
 						}
 						OutFile.EndArray();
 					}
