Don't add Mac third party libraries inside the app bundle to the BuildProducts list/manifest, we are no longer copying them there

diff --git a/Engine/Source/Programs/UnrealBuildTool/Platform/Mac/MacToolChain.cs b/Engine/Source/Programs/UnrealBuildTool/Platform/Mac/MacToolChain.cs
--- a/Engine/Source/Programs/UnrealBuildTool/Platform/Mac/MacToolChain.cs
+++ b/Engine/Source/Programs/UnrealBuildTool/Platform/Mac/MacToolChain.cs
@@ -1284,22 +1284,6 @@ namespace UnrealBuildTool
 			// We need to know what third party dylibs would be copied to the bundle
 			if (Binary.Type != UEBuildBinaryType.StaticLibrary)
 			{
-			    foreach (string AdditionalLibrary in Libraries)
-				{
-					string LibName = Path.GetFileName(AdditionalLibrary);
-					if (LibName.StartsWith("lib"))
-					{
-						if (Path.GetExtension(AdditionalLibrary) == ".dylib" && BundleContentsDirectory != null)
-						{
-							FileReference Entry = FileReference.Combine(BundleContentsDirectory, "MacOS", LibName);
-							if (!BuildProducts.ContainsKey(Entry))
-							{
-								BuildProducts.Add(Entry, BuildProductType.DynamicLibrary);
-							}
-						}
-					}
-				}
-
 			    foreach (UEBuildBundleResource Resource in BundleResources)
 				{
 					if (Directory.Exists(Resource.ResourcePath))
