Add string pooling to optimized builds.

https://udn.unrealengine.com/questions/470864/reducing-executable-size-shipping-and-test.html

diff --git a/Engine/Source/Programs/UnrealBuildTool/Platform/Windows/VCToolChain.cs b/Engine/Source/Programs/UnrealBuildTool/Platform/Windows/VCToolChain.cs
--- a/Engine/Source/Programs/UnrealBuildTool/Platform/Windows/VCToolChain.cs
+++ b/Engine/Source/Programs/UnrealBuildTool/Platform/Windows/VCToolChain.cs
@@ -299,6 +299,9 @@ namespace UnrealBuildTool
 					// Favor code speed.
 					Arguments.Add("/Ot");
 
+					// Coalesce duplicate strings
+					Arguments.Add("/GF");
+
 					// Only omit frame pointers on the PC (which is implied by /Ox) if wanted.
 					if (CompileEnvironment.bOmitFramePointers == false
 					&& ((CompileEnvironment.Platform == CppPlatform.Win32) ||
