Fix ScriptPlugin not being built when enabled in a project.

diff --git a/Engine/Plugins/ScriptPlugin/ScriptPlugin.uplugin b/Engine/Plugins/ScriptPlugin/ScriptPlugin.uplugin
--- a/Engine/Plugins/ScriptPlugin/ScriptPlugin.uplugin
+++ b/Engine/Plugins/ScriptPlugin/ScriptPlugin.uplugin
@@ -15,6 +15,7 @@
 	"IsBetaVersion" : false,
 	"Installed" : false,
 	"CanBeUsedWithUnrealHeaderTool" : true,
+	"SupportedPrograms": [ "UnrealHeaderTool" ],
 	
 	"Modules" :
 	[
diff --git a/Engine/Source/Programs/UnrealBuildTool/Configuration/UEBuildTarget.cs b/Engine/Source/Programs/UnrealBuildTool/Configuration/UEBuildTarget.cs
--- a/Engine/Source/Programs/UnrealBuildTool/Configuration/UEBuildTarget.cs
+++ b/Engine/Source/Programs/UnrealBuildTool/Configuration/UEBuildTarget.cs
@@ -3342,7 +3342,7 @@ namespace UnrealBuildTool
 			BuildPlugins = new List<UEBuildPlugin>(NameToInstance.Values);
 
 			// Determine if the project has a script plugin. We will always build UHT if there is a script plugin in the game folder.
-			bHasProjectScriptPlugin = EnabledPlugins.Any(x => x.Descriptor.SupportedPrograms != null && x.Descriptor.SupportedPrograms.Contains("UnrealHeaderTool") && !x.File.IsUnderDirectory(UnrealBuildTool.EngineDirectory));
+			bHasProjectScriptPlugin = EnabledPlugins.Any(x => x.Descriptor.SupportedPrograms != null && x.Descriptor.SupportedPrograms.Contains("UnrealHeaderTool"));
 		}
 
 		/// <summary>
