Add build types to log during packaging for Android

#ue4
#android

diff --git a/Engine/Source/Programs/UnrealBuildTool/Platform/Android/UEDeployAndroid.cs b/Engine/Source/Programs/UnrealBuildTool/Platform/Android/UEDeployAndroid.cs
--- a/Engine/Source/Programs/UnrealBuildTool/Platform/Android/UEDeployAndroid.cs
+++ b/Engine/Source/Programs/UnrealBuildTool/Platform/Android/UEDeployAndroid.cs
@@ -1298,6 +1298,21 @@ namespace UnrealBuildTool
 			}
 		}
 
+		void LogBuildSetup()
+		{
+			ConfigHierarchy Ini = GetConfigCacheIni(ConfigHierarchyType.Engine);
+			bool bBuildForES2 = false;
+			Ini.GetBool("/Script/AndroidRuntimeSettings.AndroidRuntimeSettings", "bBuildForES2", out bBuildForES2);
+			bool bBuildForES31 = false;
+			Ini.GetBool("/Script/AndroidRuntimeSettings.AndroidRuntimeSettings", "bBuildForES31", out bBuildForES31);
+			bool bSupportsVulkan = false;
+			Ini.GetBool("/Script/AndroidRuntimeSettings.AndroidRuntimeSettings", "bSupportsVulkan", out bSupportsVulkan);
+
+			Log.TraceInformation("bBuildForES2: {0}", (bBuildForES2 ? "true" : "false"));
+			Log.TraceInformation("bBuildForES31: {0}", (bBuildForES31 ? "true" : "false"));
+			Log.TraceInformation("bSupportsVulkan: {0}", (bSupportsVulkan ? "true" : "false"));
+		}
+		
 		void CopyVulkanValidationLayers(string UE4BuildPath, string UE4Arch, string NDKArch, string Configuration)
 		{
 			bool bSupportsVulkan = false;
@@ -2908,6 +2923,8 @@ namespace UnrealBuildTool
 				}
 			}
 
+			LogBuildSetup();
+
 			bool bIsBuildMachine = Environment.GetEnvironmentVariable("IsBuildMachine") == "1";
 
 			// do this here so we'll stop early if there is a problem with the SDK API level (cached so later calls will return the same)
