HTML5 - comment out terse/debug build output

#rnx

diff --git a/Engine/Source/Programs/UnrealBuildTool/Platform/HTML5/HTML5ToolChain.cs b/Engine/Source/Programs/UnrealBuildTool/Platform/HTML5/HTML5ToolChain.cs
--- a/Engine/Source/Programs/UnrealBuildTool/Platform/HTML5/HTML5ToolChain.cs
+++ b/Engine/Source/Programs/UnrealBuildTool/Platform/HTML5/HTML5ToolChain.cs
@@ -161,10 +161,7 @@ namespace UnrealBuildTool
 			if (enableMultithreading)
 			{
 				Result += " -msse2 -s USE_PTHREADS=1";
-				if ( bMultithreading_UseOffscreenCanvas )
-				{
-					Result += " -DEXPERIMENTAL_OPENGL_RHITHREAD=1";
-				}
+				Result += " -DEXPERIMENTAL_OPENGL_RHITHREAD=" + (bMultithreading_UseOffscreenCanvas ? "1" : "0");
 
 				// NOTE: use "emscripten native" video, keyboard, mouse
 			}
@@ -207,7 +204,7 @@ namespace UnrealBuildTool
 			Environment.SetEnvironmentVariable("EMCC_SKIP_SANITY_CHECK", "1");
 
 			// THESE ARE TEST/DEBUGGING
-			Environment.SetEnvironmentVariable("EMCC_DEBUG", "1");
+	//		Environment.SetEnvironmentVariable("EMCC_DEBUG", "1");
 //			Environment.SetEnvironmentVariable("EMCC_CORES", "8");
 //			Environment.SetEnvironmentVariable("EMCC_OPTIMIZE_NORMALLY", "1");
 
