HTML5 - fix regex to only read key=value pairs --> at start of line

#rnx

diff --git a/Engine/Source/Programs/UnrealBuildTool/Platform/HTML5/HTML5SDKInfo.cs b/Engine/Source/Programs/UnrealBuildTool/Platform/HTML5/HTML5SDKInfo.cs
--- a/Engine/Source/Programs/UnrealBuildTool/Platform/HTML5/HTML5SDKInfo.cs
+++ b/Engine/Source/Programs/UnrealBuildTool/Platform/HTML5/HTML5SDKInfo.cs
@@ -112,7 +112,7 @@ namespace UnrealBuildTool
 			string[] tokens = config.Split('\n');
 
 			// Parse lines of type "KEY='value'"
-			Regex regex = new Regex(variable + "\\s*=\\s*['\\\"](.*)['\\\"]");
+			Regex regex = new Regex(variable + "^\\s*=\\s*['\\\"](.*)['\\\"]");
 			foreach(string line in tokens)
 			{
 				Match m = regex.Match(line);
