Validate Android NDK before before compiling

#ue4
#android

diff --git a/Engine/Source/Programs/UnrealBuildTool/Platform/Android/AndroidToolChain.cs b/Engine/Source/Programs/UnrealBuildTool/Platform/Android/AndroidToolChain.cs
--- a/Engine/Source/Programs/UnrealBuildTool/Platform/Android/AndroidToolChain.cs
+++ b/Engine/Source/Programs/UnrealBuildTool/Platform/Android/AndroidToolChain.cs
@@ -13,6 +13,11 @@ namespace UnrealBuildTool
 {
 	class AndroidToolChain : UEToolChain, IAndroidToolChain
 	{
+		// Android NDK toolchain that must be used for C++ compiling
+		readonly int MinimumNDKToolchain = 140100;
+		readonly int MaximumNDKToolchain = 180100;
+		readonly int RecommendedNDKToolchain = 140100;
+
 		public static readonly string[] AllCpuSuffixes =
 		{
 			"-armv7",
@@ -123,6 +128,17 @@ namespace UnrealBuildTool
 				(ClangVersionMajor == Major && ClangVersionMinor == Minor && ClangVersionPatch < Patch);
 		}
 
+		private static string ToolchainIntToString(int ToolchainInt)
+		{
+			int RevisionNum = ToolchainInt / 10000;
+			int RevisionMinor = ToolchainInt - (RevisionNum * 10000);
+			int RevisionLetterNum = RevisionMinor / 100;
+			int RevisionBeta = RevisionMinor - (RevisionLetterNum * 100);
+			char RevisionLetter = Convert.ToChar('a' + RevisionLetterNum);
+
+			return "r" + RevisionNum + (RevisionLetterNum > 0 ? Char.ToString(RevisionLetter) : "");
+		}
+
 		[CommandLine("-Architectures=", ListSeparator = '+')]
 		public List<string> ArchitectureArg = new List<string>();
 
@@ -297,6 +313,13 @@ namespace UnrealBuildTool
 				throw new BuildException("Cannot find supported Android toolchain with NDKPath:" + NDKPath);
 			}
 
+			// verify NDK toolchain is supported
+			if (NDKDefineInt < MinimumNDKToolchain || NDKDefineInt > MaximumNDKToolchain)
+			{
+				throw new BuildException("Android toolchain NDK " + ToolchainIntToString(NDKDefineInt) + " not supported; please use NDK " + ToolchainIntToString(MinimumNDKToolchain) + " to NDK " + ToolchainIntToString(MaximumNDKToolchain) +
+					" (NDK " + ToolchainIntToString(RecommendedNDKToolchain) + " recommended)");
+			}
+
 			// set up the path to our toolchains
 			ClangPath = Utils.CollapseRelativeDirectories(Path.Combine(NDKPath, @"toolchains/llvm" + ClangVersion, ArchitecturePath, @"bin/clang++" + ExeExtension));
 			ArPathArm = Path.Combine(NDKPath, @"toolchains/arm-linux-androideabi-" + GccVersion, ArchitecturePath, @"bin/arm-linux-androideabi-ar" + ExeExtension);     //@todo android: use llvm-ar.exe instead?
