Fix executables not being rebuilt whenever the user's selected icon changes on Windows.

* The icon file is now added as a dependency when compiling resource files.
* The path to the chosen icon file is saved into an "external build metadata" section of the makefile, so that removing or adding the file can cause the makefile to be invalidated on startup.

diff --git a/Engine/Source/Programs/UnrealBuildTool/Configuration/UEBuildBinary.cs b/Engine/Source/Programs/UnrealBuildTool/Configuration/UEBuildBinary.cs
--- a/Engine/Source/Programs/UnrealBuildTool/Configuration/UEBuildBinary.cs
+++ b/Engine/Source/Programs/UnrealBuildTool/Configuration/UEBuildBinary.cs
@@ -744,7 +744,6 @@ namespace UnrealBuildTool
 
 						// Create a compile environment for resource files
 						CppCompileEnvironment ResourceCompileEnvironment = new CppCompileEnvironment(BinaryCompileEnvironment);
-						WindowsPlatform.SetupResourceCompileEnvironment(ResourceCompileEnvironment, Target);
 
 						// @todo: This should be in some Windows code somewhere...
 						// Set the original file name macro; used in Default.rc2 to set the binary metadata fields.
diff --git a/Engine/Source/Programs/UnrealBuildTool/Configuration/UEBuildModuleCPP.cs b/Engine/Source/Programs/UnrealBuildTool/Configuration/UEBuildModuleCPP.cs
--- a/Engine/Source/Programs/UnrealBuildTool/Configuration/UEBuildModuleCPP.cs
+++ b/Engine/Source/Programs/UnrealBuildTool/Configuration/UEBuildModuleCPP.cs
@@ -490,7 +490,6 @@ namespace UnrealBuildTool
 			if(InputFiles.RCFiles.Count > 0)
 			{
 				CppCompileEnvironment ResourceCompileEnvironment = new CppCompileEnvironment(BinaryCompileEnvironment);
-				WindowsPlatform.SetupResourceCompileEnvironment(ResourceCompileEnvironment, Target);
 				LinkInputFiles.AddRange(ToolChain.CompileRCFiles(ResourceCompileEnvironment, InputFiles.RCFiles, IntermediateDirectory, Makefile.Actions).ObjectFiles);
 			}
 
diff --git a/Engine/Source/Programs/UnrealBuildTool/Configuration/UEBuildPlatform.cs b/Engine/Source/Programs/UnrealBuildTool/Configuration/UEBuildPlatform.cs
--- a/Engine/Source/Programs/UnrealBuildTool/Configuration/UEBuildPlatform.cs
+++ b/Engine/Source/Programs/UnrealBuildTool/Configuration/UEBuildPlatform.cs
@@ -819,6 +819,27 @@ namespace UnrealBuildTool
 			GlobalLinkEnvironment.bCreateDebugInfo = GlobalCompileEnvironment.bCreateDebugInfo;
 		}
 
+		/// <summary>
+		/// Allows the platform to return various build metadata that is not tracked by other means. If the returned string changes, the makefile will be invalidated.
+		/// </summary>
+		/// <param name="ProjectFile">The project file being built</param>
+		/// <returns>String describing the current build metadata</returns>
+		public string GetExternalBuildMetadata(FileReference ProjectFile)
+		{
+			StringBuilder Result = new StringBuilder();
+			GetExternalBuildMetadata(ProjectFile, Result);
+			return Result.ToString();
+		}
+
+		/// <summary>
+		/// Allows the platform to return various build metadata that is not tracked by other means. If the returned string changes, the makefile will be invalidated.
+		/// </summary>
+		/// <param name="ProjectFile">The project file being built</param>
+		/// <param name="Metadata">String builder to contain build metadata</param>
+		public virtual void GetExternalBuildMetadata(FileReference ProjectFile, StringBuilder Metadata)
+		{
+		}
+
 		/// <summary>
 		/// Checks if platform is part of a given platform group
 		/// </summary>
diff --git a/Engine/Source/Programs/UnrealBuildTool/Configuration/UEBuildTarget.cs b/Engine/Source/Programs/UnrealBuildTool/Configuration/UEBuildTarget.cs
--- a/Engine/Source/Programs/UnrealBuildTool/Configuration/UEBuildTarget.cs
+++ b/Engine/Source/Programs/UnrealBuildTool/Configuration/UEBuildTarget.cs
@@ -1195,7 +1195,8 @@ namespace UnrealBuildTool
 			}
 
 			// Create the makefile
-			TargetMakefile Makefile = new TargetMakefile(TargetToolChain.GetVersionInfo(), ReceiptFileName, ProjectIntermediateDirectory, TargetType, bDeployAfterCompile, bHasProjectScriptPlugin);
+			string ExternalMetadata = UEBuildPlatform.GetBuildPlatform(Platform).GetExternalBuildMetadata(ProjectFile);
+			TargetMakefile Makefile = new TargetMakefile(TargetToolChain.GetVersionInfo(), ExternalMetadata, ReceiptFileName, ProjectIntermediateDirectory, TargetType, bDeployAfterCompile, bHasProjectScriptPlugin);
 
 			// Setup the hot reload module list
 			Makefile.HotReloadModuleNames = GetHotReloadModuleNames();
@@ -1283,8 +1284,6 @@ namespace UnrealBuildTool
 
 						FileItem DefaultResourceFile = FileItem.GetItemByFileReference(FileReference.Combine(UnrealBuildTool.EngineDirectory, "Build", "Windows", "Resources", "Default.rc2"));
 
-						WindowsPlatform.SetupResourceCompileEnvironment(DefaultResourceCompileEnvironment, Rules);
-
 						CPPOutput DefaultResourceOutput = TargetToolChain.CompileRCFiles(DefaultResourceCompileEnvironment, new List<FileItem> { DefaultResourceFile }, EngineIntermediateDirectory, Makefile.Actions);
 						GlobalLinkEnvironment.DefaultResourceFiles.AddRange(DefaultResourceOutput.ObjectFiles);
 					}
diff --git a/Engine/Source/Programs/UnrealBuildTool/Platform/Windows/UEBuildWindows.cs b/Engine/Source/Programs/UnrealBuildTool/Platform/Windows/UEBuildWindows.cs
--- a/Engine/Source/Programs/UnrealBuildTool/Platform/Windows/UEBuildWindows.cs
+++ b/Engine/Source/Programs/UnrealBuildTool/Platform/Windows/UEBuildWindows.cs
@@ -1488,52 +1488,6 @@ namespace UnrealBuildTool
 			return FileReference.Combine(UnrealBuildTool.EngineDirectory, "Build", "Windows", "Resources", "Default.ico");
 		}
 
-		/// <summary>
-		/// Configures the resource compile environment for the given target
-		/// </summary>
-		/// <param name="ResourceCompileEnvironment">The compile environment</param>
-		/// <param name="Target">The target being built</param>
-		public static void SetupResourceCompileEnvironment(CppCompileEnvironment ResourceCompileEnvironment, ReadOnlyTargetRules Target)
-		{
-			// Figure the icon to use. We can only use a custom icon when compiling to a project-specific intemediate directory (and not for the shared editor executable, for example).
-			FileReference IconFile;
-			if(Target.ProjectFile != null && !ResourceCompileEnvironment.bUseSharedBuildEnvironment)
-			{
-				IconFile = WindowsPlatform.GetApplicationIcon(Target.ProjectFile);
-			}
-			else
-			{
-				IconFile = WindowsPlatform.GetApplicationIcon(null);
-			}
-
-			// Setup the compile environment, setting the icon to use via a macro. This is used in Default.rc2.
-			ResourceCompileEnvironment.Definitions.Add(String.Format("BUILD_ICON_FILE_NAME=\"\\\"{0}\\\"\"", IconFile.FullName.Replace("\\", "\\\\")));
-
-			// Apply the target settings for the resources
-			if(!ResourceCompileEnvironment.bUseSharedBuildEnvironment)
-			{
-				if (!String.IsNullOrEmpty(Target.WindowsPlatform.CompanyName))
-				{
-					ResourceCompileEnvironment.Definitions.Add(String.Format("PROJECT_COMPANY_NAME={0}", SanitizeMacroValue(Target.WindowsPlatform.CompanyName)));
-				}
-
-				if (!String.IsNullOrEmpty(Target.WindowsPlatform.CopyrightNotice))
-				{
-					ResourceCompileEnvironment.Definitions.Add(String.Format("PROJECT_COPYRIGHT_STRING={0}", SanitizeMacroValue(Target.WindowsPlatform.CopyrightNotice)));
-				}
-
-				if (!String.IsNullOrEmpty(Target.WindowsPlatform.ProductName))
-				{
-					ResourceCompileEnvironment.Definitions.Add(String.Format("PROJECT_PRODUCT_NAME={0}", SanitizeMacroValue(Target.WindowsPlatform.ProductName)));
-				}
-
-				if (Target.ProjectFile != null)
-				{
-					ResourceCompileEnvironment.Definitions.Add(String.Format("PROJECT_PRODUCT_IDENTIFIER={0}", SanitizeMacroValue(Target.ProjectFile.GetFileNameWithoutExtension())));
-				}
-			}
-		}
-
 		/// <summary>
 		/// Modify the rules for a newly created module, in a target that's being built for this platform.
 		/// This is not required - but allows for hiding details of a particular platform.
@@ -1732,25 +1686,6 @@ namespace UnrealBuildTool
 			LinkEnvironment.ModuleDefinitionFile = Target.WindowsPlatform.ModuleDefinitionFile;
 		}
 
-		/// <summary>
-		/// Macros passed via the command line have their quotes stripped, and are tokenized before being re-stringized by the compiler. This conversion
-		/// back and forth is normally ok, but certain characters such as single quotes must always be paired. Remove any such characters here.
-		/// </summary>
-		/// <param name="Value">The macro value</param>
-		/// <returns>The sanitized value</returns>
-		static string SanitizeMacroValue(string Value)
-		{
-			StringBuilder Result = new StringBuilder(Value.Length);
-			for(int Idx = 0; Idx < Value.Length; Idx++)
-			{
-				if(Value[Idx] != '\'' && Value[Idx] != '\"')
-				{
-					Result.Append(Value[Idx]);
-				}
-			}
-			return Result.ToString();
-		}
-
 		/// <summary>
 		/// Setup the configuration environment for building
 		/// </summary>
@@ -1839,6 +1774,21 @@ namespace UnrealBuildTool
 			}
 		}
 
+		/// <summary>
+		/// Allows the platform to return various build metadata that is not tracked by other means. If the returned string changes, the makefile will be invalidated.
+		/// </summary>
+		/// <param name="ProjectFile">The project file being built</param>
+		/// <param name="Metadata">String builder to contain build metadata</param>
+		public override void GetExternalBuildMetadata(FileReference ProjectFile, StringBuilder Metadata)
+		{
+			base.GetExternalBuildMetadata(ProjectFile, Metadata);
+
+			if(ProjectFile != null)
+			{
+				Metadata.AppendLine("ICON: {0}", GetApplicationIcon(ProjectFile));
+			}
+		}
+
 		/// <summary>
 		/// Deploys the given target
 		/// </summary>
diff --git a/Engine/Source/Programs/UnrealBuildTool/Platform/Windows/VCToolChain.cs b/Engine/Source/Programs/UnrealBuildTool/Platform/Windows/VCToolChain.cs
--- a/Engine/Source/Programs/UnrealBuildTool/Platform/Windows/VCToolChain.cs
+++ b/Engine/Source/Programs/UnrealBuildTool/Platform/Windows/VCToolChain.cs
@@ -1210,6 +1210,45 @@ namespace UnrealBuildTool
 					}
 				}
 
+				// Figure the icon to use. We can only use a custom icon when compiling to a project-specific intemediate directory (and not for the shared editor executable, for example).
+				FileReference IconFile;
+				if(Target.ProjectFile != null && !CompileEnvironment.bUseSharedBuildEnvironment)
+				{
+					IconFile = WindowsPlatform.GetApplicationIcon(Target.ProjectFile);
+				}
+				else
+				{
+					IconFile = WindowsPlatform.GetApplicationIcon(null);
+				}
+				CompileAction.PrerequisiteItems.Add(FileItem.GetItemByFileReference(IconFile));
+
+				// Setup the compile environment, setting the icon to use via a macro. This is used in Default.rc2.
+				AddDefinition(Arguments, String.Format("BUILD_ICON_FILE_NAME=\"\\\"{0}\\\"\"", IconFile.FullName.Replace("\\", "\\\\")));
+
+				// Apply the target settings for the resources
+				if(!CompileEnvironment.bUseSharedBuildEnvironment)
+				{
+					if (!String.IsNullOrEmpty(Target.WindowsPlatform.CompanyName))
+					{
+						AddDefinition(Arguments, String.Format("PROJECT_COMPANY_NAME={0}", SanitizeMacroValue(Target.WindowsPlatform.CompanyName)));
+					}
+
+					if (!String.IsNullOrEmpty(Target.WindowsPlatform.CopyrightNotice))
+					{
+						AddDefinition(Arguments, String.Format("PROJECT_COPYRIGHT_STRING={0}", SanitizeMacroValue(Target.WindowsPlatform.CopyrightNotice)));
+					}
+
+					if (!String.IsNullOrEmpty(Target.WindowsPlatform.ProductName))
+					{
+						AddDefinition(Arguments, String.Format("PROJECT_PRODUCT_NAME={0}", SanitizeMacroValue(Target.WindowsPlatform.ProductName)));
+					}
+
+					if (Target.ProjectFile != null)
+					{
+						AddDefinition(Arguments, String.Format("PROJECT_PRODUCT_IDENTIFIER={0}", SanitizeMacroValue(Target.ProjectFile.GetFileNameWithoutExtension())));
+					}
+				}
+
 				// Add the RES file to the produced item list.
 				FileItem CompiledResourceFile = FileItem.GetItemByFileReference(
 					FileReference.Combine(
@@ -1235,6 +1274,25 @@ namespace UnrealBuildTool
 			return Result;
 		}
 
+		/// <summary>
+		/// Macros passed via the command line have their quotes stripped, and are tokenized before being re-stringized by the compiler. This conversion
+		/// back and forth is normally ok, but certain characters such as single quotes must always be paired. Remove any such characters here.
+		/// </summary>
+		/// <param name="Value">The macro value</param>
+		/// <returns>The sanitized value</returns>
+		static string SanitizeMacroValue(string Value)
+		{
+			StringBuilder Result = new StringBuilder(Value.Length);
+			for(int Idx = 0; Idx < Value.Length; Idx++)
+			{
+				if(Value[Idx] != '\'' && Value[Idx] != '\"')
+				{
+					Result.Append(Value[Idx]);
+				}
+			}
+			return Result.ToString();
+		}
+
 		public override FileItem LinkFiles(LinkEnvironment LinkEnvironment, bool bBuildImportLibraryOnly, List<Action> Actions)
 		{
 			if (LinkEnvironment.bIsBuildingDotNetAssembly)
diff --git a/Engine/Source/Programs/UnrealBuildTool/System/TargetMakefile.cs b/Engine/Source/Programs/UnrealBuildTool/System/TargetMakefile.cs
--- a/Engine/Source/Programs/UnrealBuildTool/System/TargetMakefile.cs
+++ b/Engine/Source/Programs/UnrealBuildTool/System/TargetMakefile.cs
@@ -20,7 +20,7 @@ namespace UnrealBuildTool
 		/// <summary>
 		/// The version number to write
 		/// </summary>
-		public const int CurrentVersion = 10;
+		public const int CurrentVersion = 11;
 
 		/// <summary>
 		/// The time at which the makefile was created
@@ -32,6 +32,11 @@ namespace UnrealBuildTool
 		/// </summary>
 		public string ToolchainInfo;
 
+		/// <summary>
+		/// Any additional information about the build environment which the platform can use to invalidate the makefile
+		/// </summary>
+		public string ExternalMetadata;
+
 		/// <summary>
 		/// Path to the receipt file for this target
 		/// </summary>
@@ -132,15 +137,17 @@ namespace UnrealBuildTool
 		/// Constructor
 		/// </summary>
 		/// <param name="ToolchainInfo">String describing the toolchain used to build. This will be output before executing actions.</param>
+		/// <param name="ExternalMetadata">External build metadata from the platform</param>
 		/// <param name="ReceiptFile">Path to the receipt file</param>
 		/// <param name="ProjectIntermediateDirectory">Path to the project intermediate directory</param>
 		/// <param name="TargetType">The type of target</param>
 		/// <param name="bDeployAfterCompile">Whether to deploy the target after compiling</param>
 		/// <param name="bHasProjectScriptPlugin">Whether the target has a project script plugin</param>
-		public TargetMakefile(string ToolchainInfo, FileReference ReceiptFile, DirectoryReference ProjectIntermediateDirectory, TargetType TargetType, bool bDeployAfterCompile, bool bHasProjectScriptPlugin)
+		public TargetMakefile(string ToolchainInfo, string ExternalMetadata, FileReference ReceiptFile, DirectoryReference ProjectIntermediateDirectory, TargetType TargetType, bool bDeployAfterCompile, bool bHasProjectScriptPlugin)
 		{
 			this.CreateTimeUtc = DateTime.UtcNow;
 			this.ToolchainInfo = ToolchainInfo;
+			this.ExternalMetadata = ExternalMetadata;
 			this.ReceiptFile = ReceiptFile;
 			this.ProjectIntermediateDirectory = ProjectIntermediateDirectory;
 			this.TargetType = TargetType;
@@ -167,6 +174,7 @@ namespace UnrealBuildTool
 		{
 			CreateTimeUtc = new DateTime(Reader.ReadLong(), DateTimeKind.Utc);
 			ToolchainInfo = Reader.ReadString();
+			ExternalMetadata = Reader.ReadString();
 			ReceiptFile = Reader.ReadFileReference();
 			ProjectIntermediateDirectory = Reader.ReadDirectoryReference();
 			TargetType = (TargetType)Reader.ReadInt();
@@ -196,6 +204,7 @@ namespace UnrealBuildTool
 		{
 			Writer.WriteLong(CreateTimeUtc.Ticks);
 			Writer.WriteString(ToolchainInfo);
+			Writer.WriteString(ExternalMetadata);
 			Writer.WriteFileReference(ReceiptFile);
 			Writer.WriteDirectoryReference(ProjectIntermediateDirectory);
 			Writer.WriteInt((int)TargetType);
@@ -383,6 +392,16 @@ namespace UnrealBuildTool
 						}
 					}
 				}
+
+				// Get the current build metadata from the platform
+				string CurrentExternalMetadata = UEBuildPlatform.GetBuildPlatform(Platform).GetExternalBuildMetadata(ProjectFile);
+				if(String.Compare(CurrentExternalMetadata, Makefile.ExternalMetadata, StringComparison.Ordinal) != 0)
+				{
+					Log.TraceLog("Old metadata:\n", Makefile.ExternalMetadata);
+					Log.TraceLog("New metadata:\n", CurrentExternalMetadata);
+					ReasonNotLoaded = "build metadata has changed";
+					return null;
+				}
 			}
 
 			// The makefile is ok
