Fix Mali debugger paths to work on Mac and Linux

#ue4
#android

diff --git a/Engine/Source/Programs/UnrealBuildTool/Platform/Android/UEDeployAndroid.cs b/Engine/Source/Programs/UnrealBuildTool/Platform/Android/UEDeployAndroid.cs
--- a/Engine/Source/Programs/UnrealBuildTool/Platform/Android/UEDeployAndroid.cs
+++ b/Engine/Source/Programs/UnrealBuildTool/Platform/Android/UEDeployAndroid.cs
@@ -1278,11 +1278,11 @@ namespace UnrealBuildTool
 						if (Directory.Exists(MaliGraphicsDebuggerPath))
 						{
 							Directory.CreateDirectory(Path.Combine(UE4BuildPath, "libs", NDKArch));
-							string MaliLibSrcPath = Path.Combine(MaliGraphicsDebuggerPath, @"target\android-non-root\arm", NDKArch, "libMGD.so");
+							string MaliLibSrcPath = Path.Combine(MaliGraphicsDebuggerPath, "target", "android-non-root", "arm", NDKArch, "libMGD.so");
 							if (!File.Exists(MaliLibSrcPath))
 							{
 								// in v4.3.0 library location was changed
-								MaliLibSrcPath = Path.Combine(MaliGraphicsDebuggerPath, @"target\android\arm\unrooted", NDKArch, "libMGD.so");
+								MaliLibSrcPath = Path.Combine(MaliGraphicsDebuggerPath, "target", "android", "arm", "unrooted", NDKArch, "libMGD.so");
 							}
 							string MaliLibDstPath = Path.Combine(UE4BuildPath, "libs", NDKArch, "libMGD.so");
 
