UBT: Validate that runtime dependencies from restricted folders are not copied into public folders.

#rnx

diff --git a/Engine/Source/Programs/UnrealBuildTool/Configuration/UEBuildTarget.cs b/Engine/Source/Programs/UnrealBuildTool/Configuration/UEBuildTarget.cs
--- a/Engine/Source/Programs/UnrealBuildTool/Configuration/UEBuildTarget.cs
+++ b/Engine/Source/Programs/UnrealBuildTool/Configuration/UEBuildTarget.cs
@@ -1473,6 +1473,10 @@ namespace UnrealBuildTool
 					{
 						bResult &= Binary.CheckRestrictedFolders(DirectoryReference.FromFile(ProjectFile), ModuleRestrictedFolderCache);
 					}
+					foreach(KeyValuePair<FileReference, FileReference> Pair in RuntimeDependencyTargetFileToSourceFile)
+					{
+						bResult &= CheckRestrictedFolders(Pair.Key, Pair.Value);
+					}
 
 					if(!bResult)
 					{
@@ -1510,6 +1514,53 @@ namespace UnrealBuildTool
 			return Makefile;
 		}
 
+		/// <summary>
+		/// Check that copying a file from one location to another does not violate rules regarding restricted folders
+		/// </summary>
+		/// <param name="TargetFile">The destination location for the file</param>
+		/// <param name="SourceFile">The source location of the file</param>
+		/// <returns>True if the copy is permitted, false otherwise</returns>
+		bool CheckRestrictedFolders(FileReference TargetFile, FileReference SourceFile)
+		{
+			List<RestrictedFolder> TargetRestrictedFolders = GetRestrictedFolders(TargetFile);
+			List<RestrictedFolder> SourceRestrictedFolders = GetRestrictedFolders(SourceFile);
+			foreach(RestrictedFolder SourceRestrictedFolder in SourceRestrictedFolders)
+			{
+				if(!TargetRestrictedFolders.Contains(SourceRestrictedFolder))
+				{
+					Log.TraceError("Runtime dependency '{0}' is copied to '{1}', which does not contain a '{2}' folder.", SourceFile, TargetFile, SourceRestrictedFolder);
+					return false;
+				}
+			}
+			return true;
+		}
+
+		/// <summary>
+		/// Gets the restricted folders that the given file is in
+		/// </summary>
+		/// <param name="File">The file to test</param>
+		/// <returns>List of restricted folders for the file</returns>
+		List<RestrictedFolder> GetRestrictedFolders(FileReference File)
+		{
+			// Find the base directory for this binary
+			DirectoryReference BaseDir;
+			if(File.IsUnderDirectory(UnrealBuildTool.RootDirectory))
+			{
+				BaseDir = UnrealBuildTool.RootDirectory;
+			}
+			else if(ProjectDirectory != null && File.IsUnderDirectory(ProjectDirectory))
+			{
+				BaseDir = ProjectDirectory;
+			}
+			else
+			{
+				return new List<RestrictedFolder>();
+			}
+
+			// Find the restricted folders under the base directory
+			return RestrictedFolders.FindRestrictedFolders(BaseDir, File.Directory);
+		}
+
 		/// <summary>
 		/// Creates an action which copies a file from one location to another
 		/// </summary>
