Change default Windows toolchain to 14.16.27023 (from VS2017 15.9.4).

Should fix periodic linker errors, as described here https://developercommunity.visualstudio.com/content/problem/225957/fatal-error-lnk1318-unexpected-pdb-error-ok-0.html.

#rnx

diff --git a/Engine/Source/Programs/UnrealBuildTool/Platform/Windows/UEBuildWindows.cs b/Engine/Source/Programs/UnrealBuildTool/Platform/Windows/UEBuildWindows.cs
--- a/Engine/Source/Programs/UnrealBuildTool/Platform/Windows/UEBuildWindows.cs
+++ b/Engine/Source/Programs/UnrealBuildTool/Platform/Windows/UEBuildWindows.cs
@@ -428,7 +428,7 @@ namespace UnrealBuildTool
 		/// <summary>
 		/// The default compiler version to be used, if installed. 
 		/// </summary>
-		static readonly VersionNumber DefaultVisualStudioToolChainVersion = VersionNumber.Parse("14.13.26128");
+		static readonly VersionNumber DefaultVisualStudioToolChainVersion = VersionNumber.Parse("14.16.27023");
 
 		/// <summary>
 		/// The default Windows SDK version to be used, if installed.
