Fixed the accidental removal of CLion project generation support

diff --git a/Engine/Source/Programs/UnrealBuildTool/UnrealBuildTool.cs b/Engine/Source/Programs/UnrealBuildTool/UnrealBuildTool.cs
--- a/Engine/Source/Programs/UnrealBuildTool/UnrealBuildTool.cs
+++ b/Engine/Source/Programs/UnrealBuildTool/UnrealBuildTool.cs
@@ -886,6 +886,9 @@ namespace UnrealBuildTool
 								case ProjectFileFormat.VisualStudioCode:
 									Generator = new VSCodeProjectFileGenerator(ProjectFile);
 									break;
+﻿								case ProjectFileFormat.CLion:
+									Generator = new CLionGenerator(ProjectFile);
+									break;
 								default:
 									throw new BuildException("Unhandled project file type '{0}", ProjectFileFormat);
 							}
