#DEV-VR - Adding script version string to make sure AutoSDK gets run again
[To Fix CIS Builds for UE4 Lumin]

diff --git a/Engine/Source/Programs/UnrealBuildTool/Platform/Lumin/UEBuildLumin.cs b/Engine/Source/Programs/UnrealBuildTool/Platform/Lumin/UEBuildLumin.cs
--- a/Engine/Source/Programs/UnrealBuildTool/Platform/Lumin/UEBuildLumin.cs
+++ b/Engine/Source/Programs/UnrealBuildTool/Platform/Lumin/UEBuildLumin.cs
@@ -281,6 +281,10 @@ namespace UnrealBuildTool
 		{
 			return ExpectedSDKVersion;
 		}
+		protected override String GetRequiredScriptVersionString()
+		{
+			return "Lumin_15";
+		}
 
 		private string FindVersionNumber(string StringToFind, string[] AllLines)
 		{
