MagicLeap: Add quotes around the path to the tools (clang etc) in the mlsdk directory, so that it won't fail if the MLSDK directory has a space in the path.

diff --git a/Engine/Source/Programs/UnrealBuildTool/Platform/Lumin/LuminToolChain.cs b/Engine/Source/Programs/UnrealBuildTool/Platform/Lumin/LuminToolChain.cs
--- a/Engine/Source/Programs/UnrealBuildTool/Platform/Lumin/LuminToolChain.cs
+++ b/Engine/Source/Programs/UnrealBuildTool/Platform/Lumin/LuminToolChain.cs
@@ -109,7 +109,7 @@ namespace UnrealBuildTool
 						throw new BuildException("Unexpected output from mabu in --print-tools: '{0}'", Line);
 					}
 
-					ToolsDict.Add(Split[0].Trim(), Split[1].Trim());
+					ToolsDict.Add(Split[0].Trim(), "\"" + Split[1].Trim() + "\"");
 				}
 			}
 
