#DEV-VR - Fix for Lumin BP Projects not using their project specific ini files for compiling/cooking/packaging

diff --git a/Engine/Source/Programs/UnrealBuildTool/Platform/Lumin/LuminToolChain.cs b/Engine/Source/Programs/UnrealBuildTool/Platform/Lumin/LuminToolChain.cs
--- a/Engine/Source/Programs/UnrealBuildTool/Platform/Lumin/LuminToolChain.cs
+++ b/Engine/Source/Programs/UnrealBuildTool/Platform/Lumin/LuminToolChain.cs
@@ -136,8 +136,8 @@ namespace UnrealBuildTool
 
 		public bool UseVulkan()
 		{
-			// @todo Lumin: Switch to Lumin once we have sub-platform ini
-			ConfigHierarchy Ini = ConfigCache.ReadHierarchy(ConfigHierarchyType.Engine, DirectoryReference.FromFile(ProjectFile), UnrealTargetPlatform.Android);
+			DirectoryReference DirRef = (!string.IsNullOrEmpty(UnrealBuildTool.GetRemoteIniPath()) ? new DirectoryReference(UnrealBuildTool.GetRemoteIniPath()) : (ProjectFile != null ? ProjectFile.Directory : null));
+			ConfigHierarchy Ini = ConfigCache.ReadHierarchy(ConfigHierarchyType.Engine, DirRef, UnrealTargetPlatform.Lumin);
 			// go by string
 			bool bUseVulkan = true;
 			Ini.GetBool("/Script/LuminRuntimeSettings.LuminRuntimeSettings", "bUseVulkan", out bUseVulkan);
@@ -147,8 +147,8 @@ namespace UnrealBuildTool
 
 		public bool UseMobileRendering()
 		{
-			// @todo Lumin: Switch to Lumin once we have sub-platform ini
-			ConfigHierarchy Ini = ConfigCache.ReadHierarchy(ConfigHierarchyType.Engine, DirectoryReference.FromFile(ProjectFile), UnrealTargetPlatform.Android);
+			DirectoryReference DirRef = (!string.IsNullOrEmpty(UnrealBuildTool.GetRemoteIniPath()) ? new DirectoryReference(UnrealBuildTool.GetRemoteIniPath()) : (ProjectFile != null ? ProjectFile.Directory : null));
+			ConfigHierarchy Ini = ConfigCache.ReadHierarchy(ConfigHierarchyType.Engine, DirRef, UnrealTargetPlatform.Lumin);
 
 			// go by string
 			bool bUseMobileRendering = true;
