Fix for Lumin Toolchain failing to compile on Mac due to changes to the SDK

diff --git a/Engine/Source/Programs/UnrealBuildTool/Platform/Android/AndroidToolChain.cs b/Engine/Source/Programs/UnrealBuildTool/Platform/Android/AndroidToolChain.cs
--- a/Engine/Source/Programs/UnrealBuildTool/Platform/Android/AndroidToolChain.cs
+++ b/Engine/Source/Programs/UnrealBuildTool/Platform/Android/AndroidToolChain.cs
@@ -1439,8 +1439,17 @@ namespace UnrealBuildTool
 						CompileAction.WorkingDirectory = UnrealBuildTool.EngineSourceDirectory.FullName;
 						if(bExecuteCompilerThroughShell)
 						{
-							CompileAction.CommandPath = "cmd.exe";
-							CompileAction.CommandArguments = String.Format("/c \"{0} {1}\"", ClangPath, ResponseArgument);
+							if (BuildHostPlatform.Current.Platform == UnrealTargetPlatform.Win64)
+							{
+								CompileAction.CommandPath = "cmd.exe";
+								CompileAction.CommandArguments = String.Format("/c \"{0} {1}\"", ClangPath, ResponseArgument);
+							}
+							else
+							{
+								CompileAction.CommandPath = "/bin/sh";
+								CompileAction.CommandArguments = String.Format("-c \"{0} {1}\"", ClangPath, ResponseArgument);
+								CompileAction.CommandDescription = "Compile";
+							}
 						}
 						else
 						{
@@ -1662,8 +1671,17 @@ namespace UnrealBuildTool
 
 					if(bExecuteCompilerThroughShell)
 					{
-						LinkAction.CommandArguments = String.Format("/c \"{0} {1}\"", LinkAction.CommandPath, LinkAction.CommandArguments);
-						LinkAction.CommandPath = "cmd.exe";
+						if (BuildHostPlatform.Current.Platform == UnrealTargetPlatform.Win64)
+						{
+							LinkAction.CommandArguments = String.Format("/c \"{0} {1}\"", LinkAction.CommandPath, LinkAction.CommandArguments);
+							LinkAction.CommandPath = "cmd.exe";
+						}
+						else
+						{
+							LinkAction.CommandArguments = String.Format("-c \"{0} {1}\"", LinkAction.CommandPath, LinkAction.CommandArguments);
+							LinkAction.CommandPath = "/bin/sh";
+							LinkAction.CommandDescription = "Link";
+						}
 					}
 
 					// Windows can run into an issue with too long of a commandline when clang tries to call ld to link.
diff --git a/Engine/Source/Programs/UnrealBuildTool/Platform/Lumin/LuminToolChain.cs b/Engine/Source/Programs/UnrealBuildTool/Platform/Lumin/LuminToolChain.cs
--- a/Engine/Source/Programs/UnrealBuildTool/Platform/Lumin/LuminToolChain.cs
+++ b/Engine/Source/Programs/UnrealBuildTool/Platform/Lumin/LuminToolChain.cs
@@ -120,10 +120,7 @@ namespace UnrealBuildTool
 			ObjCopyPath = ToolsDict["OBJCOPY"];
 
 			// force the compiler to be executed through a command prompt instance
-			if(BuildHostPlatform.Current.Platform == UnrealTargetPlatform.Win64)
-			{
-				bExecuteCompilerThroughShell = true;
-			}
+			bExecuteCompilerThroughShell = true;
 
 			// toolchain params
 			ToolchainParamsArm64 = " --sysroot=\"" + Path.Combine(MLSDKPath, "lumin") + "\"";
