UBT: Fix exception when building a project that does not have any modules, due to failure to compile a valid rules assembly.

diff --git a/Engine/Source/Programs/UnrealBuildTool/System/RulesCompiler.cs b/Engine/Source/Programs/UnrealBuildTool/System/RulesCompiler.cs
--- a/Engine/Source/Programs/UnrealBuildTool/System/RulesCompiler.cs
+++ b/Engine/Source/Programs/UnrealBuildTool/System/RulesCompiler.cs
@@ -477,9 +477,16 @@ namespace UnrealBuildTool
 					TargetFiles.AddRange(FindAllRulesFiles(ProjectIntermediateSourceDirectory, RulesFileType.Target));
 				}
 
-				// Compile the assembly
+				// Compile the assembly. If there are no module or target files, just use the parent assembly.
 				FileReference AssemblyFileName = FileReference.Combine(ProjectDirectory, "Intermediate", "Build", "BuildRules", ProjectFileName.GetFileNameWithoutExtension() + "ModuleRules" + FrameworkAssemblyExtension);
-				ProjectRulesAssembly = new RulesAssembly(ProjectDirectory, ProjectPlugins, ModuleFiles, TargetFiles, ModuleFileToPluginInfo, AssemblyFileName, bContainsEngineModules: false, bUseBackwardsCompatibleDefaults: true, bReadOnly: UnrealBuildTool.IsProjectInstalled(), bSkipCompile: bSkipCompile, Parent: Parent);
+				if(ModuleFiles.Count == 0 && TargetFiles.Count == 0)
+				{
+					ProjectRulesAssembly = Parent;
+				}
+				else
+				{
+					ProjectRulesAssembly = new RulesAssembly(ProjectDirectory, ProjectPlugins, ModuleFiles, TargetFiles, ModuleFileToPluginInfo, AssemblyFileName, bContainsEngineModules: false, bUseBackwardsCompatibleDefaults: true, bReadOnly: UnrealBuildTool.IsProjectInstalled(), bSkipCompile: bSkipCompile, Parent: Parent);
+				}
 				LoadedAssemblyMap.Add(ProjectFileName, ProjectRulesAssembly);
 			}
 			return ProjectRulesAssembly;
