Merging //UE4/Dev-Main to Dev-Editor (//UE4/Dev-Editor)

diff --git a/Engine/Build/Android/Java/src/com/epicgames/ue4/GameActivity.java.template b/Engine/Build/Android/Java/src/com/epicgames/ue4/GameActivity.java.template
--- a/Engine/Build/Android/Java/src/com/epicgames/ue4/GameActivity.java.template
+++ b/Engine/Build/Android/Java/src/com/epicgames/ue4/GameActivity.java.template
@@ -1,4 +1,4 @@
-// Copyright 1998-2018 Epic Games, Inc. All Rights Reserved.
+// Copyright 1998-2019 Epic Games, Inc. All Rights Reserved.
 //This file needs to be here so the "ant" build step doesnt fail when looking for a /src folder.
 
 package com.epicgames.ue4;
diff --git a/Engine/Source/Developer/SimplygonSwarm/Private/SimplygonSwarm.cpp b/Engine/Source/Developer/SimplygonSwarm/Private/SimplygonSwarm.cpp
--- a/Engine/Source/Developer/SimplygonSwarm/Private/SimplygonSwarm.cpp
+++ b/Engine/Source/Developer/SimplygonSwarm/Private/SimplygonSwarm.cpp
@@ -1375,7 +1375,7 @@ bool ZipContentsForUpload(FString InputDirectoryPath, FString OutputFileName)
 					//Vertex color
 					if (bHasVertexColor)
 					{
-						FLinearColor LinearColor = VertexInstanceColors[VertexInstanceID];
+						FLinearColor LinearColor(VertexInstanceColors[VertexInstanceID]);
 						SsfColorMap.Items[TriangleIndex * 3 + DestCornerIndex].V[0] = LinearColor.R;
 						SsfColorMap.Items[TriangleIndex * 3 + DestCornerIndex].V[1] = LinearColor.G;
 						SsfColorMap.Items[TriangleIndex * 3 + DestCornerIndex].V[2] = LinearColor.B;
diff --git a/Engine/Source/Editor/UMGEditor/Private/Widgets/SGraphNodeCreateWidget.cpp b/Engine/Source/Editor/UMGEditor/Private/Widgets/SGraphNodeCreateWidget.cpp
--- a/Engine/Source/Editor/UMGEditor/Private/Widgets/SGraphNodeCreateWidget.cpp
+++ b/Engine/Source/Editor/UMGEditor/Private/Widgets/SGraphNodeCreateWidget.cpp
@@ -1,4 +1,4 @@
-// Copyright 1998-2018 Epic Games, Inc. All Rights Reserved.
+// Copyright 1998-2019 Epic Games, Inc. All Rights Reserved.
 
 #include "Widgets/SGraphNodeCreateWidget.h"
 #include "Modules/ModuleManager.h"
diff --git a/Engine/Source/Editor/UMGEditor/Private/Widgets/SGraphNodeCreateWidget.h b/Engine/Source/Editor/UMGEditor/Private/Widgets/SGraphNodeCreateWidget.h
--- a/Engine/Source/Editor/UMGEditor/Private/Widgets/SGraphNodeCreateWidget.h
+++ b/Engine/Source/Editor/UMGEditor/Private/Widgets/SGraphNodeCreateWidget.h
@@ -1,4 +1,4 @@
-// Copyright 1998-2018 Epic Games, Inc. All Rights Reserved.
+// Copyright 1998-2019 Epic Games, Inc. All Rights Reserved.
 
 #pragma once
 
diff --git a/Engine/Source/Editor/UnrealEd/Private/Kismet2/KismetReinstanceUtilities.cpp b/Engine/Source/Editor/UnrealEd/Private/Kismet2/KismetReinstanceUtilities.cpp
--- a/Engine/Source/Editor/UnrealEd/Private/Kismet2/KismetReinstanceUtilities.cpp
+++ b/Engine/Source/Editor/UnrealEd/Private/Kismet2/KismetReinstanceUtilities.cpp
@@ -1533,6 +1533,13 @@ UClass* FBlueprintCompileReinstancer::MoveCDOToNewClass(UClass* OwnerClass, cons
 		BPGDuplicatedClass->OverridenArchetypeForCDO = BPClassToReinstance->OverridenArchetypeForCDO;
 	}
 
+#if VALIDATE_UBER_GRAPH_PERSISTENT_FRAME
+	if (BPGDuplicatedClass && BPClassToReinstance)
+	{
+		BPGDuplicatedClass->UberGraphFunctionKey = BPClassToReinstance->UberGraphFunctionKey;
+	}
+#endif
+
 	UFunction* DuplicatedClassUberGraphFunction = BPGDuplicatedClass ? BPGDuplicatedClass->UberGraphFunction : nullptr;
 	if (DuplicatedClassUberGraphFunction)
 	{
diff --git a/Engine/Source/Programs/UnrealBuildTool/Platform/Windows/UEBuildWindows.cs b/Engine/Source/Programs/UnrealBuildTool/Platform/Windows/UEBuildWindows.cs
--- a/Engine/Source/Programs/UnrealBuildTool/Platform/Windows/UEBuildWindows.cs
+++ b/Engine/Source/Programs/UnrealBuildTool/Platform/Windows/UEBuildWindows.cs
@@ -384,7 +384,12 @@ namespace UnrealBuildTool
 		/// <summary>
 		/// The default compiler version to be used, if installed. 
 		/// </summary>
-		static readonly VersionNumber DefaultToolChainVersion = VersionNumber.Parse("14.13.26128");
+		static readonly VersionNumber DefaultClangToolChainVersion = VersionNumber.Parse("6.0.0");
+
+		/// <summary>
+		/// The default compiler version to be used, if installed. 
+		/// </summary>
+		static readonly VersionNumber DefaultVisualStudioToolChainVersion = VersionNumber.Parse("14.13.26128");
 
 		/// <summary>
 		/// The default Windows SDK version to be used, if installed.
@@ -912,6 +917,16 @@ namespace UnrealBuildTool
 			}
 			else
 			{
+				VersionNumber DefaultToolChainVersion;
+				if(Compiler == WindowsCompiler.Clang)
+				{
+					DefaultToolChainVersion = DefaultClangToolChainVersion;
+				}
+				else
+				{
+					DefaultToolChainVersion = DefaultVisualStudioToolChainVersion;
+				}
+
 				if(ToolChainVersionToDir.ContainsKey(DefaultToolChainVersion))
 				{
 					ToolChainVersion = DefaultToolChainVersion;
diff --git a/Engine/Source/Runtime/Core/Public/HAL/Platform.h b/Engine/Source/Runtime/Core/Public/HAL/Platform.h
--- a/Engine/Source/Runtime/Core/Public/HAL/Platform.h
+++ b/Engine/Source/Runtime/Core/Public/HAL/Platform.h
@@ -135,7 +135,7 @@
 
 //@port Identify the platform here and include the platform header to setup the platform types, etc
 #if PLATFORM_WINDOWS
-	#include "Windows/WindowsPlatform.h"
+	#include "Windows/WIndowsPlatform.h"	// this is the actual filename on disk, alas cannot be easily renamed in source control
 #elif PLATFORM_PS4
 	#include "PS4/PS4Platform.h"
 #elif PLATFORM_XBOXONE
diff --git a/Engine/Source/Runtime/CoreUObject/Private/Internationalization/StabilizeLocalizationKeys.cpp b/Engine/Source/Runtime/CoreUObject/Private/Internationalization/StabilizeLocalizationKeys.cpp
--- a/Engine/Source/Runtime/CoreUObject/Private/Internationalization/StabilizeLocalizationKeys.cpp
+++ b/Engine/Source/Runtime/CoreUObject/Private/Internationalization/StabilizeLocalizationKeys.cpp
@@ -1,4 +1,4 @@
-// Copyright 1998-2018 Epic Games, Inc. All Rights Reserved.
+// Copyright 1998-2019 Epic Games, Inc. All Rights Reserved.
 
 #include "Internationalization/StabilizeLocalizationKeys.h"
 #include "UObject/TextProperty.h"
diff --git a/Engine/Source/Runtime/CoreUObject/Public/Internationalization/StabilizeLocalizationKeys.h b/Engine/Source/Runtime/CoreUObject/Public/Internationalization/StabilizeLocalizationKeys.h
--- a/Engine/Source/Runtime/CoreUObject/Public/Internationalization/StabilizeLocalizationKeys.h
+++ b/Engine/Source/Runtime/CoreUObject/Public/Internationalization/StabilizeLocalizationKeys.h
@@ -1,4 +1,4 @@
-// Copyright 1998-2018 Epic Games, Inc. All Rights Reserved.
+// Copyright 1998-2019 Epic Games, Inc. All Rights Reserved.
 
 #pragma once
 
diff --git a/Engine/Source/Runtime/CoreUObject/Public/Serialization/BulkDataReader.h b/Engine/Source/Runtime/CoreUObject/Public/Serialization/BulkDataReader.h
--- a/Engine/Source/Runtime/CoreUObject/Public/Serialization/BulkDataReader.h
+++ b/Engine/Source/Runtime/CoreUObject/Public/Serialization/BulkDataReader.h
@@ -1,4 +1,4 @@
-// Copyright 1998-2018 Epic Games, Inc. All Rights Reserved.
+// Copyright 1998-2019 Epic Games, Inc. All Rights Reserved.
 
 #pragma once
 
diff --git a/Engine/Source/Runtime/CoreUObject/Public/Serialization/BulkDataWriter.h b/Engine/Source/Runtime/CoreUObject/Public/Serialization/BulkDataWriter.h
--- a/Engine/Source/Runtime/CoreUObject/Public/Serialization/BulkDataWriter.h
+++ b/Engine/Source/Runtime/CoreUObject/Public/Serialization/BulkDataWriter.h
@@ -1,4 +1,4 @@
-// Copyright 1998-2018 Epic Games, Inc. All Rights Reserved.
+// Copyright 1998-2019 Epic Games, Inc. All Rights Reserved.
 
 #pragma once
 
