Fixed a problem with MacToolChain using incorrect module paths for linking if an old app bundle with dylibs inside the bundle was present in Binaries/Mac

diff --git a/Engine/Source/Programs/UnrealBuildTool/Platform/Mac/MacToolChain.cs b/Engine/Source/Programs/UnrealBuildTool/Platform/Mac/MacToolChain.cs
--- a/Engine/Source/Programs/UnrealBuildTool/Platform/Mac/MacToolChain.cs
+++ b/Engine/Source/Programs/UnrealBuildTool/Platform/Mac/MacToolChain.cs
@@ -801,25 +801,7 @@ namespace UnrealBuildTool
 				{
 					foreach (string Library in EngineAndGameLibraries)
 					{
-						string LibraryPath = Library;
-						if (!File.Exists(Library))
-						{
-							string LibraryDir = Path.GetDirectoryName(Library);
-							string LibraryName = Path.GetFileName(Library);
-							string AppBundleName = "UE4Editor";
-							if (LibraryName.Contains("UE4Editor-Mac-"))
-							{
-								string[] Parts = LibraryName.Split('-');
-								AppBundleName += "-" + Parts[1] + "-" + Parts[2];
-							}
-							AppBundleName += ".app";
-							LibraryPath = LibraryDir + "/" + AppBundleName + "/Contents/MacOS/" + LibraryName;
-							if (!File.Exists(LibraryPath))
-							{
-								LibraryPath = Library;
-							}
-						}
-						LinkCommand += " \"" + LibraryPath + "\"";
+						LinkCommand += " \"" + Library + "\"";
 					}
 				}
 				else
