UBT: Prevent DebugGame compile environment from having different ORIGINAL_FILE_NAME macro to the development build environment. This is only needed when compiling resources for binaries. Also prevent the precompiled manifest being written to if it doesn't change.

diff --git a/Engine/Source/Programs/UnrealBuildTool/Configuration/UEBuildBinary.cs b/Engine/Source/Programs/UnrealBuildTool/Configuration/UEBuildBinary.cs
--- a/Engine/Source/Programs/UnrealBuildTool/Configuration/UEBuildBinary.cs
+++ b/Engine/Source/Programs/UnrealBuildTool/Configuration/UEBuildBinary.cs
@@ -652,11 +652,6 @@ namespace UnrealBuildTool
 			CppCompileEnvironment BinaryCompileEnvironment = new CppCompileEnvironment(GlobalCompileEnvironment);
 			BinaryCompileEnvironment.bIsBuildingDLL = IsBuildingDll(Type);
 			BinaryCompileEnvironment.bIsBuildingLibrary = IsBuildingLibrary(Type);
-
-			// @todo: This should be in some Windows code somewhere...
-			// Set the original file name macro; used in Default.rc2 to set the binary metadata fields.
-			BinaryCompileEnvironment.Definitions.Add("ORIGINAL_FILE_NAME=\"" + OutputFilePaths[0].GetFileName() + "\"");
-
 			return BinaryCompileEnvironment;
 		}
 
diff --git a/Engine/Source/Programs/UnrealBuildTool/Configuration/UEBuildModuleCPP.cs b/Engine/Source/Programs/UnrealBuildTool/Configuration/UEBuildModuleCPP.cs
--- a/Engine/Source/Programs/UnrealBuildTool/Configuration/UEBuildModuleCPP.cs
+++ b/Engine/Source/Programs/UnrealBuildTool/Configuration/UEBuildModuleCPP.cs
@@ -490,6 +490,11 @@ namespace UnrealBuildTool
 			if(InputFiles.RCFiles.Count > 0)
 			{
 				CppCompileEnvironment ResourceCompileEnvironment = new CppCompileEnvironment(BinaryCompileEnvironment);
+				if(Binary != null)
+				{
+					// @todo: This should be in some Windows code somewhere...
+					ResourceCompileEnvironment.Definitions.Add("ORIGINAL_FILE_NAME=\"" + Binary.OutputFilePaths[0].GetFileName() + "\"");
+				}
 				LinkInputFiles.AddRange(ToolChain.CompileRCFiles(ResourceCompileEnvironment, InputFiles.RCFiles, IntermediateDirectory, Makefile.Actions).ObjectFiles);
 			}
 
@@ -500,7 +505,7 @@ namespace UnrealBuildTool
 
 				PrecompiledManifest Manifest = new PrecompiledManifest();
 				Manifest.OutputFiles.AddRange(LinkInputFiles.Select(x => x.Location));
-				Manifest.Write(PrecompiledManifestLocation);
+				Manifest.WriteIfModified(PrecompiledManifestLocation);
 			}
 
 			return LinkInputFiles;
diff --git a/Engine/Source/Programs/UnrealBuildTool/System/PrecompiledManifest.cs b/Engine/Source/Programs/UnrealBuildTool/System/PrecompiledManifest.cs
--- a/Engine/Source/Programs/UnrealBuildTool/System/PrecompiledManifest.cs
+++ b/Engine/Source/Programs/UnrealBuildTool/System/PrecompiledManifest.cs
@@ -1,6 +1,7 @@
 // Copyright 1998-2019 Epic Games, Inc. All Rights Reserved.
 using System;
 using System.Collections.Generic;
+using System.IO;
 using System.Linq;
 using System.Text;
 using System.Threading.Tasks;
@@ -75,10 +76,12 @@ namespace UnrealBuildTool
 		/// Write the receipt to disk.
 		/// </summary>
 		/// <param name="Location">Output filename</param>
-		public void Write(FileReference Location)
+		public void WriteIfModified(FileReference Location)
 		{
 			DirectoryReference BaseDir = Location.Directory;
-			using (JsonWriter Writer = new JsonWriter(Location.FullName))
+
+			MemoryStream MemoryStream = new MemoryStream();
+			using (JsonWriter Writer = new JsonWriter(new StreamWriter(MemoryStream)))
 			{
 				Writer.WriteObjectStart();
 
@@ -91,6 +94,12 @@ namespace UnrealBuildTool
 
 				Writer.WriteObjectEnd();
 			}
+
+			byte[] Data = MemoryStream.ToArray();
+			if(!FileReference.Exists(Location) || !Utils.ByteArraysEqual(Data, FileReference.ReadAllBytes(Location)))
+			{
+				FileReference.WriteAllBytes(Location, Data);
+			}
 		}
 	}
 }
diff --git a/Engine/Source/Programs/UnrealBuildTool/System/Utils.cs b/Engine/Source/Programs/UnrealBuildTool/System/Utils.cs
--- a/Engine/Source/Programs/UnrealBuildTool/System/Utils.cs
+++ b/Engine/Source/Programs/UnrealBuildTool/System/Utils.cs
@@ -1075,5 +1075,29 @@ namespace UnrealBuildTool
 				}
 			}
 		}
+
+		/// <summary>
+		/// Compares two byte arrays for equality
+		/// </summary>
+		/// <param name="A">The first byte array</param>
+		/// <param name="B">The second byte array</param>
+		/// <returns>True if the two arrays are equal, false otherwise</returns>
+		public static bool ByteArraysEqual(byte[] A, byte[] B)
+		{
+			if(A.Length != B.Length)
+			{
+				return false;
+			}
+
+			for(int Idx = 0; Idx < A.Length; Idx++)
+			{
+				if(A[Idx] != B[Idx])
+				{
+					return false;
+				}
+			}
+
+			return true;
+		}
 	}
 }
