HTML5 - fix log spam

replicated to //UE4/Release-4.22

#rnx

diff --git a/Engine/Build/HTML5/project_template.js b/Engine/Build/HTML5/project_template.js
--- a/Engine/Build/HTML5/project_template.js
+++ b/Engine/Build/HTML5/project_template.js
@@ -54,6 +54,9 @@ const serveCompressedAssets = %SERVE_COMPRESSED%;
 // If true, serve out 'UE4Game.data.gz', if false, serve out 'UE4Game.data'.
 //const dataFileIsGzipCompressed = false;
 
+console.log("Emscripten version: %EMSDK_VERSION%");
+console.log("Emscripten configuration: %EMSDK_CONFIG%");
+
 
 
 
diff --git a/Engine/Source/Programs/AutomationTool/HTML5/HTML5Platform.Automation.cs b/Engine/Source/Programs/AutomationTool/HTML5/HTML5Platform.Automation.cs
--- a/Engine/Source/Programs/AutomationTool/HTML5/HTML5Platform.Automation.cs
+++ b/Engine/Source/Programs/AutomationTool/HTML5/HTML5Platform.Automation.cs
@@ -396,6 +396,18 @@ public class HTML5Platform : Platform
 					LineStr = LineStr.Replace("%OFFSCREENCANVAS%", bMultithreading_UseOffscreenCanvas ? "true" : "false");
 				}
 
+				if (LineStr.Contains("%EMSDK_VERSION%"))
+				{
+					string escpath = HTML5SDKInfo.EMSCRIPTEN_ROOT.Replace("\\", "/");
+					LineStr = LineStr.Replace("%EMSDK_VERSION%", (ProjectConfiguration == "Shipping") ? HTML5SDKInfo.EmscriptenVersion() : escpath);
+				}
+
+				if (LineStr.Contains("%EMSDK_CONFIG%"))
+				{
+					string escpath = HTML5SDKInfo.DOT_EMSCRIPTEN.Replace("\\", "/");
+					LineStr = LineStr.Replace("%EMSDK_CONFIG%", (ProjectConfiguration == "Shipping") ? "" : escpath);
+				}
+
 				outputContents.AppendLine(LineStr);
 			}
 		}
diff --git a/Engine/Source/Programs/UnrealBuildTool/Platform/HTML5/HTML5SDKInfo.cs b/Engine/Source/Programs/UnrealBuildTool/Platform/HTML5/HTML5SDKInfo.cs
--- a/Engine/Source/Programs/UnrealBuildTool/Platform/HTML5/HTML5SDKInfo.cs
+++ b/Engine/Source/Programs/UnrealBuildTool/Platform/HTML5/HTML5SDKInfo.cs
@@ -98,7 +98,6 @@ namespace UnrealBuildTool
 			string config = Environment.GetEnvironmentVariable("EM_CONFIG"); // This is either a string containing the config directly, or points to a file
 			if (config != null && File.Exists(config))
 			{
-				Log.TraceInformation( "NOTE [ReadEmscriptenConfigFile]: using EM_CONFIG=" + config );
 				config = File.ReadAllText(config);
 			}
 			return config;
diff --git a/Engine/Source/Programs/UnrealBuildTool/Platform/HTML5/HTML5ToolChain.cs b/Engine/Source/Programs/UnrealBuildTool/Platform/HTML5/HTML5ToolChain.cs
--- a/Engine/Source/Programs/UnrealBuildTool/Platform/HTML5/HTML5ToolChain.cs
+++ b/Engine/Source/Programs/UnrealBuildTool/Platform/HTML5/HTML5ToolChain.cs
@@ -88,6 +88,8 @@ namespace UnrealBuildTool
 				Environment.SetEnvironmentVariable("EM_CACHE", HTML5SDKInfo.EMSCRIPTEN_CACHE);
 				Environment.SetEnvironmentVariable("EMCC_TEMP_DIR", TempDir);
 			}
+
+			Log.TraceInformation("*** Emscripten Config File: " + Environment.GetEnvironmentVariable("EM_CONFIG"));
 		}
 
 		string GetSharedArguments_Global(CppConfiguration Configuration, bool bOptimizeForSize, string Architecture, bool bEnableShadowVariableWarnings, bool bShadowVariableWarningsAsErrors, bool bEnableUndefinedIdentifierWarnings, bool bUndefinedIdentifierWarningsAsErrors, bool bUseInlining)
@@ -205,14 +207,14 @@ namespace UnrealBuildTool
 //			Environment.SetEnvironmentVariable("EMCC_CORES", "8");
 //			Environment.SetEnvironmentVariable("EMCC_OPTIMIZE_NORMALLY", "1");
 
-			// Linux builds needs this - or else system clang will be attempted to be picked up instead of UE4's
-			// TODO: test on other platforms to remove this if() check
 			if (BuildHostPlatform.Current.Platform == UnrealTargetPlatform.Linux)
 			{
+				// Packaging on Linux needs this - or else system clang will be attempted to be picked up instead of UE4's included emsdk
 				Environment.SetEnvironmentVariable(HTML5SDKInfo.PLATFORM_USER_HOME, HTML5SDKInfo.HTML5Intermediatory);
 			}
 			if (BuildHostPlatform.Current.Platform == UnrealTargetPlatform.Win64 || BuildHostPlatform.Current.Platform == UnrealTargetPlatform.Win32)
 			{
+				// Packaging on Window needs this - zap any existing HOME environment variables to prevent any accidental pick ups
 				Environment.SetEnvironmentVariable("HOME", "");
 			}
 			return Result;
