Disable XGE when building on Linux for Android

diff --git a/Engine/Source/Programs/UnrealBuildTool/Platform/Android/UEBuildAndroid.cs b/Engine/Source/Programs/UnrealBuildTool/Platform/Android/UEBuildAndroid.cs
--- a/Engine/Source/Programs/UnrealBuildTool/Platform/Android/UEBuildAndroid.cs
+++ b/Engine/Source/Programs/UnrealBuildTool/Platform/Android/UEBuildAndroid.cs
@@ -108,7 +108,8 @@ namespace UnrealBuildTool
 
 		public override bool CanUseXGE()
 		{
-			return !(Environment.GetEnvironmentVariable("IsBuildMachine") == "1");
+			// Disable when building on Linux
+			return BuildHostPlatform.Current.Platform != UnrealTargetPlatform.Linux && !(Environment.GetEnvironmentVariable("IsBuildMachine") == "1");
 		}
 
 		public override bool IsBuildProduct(string FileName, string[] NamePrefixes, string[] NameSuffixes)
