add missing XMA2 streaming define, currently disabled.

diff --git a/Engine/Source/Programs/UnrealBuildTool/Configuration/UEBuildTarget.cs b/Engine/Source/Programs/UnrealBuildTool/Configuration/UEBuildTarget.cs
--- a/Engine/Source/Programs/UnrealBuildTool/Configuration/UEBuildTarget.cs
+++ b/Engine/Source/Programs/UnrealBuildTool/Configuration/UEBuildTarget.cs
@@ -2901,6 +2901,9 @@ namespace UnrealBuildTool
 			// Toggle to enable vorbis for audio streaming where available
 			GlobalCompileEnvironment.Definitions.Add("USE_VORBIS_FOR_STREAMING=1");
 
+			// Toggle to enable XMA for audio streaming where available (XMA2 only supports up to stereo streams - surround streams will fall back to Vorbis etc)
+			GlobalCompileEnvironment.Definitions.Add("USE_XMA2_FOR_STREAMING=0");
+
 			// Add the 'Engine/Source' path as a global include path for all modules
 			GlobalCompileEnvironment.UserIncludePaths.Add(UnrealBuildTool.EngineSourceDirectory);
 
