enable XMA2 streaming on platforms that support it

diff --git a/Engine/Source/Programs/UnrealBuildTool/Configuration/UEBuildTarget.cs b/Engine/Source/Programs/UnrealBuildTool/Configuration/UEBuildTarget.cs
--- a/Engine/Source/Programs/UnrealBuildTool/Configuration/UEBuildTarget.cs
+++ b/Engine/Source/Programs/UnrealBuildTool/Configuration/UEBuildTarget.cs
@@ -2902,7 +2902,7 @@ namespace UnrealBuildTool
 			GlobalCompileEnvironment.Definitions.Add("USE_VORBIS_FOR_STREAMING=1");
 
 			// Toggle to enable XMA for audio streaming where available (XMA2 only supports up to stereo streams - surround streams will fall back to Vorbis etc)
-			GlobalCompileEnvironment.Definitions.Add("USE_XMA2_FOR_STREAMING=0");
+			GlobalCompileEnvironment.Definitions.Add("USE_XMA2_FOR_STREAMING=1");
 
 			// Add the 'Engine/Source' path as a global include path for all modules
 			GlobalCompileEnvironment.UserIncludePaths.Add(UnrealBuildTool.EngineSourceDirectory);
