Linux: harden the UBT makefile generator against .target.cs

diff --git a/Engine/Source/Programs/UnrealBuildTool/ProjectFiles/Make/MakefileGenerator.cs b/Engine/Source/Programs/UnrealBuildTool/ProjectFiles/Make/MakefileGenerator.cs
--- a/Engine/Source/Programs/UnrealBuildTool/ProjectFiles/Make/MakefileGenerator.cs
+++ b/Engine/Source/Programs/UnrealBuildTool/ProjectFiles/Make/MakefileGenerator.cs
@@ -106,7 +106,7 @@ namespace UnrealBuildTool
 					}
 
 					string TargetFileName = TargetFile.TargetFilePath.GetFileNameWithoutExtension();
-					string Basename = TargetFileName.Substring(0, TargetFileName.LastIndexOf(".Target"));
+					string Basename = TargetFileName.Substring(0, TargetFileName.LastIndexOf(".Target", StringComparison.InvariantCultureIgnoreCase));
 
 					foreach (UnrealTargetConfiguration CurConfiguration in Enum.GetValues(typeof(UnrealTargetConfiguration)))
 					{
@@ -141,7 +141,7 @@ namespace UnrealBuildTool
 					}
 
 					string TargetFileName = TargetFile.TargetFilePath.GetFileNameWithoutExtension();
-					string Basename = TargetFileName.Substring(0, TargetFileName.LastIndexOf(".Target"));
+					string Basename = TargetFileName.Substring(0, TargetFileName.LastIndexOf(".Target", StringComparison.InvariantCultureIgnoreCase));
 
 					if (Basename == GameProjectName || Basename == (GameProjectName + "Editor"))
 					{
