UBT: Fix the user's selected Visual Studio version not being respected when generating project files.

diff --git a/Engine/Source/Programs/UnrealBuildTool/ProjectFiles/ProjectFileGenerator.cs b/Engine/Source/Programs/UnrealBuildTool/ProjectFiles/ProjectFileGenerator.cs
--- a/Engine/Source/Programs/UnrealBuildTool/ProjectFiles/ProjectFileGenerator.cs
+++ b/Engine/Source/Programs/UnrealBuildTool/ProjectFiles/ProjectFileGenerator.cs
@@ -540,6 +540,26 @@ namespace UnrealBuildTool
 					Format = ProjectFileFormat.KDevelop;
 					return true;
 				}
+				else if (PreferredAccessor == "visualstudiosourcecodeaccessor")
+				{
+					Format = ProjectFileFormat.VisualStudio;
+					return true;
+				}
+				else if (PreferredAccessor == "visualstudio2015")
+				{
+					Format = ProjectFileFormat.VisualStudio2015;
+					return true;
+				}
+				else if (PreferredAccessor == "visualstudio2017")
+				{
+					Format = ProjectFileFormat.VisualStudio2017;
+					return true;
+				}
+				else if (PreferredAccessor == "visualstudio2019")
+				{
+					Format = ProjectFileFormat.VisualStudio2019;
+					return true;
+				}
 			}
 
 			Format = ProjectFileFormat.VisualStudio;
