UBT: Prevent hot reload attempting to build non-C++ modules.

diff --git a/Engine/Source/Programs/UnrealBuildTool/Configuration/UEBuildTarget.cs b/Engine/Source/Programs/UnrealBuildTool/Configuration/UEBuildTarget.cs
--- a/Engine/Source/Programs/UnrealBuildTool/Configuration/UEBuildTarget.cs
+++ b/Engine/Source/Programs/UnrealBuildTool/Configuration/UEBuildTarget.cs
@@ -2035,7 +2035,7 @@ namespace UnrealBuildTool
 				{
 					if(!UnrealBuildTool.IsProjectInstalled() || EnabledPlugins.Where(x => x.Type == PluginType.Mod).Any(x => Binary.OutputFilePaths[0].IsUnderDirectory(x.Directory)))
 					{
-						HotReloadModuleNames.UnionWith(GameModules.Select(x => x.Name));
+						HotReloadModuleNames.UnionWith(GameModules.OfType<UEBuildModuleCPP>().Select(x => x.Name));
 					}
 				}
 			}
