UBT: Add a global mutex around UBT initializing the XML config cache. This happens before the global mutex is acquired, so needs to be guarded separately.

diff --git a/Engine/Source/Programs/UnrealBuildTool/UnrealBuildTool.cs b/Engine/Source/Programs/UnrealBuildTool/UnrealBuildTool.cs
--- a/Engine/Source/Programs/UnrealBuildTool/UnrealBuildTool.cs
+++ b/Engine/Source/Programs/UnrealBuildTool/UnrealBuildTool.cs
@@ -390,8 +390,12 @@ namespace UnrealBuildTool
 				{
 					using(Timeline.ScopeEvent("XmlConfig.ReadConfigFiles()"))
 					{
-						FileReference XmlConfigCache = Arguments.GetFileReferenceOrDefault("-XmlConfigCache=", null);
-						XmlConfig.ReadConfigFiles(XmlConfigCache);
+						string XmlConfigMutexName = SingleInstanceMutex.GetUniqueMutexForPath("UnrealBuildTool_Mutex_XmlConfig", Assembly.GetExecutingAssembly().CodeBase);
+						using(SingleInstanceMutex XmlConfigMutex = new SingleInstanceMutex(XmlConfigMutexName, true))
+						{
+							FileReference XmlConfigCache = Arguments.GetFileReferenceOrDefault("-XmlConfigCache=", null);
+							XmlConfig.ReadConfigFiles(XmlConfigCache);
+						}
 					}
 				}
 
