UBT: Update toolset version number for Visual Studio 2019 preview 2.

diff --git a/Engine/Source/Programs/UnrealBuildTool/ProjectFiles/VisualStudio/VCProjectFileGenerator.cs b/Engine/Source/Programs/UnrealBuildTool/ProjectFiles/VisualStudio/VCProjectFileGenerator.cs
--- a/Engine/Source/Programs/UnrealBuildTool/ProjectFiles/VisualStudio/VCProjectFileGenerator.cs
+++ b/Engine/Source/Programs/UnrealBuildTool/ProjectFiles/VisualStudio/VCProjectFileGenerator.cs
@@ -222,7 +222,7 @@ namespace UnrealBuildTool
 				case VCProjectFileFormat.VisualStudio2017:
                     return "v141";
 				case VCProjectFileFormat.VisualStudio2019:
-					return "v141"; // Correct as of VS2019 Preview 1
+					return "v142"; // Correct as of VS2019 Preview 2
 
             }
 			return string.Empty;
