UBT: Write exception info if XGE fails to start.

#rnx

diff --git a/Engine/Source/Programs/UnrealBuildTool/Executors/XGE.cs b/Engine/Source/Programs/UnrealBuildTool/Executors/XGE.cs
--- a/Engine/Source/Programs/UnrealBuildTool/Executors/XGE.cs
+++ b/Engine/Source/Programs/UnrealBuildTool/Executors/XGE.cs
@@ -507,8 +507,9 @@ namespace UnrealBuildTool
 				XGEProcess.WaitForExit();
 				return XGEProcess.ExitCode == 0;
 			}
-			catch (Exception)
+			catch (Exception Ex)
 			{
+				Log.WriteException(Ex, null);
 				return false;
 			}
 		}
