Missing file.

#rnx

diff --git a/Engine/Source/Programs/UnrealBuildTool/System/PrecompiledManifest.cs b/Engine/Source/Programs/UnrealBuildTool/System/PrecompiledManifest.cs
--- a/Engine/Source/Programs/UnrealBuildTool/System/PrecompiledManifest.cs
+++ b/Engine/Source/Programs/UnrealBuildTool/System/PrecompiledManifest.cs
@@ -95,11 +95,7 @@ namespace UnrealBuildTool
 				Writer.WriteObjectEnd();
 			}
 
-			byte[] Data = MemoryStream.ToArray();
-			if(!FileReference.Exists(Location) || !Utils.ByteArraysEqual(Data, FileReference.ReadAllBytes(Location)))
-			{
-				FileReference.WriteAllBytes(Location, Data);
-			}
+			FileReference.WriteAllBytesIfDifferent(Location, MemoryStream.ToArray());
 		}
 	}
 }
