UBT: Fix module manifest not being updated correctly when -ModuleWithSuffix is passed on the command line.

#rnx

diff --git a/Engine/Source/Programs/UnrealBuildTool/Configuration/UEBuildTarget.cs b/Engine/Source/Programs/UnrealBuildTool/Configuration/UEBuildTarget.cs
--- a/Engine/Source/Programs/UnrealBuildTool/Configuration/UEBuildTarget.cs
+++ b/Engine/Source/Programs/UnrealBuildTool/Configuration/UEBuildTarget.cs
@@ -1401,7 +1401,7 @@ namespace UnrealBuildTool
 				// Create an action which to generate the receipts
 				WriteMetadataTargetInfo MetadataTargetInfo = new WriteMetadataTargetInfo(ProjectFile, VersionFile, ReceiptFileName, Receipt, FileNameToModuleManifest);
 				FileReference MetadataTargetFile = FileReference.Combine(ProjectIntermediateDirectory, "Metadata.dat");
-				BinaryFormatterUtils.Save(MetadataTargetFile, MetadataTargetInfo);
+				BinaryFormatterUtils.SaveIfDifferent(MetadataTargetFile, MetadataTargetInfo);
 
 				StringBuilder WriteMetadataArguments = new StringBuilder();
 				WriteMetadataArguments.AppendFormat("-Input={0}", Utils.MakePathSafeToUseWithCommandLine(MetadataTargetFile));
@@ -1415,6 +1415,7 @@ namespace UnrealBuildTool
 				WriteMetadataAction.WorkingDirectory = UnrealBuildTool.EngineSourceDirectory;
 				WriteMetadataAction.StatusDescription = ReceiptFileName.GetFileName();
 				WriteMetadataAction.bCanExecuteRemotely = false;
+				WriteMetadataAction.PrerequisiteItems.Add(FileItem.GetItemByFileReference(MetadataTargetFile));
 				WriteMetadataAction.PrerequisiteItems.AddRange(Makefile.OutputItems);
 				WriteMetadataAction.ProducedItems.Add(FileItem.GetItemByFileReference(ReceiptFileName));
 				Makefile.Actions.Add(WriteMetadataAction);
diff --git a/Engine/Source/Programs/UnrealBuildTool/System/HotReload.cs b/Engine/Source/Programs/UnrealBuildTool/System/HotReload.cs
--- a/Engine/Source/Programs/UnrealBuildTool/System/HotReload.cs
+++ b/Engine/Source/Programs/UnrealBuildTool/System/HotReload.cs
@@ -643,7 +643,7 @@ namespace UnrealBuildTool
 					{
 						foreach (KeyValuePair<string, string> Manifest in FileNameToVersionManifest.Value.ModuleNameToFileName)
 						{
-							FileReference OriginalFile = FileReference.Combine(FileNameToVersionManifest.Key.Directory, Manifest.Key);
+							FileReference OriginalFile = FileReference.Combine(FileNameToVersionManifest.Key.Directory, Manifest.Value);
 
 							FileReference HotReloadFile;
 							if(OriginalFileToHotReloadFile.TryGetValue(OriginalFile, out HotReloadFile))
@@ -656,8 +656,12 @@ namespace UnrealBuildTool
 
 					// Write the hot-reload metadata file and update the argument list
 					FileReference HotReloadTargetInfoFile = FileReference.Combine(TargetInfoFile.Directory, "Metadata-HotReload.dat");
-					BinaryFormatterUtils.Save(HotReloadTargetInfoFile, TargetInfo);
-					Arguments = Arguments.Substring(0, FileNameIdx) + HotReloadTargetInfoFile + Arguments.Substring(FileNameEndIdx);
+					BinaryFormatterUtils.SaveIfDifferent(HotReloadTargetInfoFile, TargetInfo);
+
+					Action.PrerequisiteItems.RemoveAll(x => x.Location == TargetInfoFile);
+					Action.PrerequisiteItems.Add(FileItem.GetItemByFileReference(HotReloadTargetInfoFile));
+
+					Action.CommandArguments = Arguments.Substring(0, FileNameIdx) + HotReloadTargetInfoFile + Arguments.Substring(FileNameEndIdx);
 				}
 			}
 		}
