PR #5488: Fixed packaging of framework resource bundles when doing a remote build for iOS (Contributed by [REDACTED_PERSON])

diff --git a/Engine/Source/Programs/UnrealBuildTool/ToolChain/RemoteMac.cs b/Engine/Source/Programs/UnrealBuildTool/ToolChain/RemoteMac.cs
--- a/Engine/Source/Programs/UnrealBuildTool/ToolChain/RemoteMac.cs
+++ b/Engine/Source/Programs/UnrealBuildTool/ToolChain/RemoteMac.cs
@@ -535,12 +535,20 @@ namespace UnrealBuildTool
 			FilesToDownload.AddRange(Manifest.BuildProducts.Select(x => new FileReference(x)));
 			DownloadFiles(FilesToDownload);
 
+			// Copy remote FrameworkAssets directory as it could contain resource bundles that must be packaged locally.
+			DirectoryReference BaseDir = DirectoryReference.FromFile(TargetDesc.ProjectFile) ?? UnrealBuildTool.EngineDirectory;
+			DirectoryReference FrameworkAssetsDir = DirectoryReference.Combine(BaseDir, "Intermediate", "IOS", "FrameworkAssets");
+
+			Log.TraceInformation("[Remote] Downloading {0}", FrameworkAssetsDir);
+			DownloadDirectory(FrameworkAssetsDir);
+
 			// Write out all the local manifests
 			foreach(FileReference LocalManifestFile in LocalManifestFiles)
 			{
 				Log.TraceInformation("[Remote] Writing {0}", LocalManifestFile);
 				Utils.WriteClass<BuildManifest>(Manifest, LocalManifestFile.FullName, "");
 			}
+
 			return true;
 		}
 
@@ -947,6 +955,28 @@ namespace UnrealBuildTool
 			}
 		}
 
+		/// <summary>
+		/// Download a directory from the remote Mac
+		/// </summary>
+		/// <param name="LocalDirectory">Directory to download</param>
+		private void DownloadDirectory(DirectoryReference LocalDirectory)
+		{
+			DirectoryReference.CreateDirectory(LocalDirectory);
+
+			string RemoteDirectory = GetRemotePath(LocalDirectory);
+
+			List<string> Arguments = new List<string>(CommonRsyncArguments);
+			Arguments.Add(String.Format("--rsync-path=\"[ ! -d {0} ] || rsync\"", EscapeShellArgument(RemoteDirectory)));
+			Arguments.Add(String.Format("\"{0}@{1}\":'{2}/'", UserName, ServerName, RemoteDirectory));
+			Arguments.Add(String.Format("\"{0}/\"", GetLocalCygwinPath(LocalDirectory)));
+
+			int Result = Rsync(String.Join(" ", Arguments));
+			if (Result != 0)
+			{
+				throw new BuildException("Unable to download '{0}' from the remote Mac (exit code {1}).", LocalDirectory, Result);
+			}
+		}
+
 		/// <summary>
 		/// Execute Rsync
 		/// </summary>
