UBT: Don't throw an exception when scanning for restricted folder references from modules that are precompiled; they don't have source code, so won't have a list of source directories.

diff --git a/Engine/Source/Programs/UnrealBuildTool/Configuration/UEBuildModuleCPP.cs b/Engine/Source/Programs/UnrealBuildTool/Configuration/UEBuildModuleCPP.cs
--- a/Engine/Source/Programs/UnrealBuildTool/Configuration/UEBuildModuleCPP.cs
+++ b/Engine/Source/Programs/UnrealBuildTool/Configuration/UEBuildModuleCPP.cs
@@ -64,12 +64,14 @@ namespace UnrealBuildTool
 		{
 			base.GetReferencedDirectories(Directories);
 
-			if(SourceDirectories == null)
+			if(!Rules.bUsePrecompiled)
 			{
-				throw new BuildException("GetReferencedDirectories() should not be called before building.");
+				if(SourceDirectories == null)
+				{
+					throw new BuildException("GetReferencedDirectories() should not be called before building.");
+				}
+				Directories.UnionWith(SourceDirectories);
 			}
-
-			Directories.UnionWith(SourceDirectories);
 		}
 
 		/// <summary>
