Add Visual Studio 2019 as a compiler for generating intellisense data.

#rnx

diff --git a/Engine/Source/Programs/UnrealBuildTool/ProjectFiles/VisualStudio/VCProject.cs b/Engine/Source/Programs/UnrealBuildTool/ProjectFiles/VisualStudio/VCProject.cs
--- a/Engine/Source/Programs/UnrealBuildTool/ProjectFiles/VisualStudio/VCProject.cs
+++ b/Engine/Source/Programs/UnrealBuildTool/ProjectFiles/VisualStudio/VCProject.cs
@@ -534,6 +534,8 @@ namespace UnrealBuildTool
 		{
 			switch(ProjectFileFormat)
 			{
+				case VCProjectFileFormat.VisualStudio2019:
+					return WindowsCompiler.VisualStudio2019;
 				case VCProjectFileFormat.VisualStudio2017:
 					return WindowsCompiler.VisualStudio2017;
 				default:
