Getting rid of Lumin spam which appears for builds and generating project files for any platform.

diff --git a/Engine/Source/Programs/UnrealBuildTool/Platform/Lumin/UEBuildLumin.cs b/Engine/Source/Programs/UnrealBuildTool/Platform/Lumin/UEBuildLumin.cs
--- a/Engine/Source/Programs/UnrealBuildTool/Platform/Lumin/UEBuildLumin.cs
+++ b/Engine/Source/Programs/UnrealBuildTool/Platform/Lumin/UEBuildLumin.cs
@@ -317,7 +317,7 @@ namespace UnrealBuildTool
 			string MLSDKPath = Environment.GetEnvironmentVariable(EnvVarKey);
 			if (String.IsNullOrEmpty(MLSDKPath))
 			{
-				Console.WriteLine("*** Unable to determine MLSDK location ***");
+				Log.TraceLog("*** Unable to determine MLSDK location ***");
 				return false;
 			}
 
@@ -336,18 +336,18 @@ namespace UnrealBuildTool
 			}
 			else
 			{
-				Console.WriteLine("*** Unable to locate MLSDK version file ml_version.h ***");
+				Log.TraceLog("*** Unable to locate MLSDK version file ml_version.h ***");
 				return false;
 			}
 
 			if (DetectedMajorVersion < MinimumSDKVersionMajor || (DetectedMajorVersion == MinimumSDKVersionMajor && DetectedMinorVersion < MinimumSDKVersionMinor))
 			{
-				Console.WriteLine("*** Found installed MLSDK version {0}.{1} at '{2}' but require at least {3}.{4} ***",
+				Log.TraceLog("*** Found installed MLSDK version {0}.{1} at '{2}' but require at least {3}.{4} ***",
 					DetectedMajorVersion, DetectedMinorVersion, MLSDKPath, MinimumSDKVersionMajor, MinimumSDKVersionMinor);
 				return false;
 			}
 
-			Console.WriteLine("*** Found installed MLSDK version {0}.{1} at '{2}' ***", DetectedMajorVersion, DetectedMinorVersion, MLSDKPath);
+			Log.TraceLog("*** Found installed MLSDK version {0}.{1} at '{2}' ***", DetectedMajorVersion, DetectedMinorVersion, MLSDKPath);
 			return true;
 		}
 
