UBT: Make sure the global mutex is disposed of last. If we dispose of it before we shutdown the logging system, we can get file access exceptions when the next process starts up and tries to write to the log.

#rnx

diff --git a/Engine/Source/Programs/UnrealBuildTool/UnrealBuildTool.cs b/Engine/Source/Programs/UnrealBuildTool/UnrealBuildTool.cs
--- a/Engine/Source/Programs/UnrealBuildTool/UnrealBuildTool.cs
+++ b/Engine/Source/Programs/UnrealBuildTool/UnrealBuildTool.cs
@@ -451,12 +451,6 @@ namespace UnrealBuildTool
 			}
 			finally
 			{
-				// Dispose of the mutex
-				if(Mutex != null)
-				{
-					Mutex.Dispose();
-				}
-
 				// Cancel the prefetcher
 				using(Timeline.ScopeEvent("FileMetadataPrefetch.Stop()"))
 				{
@@ -468,6 +462,12 @@ namespace UnrealBuildTool
 
 				// Make sure we flush the logs however we exit
 				Trace.Close();
+
+				// Dispose of the mutex. Must be done last to ensure that another process does not startup and start trying to write to the same log file.
+				if(Mutex != null)
+				{
+					Mutex.Dispose();
+				}
 			}
 		}
 	}
