Add support for building Linux servers with ThinLTO enabled (via the -ThinLTO command line argument).

diff --git a/Engine/Source/Programs/UnrealBuildTool/Platform/Linux/LinuxToolChain.cs b/Engine/Source/Programs/UnrealBuildTool/Platform/Linux/LinuxToolChain.cs
--- a/Engine/Source/Programs/UnrealBuildTool/Platform/Linux/LinuxToolChain.cs
+++ b/Engine/Source/Programs/UnrealBuildTool/Platform/Linux/LinuxToolChain.cs
@@ -36,6 +36,11 @@ namespace UnrealBuildTool
 		/// Enable undefined behavior sanitizer
 		/// </summary>
 		EnableUndefinedBehaviorSanitizer = 0x4,
+
+		/// <summary>
+		/// Enable thin LTO
+		/// </summary>
+		EnableThinLTO = 0x8,
 	}
 
 	class LinuxToolChain : UEToolChain
@@ -664,7 +669,14 @@ namespace UnrealBuildTool
 			// Whether we actually can enable that is checked in CanUseAdvancedLinkerFeatures() earlier
 			if (CompileEnvironment.bAllowLTCG)
 			{
-				Result += " -flto";
+				if((Options & LinuxToolChainOptions.EnableThinLTO) != 0)
+				{
+					Result += " -flto=thin";
+				}
+				else
+				{
+					Result += " -flto";
+				}
 			}
 
 			if (CompileEnvironment.bEnableShadowVariableWarnings)
@@ -973,7 +985,14 @@ namespace UnrealBuildTool
 			// whether we actually can do that is checked in CanUseAdvancedLinkerFeatures() earlier
 			if (LinkEnvironment.bAllowLTCG)
 			{
-				Result += " -flto";
+				if((Options & LinuxToolChainOptions.EnableThinLTO) != 0)
+				{
+					Result += String.Format(" -flto=thin -Wl,--thinlto-jobs={0}", Utils.GetPhysicalProcessorCount());
+				}
+				else
+				{
+					Result += " -flto";
+				}
 			}
 
 			if (CrossCompiling())
diff --git a/Engine/Source/Programs/UnrealBuildTool/Platform/Linux/UEBuildLinux.cs b/Engine/Source/Programs/UnrealBuildTool/Platform/Linux/UEBuildLinux.cs
--- a/Engine/Source/Programs/UnrealBuildTool/Platform/Linux/UEBuildLinux.cs
+++ b/Engine/Source/Programs/UnrealBuildTool/Platform/Linux/UEBuildLinux.cs
@@ -60,6 +60,12 @@ namespace UnrealBuildTool
 		[XmlConfigFile(Category = "BuildConfiguration", Name = "bEnableUndefinedBehaviorSanitizer")]
 		public bool bEnableUndefinedBehaviorSanitizer = false;
 
+		/// <summary>
+		/// Enables "thin" LTO
+		/// </summary>
+		[CommandLine("-ThinLTO")]
+		public bool bEnableThinLTO = false;
+
 		/// <summary>
 		/// Whether or not to preserve the portable symbol file produced by dump_syms
 		/// </summary>
@@ -114,6 +120,11 @@ namespace UnrealBuildTool
 			get { return Inner.bEnableUndefinedBehaviorSanitizer; }
 		}
 
+		public bool bEnableThinLTO
+		{
+			get { return Inner.bEnableThinLTO; }
+		}
+
 		#if !__MonoCS__
 		#pragma warning restore CS1591
 		#endif
@@ -233,6 +244,11 @@ namespace UnrealBuildTool
 
 		public override void ValidateTarget(TargetRules Target)
 		{
+			if(Target.LinuxPlatform.bEnableThinLTO)
+			{
+				Target.bAllowLTCG = true;
+			}
+
 			if (Target.bAllowLTCG && Target.LinkType != TargetLinkType.Monolithic)
 			{
 				throw new BuildException("LTO (LTCG) for modular builds is not supported (lld is not currently used for dynamic libraries).");
@@ -555,6 +571,10 @@ namespace UnrealBuildTool
 			{
 				Options |= LinuxToolChainOptions.EnableUndefinedBehaviorSanitizer;
 			}
+			if(Target.LinuxPlatform.bEnableThinLTO)
+			{
+				Options |= LinuxToolChainOptions.EnableThinLTO;
+			}
 
 			return new LinuxToolChain(Target.Architecture, SDK, Target.LinuxPlatform.bPreservePSYM, Options);
 		}
