Fix for needing the NDK installed for building Lumin / Magic Leap projects.  Added bypass the to NDK version check if bAllowMissingNDK.  Opted to not early out in case other Android functionality that Lumin needs is added afterwards.

diff --git a/Engine/Source/Programs/UnrealBuildTool/Platform/Android/AndroidToolChain.cs b/Engine/Source/Programs/UnrealBuildTool/Platform/Android/AndroidToolChain.cs
--- a/Engine/Source/Programs/UnrealBuildTool/Platform/Android/AndroidToolChain.cs
+++ b/Engine/Source/Programs/UnrealBuildTool/Platform/Android/AndroidToolChain.cs
@@ -315,7 +315,8 @@ namespace UnrealBuildTool
 			}
 
 			// verify NDK toolchain is supported
-			if (NDKDefineInt < MinimumNDKToolchain || NDKDefineInt > MaximumNDKToolchain)
+			if ((NDKDefineInt < MinimumNDKToolchain || NDKDefineInt > MaximumNDKToolchain)
+				&& !bAllowMissingNDK)
 			{
 				throw new BuildException("Android toolchain NDK " + ToolchainIntToString(NDKDefineInt) + " not supported; please use NDK " + ToolchainIntToString(MinimumNDKToolchain) + " to NDK " + ToolchainIntToString(MaximumNDKToolchain) +
 					" (NDK " + ToolchainIntToString(RecommendedNDKToolchain) + " recommended)");
