Linux: reduce disk size by removing -fstandalone-debug.

- Benefits are dubious, while its impact on the binary size is significant.

diff --git a/Engine/Source/Programs/UnrealBuildTool/Platform/Linux/LinuxToolChain.cs b/Engine/Source/Programs/UnrealBuildTool/Platform/Linux/LinuxToolChain.cs
--- a/Engine/Source/Programs/UnrealBuildTool/Platform/Linux/LinuxToolChain.cs
+++ b/Engine/Source/Programs/UnrealBuildTool/Platform/Linux/LinuxToolChain.cs
@@ -715,12 +715,6 @@ namespace UnrealBuildTool
 
 				// Make debug info LLDB friendly
 				Result += " -glldb";
-
-				if (CompileEnvironment.bIsBuildingDLL)
-				{
-					// Makes debugging .so libraries better
-					Result += " -fstandalone-debug";
-				}
 			}
 
 			// optimization level
