UBT: Use the cached filesystem attributes to figure out which files are part of the adaptive non-unity working set. Much faster than querying directly because (a) they're cached, and (b) they're enumerated via a directory scan, which fetches the metadata up front.

#rnx

diff --git a/Engine/Source/Programs/UnrealBuildTool/System/FileItem.cs b/Engine/Source/Programs/UnrealBuildTool/System/FileItem.cs
--- a/Engine/Source/Programs/UnrealBuildTool/System/FileItem.cs
+++ b/Engine/Source/Programs/UnrealBuildTool/System/FileItem.cs
@@ -76,6 +76,14 @@ namespace UnrealBuildTool
 			get { return Info.Length; }
 		}
 
+		/// <summary>
+		/// The attributes for this file
+		/// </summary>
+		public FileAttributes Attributes
+		{
+			get { return Info.Attributes; }
+		}
+
 		/// <summary>
 		/// The last write time of the file.
 		/// </summary>
diff --git a/Engine/Source/Programs/UnrealBuildTool/System/SourceFileWorkingSet.cs b/Engine/Source/Programs/UnrealBuildTool/System/SourceFileWorkingSet.cs
--- a/Engine/Source/Programs/UnrealBuildTool/System/SourceFileWorkingSet.cs
+++ b/Engine/Source/Programs/UnrealBuildTool/System/SourceFileWorkingSet.cs
@@ -21,7 +21,7 @@ namespace UnrealBuildTool
 		/// </summary>
 		/// <param name="File">File to check</param>
 		/// <returns>True if the file is part of the working set, false otherwise</returns>
-		bool Contains(FileReference File);
+		bool Contains(FileItem File);
 	}
 
 	/// <summary>
@@ -41,7 +41,7 @@ namespace UnrealBuildTool
 		/// </summary>
 		/// <param name="File">File to check</param>
 		/// <returns>True if the file is part of the working set, false otherwise</returns>
-		public bool Contains(FileReference File)
+		public bool Contains(FileItem File)
 		{
 			return false;
 		}
@@ -64,7 +64,7 @@ namespace UnrealBuildTool
 		/// </summary>
 		/// <param name="File">File to check</param>
 		/// <returns>True if the file is part of the working set, false otherwise</returns>
-		public bool Contains(FileReference File)
+		public bool Contains(FileItem File)
 		{
 			// Generated .cpp files should never be treated as part of the working set
 			if (File.HasExtension(".gen.cpp"))
@@ -75,7 +75,7 @@ namespace UnrealBuildTool
 			// Check if the file is read-only
 			try
 			{
-				return !FileReference.GetAttributes(File).HasFlag(FileAttributes.ReadOnly);
+				return !File.Attributes.HasFlag(FileAttributes.ReadOnly);
 			}
 			catch (FileNotFoundException)
 			{
@@ -201,10 +201,10 @@ namespace UnrealBuildTool
 		/// </summary>
 		/// <param name="File">File to check</param>
 		/// <returns>True if the file is part of the working set, false otherwise</returns>
-		public bool Contains(FileReference File)
+		public bool Contains(FileItem File)
 		{
 			WaitForBackgroundProcess();
-			if(Files.Contains(File) || Directories.Any(x => File.IsUnderDirectory(x)))
+			if(Files.Contains(File.Location) || Directories.Any(x => File.Location.IsUnderDirectory(x)))
 			{
 				return true;
 			}
diff --git a/Engine/Source/Programs/UnrealBuildTool/System/TargetMakefile.cs b/Engine/Source/Programs/UnrealBuildTool/System/TargetMakefile.cs
--- a/Engine/Source/Programs/UnrealBuildTool/System/TargetMakefile.cs
+++ b/Engine/Source/Programs/UnrealBuildTool/System/TargetMakefile.cs
@@ -532,7 +532,7 @@ namespace UnrealBuildTool
 				// Check if any source files in the working set no longer belong in it
 				foreach (FileItem SourceFile in Makefile.WorkingSet)
 				{
-					if (!WorkingSet.Contains(SourceFile.Location) && SourceFile.LastWriteTimeUtc > Makefile.CreateTimeUtc)
+					if (!WorkingSet.Contains(SourceFile) && SourceFile.LastWriteTimeUtc > Makefile.CreateTimeUtc)
 					{
 						Log.TraceLog("{0} was part of source working set and now is not; invalidating makefile", SourceFile.AbsolutePath);
 						ReasonNotLoaded = string.Format("working set of source files changed");
@@ -543,7 +543,7 @@ namespace UnrealBuildTool
 				// Check if any source files that are eligible for being in the working set have been modified
 				foreach (FileItem SourceFile in Makefile.CandidatesForWorkingSet)
 				{
-					if (WorkingSet.Contains(SourceFile.Location) && SourceFile.LastWriteTimeUtc > Makefile.CreateTimeUtc)
+					if (WorkingSet.Contains(SourceFile) && SourceFile.LastWriteTimeUtc > Makefile.CreateTimeUtc)
 					{
 						Log.TraceLog("{0} was part of source working set and now is not", SourceFile.AbsolutePath);
 						ReasonNotLoaded = string.Format("working set of source files changed");
diff --git a/Engine/Source/Programs/UnrealBuildTool/System/Unity.cs b/Engine/Source/Programs/UnrealBuildTool/System/Unity.cs
--- a/Engine/Source/Programs/UnrealBuildTool/System/Unity.cs
+++ b/Engine/Source/Programs/UnrealBuildTool/System/Unity.cs
@@ -211,7 +211,7 @@ namespace UnrealBuildTool
 						++CandidateWorkingSetSourceFileCount;
 
 						// Don't include writable source files into unity blobs
-						if (WorkingSet.Contains(CPPFile.Location))
+						if (WorkingSet.Contains(CPPFile))
 						{
 							++WorkingSetSourceFileCount;
 
