Improve accuracy of timer for running UBT. Now includes startup time as well as time spent in the build mode.

diff --git a/Engine/Source/Programs/UnrealBuildTool/Modes/BuildMode.cs b/Engine/Source/Programs/UnrealBuildTool/Modes/BuildMode.cs
--- a/Engine/Source/Programs/UnrealBuildTool/Modes/BuildMode.cs
+++ b/Engine/Source/Programs/UnrealBuildTool/Modes/BuildMode.cs
@@ -42,7 +42,7 @@ namespace UnrealBuildTool
 	/// <summary>
 	/// Builds a target
 	/// </summary>
-	[ToolMode("Build", ToolModeOptions.XmlConfig | ToolModeOptions.BuildPlatforms | ToolModeOptions.SingleInstance | ToolModeOptions.StartPrefetchingEngine)]
+	[ToolMode("Build", ToolModeOptions.XmlConfig | ToolModeOptions.BuildPlatforms | ToolModeOptions.SingleInstance | ToolModeOptions.StartPrefetchingEngine | ToolModeOptions.ShowExecutionTime)]
 	class BuildMode : ToolMode
 	{
 		/// <summary>
@@ -77,7 +77,6 @@ namespace UnrealBuildTool
 		/// <returns>One of the values of ECompilationResult</returns>
 		public override int Execute(CommandLineArguments Arguments)
 		{
-			Stopwatch BuildTimer = Stopwatch.StartNew();
 			Arguments.ApplyTo(this);
 
 			// Initialize the log system, buffering the output until we can create the log file
@@ -204,9 +203,6 @@ namespace UnrealBuildTool
 				SourceFileMetadataCache.SaveAll();
 				CppDependencyCache.SaveAll();
 			}
-
-			// Figure out how long we took to execute.
-			Log.TraceInformation("Total build time: {0:0.00} seconds", BuildTimer.Elapsed.TotalSeconds);
 			return 0;
 		}
 
diff --git a/Engine/Source/Programs/UnrealBuildTool/System/Timeline.cs b/Engine/Source/Programs/UnrealBuildTool/System/Timeline.cs
--- a/Engine/Source/Programs/UnrealBuildTool/System/Timeline.cs
+++ b/Engine/Source/Programs/UnrealBuildTool/System/Timeline.cs
@@ -89,6 +89,14 @@ namespace UnrealBuildTool
 		/// </summary>
 		static List<Event> Events = new List<Event>();
 
+		/// <summary>
+		/// Property for the total time elapsed
+		/// </summary>
+		public static TimeSpan Elapsed
+		{
+			get { return Stopwatch.Elapsed; }
+		}
+
 		/// <summary>
 		/// Start the stopwatch
 		/// </summary>
diff --git a/Engine/Source/Programs/UnrealBuildTool/System/ToolMode.cs b/Engine/Source/Programs/UnrealBuildTool/System/ToolMode.cs
--- a/Engine/Source/Programs/UnrealBuildTool/System/ToolMode.cs
+++ b/Engine/Source/Programs/UnrealBuildTool/System/ToolMode.cs
@@ -44,6 +44,11 @@ namespace UnrealBuildTool
 		/// Only allow a single instance running in the branch at once
 		/// </summary>
 		SingleInstance = 16,
+
+		/// <summary>
+		/// Print out the total time taken to execute
+		/// </summary>
+		ShowExecutionTime = 32,
 	}
 
 	/// <summary>
diff --git a/Engine/Source/Programs/UnrealBuildTool/UnrealBuildTool.cs b/Engine/Source/Programs/UnrealBuildTool/UnrealBuildTool.cs
--- a/Engine/Source/Programs/UnrealBuildTool/UnrealBuildTool.cs
+++ b/Engine/Source/Programs/UnrealBuildTool/UnrealBuildTool.cs
@@ -427,7 +427,14 @@ namespace UnrealBuildTool
 
 				// Create the appropriate handler
 				ToolMode Mode = (ToolMode)Activator.CreateInstance(ModeType);
-				return Mode.Execute(Arguments);
+
+				// Execute the mode
+				int Result = Mode.Execute(Arguments);
+				if((ModeOptions & ToolModeOptions.ShowExecutionTime) != 0)
+				{
+					Log.TraceInformation("Total execution time: {0:0.00} seconds", Timeline.Elapsed.TotalSeconds);
+				}
+				return Result;
 			}
 			catch (CompilationResultException Ex)
 			{
@@ -459,7 +466,7 @@ namespace UnrealBuildTool
 
 				// Print out all the performance info
 				Timeline.Print(TimeSpan.FromMilliseconds(20.0), LogEventType.Log);
-
+			
 				// Make sure we flush the logs however we exit
 				Trace.Close();
 
