UBT: Store hashes in the action history rather than full command lines. Improves serialization performance by ~50ms.

diff --git a/Engine/Source/Programs/DotNETCommon/DotNETUtilities/BinaryArchiveReader.cs b/Engine/Source/Programs/DotNETCommon/DotNETUtilities/BinaryArchiveReader.cs
--- a/Engine/Source/Programs/DotNETCommon/DotNETUtilities/BinaryArchiveReader.cs
+++ b/Engine/Source/Programs/DotNETCommon/DotNETUtilities/BinaryArchiveReader.cs
@@ -230,6 +230,49 @@ namespace Tools.DotNETCommon
 			}
 		}
 
+		/// <summary>
+		/// Reads a byte array from the stream
+		/// </summary>
+		/// <param name="Length">Length of the array to read</param>
+		/// <returns>The data that was read</returns>
+		public byte[] ReadFixedSizeByteArray(int Length)
+		{
+			return ReadFixedSizePrimitiveArray<byte>(sizeof(byte), Length);
+		}
+
+		/// <summary>
+		/// Reads a short array from the stream
+		/// </summary>
+		/// <param name="Length">Length of the array to read</param>
+		/// <returns>The data that was read</returns>
+		public short[] ReadFixedSizeShortArray(int Length)
+		{
+			return ReadFixedSizePrimitiveArray<short>(sizeof(short), Length);
+		}
+
+		/// <summary>
+		/// Reads an int array from the stream
+		/// </summary>
+		/// <param name="Length">Length of the array to read</param>
+		/// <returns>The data that was read</returns>
+		public int[] ReadFixedSizeIntArray(int Length)
+		{
+			return ReadFixedSizePrimitiveArray<int>(sizeof(int), Length);
+		}
+
+		/// <summary>
+		/// Reads an array of primitive types from the stream
+		/// </summary>
+		/// <param name="ElementSize">Size of a single element</param>
+		/// <param name="ElementCount">Number of elements to read</param>
+		/// <returns>The data that was read</returns>
+		private T[] ReadFixedSizePrimitiveArray<T>(int ElementSize, int ElementCount) where T : struct
+		{
+			T[] Result = new T[ElementCount];
+			ReadBulkData(Result, ElementSize * ElementCount);
+			return Result;
+		}
+
 		/// <summary>
 		/// Reads bulk data from the stream into the given buffer
 		/// </summary>
diff --git a/Engine/Source/Programs/DotNETCommon/DotNETUtilities/BinaryArchiveWriter.cs b/Engine/Source/Programs/DotNETCommon/DotNETUtilities/BinaryArchiveWriter.cs
--- a/Engine/Source/Programs/DotNETCommon/DotNETUtilities/BinaryArchiveWriter.cs
+++ b/Engine/Source/Programs/DotNETCommon/DotNETUtilities/BinaryArchiveWriter.cs
@@ -287,6 +287,42 @@ namespace Tools.DotNETCommon
 			}
 		}
 
+		/// <summary>
+		/// Writes an array of bytes to the output
+		/// </summary>
+		/// <param name="Data">Data to write. May be null.</param>
+		public void WriteFixedSizeByteArray(byte[] Data)
+		{
+			WriteFixedSizePrimitiveArray(Data, sizeof(byte));
+		}
+
+		/// <summary>
+		/// Writes an array of shorts to the output
+		/// </summary>
+		/// <param name="Data">Data to write. May be null.</param>
+		public void WriteFixedSizeShortArray(short[] Data)
+		{
+			WriteFixedSizePrimitiveArray(Data, sizeof(short));
+		}
+
+		/// <summary>
+		/// Writes an array of ints to the output
+		/// </summary>
+		/// <param name="Data">Data to write. May be null.</param>
+		public void WriteFixedSizeIntArray(int[] Data)
+		{
+			WriteFixedSizePrimitiveArray(Data, sizeof(int));
+		}
+
+		/// <summary>
+		/// Writes an array of primitive types to the output.
+		/// </summary>
+		/// <param name="Data">Data to write. May be null.</param>
+		private void WriteFixedSizePrimitiveArray<T>(T[] Data, int ElementSize) where T : struct
+		{
+			WriteBulkData(Data, Data.Length * ElementSize);
+		}
+
 		/// <summary>
 		/// Writes primitive data from the given array to the output buffer.
 		/// </summary>
diff --git a/Engine/Source/Programs/UnrealBuildTool/System/ActionGraph.cs b/Engine/Source/Programs/UnrealBuildTool/System/ActionGraph.cs
--- a/Engine/Source/Programs/UnrealBuildTool/System/ActionGraph.cs
+++ b/Engine/Source/Programs/UnrealBuildTool/System/ActionGraph.cs
@@ -384,26 +384,20 @@ namespace UnrealBuildTool
 			foreach (FileItem ProducedItem in RootAction.ProducedItems)
 			{
 				// Check if the command-line of the action previously used to produce the item is outdated.
-				string OldProducingCommandLine = "";
 				string NewProducingCommandLine = RootAction.CommandPath.FullName + " " + RootAction.CommandArguments;
-				if (!ActionHistory.TryGetProducingCommandLine(ProducedItem, out OldProducingCommandLine)
-				|| !String.Equals(OldProducingCommandLine, NewProducingCommandLine, StringComparison.OrdinalIgnoreCase))
+				if (ActionHistory.UpdateProducingCommandLine(ProducedItem, NewProducingCommandLine))
 				{
 					if(ProducedItem.Exists)
 					{
 						Log.TraceLog(
-							"{0}: Produced item \"{1}\" was produced by outdated command-line.\n  Old command-line: {2}\n  New command-line: {3}",
+							"{0}: Produced item \"{1}\" was produced by outdated command-line.\n  New command-line: {2}",
 							RootAction.StatusDescription,
 							Path.GetFileName(ProducedItem.AbsolutePath),
-							OldProducingCommandLine,
 							NewProducingCommandLine
 							);
 					}
 
 					bIsOutdated = true;
-
-					// Update the command-line used to produce this item in the action history.
-					ActionHistory.SetProducingCommandLine(ProducedItem, NewProducingCommandLine);
 				}
 
 				// If the produced file doesn't exist or has zero size, consider it outdated.  The zero size check is to detect cases
diff --git a/Engine/Source/Programs/UnrealBuildTool/System/ActionHistory.cs b/Engine/Source/Programs/UnrealBuildTool/System/ActionHistory.cs
--- a/Engine/Source/Programs/UnrealBuildTool/System/ActionHistory.cs
+++ b/Engine/Source/Programs/UnrealBuildTool/System/ActionHistory.cs
@@ -5,6 +5,8 @@ using System.Collections.Generic;
 using System.Diagnostics;
 using System.IO;
 using System.Runtime.Serialization.Formatters.Binary;
+using System.Security.Cryptography;
+using System.Text;
 using Tools.DotNETCommon;
 
 namespace UnrealBuildTool
@@ -17,7 +19,12 @@ namespace UnrealBuildTool
 		/// <summary>
 		/// Version number to check
 		/// </summary>
-		const int CurrentVersion = 1;
+		const int CurrentVersion = 2;
+
+		/// <summary>
+		/// Size of each hash value
+		/// </summary>
+		const int HashLength = 16;
 
 		/// <summary>
 		/// Path to store the cache data to.
@@ -37,7 +44,7 @@ namespace UnrealBuildTool
 		/// <summary>
 		/// The command lines used to produce files, keyed by the absolute file paths.
 		/// </summary>
-		Dictionary<FileItem, string> OutputItemToCommandLine = new Dictionary<FileItem, string>();
+		Dictionary<FileItem, byte[]> OutputItemToCommandLineHash = new Dictionary<FileItem, byte[]>();
 
 		/// <summary>
 		/// Whether the dependency cache is dirty and needs to be saved.
@@ -88,7 +95,7 @@ namespace UnrealBuildTool
 						return;
 					}
 
-					OutputItemToCommandLine = Reader.ReadDictionary(() => Reader.ReadFileItem(), () => Reader.ReadString());
+					OutputItemToCommandLineHash = Reader.ReadDictionary(() => Reader.ReadFileItem(), () => Reader.ReadFixedSizeByteArray(HashLength));
 				}
 			}
 			catch(Exception Ex)
@@ -107,50 +114,67 @@ namespace UnrealBuildTool
 			using(BinaryArchiveWriter Writer = new BinaryArchiveWriter(Location))
 			{
 				Writer.WriteInt(CurrentVersion);
-				Writer.WriteDictionary(OutputItemToCommandLine, Key => Writer.WriteFileItem(Key), Value => Writer.WriteString(Value));
+				Writer.WriteDictionary(OutputItemToCommandLineHash, Key => Writer.WriteFileItem(Key), Value => Writer.WriteFixedSizeByteArray(Value));
 			}
 			bModified = false;
 		}
 
 		/// <summary>
-		/// Gets the producing command line for the given file
+		/// Computes the case-invariant hash for a string
 		/// </summary>
-		/// <param name="File">The output file to look for</param>
-		/// <param name="CommandLine">Receives the command line used to produce this file</param>
-		/// <returns>True if the output item exists</returns>
-		public bool TryGetProducingCommandLine(FileItem File, out string CommandLine)
+		/// <param name="Text">The text to hash</param>
+		/// <returns>Hash of the string</returns>
+		static byte[] ComputeHash(string Text)
 		{
-			if(File.Location.IsUnderDirectory(BaseDirectory) || Parent == null)
+			string InvariantText = Text.ToUpperInvariant();
+			byte[] InvariantBytes = Encoding.Unicode.GetBytes(InvariantText);
+			return new MD5CryptoServiceProvider().ComputeHash(InvariantBytes);
+		}
+
+		/// <summary>
+		/// Compares two hashes for equality
+		/// </summary>
+		/// <param name="A">The first hash value</param>
+		/// <param name="B">The second hash value</param>
+		/// <returns>True if the hashes are equal</returns>
+		static bool CompareHashes(byte[] A, byte[] B)
+		{
+			for(int Idx = 0; Idx < HashLength; Idx++)
 			{
-				lock (LockObject)
+				if(A[Idx] != B[Idx])
 				{
-					return OutputItemToCommandLine.TryGetValue(File, out CommandLine);
+					return false;
 				}
 			}
-			else
-			{
-				return Parent.TryGetProducingCommandLine(File, out CommandLine);
-			}
+			return true;
 		}
 
 		/// <summary>
-		/// Updates the producing command line for the given file
+		/// Gets the producing command line for the given file
 		/// </summary>
-		/// <param name="File">The output file to update</param>
-		/// <param name="CommandLine">The command line used to produce this file</param>
-		public void SetProducingCommandLine(FileItem File, string CommandLine)
+		/// <param name="File">The output file to look for</param>
+		/// <param name="CommandLine">Receives the command line used to produce this file</param>
+		/// <returns>True if the output item exists</returns>
+		public bool UpdateProducingCommandLine(FileItem File, string CommandLine)
 		{
 			if(File.Location.IsUnderDirectory(BaseDirectory) || Parent == null)
 			{
-				lock(LockObject)
+				byte[] NewHash = ComputeHash(CommandLine);
+				lock (LockObject)
 				{
-					OutputItemToCommandLine[File] = CommandLine;
-					bModified = true;
+					byte[] CurrentHash;
+					if(!OutputItemToCommandLineHash.TryGetValue(File, out CurrentHash) || !CompareHashes(CurrentHash, NewHash))
+					{
+						OutputItemToCommandLineHash[File] = NewHash;
+						bModified = true;
+						return true;
+					}
+					return false;
 				}
 			}
 			else
 			{
-				Parent.SetProducingCommandLine(File, CommandLine);
+				return Parent.UpdateProducingCommandLine(File, CommandLine);
 			}
 		}
 
@@ -173,7 +197,7 @@ namespace UnrealBuildTool
 				AppName = UEBuildTarget.GetAppNameForTargetType(TargetType);
 			}
 
-			return FileReference.Combine(UnrealBuildTool.EngineDirectory, "Intermediate", "Build", Platform.ToString(), AppName, "ActionHistory.dat");
+			return FileReference.Combine(UnrealBuildTool.EngineDirectory, "Intermediate", "Build", Platform.ToString(), AppName, "ActionHistory.bin");
 		}
 
 		/// <summary>
@@ -185,7 +209,7 @@ namespace UnrealBuildTool
 		/// <returns>Path to the project action history</returns>
 		public static FileReference GetProjectLocation(FileReference ProjectFile, string TargetName, UnrealTargetPlatform Platform)
 		{
-			return FileReference.Combine(ProjectFile.Directory, "Intermediate", "Build", Platform.ToString(), TargetName, "ActionHistory.dat");
+			return FileReference.Combine(ProjectFile.Directory, "Intermediate", "Build", Platform.ToString(), TargetName, "ActionHistory.bin");
 		}
 
 		/// <summary>
