UBT: Serialize compact FileItem objects. Improves serialization performance and significantly reduces size of cpp dependency cache.

diff --git a/Engine/Source/Programs/UnrealBuildTool/System/CppDependencyCache.cs b/Engine/Source/Programs/UnrealBuildTool/System/CppDependencyCache.cs
--- a/Engine/Source/Programs/UnrealBuildTool/System/CppDependencyCache.cs
+++ b/Engine/Source/Programs/UnrealBuildTool/System/CppDependencyCache.cs
@@ -34,7 +34,7 @@ namespace UnrealBuildTool
 			public static DependencyInfo Read(BinaryArchiveReader Reader)
 			{
 				long LastWriteTimeUtc = Reader.ReadLong();
-				List<FileItem> Files = Reader.ReadList(() => Reader.ReadFileItem());
+				List<FileItem> Files = Reader.ReadList(() => Reader.ReadCompactFileItem());
 
 				return new DependencyInfo(LastWriteTimeUtc, Files);
 			}
@@ -42,14 +42,14 @@ namespace UnrealBuildTool
 			public void Write(BinaryArchiveWriter Writer)
 			{
 				Writer.WriteLong(LastWriteTimeUtc);
-				Writer.WriteList<FileItem>(Files, File => Writer.WriteFileItem(File));
+				Writer.WriteList<FileItem>(Files, File => Writer.WriteCompactFileItem(File));
 			}
 		}
 
 		/// <summary>
 		/// The current file version
 		/// </summary>
-		public const int CurrentVersion = 1;
+		public const int CurrentVersion = 2;
 
 		/// <summary>
 		/// Location of this dependency cache
diff --git a/Engine/Source/Programs/UnrealBuildTool/System/DirectoryItem.cs b/Engine/Source/Programs/UnrealBuildTool/System/DirectoryItem.cs
--- a/Engine/Source/Programs/UnrealBuildTool/System/DirectoryItem.cs
+++ b/Engine/Source/Programs/UnrealBuildTool/System/DirectoryItem.cs
@@ -280,7 +280,9 @@ namespace UnrealBuildTool
 				{
 					foreach(FileInfo FileInfo in Info.EnumerateFiles())
 					{
-						NewFiles[FileInfo.Name] = FileItem.GetItemByFileInfo(FileInfo);
+						FileItem FileItem = FileItem.GetItemByFileInfo(FileInfo);
+						FileItem.UpdateCachedDirectory(this);
+						NewFiles[FileInfo.Name] = FileItem;
 					}
 				}
 				Files = NewFiles;
@@ -298,8 +300,8 @@ namespace UnrealBuildTool
 		}
 
 		/// <summary>
-		/// Attempts to get a file from this directory by name. Unlike creating a file item and checking whether it exists, this does not
-		/// cause any I/O to check whether it exists, nor does it create a permanent FileItem object.
+		/// Attempts to get a file from this directory by name. Unlike creating a file item and checking whether it exists, this will
+		/// not create a permanent FileItem object if it does not exist.
 		/// </summary>
 		/// <param name="Name">Name of the file</param>
 		/// <param name="OutFile">If successful receives the matching file item with this name</param>
diff --git a/Engine/Source/Programs/UnrealBuildTool/System/FileItem.cs b/Engine/Source/Programs/UnrealBuildTool/System/FileItem.cs
--- a/Engine/Source/Programs/UnrealBuildTool/System/FileItem.cs
+++ b/Engine/Source/Programs/UnrealBuildTool/System/FileItem.cs
@@ -19,9 +19,14 @@ namespace UnrealBuildTool
 	class FileItem
 	{
 		/// <summary>
-		/// The file reference
+		/// The directory containing this file
 		/// </summary>
-		public FileReference Location;
+		DirectoryItem CachedDirectory;
+
+		/// <summary>
+		/// Location of this file
+		/// </summary>
+		public readonly FileReference Location;
 
 		/// <summary>
 		/// The information about the file.
@@ -60,6 +65,21 @@ namespace UnrealBuildTool
 			get { return Location.FullName; }
 		}
 
+		/// <summary>
+		/// Gets the directory that this file is in
+		/// </summary>
+		public DirectoryItem Directory
+		{
+			get
+			{
+				if(CachedDirectory == null)
+				{
+					CachedDirectory = DirectoryItem.GetItemByDirectoryReference(Location.Directory);
+				}
+				return CachedDirectory;
+			}
+		}
+
 		/// <summary>
 		/// Whether the file exists.
 		/// </summary>
@@ -108,7 +128,17 @@ namespace UnrealBuildTool
 		/// <returns>DirectoryItem for the directory containing this file</returns>
 		public DirectoryItem GetDirectoryItem()
 		{
-			return DirectoryItem.GetItemByDirectoryInfo(Info.Directory);
+			return Directory;
+		}
+
+		/// <summary>
+		/// Updates the cached directory for this file. Used by DirectoryItem when enumerating files, to avoid having to look this up later.
+		/// </summary>
+		/// <param name="Directory">The directory that this file is in</param>
+		public void UpdateCachedDirectory(DirectoryItem Directory)
+		{
+			Debug.Assert(Directory.Location == Location.Directory);
+			CachedDirectory = Directory;
 		}
 
 		/// <summary>
@@ -335,5 +365,41 @@ namespace UnrealBuildTool
 		{
 			Writer.WriteObjectReference<FileItem>(FileItem, () => Writer.WriteFileReference(FileItem.Location));
 		}
+
+		/// <summary>
+		/// Read a file item as a DirectoryItem and name. This is slower than reading it directly, but results in a significantly smaller archive
+		/// where most files are in the same directories.
+		/// </summary>
+		/// <param name="Reader">Archive to read from</param>
+		/// <returns>FileItem read from the archive</returns>
+		static FileItem ReadCompactFileItemData(this BinaryArchiveReader Reader)
+		{
+			DirectoryItem Directory = Reader.ReadDirectoryItem();
+			string Name = Reader.ReadString();
+
+			FileItem FileItem = FileItem.GetItemByFileReference(FileReference.Combine(Directory.Location, Name));
+            FileItem.UpdateCachedDirectory(Directory);
+            return FileItem;
+		}
+
+		/// <summary>
+		/// Read a file item in a format which de-duplicates directory names.
+		/// </summary>
+		/// <param name="Reader">Reader to serialize data from</param>
+		/// <returns>Instance of the serialized file item</returns>
+		public static FileItem ReadCompactFileItem(this BinaryArchiveReader Reader)
+		{
+			return Reader.ReadObjectReference<FileItem>(() => ReadCompactFileItemData(Reader));
+		}
+
+		/// <summary>
+		/// Writes a file item in a format which de-duplicates directory names.
+		/// </summary>
+		/// <param name="Writer">Writer to serialize data to</param>
+		/// <param name="FileItem">File item to write</param>
+		public static void WriteCompactFileItem(this BinaryArchiveWriter Writer, FileItem FileItem)
+		{
+			Writer.WriteObjectReference<FileItem>(FileItem, () => { Writer.WriteDirectoryItem(FileItem.GetDirectoryItem()); Writer.WriteString(FileItem.Name); });
+		}
 	}
 }
