UBT: Use compact FileItem objects in source file metadata cache. Improves serialization performance by another 50-60ms.

diff --git a/Engine/Source/Programs/UnrealBuildTool/System/SourceFileMetadataCache.cs b/Engine/Source/Programs/UnrealBuildTool/System/SourceFileMetadataCache.cs
--- a/Engine/Source/Programs/UnrealBuildTool/System/SourceFileMetadataCache.cs
+++ b/Engine/Source/Programs/UnrealBuildTool/System/SourceFileMetadataCache.cs
@@ -53,7 +53,7 @@ namespace UnrealBuildTool
 		/// <summary>
 		/// The current file version
 		/// </summary>
-		public const int CurrentVersion = 2;
+		public const int CurrentVersion = 3;
 
 		/// <summary>
 		/// Location of this dependency cache
@@ -119,7 +119,10 @@ namespace UnrealBuildTool
 
 			if(FileReference.Exists(Location))
 			{
-				Read();
+				using(Timeline.ScopeEvent("Reading source file metadata cache"))
+				{
+					Read();
+				}
 			}
 		}
 
@@ -305,7 +308,7 @@ namespace UnrealBuildTool
 					int FileToFirstIncludeCount = Reader.ReadInt();
 					for(int Idx = 0; Idx < FileToFirstIncludeCount; Idx++)
 					{
-						FileItem File = Reader.ReadFileItem();
+						FileItem File = Reader.ReadCompactFileItem();
 						
 						IncludeInfo IncludeInfo = new IncludeInfo();
 						IncludeInfo.LastWriteTimeUtc = Reader.ReadLong();
@@ -317,7 +320,7 @@ namespace UnrealBuildTool
 					int FileToMarkupFlagCount = Reader.ReadInt();
 					for(int Idx = 0; Idx < FileToMarkupFlagCount; Idx++)
 					{
-						FileItem File = Reader.ReadFileItem();
+						FileItem File = Reader.ReadCompactFileItem();
 
 						ReflectionInfo ReflectionInfo = new ReflectionInfo();
 						ReflectionInfo.LastWriteTimeUtc = Reader.ReadLong();
@@ -349,7 +352,7 @@ namespace UnrealBuildTool
 					Writer.WriteInt(FileToIncludeInfo.Count);
 					foreach(KeyValuePair<FileItem, IncludeInfo> Pair in FileToIncludeInfo)
 					{
-						Writer.WriteFileItem(Pair.Key);
+						Writer.WriteCompactFileItem(Pair.Key);
 						Writer.WriteLong(Pair.Value.LastWriteTimeUtc);
 						Writer.WriteString(Pair.Value.IncludeText);
 					}
@@ -357,7 +360,7 @@ namespace UnrealBuildTool
 					Writer.WriteInt(FileToReflectionInfo.Count);
 					foreach(KeyValuePair<FileItem, ReflectionInfo> Pair in FileToReflectionInfo)
 					{
-						Writer.WriteFileItem(Pair.Key);
+						Writer.WriteCompactFileItem(Pair.Key);
 						Writer.WriteLong(Pair.Value.LastWriteTimeUtc);
 						Writer.WriteBool(Pair.Value.bContainsMarkup);
 					}
