- Attempt 3 at getting dylib frameworks working with EC splitting up build and stage on different machines

diff --git a/Engine/Source/Programs/UnrealBuildTool/Platform/IOS/IOSToolChain.cs b/Engine/Source/Programs/UnrealBuildTool/Platform/IOS/IOSToolChain.cs
--- a/Engine/Source/Programs/UnrealBuildTool/Platform/IOS/IOSToolChain.cs
+++ b/Engine/Source/Programs/UnrealBuildTool/Platform/IOS/IOSToolChain.cs
@@ -1425,20 +1425,10 @@ namespace UnrealBuildTool
 				{
 					if (Framework.OutputDirectory != null)
 					{
-						// copy entire framework in for dynamic lib frameworks, not just a bundle
-						// @todo add a bool to the UEBuildFramework class that denotes it as a dylib framework, or possibly use RuntimeDependencies instead?
-						if (Framework.ZipFile.FullName.EndsWith(".framework.zip"))
-						{
-							string BundleName = Framework.OutputDirectory.FullName.Substring(Framework.OutputDirectory.FullName.LastIndexOf('/') + 1);
-							FrameworkNameToSourceDir[BundleName] = Framework.OutputDirectory;
-						}
-						else if (!String.IsNullOrEmpty(Framework.CopyBundledAssets))
-						{
-							// For now, this is hard coded, but we need to loop over all modules, and copy bundled assets that need it
-							DirectoryReference LocalSource = DirectoryReference.Combine(Framework.OutputDirectory, Framework.CopyBundledAssets);
-							string BundleName = Framework.CopyBundledAssets.Substring(Framework.CopyBundledAssets.LastIndexOf('/') + 1);
-							FrameworkNameToSourceDir[BundleName] = LocalSource;
-						}
+						// For now, this is hard coded, but we need to loop over all modules, and copy bundled assets that need it
+						DirectoryReference LocalSource = DirectoryReference.Combine(Framework.OutputDirectory, Framework.CopyBundledAssets);
+						string BundleName = Framework.CopyBundledAssets.Substring(Framework.CopyBundledAssets.LastIndexOf('/') + 1);
+						FrameworkNameToSourceDir[BundleName] = LocalSource;
 					}
 				}
 
@@ -1799,26 +1789,13 @@ namespace UnrealBuildTool
 				{
 					string UnpackedZipPath = Pair.Value.FullName;
 
-					// some framework bundles are an entire framework, not a .bundle, so copy them into a Frameworks directory
-					if (Pair.Key.EndsWith(".framework"))
-					{
-						string LocalDest = LocalFrameworkAssets + "/Frameworks";
-						Console.WriteLine("Copying dynamic framework... LocalSource: {0}, LocalDest: {1}", UnpackedZipPath, LocalDest);
-						string ResultsText;
-						// Create the intermediate local directory
-						RunExecutableAndWait("mkdir", String.Format("-p \"{0}\"", LocalDest), out ResultsText);
-						RunExecutableAndWait("cp", String.Format("-R -L \"{0}\" \"{1}\"", Pair.Value, LocalDest), out ResultsText);
-					}
-					else
-					{
-						// For now, this is hard coded, but we need to loop over all modules, and copy bundled assets that need it
-						string LocalDest = LocalFrameworkAssets + "/" + Pair.Key;
-	
-						Log.TraceInformation("Copying bundled asset... LocalSource: {0}, LocalDest: {1}", Pair.Value, LocalDest);
-	
-						string ResultsText;
-						RunExecutableAndWait("cp", String.Format("-R -L \"{0}\" \"{1}\"", Pair.Value, LocalDest), out ResultsText);
-					}
+					// For now, this is hard coded, but we need to loop over all modules, and copy bundled assets that need it
+					string LocalDest = LocalFrameworkAssets + "/" + Pair.Key;
+
+					Log.TraceInformation("Copying bundled asset... LocalSource: {0}, LocalDest: {1}", Pair.Value, LocalDest);
+
+					string ResultsText;
+					RunExecutableAndWait("cp", String.Format("-R -L \"{0}\" \"{1}\"", Pair.Value, LocalDest), out ResultsText);
                 }
             }
 		}
diff --git a/Engine/Source/Programs/UnrealBuildTool/ProjectFiles/Xcode/XcodeProject.cs b/Engine/Source/Programs/UnrealBuildTool/ProjectFiles/Xcode/XcodeProject.cs
--- a/Engine/Source/Programs/UnrealBuildTool/ProjectFiles/Xcode/XcodeProject.cs
+++ b/Engine/Source/Programs/UnrealBuildTool/ProjectFiles/Xcode/XcodeProject.cs
@@ -716,15 +716,64 @@ namespace UnrealBuildTool
 			Content.Append("/* End PBXNativeTarget section */" + ProjectFileGenerator.NewLine + ProjectFileGenerator.NewLine);
 		}
 
-		private void AppendShellScriptSection(StringBuilder Content, string ShellScriptGuid)
+		private void AppendShellScriptSection(StringBuilder Content, string ShellScriptGuid, FileReference UProjectPath)
 		{
+			StringBuilder FrameworkScript = new StringBuilder();
+
+			// nothing to do without a project
+			if (UProjectPath == null)
+			{
+				return;
+			}
+			
+			// @todo: look also in Project/Build/Frameworks directory!
+			ProjectDescriptor Project = ProjectDescriptor.FromFile(UProjectPath);
+			List<PluginInfo> AvailablePlugins = Plugins.ReadAvailablePlugins(UnrealBuildTool.EngineDirectory, UProjectPath, Project.AdditionalPluginDirectories);
+
+			// look in each plugin for frameworks
+			// @todo: Cache this kind of things since every target will re-do this work!
+			foreach (PluginInfo PI in AvailablePlugins)
+			{
+				if (!Plugins.IsPluginEnabledForProject(PI, Project, UnrealTargetPlatform.IOS, UnrealTargetConfiguration.Development, TargetRules.TargetType.Game))
+				{
+					continue;
+				}
+
+				// for now, we copy and code sign all *.framework.zip, even if the have no code (non-code frameworks are assumed to be *.embeddedframework.zip
+				DirectoryReference FrameworkDir = DirectoryReference.Combine(PI.Directory, "Source/Frameworks");
+				if (!DirectoryReference.Exists(FrameworkDir))
+				{
+					FrameworkDir = DirectoryReference.Combine(PI.Directory, "Frameworks");
+				}
+				if (DirectoryReference.Exists(FrameworkDir))
+				{
+					// look at each zip
+					foreach (FileInfo FI in new System.IO.DirectoryInfo(FrameworkDir.FullName).EnumerateFiles("*.framework.zip"))
+					{
+						string Guid = XcodeProjectFileGenerator.MakeXcodeGuid();
+						string RefGuid = XcodeProjectFileGenerator.MakeXcodeGuid();
+
+						// for FI of foo.framework.zip, this will give us foo.framework
+						string Framework = Path.GetFileNameWithoutExtension(FI.FullName);
+
+						// unzip the framework right into the .app
+						FrameworkScript.AppendFormat("\\techo Unzipping {0}...\\n", FI.FullName);
+						FrameworkScript.AppendFormat("\\tunzip -o -q {0} -d ${{FRAMEWORK_DIR}} -x \\\"__MACOSX/*\\\" \\\"*/.DS_Store\\\"\\n", FI.FullName);
+					}
+				}
+			}
+
 			string ShellScript = "set -e\\n\\n" +
-				"FRAMEWORK_DIR=$TARGET_BUILD_DIR/$EXECUTABLE_FOLDER_PATH/Frameworks\\n" +
-				"for FRAMEWORK in ${FRAMEWORK_DIR}/*.framework; do\\n" +
-					"\\t[ -d \\\"${FRAMEWORK}\\\" ] || continue\\n" + 
-					"\\techo Codesigning ${FRAMEWORK}\\n" +
-					"\\tcodesign --force --sign ${EXPANDED_CODE_SIGN_IDENTITY} --verbose --preserve-metadata=identifier,entitlements,flags --timestamp=none \\\"${FRAMEWORK}\\\"\\n" +
-				"done";
+				"if [ $PLATFORM_NAME = iphoneos ] || [ $PLATFORM_NAME = tvos ]; then \\n" +
+				"\\tFRAMEWORK_DIR=$TARGET_BUILD_DIR/$EXECUTABLE_FOLDER_PATH/Frameworks\\n" +
+				FrameworkScript.ToString() + 
+				// and now code sign anything that has been unzipped above
+				"\\tfor FRAMEWORK in ${FRAMEWORK_DIR}/*.framework; do\\n" +
+					"\\t\\t[ -d \\\"${FRAMEWORK}\\\" ] || continue\\n" +
+					"\\t\\techo Codesigning ${FRAMEWORK}\\n" +
+					"\\t\\tcodesign --force --sign ${EXPANDED_CODE_SIGN_IDENTITY} --verbose --preserve-metadata=identifier,entitlements,flags --timestamp=none \\\"${FRAMEWORK}\\\"\\n" +
+				"\\tdone\\n" +
+				"fi\\n";
 
 			Content.Append("/* Begin PBXShellScriptBuildPhase section */" + ProjectFileGenerator.NewLine);
 			Content.Append(string.Format("\t\t{0} /* Sign Frameworks */ = {{" + ProjectFileGenerator.NewLine, ShellScriptGuid));
@@ -1756,7 +1805,7 @@ namespace UnrealBuildTool
 			AppendXCBuildConfigurationSection(ProjectFileContent, ProjectBuildConfigs, TargetBuildConfigs, BuildTargetBuildConfigs, IndexTargetBuildConfigs, GameProjectPath, AllExtensions);
 			AppendXCConfigurationListSection(ProjectFileContent, TargetName, BuildTargetName, IndexTargetName, ProjectConfigListGuid, ProjectBuildConfigs,
 				TargetConfigListGuid, TargetBuildConfigs, BuildTargetConfigListGuid, BuildTargetBuildConfigs, IndexTargetConfigListGuid, IndexTargetBuildConfigs, AllExtensions);
-			AppendShellScriptSection(ProjectFileContent, ShellScriptSectionGuid);
+			AppendShellScriptSection(ProjectFileContent, ShellScriptSectionGuid, GameProjectPath);
 
 			ProjectFileContent.Append("\t};" + ProjectFileGenerator.NewLine);
 			ProjectFileContent.Append("\trootObject = " + ProjectGuid + " /* Project object */;" + ProjectFileGenerator.NewLine);
