- Fixed build error afger merge

diff --git a/Engine/Source/Programs/UnrealBuildTool/Platform/IOS/IOSToolChain.cs b/Engine/Source/Programs/UnrealBuildTool/Platform/IOS/IOSToolChain.cs
--- a/Engine/Source/Programs/UnrealBuildTool/Platform/IOS/IOSToolChain.cs
+++ b/Engine/Source/Programs/UnrealBuildTool/Platform/IOS/IOSToolChain.cs
@@ -1423,7 +1423,7 @@ namespace UnrealBuildTool
 				Dictionary<string, DirectoryReference> FrameworkNameToSourceDir = new Dictionary<string, DirectoryReference>();
 				foreach (UEBuildFramework Framework in BinaryLinkEnvironment.AdditionalFrameworks)
 				{
-					if (Framework.OutputDirectory != null)
+					if (Framework.OutputDirectory != null && !String.IsNullOrEmpty(Framework.CopyBundledAssets))
 					{
 						// For now, this is hard coded, but we need to loop over all modules, and copy bundled assets that need it
 						DirectoryReference LocalSource = DirectoryReference.Combine(Framework.OutputDirectory, Framework.CopyBundledAssets);
