Fix issue with changelist as store version update ordering

#ue4
#android

diff --git a/Engine/Source/Programs/UnrealBuildTool/Platform/Android/UEDeployAndroid.cs b/Engine/Source/Programs/UnrealBuildTool/Platform/Android/UEDeployAndroid.cs
--- a/Engine/Source/Programs/UnrealBuildTool/Platform/Android/UEDeployAndroid.cs
+++ b/Engine/Source/Programs/UnrealBuildTool/Platform/Android/UEDeployAndroid.cs
@@ -863,6 +863,9 @@ namespace UnrealBuildTool
 				// override store version with changelist if enabled and is build machine
 				if (bUseChangeListAsStoreVersion && Environment.GetEnvironmentVariable("IsBuildMachine") == "1")
 				{
+					// make sure changelist is cached
+					string EngineVersion = ReadEngineVersion();
+					
 					int Changelist = 0;
 					if (int.TryParse(EngineChangelist, out Changelist))
 					{
@@ -2101,7 +2104,7 @@ namespace UnrealBuildTool
 		private string EngineChangelist = "0";
 		private string EngineBranch = "UE4";
 
-		private string ReadEngineVersion(string EngineDirectory)
+		private string ReadEngineVersion()
 		{
 			if (!bHaveReadEngineVersion)
 			{
@@ -2123,7 +2126,7 @@ namespace UnrealBuildTool
 		private string GenerateManifest(AndroidToolChain ToolChain, string ProjectName, TargetType InTargetType, string EngineDirectory, bool bIsForDistribution, bool bPackageDataInsideApk, string GameBuildFilesPath, bool bHasOBBFiles, bool bDisableVerifyOBBOnStartUp, string UE4Arch, string GPUArch, string CookFlavor, bool bUseExternalFilesDir, string Configuration, int SDKLevelInt)
 		{
 			// Read the engine version
-			string EngineVersion = ReadEngineVersion(EngineDirectory);
+			string EngineVersion = ReadEngineVersion();
 
 			int StoreVersion = GetStoreVersion();
 
@@ -3225,6 +3228,9 @@ namespace UnrealBuildTool
 				throw new BuildException("Cannot make APK with UPL errors");
 			}
 
+			// make sure it is cached
+			string EngineVersion = ReadEngineVersion();
+
 			SetMinimumSDKLevelForGradle();
 
 			// check for Gradle enabled for this project
@@ -3252,9 +3258,6 @@ namespace UnrealBuildTool
 				PatchAntBatIfNeeded();
 			}
 
-			// make sure it is cached
-			string EngineVersion = ReadEngineVersion(EngineDirectory);
-
 			// cache some tools paths
 			string NDKBuildPath = Environment.ExpandEnvironmentVariables("%NDKROOT%/ndk-build" + (Utils.IsRunningOnMono ? "" : ".cmd"));
 
