UBT: Fixed Copy Action on Mac and Linux for paths containing space character

diff --git a/Engine/Source/Programs/UnrealBuildTool/Configuration/UEBuildTarget.cs b/Engine/Source/Programs/UnrealBuildTool/Configuration/UEBuildTarget.cs
--- a/Engine/Source/Programs/UnrealBuildTool/Configuration/UEBuildTarget.cs
+++ b/Engine/Source/Programs/UnrealBuildTool/Configuration/UEBuildTarget.cs
@@ -1604,7 +1604,7 @@ namespace UnrealBuildTool
 			}
 			else
 			{
-				CopyAction.CommandArguments = String.Format("-c 'cp -f {0} {1}'", Utils.EscapeShellArgument(SourceFile.FullName), Utils.EscapeShellArgument(TargetFile.FullName));
+				CopyAction.CommandArguments = String.Format("-c 'cp -f \"{0}\" \"{1}\"'", SourceFile.FullName, TargetFile.FullName);
 			}
 			CopyAction.WorkingDirectory = UnrealBuildTool.EngineSourceDirectory;
 			CopyAction.PrerequisiteItems.Add(SourceFileItem);
