UE-70358 - Failure to sign when packaging for iOS Distribution when using Remote Compile from a code project on Windows

diff --git a/Engine/Source/Programs/AutomationTool/Scripts/BuildProjectCommand.Automation.cs b/Engine/Source/Programs/AutomationTool/Scripts/BuildProjectCommand.Automation.cs
--- a/Engine/Source/Programs/AutomationTool/Scripts/BuildProjectCommand.Automation.cs
+++ b/Engine/Source/Programs/AutomationTool/Scripts/BuildProjectCommand.Automation.cs
@@ -137,6 +137,11 @@ public partial class Project : CommandUtils
 			AdditionalArgs += " -skipdeploy"; // skip deploy step in UBT if we going to do it later anyway
 		}
 
+		if (Params.Distribution)
+		{
+			AdditionalArgs += " -distribution";
+		}
+
 		// Config overrides (-ini)
 		foreach (string ConfigOverrideParam in Params.ConfigOverrideParams)
 		{
diff --git a/Engine/Source/Programs/UnrealBuildTool/Modes/GenerateProjectFilesMode.cs b/Engine/Source/Programs/UnrealBuildTool/Modes/GenerateProjectFilesMode.cs
--- a/Engine/Source/Programs/UnrealBuildTool/Modes/GenerateProjectFilesMode.cs
+++ b/Engine/Source/Programs/UnrealBuildTool/Modes/GenerateProjectFilesMode.cs
@@ -143,7 +143,7 @@ namespace UnrealBuildTool
 						Generator = new VCProjectFileGenerator(ProjectFile, VCProjectFileFormat.VisualStudio2019, Arguments);
 						break;
 					case ProjectFileFormat.XCode:
-						Generator = new XcodeProjectFileGenerator(ProjectFile);
+						Generator = new XcodeProjectFileGenerator(ProjectFile, Arguments);
 						break;
 					case ProjectFileFormat.Eddie:
 						Generator = new EddieProjectFileGenerator(ProjectFile);
diff --git a/Engine/Source/Programs/UnrealBuildTool/Platform/IOS/IOSPostBuildSyncMode.cs b/Engine/Source/Programs/UnrealBuildTool/Platform/IOS/IOSPostBuildSyncMode.cs
--- a/Engine/Source/Programs/UnrealBuildTool/Platform/IOS/IOSPostBuildSyncMode.cs
+++ b/Engine/Source/Programs/UnrealBuildTool/Platform/IOS/IOSPostBuildSyncMode.cs
@@ -28,6 +28,7 @@ namespace UnrealBuildTool
 		public string ImportCertificate;
 		public string ImportCertificatePassword;
 		public Dictionary<string, DirectoryReference> FrameworkNameToSourceDir;
+		public bool bForDistribution = false;
 
 		public IOSPostBuildSyncTarget(ReadOnlyTargetRules Target, FileReference OutputPath, DirectoryReference ProjectIntermediateDirectory, List<string> UPLScripts, VersionNumber SdkVersion, Dictionary<string, DirectoryReference> FrameworkNameToSourceDir)
 		{
@@ -47,6 +48,7 @@ namespace UnrealBuildTool
 			this.ImportCertificate = Target.IOSPlatform.ImportCertificate;
 			this.ImportCertificatePassword = Target.IOSPlatform.ImportCertificatePassword;
 			this.FrameworkNameToSourceDir = FrameworkNameToSourceDir;
+			this.bForDistribution = Target.IOSPlatform.bForDistribution;
 		}
 	}
 
diff --git a/Engine/Source/Programs/UnrealBuildTool/Platform/IOS/IOSToolChain.cs b/Engine/Source/Programs/UnrealBuildTool/Platform/IOS/IOSToolChain.cs
--- a/Engine/Source/Programs/UnrealBuildTool/Platform/IOS/IOSToolChain.cs
+++ b/Engine/Source/Programs/UnrealBuildTool/Platform/IOS/IOSToolChain.cs
@@ -1549,7 +1549,7 @@ namespace UnrealBuildTool
 				PlatformProjectGenerators.RegisterPlatformProjectGenerator(UnrealTargetPlatform.IOS, new IOSProjectGenerator(CmdLine));
 				PlatformProjectGenerators.RegisterPlatformProjectGenerator(UnrealTargetPlatform.TVOS, new TVOSProjectGenerator(CmdLine));
 
-				XcodeProjectFileGenerator Generator = new XcodeProjectFileGenerator(ProjectFile);
+				XcodeProjectFileGenerator Generator = new XcodeProjectFileGenerator(ProjectFile, CmdLine);
 				return Generator.GenerateProjectFiles(PlatformProjectGenerators, Arguments);
 			}
 			finally
@@ -1598,6 +1598,7 @@ namespace UnrealBuildTool
 
             // ensure the plist, entitlements, and provision files are properly copied
             UEDeployIOS DeployHandler = (Target.Platform == UnrealTargetPlatform.IOS ? new UEDeployIOS() : new UEDeployTVOS());
+			DeployHandler.ForDistribution = Target.bForDistribution;
             DeployHandler.PrepTargetForDeployment(Target.ProjectFile, Target.TargetName, Target.Platform, Target.Configuration, Target.UPLScripts, Target.SdkVersion, Target.bCreateStubIPA);
 
 			// copy the executable
@@ -1615,12 +1616,12 @@ namespace UnrealBuildTool
 				DirectoryReference XcodeWorkspaceDir;
 				if (AppName == "UE4Game" || AppName == "UE4Client" || Target.ProjectFile == null || Target.ProjectFile.IsUnderDirectory(UnrealBuildTool.EngineDirectory))
 				{
-					GenerateProjectFiles(Target.ProjectFile, new string[] { "-platforms=" + (Target.Platform == UnrealTargetPlatform.IOS ? "IOS" : "TVOS"), "-NoIntellIsense", (Target.Platform == UnrealTargetPlatform.IOS ? "-iosdeployonly" : "-tvosdeployonly"), "-ignorejunk" });
+					GenerateProjectFiles(Target.ProjectFile, new string[] { "-platforms=" + (Target.Platform == UnrealTargetPlatform.IOS ? "IOS" : "TVOS"), "-NoIntellIsense", (Target.Platform == UnrealTargetPlatform.IOS ? "-iosdeployonly" : "-tvosdeployonly"), "-ignorejunk", (Target.bForDistribution ? "-distribution" : "-development") });
 					XcodeWorkspaceDir = DirectoryReference.Combine(UnrealBuildTool.RootDirectory, String.Format("UE4_{0}.xcworkspace", (Target.Platform == UnrealTargetPlatform.IOS ? "IOS" : "TVOS")));
 				}
 				else
 				{
-					GenerateProjectFiles(Target.ProjectFile, new string[] { "-platforms=" + (Target.Platform == UnrealTargetPlatform.IOS ? "IOS" : "TVOS"), "-NoIntellIsense", (Target.Platform == UnrealTargetPlatform.IOS ? "-iosdeployonly" : "-tvosdeployonly"), "-ignorejunk", String.Format("-project={0}", Target.ProjectFile), "-game" });
+					GenerateProjectFiles(Target.ProjectFile, new string[] { "-platforms=" + (Target.Platform == UnrealTargetPlatform.IOS ? "IOS" : "TVOS"), "-NoIntellIsense", (Target.Platform == UnrealTargetPlatform.IOS ? "-iosdeployonly" : "-tvosdeployonly"), "-ignorejunk", (Target.bForDistribution ? "-distribution" : "-development"), String.Format("-project={0}", Target.ProjectFile), "-game" });
 					XcodeWorkspaceDir = DirectoryReference.Combine(Target.ProjectDirectory, String.Format("{0}_{1}.xcworkspace", Target.ProjectFile.GetFileNameWithoutExtension(), (Target.Platform == UnrealTargetPlatform.IOS ? "IOS" : "TVOS")));
 				}
 
@@ -1632,6 +1633,7 @@ namespace UnrealBuildTool
 
 				// ensure the plist, entitlements, and provision files are properly copied
 				DeployHandler = (Target.Platform == UnrealTargetPlatform.IOS ? new UEDeployIOS() : new UEDeployTVOS());
+				DeployHandler.ForDistribution = Target.bForDistribution;
 				DeployHandler.PrepTargetForDeployment(Target.ProjectFile, Target.TargetName, Target.Platform, Target.Configuration, Target.UPLScripts, Target.SdkVersion, true);
 
 				FileReference SignProjectScript = FileReference.Combine(Target.ProjectIntermediateDirectory, "SignProject.sh");
@@ -1656,7 +1658,7 @@ namespace UnrealBuildTool
 					if(Target.ImportCertificate == null)
 					{
 						// Take it from the standard settings
-						IOSProvisioningData ProvisioningData = ((IOSPlatform)UEBuildPlatform.GetBuildPlatform(Target.Platform)).ReadProvisioningData(Target.ProjectFile);
+						IOSProvisioningData ProvisioningData = ((IOSPlatform)UEBuildPlatform.GetBuildPlatform(Target.Platform)).ReadProvisioningData(Target.ProjectFile, Target.bForDistribution);
 						SigningCertificate = ProvisioningData.SigningCertificate;
 
 						// Set the identity on the command line
@@ -1703,7 +1705,7 @@ namespace UnrealBuildTool
 					string TeamUUID;
 					if(Target.ImportProvision == null)
 					{
-						IOSProvisioningData ProvisioningData = ((IOSPlatform)UEBuildPlatform.GetBuildPlatform(Target.Platform)).ReadProvisioningData(ProjectSettings);
+						IOSProvisioningData ProvisioningData = ((IOSPlatform)UEBuildPlatform.GetBuildPlatform(Target.Platform)).ReadProvisioningData(ProjectSettings, Target.bForDistribution);
 						MobileProvisionFile = ProvisioningData.MobileProvisionFile;
 						MobileProvisionUUID = ProvisioningData.MobileProvisionUUID;
 						TeamUUID = ProvisioningData.TeamUUID;
@@ -1753,6 +1755,7 @@ namespace UnrealBuildTool
 					if(TempKeychain != null)
 					{
 						Writer.WriteLine("security delete-keychain \"{0}\" || true", TempKeychain);
+						Writer.WriteLine("security list-keychain -s login.keychain");
 					}
 				}
 
diff --git a/Engine/Source/Programs/UnrealBuildTool/Platform/IOS/UEBuildIOS.cs b/Engine/Source/Programs/UnrealBuildTool/Platform/IOS/UEBuildIOS.cs
--- a/Engine/Source/Programs/UnrealBuildTool/Platform/IOS/UEBuildIOS.cs
+++ b/Engine/Source/Programs/UnrealBuildTool/Platform/IOS/UEBuildIOS.cs
@@ -35,6 +35,12 @@ namespace UnrealBuildTool
 		[CommandLine("-skipcrashlytics")]
 		public bool bSkipCrashlytics = false;
 
+		/// <summary>
+		/// Mark the build for distribution
+		/// </summary>
+		[CommandLine("-distribution")]
+		public bool bForDistribution = false;
+
 		/// <summary>
 		/// Manual override for the provision to use. Should be a full path.
 		/// </summary>
@@ -100,6 +106,11 @@ namespace UnrealBuildTool
 			get { return Inner.bSkipCrashlytics; }
 		}
 
+		public bool bForDistribution
+		{
+			get { return Inner.bForDistribution; }
+		}
+
 		public string ImportProvision
 		{
 			get { return Inner.ImportProvision; }
@@ -289,6 +300,12 @@ namespace UnrealBuildTool
 		[ConfigFile(ConfigHierarchyType.Engine, "/Script/IOSRuntimeSettings.IOSRuntimeSettings", "bEnableAdvertisingIdentifier")]
 		public readonly bool bEnableAdvertisingIdentifier = false;
 
+		/// <summary>
+		/// true when building for distribution
+		/// </summary>
+		[ConfigFile(ConfigHierarchyType.Game, "/Script/UnrealEd.ProjectPackagingSettings", "ForDistribution")]
+		public readonly bool bForDistribution = false;
+
 		/// <summary>
 		/// Returns a list of all the non-shipping architectures which are supported
 		/// </summary>
diff --git a/Engine/Source/Programs/UnrealBuildTool/Platform/IOS/UEDeployIOS.cs b/Engine/Source/Programs/UnrealBuildTool/Platform/IOS/UEDeployIOS.cs
--- a/Engine/Source/Programs/UnrealBuildTool/Platform/IOS/UEDeployIOS.cs
+++ b/Engine/Source/Programs/UnrealBuildTool/Platform/IOS/UEDeployIOS.cs
@@ -22,6 +22,13 @@ namespace UnrealBuildTool
 		protected UnrealPluginLanguage UPL = null;
 		protected delegate bool FilenameFilter(string InFilename);
 
+		public bool ForDistribution
+		{
+			get { return bForDistribution; }
+			set { bForDistribution = value; }
+		}
+		bool bForDistribution = false;
+
 		protected class VersionUtilities
 		{
 			public static string BuildDirectory
@@ -29,22 +36,24 @@ namespace UnrealBuildTool
 				get;
 			set;
 			}
-				public static string GameName
+			public static string GameName
 			{
 				get;
-			set;
+				set;
 			}
 
-				static string RunningVersionFilename
+			
+
+			static string RunningVersionFilename
 			{
 				get { return Path.Combine(BuildDirectory, GameName + ".PackageVersionCounter"); }
 			}
 
-				/// <summary>
-				/// Reads the GameName.PackageVersionCounter from disk and bumps the minor version number in it
-				/// </summary>
-				/// <returns></returns>
-				public static string ReadRunningVersion()
+			/// <summary>
+			/// Reads the GameName.PackageVersionCounter from disk and bumps the minor version number in it
+			/// </summary>
+			/// <returns></returns>
+			public static string ReadRunningVersion()
 			{
 				string CurrentVersion = "0.0";
 				if (File.Exists(RunningVersionFilename))
@@ -1154,7 +1163,7 @@ namespace UnrealBuildTool
 
 			if (BuildHostPlatform.Current.Platform == UnrealTargetPlatform.Mac && Environment.GetEnvironmentVariable("UBT_NO_POST_DEPLOY") != "true")
 			{
-				return PrepForUATPackageOrDeploy(Configuration, ProjectFile, GameName, ProjectDirectory, BuildPath + "/" + DecoratedGameName, "../../Engine", false, "", false, bCreateStubIPA, UPLScripts, SdkVersion);
+				return PrepForUATPackageOrDeploy(Configuration, ProjectFile, GameName, ProjectDirectory, BuildPath + "/" + DecoratedGameName, "../../Engine", bForDistribution, "", false, bCreateStubIPA, UPLScripts, SdkVersion);
 			}
 			else
 			{
diff --git a/Engine/Source/Programs/UnrealBuildTool/ProjectFiles/Xcode/XcodeProject.cs b/Engine/Source/Programs/UnrealBuildTool/ProjectFiles/Xcode/XcodeProject.cs
--- a/Engine/Source/Programs/UnrealBuildTool/ProjectFiles/Xcode/XcodeProject.cs
+++ b/Engine/Source/Programs/UnrealBuildTool/ProjectFiles/Xcode/XcodeProject.cs
@@ -133,10 +133,12 @@ namespace UnrealBuildTool
 		/// </summary>
 		/// <param name="InitFilePath">The path to the project file on disk</param>
 		/// <param name="InOnlyGameProject"></param>
-		public XcodeProjectFile(FileReference InitFilePath, FileReference InOnlyGameProject)
+		/// <param name="IsForDistribution">True for distribution builds</param>
+		public XcodeProjectFile(FileReference InitFilePath, FileReference InOnlyGameProject, bool IsForDistribution)
 			: base(InitFilePath)
 		{
 			OnlyGameProject = InOnlyGameProject;
+			bForDistribution = IsForDistribution;
 		}
 
 		public override string ToString()
@@ -144,6 +146,11 @@ namespace UnrealBuildTool
 			return ProjectFilePath.GetFileNameWithoutExtension();
 		}
 
+		/// <summary>
+		///  Used to mark the project for distribution (some platforms require this)
+		/// </summary>
+		bool bForDistribution = false;
+
 		/// <summary>
 		/// Gets Xcode file category based on its extension
 		/// </summary>
@@ -828,7 +835,7 @@ namespace UnrealBuildTool
 			{
 				TVOSPlatform TVOSPlatform = ((TVOSPlatform)UEBuildPlatform.GetBuildPlatform(UnrealTargetPlatform.TVOS));
 				TVOSProjectSettings ProjectSettings = TVOSPlatform.ReadProjectSettings(ProjectFile);
-				TVOSProvisioningData ProvisioningData = TVOSPlatform.ReadProvisioningData(ProjectSettings);
+				TVOSProvisioningData ProvisioningData = TVOSPlatform.ReadProvisioningData(ProjectSettings, bForDistribution);
 				bAutomaticSigning = ProjectSettings.bAutomaticSigning;
 			}
 
@@ -980,7 +987,7 @@ namespace UnrealBuildTool
                 {
 					IOSPlatform IOSPlatform = ((IOSPlatform)UEBuildPlatform.GetBuildPlatform(UnrealTargetPlatform.IOS));
 					IOSProjectSettings ProjectSettings = IOSPlatform.ReadProjectSettings(ProjectFile);
-					IOSProvisioningData ProvisioningData = IOSPlatform.ReadProvisioningData(ProjectSettings);
+					IOSProvisioningData ProvisioningData = IOSPlatform.ReadProvisioningData(ProjectSettings, bForDistribution);
 					IOSRunTimeVersion = ProjectSettings.RuntimeVersion;
 					IOSRunTimeDevices = ProjectSettings.RuntimeDevices;
 					ValidArchs += " arm64 armv7 armv7s";
@@ -999,7 +1006,7 @@ namespace UnrealBuildTool
 				{
 					TVOSPlatform TVOSPlatform = ((TVOSPlatform)UEBuildPlatform.GetBuildPlatform(UnrealTargetPlatform.TVOS));
 					TVOSProjectSettings ProjectSettings = TVOSPlatform.ReadProjectSettings(ProjectFile);
-					TVOSProvisioningData ProvisioningData = TVOSPlatform.ReadProvisioningData(ProjectSettings);
+					TVOSProvisioningData ProvisioningData = TVOSPlatform.ReadProvisioningData(ProjectSettings, bForDistribution);
 					TVOSRunTimeVersion = ProjectSettings.RuntimeVersion;
 					TVOSRunTimeDevices = ProjectSettings.RuntimeDevices;
 					if (!ValidArchs.Contains("arm64"))
diff --git a/Engine/Source/Programs/UnrealBuildTool/ProjectFiles/Xcode/XcodeProjectFileGenerator.cs b/Engine/Source/Programs/UnrealBuildTool/ProjectFiles/Xcode/XcodeProjectFileGenerator.cs
--- a/Engine/Source/Programs/UnrealBuildTool/ProjectFiles/Xcode/XcodeProjectFileGenerator.cs
+++ b/Engine/Source/Programs/UnrealBuildTool/ProjectFiles/Xcode/XcodeProjectFileGenerator.cs
@@ -27,9 +27,18 @@ namespace UnrealBuildTool
 		// always seed the random number the same, so multiple runs of the generator will generate the same project
 		static Random Rand = new Random(0);
 
-		public XcodeProjectFileGenerator(FileReference InOnlyGameProject)
+		/// <summary>
+		/// Mark for distribution builds
+		/// </summary>
+		bool bForDistribution = false;
+
+		public XcodeProjectFileGenerator(FileReference InOnlyGameProject, CommandLineArguments CommandLine)
 			: base(InOnlyGameProject)
 		{
+			if (CommandLine.HasOption("-distribution"))
+			{
+				bForDistribution = true;
+			}
 		}
 
 		/// <summary>
@@ -90,7 +99,7 @@ namespace UnrealBuildTool
 		/// <returns>The newly allocated project file object</returns>
 		protected override ProjectFile AllocateProjectFile(FileReference InitFilePath)
 		{
-			return new XcodeProjectFile(InitFilePath, OnlyGameProject);
+			return new XcodeProjectFile(InitFilePath, OnlyGameProject, bForDistribution);
 		}
 
 		/// ProjectFileGenerator interface
diff --git a/Engine/Source/Programs/UnrealBuildTool/ToolChain/RemoteMac.cs b/Engine/Source/Programs/UnrealBuildTool/ToolChain/RemoteMac.cs
--- a/Engine/Source/Programs/UnrealBuildTool/ToolChain/RemoteMac.cs
+++ b/Engine/Source/Programs/UnrealBuildTool/ToolChain/RemoteMac.cs
@@ -400,7 +400,7 @@ namespace UnrealBuildTool
 			RulesAssembly RulesAssembly = RulesCompiler.CreateTargetRulesAssembly(TargetDesc.ProjectFile, TargetDesc.Name, false, false, TargetDesc.ForeignPlugin);
 
 			// Create the target rules
-			TargetRules Rules = RulesAssembly.CreateTargetRules(TargetDesc.Name, TargetDesc.Platform, TargetDesc.Configuration, TargetDesc.Architecture, TargetDesc.ProjectFile, null);
+			TargetRules Rules = RulesAssembly.CreateTargetRules(TargetDesc.Name, TargetDesc.Platform, TargetDesc.Configuration, TargetDesc.Architecture, TargetDesc.ProjectFile, TargetDesc.AdditionalArguments);
 
 			// Check if we need to enable a nativized plugin, and compile the assembly for that if we do
 			FileReference NativizedPluginFile = Rules.GetNativizedPlugin();
@@ -430,7 +430,7 @@ namespace UnrealBuildTool
 				RemoteArguments.Add("-NoUBTMakefiles");
 
 				// Get the provisioning data for this project
-				IOSProvisioningData ProvisioningData = ((IOSPlatform)UEBuildPlatform.GetBuildPlatform(TargetDesc.Platform)).ReadProvisioningData(TargetDesc.ProjectFile);
+				IOSProvisioningData ProvisioningData = ((IOSPlatform)UEBuildPlatform.GetBuildPlatform(TargetDesc.Platform)).ReadProvisioningData(TargetDesc.ProjectFile, TargetDesc.AdditionalArguments.HasOption("-distribution"));
 				if(ProvisioningData == null || ProvisioningData.MobileProvisionFile == null)
 				{
 					throw new BuildException("Unable to find mobile provision for {0}. See log for more information.", TargetDesc.Name);
@@ -585,7 +585,7 @@ namespace UnrealBuildTool
 				RemoteArguments.Add(String.Format("-Project={0}", GetRemotePath(TargetDesc.ProjectFile)));
 			}
 
-			foreach(string LocalArgument in TargetDesc.AdditionalArguments)
+			foreach (string LocalArgument in TargetDesc.AdditionalArguments)
 			{
 				int EqualsIdx = LocalArgument.IndexOf('=');
 				if(EqualsIdx == -1)
