UBT: Fix generated config files not being read by UBT, or being considered when determining if makefiles are out of date.

diff --git a/Engine/Source/Programs/UnrealBuildTool/ProjectFiles/ProjectFileGenerator.cs b/Engine/Source/Programs/UnrealBuildTool/ProjectFiles/ProjectFileGenerator.cs
--- a/Engine/Source/Programs/UnrealBuildTool/ProjectFiles/ProjectFileGenerator.cs
+++ b/Engine/Source/Programs/UnrealBuildTool/ProjectFiles/ProjectFileGenerator.cs
@@ -500,16 +500,7 @@ namespace UnrealBuildTool
 		/// <returns>True if a preferred IDE was set, false otherwise</returns>
 		public static bool GetPreferredSourceCodeAccessor(FileReference ProjectFile, out ProjectFileFormat Format)
 		{
-			DirectoryReference EngineSavedDir = DirectoryReference.Combine(UnrealBuildTool.EngineDirectory, "Saved");
-			if(UnrealBuildTool.IsEngineInstalled())
-			{
-				BuildVersion Version;
-				if(BuildVersion.TryRead(BuildVersion.GetDefaultFileName(), out Version))
-				{
-					EngineSavedDir = DirectoryReference.Combine(Utils.GetUserSettingDirectory(), "UnrealEngine", String.Format("{0}.{1}", Version.MajorVersion, Version.MinorVersion), "Saved");
-				}
-			}
-			ConfigHierarchy Ini = ConfigCache.ReadHierarchy(ConfigHierarchyType.EditorSettings, DirectoryReference.FromFile(ProjectFile), BuildHostPlatform.Current.Platform, EngineSavedDir);
+			ConfigHierarchy Ini = ConfigCache.ReadHierarchy(ConfigHierarchyType.EditorSettings, DirectoryReference.FromFile(ProjectFile), BuildHostPlatform.Current.Platform);
 
 			string PreferredAccessor;
 			if (Ini.GetString("/Script/SourceCodeAccess.SourceCodeAccessSettings", "PreferredAccessor", out PreferredAccessor))
diff --git a/Engine/Source/Programs/UnrealBuildTool/System/ConfigCache.cs b/Engine/Source/Programs/UnrealBuildTool/System/ConfigCache.cs
--- a/Engine/Source/Programs/UnrealBuildTool/System/ConfigCache.cs
+++ b/Engine/Source/Programs/UnrealBuildTool/System/ConfigCache.cs
@@ -1,4 +1,4 @@
-﻿// Copyright 1998-2019 Epic Games, Inc. All Rights Reserved.
+// Copyright 1998-2019 Epic Games, Inc. All Rights Reserved.
 
 using System;
 using System.Collections.Generic;
@@ -156,9 +156,8 @@ namespace UnrealBuildTool
 		/// <param name="Type">The type of hierarchy to read</param>
 		/// <param name="ProjectDir">The project directory to read the hierarchy for</param>
 		/// <param name="Platform">Which platform to read platform-specific config files for</param>
-		/// <param name="GeneratedConfigDir">Base directory for generated configs</param>
 		/// <returns>The requested config hierarchy</returns>
-		public static ConfigHierarchy ReadHierarchy(ConfigHierarchyType Type, DirectoryReference ProjectDir, UnrealTargetPlatform Platform, DirectoryReference GeneratedConfigDir = null)
+		public static ConfigHierarchy ReadHierarchy(ConfigHierarchyType Type, DirectoryReference ProjectDir, UnrealTargetPlatform Platform)
 		{
 			// Get the key to use for the cache. It cannot be null, so we use the engine directory if a project directory is not given.
 			ConfigHierarchyKey Key = new ConfigHierarchyKey(Type, ProjectDir, Platform);
@@ -169,6 +168,7 @@ namespace UnrealBuildTool
 			{
 				if (!HierarchyKeyToHierarchy.TryGetValue(Key, out Hierarchy))
 				{
+					// Find all the input files
 					List<ConfigFile> Files = new List<ConfigFile>();
 					foreach (FileReference IniFileName in ConfigHierarchy.EnumerateConfigFileLocations(Type, ProjectDir, Platform))
 					{
@@ -179,27 +179,6 @@ namespace UnrealBuildTool
 						}
 					}
 
-					// If we haven't been given a generated project dir, but we do have a project then the generated configs
-					// should go into ProjectDir/Saved
-					if (GeneratedConfigDir == null && ProjectDir != null)
-					{
-						GeneratedConfigDir = DirectoryReference.Combine(ProjectDir, "Saved");
-					}
-
-					if (GeneratedConfigDir != null)
-					{
-						// We know where the generated version of this config file lives, so we can read it back in
-						// and include any user settings from there in our hierarchy
-						string BaseIniName = Enum.GetName(typeof(ConfigHierarchyType), Type);
-						string PlatformName = ConfigHierarchy.GetIniPlatformName(Platform);
-						FileReference DestinationIniFilename = FileReference.Combine(GeneratedConfigDir, "Config", PlatformName, BaseIniName + ".ini");
-						ConfigFile File;
-						if (TryReadFile(DestinationIniFilename, out File))
-						{
-							Files.Add(File);
-						}
-					}
-
 					// Handle command line overrides
 					string[] CmdLine = Environment.GetCommandLineArgs();
 					string IniConfigArgPrefix = "-ini:" + Enum.GetName(typeof(ConfigHierarchyType), Type) + ":";
@@ -212,6 +191,7 @@ namespace UnrealBuildTool
 						}
 					}
 
+					// Create the hierarchy
 					Hierarchy = new ConfigHierarchy(Files);
 					HierarchyKeyToHierarchy.Add(Key, Hierarchy);
 				}
diff --git a/Engine/Source/Programs/UnrealBuildTool/System/ConfigHierarchy.cs b/Engine/Source/Programs/UnrealBuildTool/System/ConfigHierarchy.cs
--- a/Engine/Source/Programs/UnrealBuildTool/System/ConfigHierarchy.cs
+++ b/Engine/Source/Programs/UnrealBuildTool/System/ConfigHierarchy.cs
@@ -850,6 +850,48 @@ namespace UnrealBuildTool
 			{
 				yield return FileReference.Combine(ProjectDir, "Config", "User" + BaseIniName + ".ini");
 			}
+
+			// Get the generated config file too. EditorSettings overrides this from 
+			if(Type == ConfigHierarchyType.EditorSettings)
+			{
+				yield return FileReference.Combine(GetGameAgnosticSavedDir(), "Config", PlatformName, BaseIniName + ".ini");
+			}
+			else
+			{
+				yield return FileReference.Combine(GetGeneratedConfigDir(ProjectDir), PlatformName, BaseIniName + ".ini");
+			}
+		}
+
+		/// <summary>
+		/// Determines the path to the generated config directory (same as FPaths::GeneratedConfigDir())
+		/// </summary>
+		/// <returns></returns>
+		public static DirectoryReference GetGeneratedConfigDir(DirectoryReference ProjectDir)
+		{
+			if(ProjectDir == null)
+			{
+				return DirectoryReference.Combine(UnrealBuildTool.EngineDirectory, "Saved", "Config");
+			}
+			else
+			{
+				return DirectoryReference.Combine(ProjectDir, "Saved", "Config");
+			}
+		}
+
+		/// <summary>
+		/// Determes the path to the game-agnostic saved directory (same as FPaths::GameAgnosticSavedDir())
+		/// </summary>
+		/// <returns></returns>
+		public static DirectoryReference GetGameAgnosticSavedDir()
+		{
+			if(UnrealBuildTool.IsEngineInstalled())
+			{
+				return DirectoryReference.Combine(Utils.GetUserSettingDirectory(), "UnrealEngine", String.Format("{0}.{1}", ReadOnlyBuildVersion.Current.MajorVersion, ReadOnlyBuildVersion.Current.MinorVersion), "Saved");
+			}
+			else
+			{
+				return DirectoryReference.Combine(UnrealBuildTool.EngineDirectory, "Saved");
+			}
 		}
 
 		/// <summary>
