UE-70766 - iOS remote build unloads the Login Keychain on the remote Mac - additional

diff --git a/Engine/Source/Programs/UnrealBuildTool/Platform/IOS/IOSToolChain.cs b/Engine/Source/Programs/UnrealBuildTool/Platform/IOS/IOSToolChain.cs
--- a/Engine/Source/Programs/UnrealBuildTool/Platform/IOS/IOSToolChain.cs
+++ b/Engine/Source/Programs/UnrealBuildTool/Platform/IOS/IOSToolChain.cs
@@ -1509,6 +1509,14 @@ namespace UnrealBuildTool
 					}
 				}
 			}
+
+			public void OutputReceivedDataEventLogger(Object Sender, DataReceivedEventArgs Line)
+			{
+				if ((Line != null) && (Line.Data != null))
+				{
+					Log.TraceInformation(Line.Data);
+				}
+			}
 		}
 
 		private static void GenerateCrashlyticsData(string ExecutableDirectory, string ExecutableName, string ProjectDir, string ProjectName)
@@ -1636,6 +1644,9 @@ namespace UnrealBuildTool
 				DeployHandler.ForDistribution = Target.bForDistribution;
 				DeployHandler.PrepTargetForDeployment(Target.ProjectFile, Target.TargetName, Target.Platform, Target.Configuration, Target.UPLScripts, Target.SdkVersion, true);
 
+				// Path to the temporary keychain. When -ImportCertificate is specified, we will temporarily add this to the list of keychains to search, and remove it later.
+				FileReference TempKeychain = FileReference.Combine(Target.ProjectIntermediateDirectory, "TempKeychain.keychain");
+
 				FileReference SignProjectScript = FileReference.Combine(Target.ProjectIntermediateDirectory, "SignProject.sh");
 				using(StreamWriter Writer = new StreamWriter(SignProjectScript.FullName))
 				{
@@ -1650,9 +1661,6 @@ namespace UnrealBuildTool
 						Writer.WriteLine("cp -f {0} ~/Library/MobileDevice/Provisioning\\ Profiles/", Utils.EscapeShellArgument(Target.ImportProvision));
 					}
 
-					// Path to the temporary keychain. When -ImportCertificate is specified, we will temporarily add this to the list of keychains to search, and remove it later.
-					FileReference TempKeychain = null;
-
 					// Get the signing certificate to use
 					string SigningCertificate;
 					if(Target.ImportCertificate == null)
@@ -1681,9 +1689,6 @@ namespace UnrealBuildTool
 						}
 						SigningCertificate = Certificate.GetNameInfo(X509NameType.SimpleName, false);
 
-						// Set the path to the temporary keychain
-						TempKeychain = FileReference.Combine(Target.ProjectIntermediateDirectory, "TempKeychain.keychain");//(DirectoryReference.GetSpecialFolder(Environment.SpecialFolder.UserProfile), "Library", "Keychains/UE4TempKeychain.keychain";
-
 						// Install a certificate given on the command line to a temporary keychain
 						Writer.WriteLine("security delete-keychain \"{0}\" || true", TempKeychain);
 						Writer.WriteLine("security create-keychain -p \"A\" \"{0}\"", TempKeychain);
@@ -1750,13 +1755,6 @@ namespace UnrealBuildTool
 						CmdLine += (!string.IsNullOrEmpty(MobileProvisionUUID) ? (" PROVISIONING_PROFILE_SPECIFIER=" + MobileProvisionUUID) : "");
 					}
 					Writer.WriteLine("/usr/bin/xcrun {0}", CmdLine);
-
-					// Remove the temporary keychain from the search list
-					if(TempKeychain != null)
-					{
-						Writer.WriteLine("security delete-keychain \"{0}\" || true", TempKeychain);
-						Writer.WriteLine("security list-keychain -s login.keychain");
-					}
 				}
 
 				Log.TraceInformation("Executing {0}", SignProjectScript);
@@ -1770,11 +1768,37 @@ namespace UnrealBuildTool
 
 				SignProcess.OutputDataReceived += new DataReceivedEventHandler(Output.OutputReceivedDataEventHandler);
 				SignProcess.ErrorDataReceived += new DataReceivedEventHandler(Output.OutputReceivedDataEventHandler);
-
+				
 				Output.OutputReceivedDataEventHandlerEncounteredError = false;
 				Output.OutputReceivedDataEventHandlerEncounteredErrorMessage = "";
 				Utils.RunLocalProcess(SignProcess);
 
+				// cleanup
+				if (Target.ImportCertificate != null)
+				{
+					FileReference CleanProjectScript = FileReference.Combine(Target.ProjectIntermediateDirectory, "CleanProject.sh");
+					using (StreamWriter CleanWriter = new StreamWriter(CleanProjectScript.FullName))
+					{
+						// Remove the temporary keychain from the search list
+						CleanWriter.WriteLine("security delete-keychain \"{0}\" || true", TempKeychain);
+						CleanWriter.WriteLine("security list-keychain -s login.keychain");
+					}
+					
+					Log.TraceInformation("Executing {0}", CleanProjectScript);
+
+					Process CleanProcess = new Process();
+					CleanProcess.StartInfo.WorkingDirectory = RemoteShadowDirectoryMac;
+					CleanProcess.StartInfo.FileName = "/bin/sh";
+					CleanProcess.StartInfo.Arguments = CleanProjectScript.FullName;
+
+					ProcessOutput CleanOutput = new ProcessOutput();
+
+					SignProcess.OutputDataReceived += new DataReceivedEventHandler(CleanOutput.OutputReceivedDataEventLogger);
+					SignProcess.ErrorDataReceived += new DataReceivedEventHandler(CleanOutput.OutputReceivedDataEventLogger);
+
+					Utils.RunLocalProcess(CleanProcess);
+				}
+
 				// delete the temp project
 				DirectoryReference.Delete(XcodeWorkspaceDir, true);
 
