UE-67484 - tvOS LaunchOn error: IBCocoaTouchImageCatalogTool[7050:421228] CoreThemeDefinition: Unable to get image stack data for file

diff --git a/Engine/Build/TVOS/Resources/Assets.xcassets/App Icon & Top Shelf Image.brandassets/App Icon - App Store.imagestack/Back.imagestacklayer/Content.imageset/Contents.json b/Engine/Build/TVOS/Resources/Assets.xcassets/App Icon & Top Shelf Image.brandassets/App Icon - App Store.imagestack/Back.imagestacklayer/Content.imageset/Contents.json
new file mode 100644
--- /dev/null
+++ b/Engine/Build/TVOS/Resources/Assets.xcassets/App Icon & Top Shelf Image.brandassets/App Icon - App Store.imagestack/Back.imagestacklayer/Content.imageset/Contents.json	
@@ -0,0 +1,17 @@
+{
+	"images" : [
+		{
+			"idiom" : "tv",
+			"filename" : "Icon_Large_Back-1280x768.png",
+			"scale" : "1x"
+    },
+    {
+      "idiom" : "tv",
+      "scale" : "2x"
+    }
+	],
+	"info" : {
+		"version" : 1,
+		"author" : "xcode"
+	}
+}
\ No newline at end of file
diff --git a/Engine/Build/TVOS/Resources/Assets.xcassets/App Icon & Top Shelf Image.brandassets/App Icon - App Store.imagestack/Back.imagestacklayer/Contents.json b/Engine/Build/TVOS/Resources/Assets.xcassets/App Icon & Top Shelf Image.brandassets/App Icon - App Store.imagestack/Back.imagestacklayer/Contents.json
new file mode 100644
--- /dev/null
+++ b/Engine/Build/TVOS/Resources/Assets.xcassets/App Icon & Top Shelf Image.brandassets/App Icon - App Store.imagestack/Back.imagestacklayer/Contents.json	
@@ -0,0 +1,6 @@
+{
+	"info" : {
+		"version" : 1,
+		"author" : "xcode"
+	}
+}
\ No newline at end of file
diff --git a/Engine/Build/TVOS/Resources/Assets.xcassets/App Icon & Top Shelf Image.brandassets/App Icon - App Store.imagestack/Contents.json b/Engine/Build/TVOS/Resources/Assets.xcassets/App Icon & Top Shelf Image.brandassets/App Icon - App Store.imagestack/Contents.json
new file mode 100644
--- /dev/null
+++ b/Engine/Build/TVOS/Resources/Assets.xcassets/App Icon & Top Shelf Image.brandassets/App Icon - App Store.imagestack/Contents.json	
@@ -0,0 +1,17 @@
+{
+	"layers" : [
+		{
+			"filename" : "Front.imagestacklayer"
+		},
+		{
+			"filename" : "Middle.imagestacklayer"
+		},
+		{
+			"filename" : "Back.imagestacklayer"
+		}
+	],
+	"info" : {
+		"version" : 1,
+		"author" : "xcode"
+	}
+}
\ No newline at end of file
diff --git a/Engine/Build/TVOS/Resources/Assets.xcassets/App Icon & Top Shelf Image.brandassets/App Icon - App Store.imagestack/Front.imagestacklayer/Content.imageset/Contents.json b/Engine/Build/TVOS/Resources/Assets.xcassets/App Icon & Top Shelf Image.brandassets/App Icon - App Store.imagestack/Front.imagestacklayer/Content.imageset/Contents.json
new file mode 100644
--- /dev/null
+++ b/Engine/Build/TVOS/Resources/Assets.xcassets/App Icon & Top Shelf Image.brandassets/App Icon - App Store.imagestack/Front.imagestacklayer/Content.imageset/Contents.json	
@@ -0,0 +1,17 @@
+{
+	"images" : [
+		{
+			"idiom" : "tv",
+			"filename" : "Icon_Large_Front-1280x768.png",
+			"scale" : "1x"
+    },
+    {
+      "idiom" : "tv",
+      "scale" : "2x"
+    }
+	],
+	"info" : {
+		"version" : 1,
+		"author" : "xcode"
+	}
+}
\ No newline at end of file
diff --git a/Engine/Build/TVOS/Resources/Assets.xcassets/App Icon & Top Shelf Image.brandassets/App Icon - App Store.imagestack/Front.imagestacklayer/Contents.json b/Engine/Build/TVOS/Resources/Assets.xcassets/App Icon & Top Shelf Image.brandassets/App Icon - App Store.imagestack/Front.imagestacklayer/Contents.json
new file mode 100644
--- /dev/null
+++ b/Engine/Build/TVOS/Resources/Assets.xcassets/App Icon & Top Shelf Image.brandassets/App Icon - App Store.imagestack/Front.imagestacklayer/Contents.json	
@@ -0,0 +1,6 @@
+{
+	"info" : {
+		"version" : 1,
+		"author" : "xcode"
+	}
+}
\ No newline at end of file
diff --git a/Engine/Build/TVOS/Resources/Assets.xcassets/App Icon & Top Shelf Image.brandassets/App Icon - App Store.imagestack/Middle.imagestacklayer/Content.imageset/Contents.json b/Engine/Build/TVOS/Resources/Assets.xcassets/App Icon & Top Shelf Image.brandassets/App Icon - App Store.imagestack/Middle.imagestacklayer/Content.imageset/Contents.json
new file mode 100644
--- /dev/null
+++ b/Engine/Build/TVOS/Resources/Assets.xcassets/App Icon & Top Shelf Image.brandassets/App Icon - App Store.imagestack/Middle.imagestacklayer/Content.imageset/Contents.json	
@@ -0,0 +1,17 @@
+{
+	"images" : [
+		{
+			"idiom" : "tv",
+			"filename" : "Icon_Large_Middle-1280x768.png",
+			"scale" : "1x"
+    },
+    {
+      "idiom" : "tv",
+      "scale" : "2x"
+    }
+	],
+	"info" : {
+		"version" : 1,
+		"author" : "xcode"
+	}
+}
\ No newline at end of file
diff --git a/Engine/Build/TVOS/Resources/Assets.xcassets/App Icon & Top Shelf Image.brandassets/App Icon - App Store.imagestack/Middle.imagestacklayer/Contents.json b/Engine/Build/TVOS/Resources/Assets.xcassets/App Icon & Top Shelf Image.brandassets/App Icon - App Store.imagestack/Middle.imagestacklayer/Contents.json
new file mode 100644
--- /dev/null
+++ b/Engine/Build/TVOS/Resources/Assets.xcassets/App Icon & Top Shelf Image.brandassets/App Icon - App Store.imagestack/Middle.imagestacklayer/Contents.json	
@@ -0,0 +1,6 @@
+{
+	"info" : {
+		"version" : 1,
+		"author" : "xcode"
+	}
+}
\ No newline at end of file
diff --git a/Engine/Build/TVOS/Resources/Assets.xcassets/App Icon & Top Shelf Image.brandassets/App Icon.imagestack/Back.imagestacklayer/Content.imageset/Contents.json b/Engine/Build/TVOS/Resources/Assets.xcassets/App Icon & Top Shelf Image.brandassets/App Icon.imagestack/Back.imagestacklayer/Content.imageset/Contents.json
new file mode 100644
--- /dev/null
+++ b/Engine/Build/TVOS/Resources/Assets.xcassets/App Icon & Top Shelf Image.brandassets/App Icon.imagestack/Back.imagestacklayer/Content.imageset/Contents.json	
@@ -0,0 +1,18 @@
+{
+	"images" : [
+		{
+			"idiom" : "tv",
+			"filename" : "Icon_Small_Back-400x240.png",
+			"scale" : "1x"
+		},
+		{
+			"idiom" : "tv",
+			"filename" : "Icon_Small_Back-800x480.png",
+			"scale" : "2x"
+		}
+	],
+	"info" : {
+		"version" : 1,
+		"author" : "xcode"
+	}
+}
\ No newline at end of file
diff --git a/Engine/Build/TVOS/Resources/Assets.xcassets/App Icon & Top Shelf Image.brandassets/App Icon.imagestack/Back.imagestacklayer/Contents.json b/Engine/Build/TVOS/Resources/Assets.xcassets/App Icon & Top Shelf Image.brandassets/App Icon.imagestack/Back.imagestacklayer/Contents.json
new file mode 100644
--- /dev/null
+++ b/Engine/Build/TVOS/Resources/Assets.xcassets/App Icon & Top Shelf Image.brandassets/App Icon.imagestack/Back.imagestacklayer/Contents.json	
@@ -0,0 +1,6 @@
+{
+	"info" : {
+		"version" : 1,
+		"author" : "xcode"
+	}
+}
\ No newline at end of file
diff --git a/Engine/Build/TVOS/Resources/Assets.xcassets/App Icon & Top Shelf Image.brandassets/App Icon.imagestack/Contents.json b/Engine/Build/TVOS/Resources/Assets.xcassets/App Icon & Top Shelf Image.brandassets/App Icon.imagestack/Contents.json
new file mode 100644
--- /dev/null
+++ b/Engine/Build/TVOS/Resources/Assets.xcassets/App Icon & Top Shelf Image.brandassets/App Icon.imagestack/Contents.json	
@@ -0,0 +1,17 @@
+{
+	"layers" : [
+		{
+			"filename" : "Front.imagestacklayer"
+		},
+		{
+			"filename" : "Middle.imagestacklayer"
+		},
+		{
+			"filename" : "Back.imagestacklayer"
+		}
+	],
+	"info" : {
+		"version" : 1,
+		"author" : "xcode"
+	}
+}
\ No newline at end of file
diff --git a/Engine/Build/TVOS/Resources/Assets.xcassets/App Icon & Top Shelf Image.brandassets/App Icon.imagestack/Front.imagestacklayer/Content.imageset/Contents.json b/Engine/Build/TVOS/Resources/Assets.xcassets/App Icon & Top Shelf Image.brandassets/App Icon.imagestack/Front.imagestacklayer/Content.imageset/Contents.json
new file mode 100644
--- /dev/null
+++ b/Engine/Build/TVOS/Resources/Assets.xcassets/App Icon & Top Shelf Image.brandassets/App Icon.imagestack/Front.imagestacklayer/Content.imageset/Contents.json	
@@ -0,0 +1,18 @@
+{
+	"images" : [
+		{
+			"idiom" : "tv",
+			"filename" : "Icon_Small_Front-400x240.png",
+			"scale" : "1x"
+		},
+		{
+			"idiom" : "tv",
+			"filename" : "Icon_Small_Front-800x480.png",
+			"scale" : "2x"
+		}
+	],
+	"info" : {
+		"version" : 1,
+		"author" : "xcode"
+	}
+}
\ No newline at end of file
diff --git a/Engine/Build/TVOS/Resources/Assets.xcassets/App Icon & Top Shelf Image.brandassets/App Icon.imagestack/Front.imagestacklayer/Contents.json b/Engine/Build/TVOS/Resources/Assets.xcassets/App Icon & Top Shelf Image.brandassets/App Icon.imagestack/Front.imagestacklayer/Contents.json
new file mode 100644
--- /dev/null
+++ b/Engine/Build/TVOS/Resources/Assets.xcassets/App Icon & Top Shelf Image.brandassets/App Icon.imagestack/Front.imagestacklayer/Contents.json	
@@ -0,0 +1,6 @@
+{
+	"info" : {
+		"version" : 1,
+		"author" : "xcode"
+	}
+}
\ No newline at end of file
diff --git a/Engine/Build/TVOS/Resources/Assets.xcassets/App Icon & Top Shelf Image.brandassets/App Icon.imagestack/Middle.imagestacklayer/Content.imageset/Contents.json b/Engine/Build/TVOS/Resources/Assets.xcassets/App Icon & Top Shelf Image.brandassets/App Icon.imagestack/Middle.imagestacklayer/Content.imageset/Contents.json
new file mode 100644
--- /dev/null
+++ b/Engine/Build/TVOS/Resources/Assets.xcassets/App Icon & Top Shelf Image.brandassets/App Icon.imagestack/Middle.imagestacklayer/Content.imageset/Contents.json	
@@ -0,0 +1,18 @@
+{
+	"images" : [
+		{
+			"idiom" : "tv",
+			"filename" : "Icon_Small_Middle-400x240.png",
+			"scale" : "1x"
+		},
+		{
+			"idiom" : "tv",
+			"filename" : "Icon_Small_Middle-800x480.png",
+			"scale" : "2x"
+		}
+	],
+	"info" : {
+		"version" : 1,
+		"author" : "xcode"
+	}
+}
\ No newline at end of file
diff --git a/Engine/Build/TVOS/Resources/Assets.xcassets/App Icon & Top Shelf Image.brandassets/App Icon.imagestack/Middle.imagestacklayer/Contents.json b/Engine/Build/TVOS/Resources/Assets.xcassets/App Icon & Top Shelf Image.brandassets/App Icon.imagestack/Middle.imagestacklayer/Contents.json
new file mode 100644
--- /dev/null
+++ b/Engine/Build/TVOS/Resources/Assets.xcassets/App Icon & Top Shelf Image.brandassets/App Icon.imagestack/Middle.imagestacklayer/Contents.json	
@@ -0,0 +1,6 @@
+{
+	"info" : {
+		"version" : 1,
+		"author" : "xcode"
+	}
+}
\ No newline at end of file
diff --git a/Engine/Build/TVOS/Resources/Assets.xcassets/App Icon & Top Shelf Image.brandassets/Contents.json b/Engine/Build/TVOS/Resources/Assets.xcassets/App Icon & Top Shelf Image.brandassets/Contents.json
new file mode 100644
--- /dev/null
+++ b/Engine/Build/TVOS/Resources/Assets.xcassets/App Icon & Top Shelf Image.brandassets/Contents.json	
@@ -0,0 +1,32 @@
+{
+  "assets" : [
+    {
+      "size" : "1280x768",
+      "idiom" : "tv",
+      "filename" : "App Icon - App Store.imagestack",
+      "role" : "primary-app-icon"
+    },
+    {
+      "size" : "400x240",
+      "idiom" : "tv",
+      "filename" : "App Icon.imagestack",
+      "role" : "primary-app-icon"
+    },
+    {
+      "size" : "2320x720",
+      "idiom" : "tv",
+      "filename" : "Top Shelf Image Wide.imageset",
+      "role" : "top-shelf-image-wide"
+    },
+    {
+      "size" : "1920x720",
+      "idiom" : "tv",
+      "filename" : "Top Shelf Image.imageset",
+      "role" : "top-shelf-image"
+    }
+  ],
+  "info" : {
+    "version" : 1,
+    "author" : "xcode"
+  }
+}
diff --git a/Engine/Build/TVOS/Resources/Assets.xcassets/App Icon & Top Shelf Image.brandassets/Top Shelf Image Wide.imageset/Contents.json b/Engine/Build/TVOS/Resources/Assets.xcassets/App Icon & Top Shelf Image.brandassets/Top Shelf Image Wide.imageset/Contents.json
new file mode 100644
--- /dev/null
+++ b/Engine/Build/TVOS/Resources/Assets.xcassets/App Icon & Top Shelf Image.brandassets/Top Shelf Image Wide.imageset/Contents.json	
@@ -0,0 +1,18 @@
+{
+	"images" : [
+		{
+			"idiom" : "tv",
+			"filename" : "TopShelfWide-1920x720.png",
+			"scale" : "1x"
+		},
+		{
+			"idiom" : "tv",
+			"filename" : "[REDACTED_EMAIL]",
+			"scale" : "2x"
+		}
+	],
+	"info" : {
+		"version" : 1,
+		"author" : "xcode"
+	}
+}
\ No newline at end of file
diff --git a/Engine/Build/TVOS/Resources/Assets.xcassets/App Icon & Top Shelf Image.brandassets/Top Shelf Image.imageset/Contents.json b/Engine/Build/TVOS/Resources/Assets.xcassets/App Icon & Top Shelf Image.brandassets/Top Shelf Image.imageset/Contents.json
new file mode 100644
--- /dev/null
+++ b/Engine/Build/TVOS/Resources/Assets.xcassets/App Icon & Top Shelf Image.brandassets/Top Shelf Image.imageset/Contents.json	
@@ -0,0 +1,18 @@
+{
+	"images" : [
+		{
+			"idiom" : "tv",
+			"filename" : "TopShelf.png",
+			"scale" : "1x"
+		},
+		{
+			"idiom" : "tv",
+			"filename" : "[REDACTED_EMAIL]",
+			"scale" : "2x"
+		}
+	],
+	"info" : {
+		"version" : 1,
+		"author" : "xcode"
+	}
+}
\ No newline at end of file
diff --git a/Engine/Build/TVOS/Resources/Assets.xcassets/AppIcon.brandassets/App Icon - Large.imagestack/Back.imagestacklayer/Content.imageset/Contents.json b/Engine/Build/TVOS/Resources/Assets.xcassets/AppIcon.brandassets/App Icon - Large.imagestack/Back.imagestacklayer/Content.imageset/Contents.json
deleted file mode 100644
diff --git a/Engine/Build/TVOS/Resources/Assets.xcassets/AppIcon.brandassets/App Icon - Large.imagestack/Back.imagestacklayer/Contents.json b/Engine/Build/TVOS/Resources/Assets.xcassets/AppIcon.brandassets/App Icon - Large.imagestack/Back.imagestacklayer/Contents.json
deleted file mode 100644
diff --git a/Engine/Build/TVOS/Resources/Assets.xcassets/AppIcon.brandassets/App Icon - Large.imagestack/Contents.json b/Engine/Build/TVOS/Resources/Assets.xcassets/AppIcon.brandassets/App Icon - Large.imagestack/Contents.json
deleted file mode 100644
diff --git a/Engine/Build/TVOS/Resources/Assets.xcassets/AppIcon.brandassets/App Icon - Large.imagestack/Front.imagestacklayer/Content.imageset/Contents.json b/Engine/Build/TVOS/Resources/Assets.xcassets/AppIcon.brandassets/App Icon - Large.imagestack/Front.imagestacklayer/Content.imageset/Contents.json
deleted file mode 100644
diff --git a/Engine/Build/TVOS/Resources/Assets.xcassets/AppIcon.brandassets/App Icon - Large.imagestack/Front.imagestacklayer/Contents.json b/Engine/Build/TVOS/Resources/Assets.xcassets/AppIcon.brandassets/App Icon - Large.imagestack/Front.imagestacklayer/Contents.json
deleted file mode 100644
diff --git a/Engine/Build/TVOS/Resources/Assets.xcassets/AppIcon.brandassets/App Icon - Large.imagestack/Middle.imagestacklayer/Content.imageset/Contents.json b/Engine/Build/TVOS/Resources/Assets.xcassets/AppIcon.brandassets/App Icon - Large.imagestack/Middle.imagestacklayer/Content.imageset/Contents.json
deleted file mode 100644
diff --git a/Engine/Build/TVOS/Resources/Assets.xcassets/AppIcon.brandassets/App Icon - Large.imagestack/Middle.imagestacklayer/Contents.json b/Engine/Build/TVOS/Resources/Assets.xcassets/AppIcon.brandassets/App Icon - Large.imagestack/Middle.imagestacklayer/Contents.json
deleted file mode 100644
diff --git a/Engine/Build/TVOS/Resources/Assets.xcassets/AppIcon.brandassets/App Icon - Small.imagestack/Back.imagestacklayer/Content.imageset/Contents.json b/Engine/Build/TVOS/Resources/Assets.xcassets/AppIcon.brandassets/App Icon - Small.imagestack/Back.imagestacklayer/Content.imageset/Contents.json
deleted file mode 100644
diff --git a/Engine/Build/TVOS/Resources/Assets.xcassets/AppIcon.brandassets/App Icon - Small.imagestack/Back.imagestacklayer/Contents.json b/Engine/Build/TVOS/Resources/Assets.xcassets/AppIcon.brandassets/App Icon - Small.imagestack/Back.imagestacklayer/Contents.json
deleted file mode 100644
diff --git a/Engine/Build/TVOS/Resources/Assets.xcassets/AppIcon.brandassets/App Icon - Small.imagestack/Contents.json b/Engine/Build/TVOS/Resources/Assets.xcassets/AppIcon.brandassets/App Icon - Small.imagestack/Contents.json
deleted file mode 100644
diff --git a/Engine/Build/TVOS/Resources/Assets.xcassets/AppIcon.brandassets/App Icon - Small.imagestack/Front.imagestacklayer/Content.imageset/Contents.json b/Engine/Build/TVOS/Resources/Assets.xcassets/AppIcon.brandassets/App Icon - Small.imagestack/Front.imagestacklayer/Content.imageset/Contents.json
deleted file mode 100644
diff --git a/Engine/Build/TVOS/Resources/Assets.xcassets/AppIcon.brandassets/App Icon - Small.imagestack/Front.imagestacklayer/Contents.json b/Engine/Build/TVOS/Resources/Assets.xcassets/AppIcon.brandassets/App Icon - Small.imagestack/Front.imagestacklayer/Contents.json
deleted file mode 100644
diff --git a/Engine/Build/TVOS/Resources/Assets.xcassets/AppIcon.brandassets/App Icon - Small.imagestack/Middle.imagestacklayer/Content.imageset/Contents.json b/Engine/Build/TVOS/Resources/Assets.xcassets/AppIcon.brandassets/App Icon - Small.imagestack/Middle.imagestacklayer/Content.imageset/Contents.json
deleted file mode 100644
diff --git a/Engine/Build/TVOS/Resources/Assets.xcassets/AppIcon.brandassets/App Icon - Small.imagestack/Middle.imagestacklayer/Contents.json b/Engine/Build/TVOS/Resources/Assets.xcassets/AppIcon.brandassets/App Icon - Small.imagestack/Middle.imagestacklayer/Contents.json
deleted file mode 100644
diff --git a/Engine/Build/TVOS/Resources/Assets.xcassets/AppIcon.brandassets/Contents.json b/Engine/Build/TVOS/Resources/Assets.xcassets/AppIcon.brandassets/Contents.json
deleted file mode 100644
diff --git a/Engine/Build/TVOS/Resources/Assets.xcassets/AppIcon.brandassets/Top Shelf Image Wide.imageset/Contents.json b/Engine/Build/TVOS/Resources/Assets.xcassets/AppIcon.brandassets/Top Shelf Image Wide.imageset/Contents.json
deleted file mode 100644
diff --git a/Engine/Build/TVOS/Resources/Assets.xcassets/AppIcon.brandassets/Top Shelf Image.imageset/Contents.json b/Engine/Build/TVOS/Resources/Assets.xcassets/AppIcon.brandassets/Top Shelf Image.imageset/Contents.json
deleted file mode 100644
diff --git a/Engine/Build/TVOS/Resources/Assets.xcassets/Launch Image.launchimage/Contents.json b/Engine/Build/TVOS/Resources/Assets.xcassets/Launch Image.launchimage/Contents.json
new file mode 100644
--- /dev/null
+++ b/Engine/Build/TVOS/Resources/Assets.xcassets/Launch Image.launchimage/Contents.json	
@@ -0,0 +1,16 @@
+{
+	"images" : [
+		{
+			"orientation" : "landscape",
+			"idiom" : "tv",
+			"filename" : "Launch.png",
+			"extent" : "full-screen",
+			"minimum-system-version" : "9.0",
+			"scale" : "1x"
+		}
+	],
+	"info" : {
+		"version" : 1,
+		"author" : "xcode"
+	}
+}
\ No newline at end of file
diff --git a/Engine/Build/TVOS/Resources/Assets.xcassets/LaunchImage.launchimage/Contents.json b/Engine/Build/TVOS/Resources/Assets.xcassets/LaunchImage.launchimage/Contents.json
deleted file mode 100644
diff --git a/Engine/Source/Programs/UnrealBuildTool/Platform/IOS/IOSToolChain.cs b/Engine/Source/Programs/UnrealBuildTool/Platform/IOS/IOSToolChain.cs
--- a/Engine/Source/Programs/UnrealBuildTool/Platform/IOS/IOSToolChain.cs
+++ b/Engine/Source/Programs/UnrealBuildTool/Platform/IOS/IOSToolChain.cs
@@ -1240,25 +1240,21 @@ namespace UnrealBuildTool
 				}
 				// copy the icons from the game directory if it has any
 				string[][] Images = {
-					new string []{ "Icon_Large_Back-1280x768.png", "AppIcon.brandassets/App Icon - Large.imagestack/Back.imagestacklayer/Content.imageset" },
-					new string []{ "Icon_Large_Front-1280x768.png", "AppIcon.brandassets/App Icon - Large.imagestack/Front.imagestacklayer/Content.imageset" },
-					new string []{ "Icon_Large_Middle-1280x768.png", "AppIcon.brandassets/App Icon - Large.imagestack/Middle.imagestacklayer/Content.imageset" },
-					new string []{ "Icon_Small_Back-400x240.png", "AppIcon.brandassets/App Icon - Small.imagestack/Back.imagestacklayer/Content.imageset" },
-					new string []{ "Icon_Small_Back-800x480.png", "AppIcon.brandassets/App Icon - Small.imagestack/Back.imagestacklayer/Content.imageset" },
-					new string []{ "Icon_Small_Front-400x240.png", "AppIcon.brandassets/App Icon - Small.imagestack/Front.imagestacklayer/Content.imageset" },
-					new string []{ "Icon_Small_Front-800x480.png", "AppIcon.brandassets/App Icon - Small.imagestack/Front.imagestacklayer/Content.imageset" },
-					new string []{ "Icon_Small_Middle-400x240.png", "AppIcon.brandassets/App Icon - Small.imagestack/Middle.imagestacklayer/Content.imageset" },
-					new string []{ "Icon_Small_Middle-800x480.png", "AppIcon.brandassets/App Icon - Small.imagestack/Middle.imagestacklayer/Content.imageset" },
-					//new string []{ "TopShelfMarketWide.png", "AppIcon.brandassets/Top Shelf Image Wide.imageset" },
-					//new string []{ "[REDACTED_EMAIL]", "AppIcon.brandassets/Top Shelf Image Wide.imageset" },
-					new string []{ "[REDACTED_EMAIL]", "AppIcon.brandassets/Top Shelf Image Wide.imageset" },
-					new string []{ "TopShelfWide-1920x720.png", "AppIcon.brandassets/Top Shelf Image Wide.imageset" },
-					//new string []{ "TopShelfMarket.png", "AppIcon.brandassets/Top Shelf Image.imageset" },
-					//new string []{ "[REDACTED_EMAIL]", "AppIcon.brandassets/Top Shelf Image.imageset" },
-					new string []{ "TopShelf.png", "AppIcon.brandassets/Top Shelf Image.imageset" },
-					new string []{ "[REDACTED_EMAIL]", "AppIcon.brandassets/Top Shelf Image.imageset" },
-					new string []{ "Launch.png", "LaunchImage.launchimage" },
-					new string []{ "[REDACTED_EMAIL]", "LaunchImage.launchimage" },
+					new string []{ "Icon_Large_Back-1280x768.png", "App Icon & Top Shelf Image.brandassets/App Icon - Large.imagestack/Back.imagestacklayer/Content.imageset" },
+					new string []{ "Icon_Large_Front-1280x768.png", "App Icon & Top Shelf Image.brandassets/App Icon - Large.imagestack/Front.imagestacklayer/Content.imageset" },
+					new string []{ "Icon_Large_Middle-1280x768.png", "App Icon & Top Shelf Image.brandassets/App Icon - Large.imagestack/Middle.imagestacklayer/Content.imageset" },
+					new string []{ "Icon_Small_Back-400x240.png", "App Icon & Top Shelf Image.brandassets/App Icon.imagestack/Back.imagestacklayer/Content.imageset" },
+					new string []{ "Icon_Small_Back-800x480.png", "App Icon & Top Shelf Image.brandassets/App Icon.imagestack/Back.imagestacklayer/Content.imageset" },
+					new string []{ "Icon_Small_Front-400x240.png", "App Icon & Top Shelf Image.brandassets/App Icon.imagestack/Front.imagestacklayer/Content.imageset" },
+					new string []{ "Icon_Small_Front-800x480.png", "App Icon & Top Shelf Image.brandassets/App Icon.imagestack/Front.imagestacklayer/Content.imageset" },
+					new string []{ "Icon_Small_Middle-400x240.png", "App Icon & Top Shelf Image.brandassets/App Icon.imagestack/Middle.imagestacklayer/Content.imageset" },
+					new string []{ "Icon_Small_Middle-800x480.png", "App Icon & Top Shelf Image.brandassets/App Icon.imagestack/Middle.imagestacklayer/Content.imageset" },
+					new string []{ "[REDACTED_EMAIL]", "App Icon & Top Shelf Image.brandassets/Top Shelf Image Wide.imageset" },
+					new string []{ "TopShelfWide-1920x720.png", "App Icon & Top Shelf Image.brandassets/Top Shelf Image Wide.imageset" },
+					new string []{ "TopShelf.png", "App Icon & Top Shelf Image.brandassets/Top Shelf Image.imageset" },
+					new string []{ "[REDACTED_EMAIL]", "App Icon & Top Shelf Image.brandassets/Top Shelf Image.imageset" },
+					new string []{ "Launch.png", "Launch Image.launchimage" },
+					new string []{ "[REDACTED_EMAIL]", "Launch Image.launchimage" },
 				};
 				Dir = Path.Combine(IntermediateDir, "Resources", "Assets.xcassets");
 
