UE-71596 iOS packaging error with CloudKit enabled - Provisioning profile <...> doesn't match the entitlements file's values

#iOS
#tvOS
#4.22

Moved the entitlements file generation from preinit to postbuild, so we can write the correct iCloud container settings

diff --git a/Engine/Source/Programs/UnrealBuildTool/Platform/IOS/IOSToolChain.cs b/Engine/Source/Programs/UnrealBuildTool/Platform/IOS/IOSToolChain.cs
--- a/Engine/Source/Programs/UnrealBuildTool/Platform/IOS/IOSToolChain.cs
+++ b/Engine/Source/Programs/UnrealBuildTool/Platform/IOS/IOSToolChain.cs
@@ -1567,6 +1567,82 @@ namespace UnrealBuildTool
 			return FileReference.Combine(Executable.Directory, "Payload", TargetName + ".app", TargetName);
 		}
 
+		private static void WriteEntitlementsFile(string OutputFilename, FileReference ProjectFile, bool bForDistribution, string iCloudContainerIdentifiersXML, string UbiquityContainerIdentifiersXML)
+		{
+			// get the settings from the ini file
+			ConfigHierarchy Ini = ConfigCache.ReadHierarchy(ConfigHierarchyType.Engine, DirectoryReference.FromFile(ProjectFile), UnrealTargetPlatform.IOS);
+			bool bCloudKitSupported = false;
+			Ini.GetBool("/Script/IOSRuntimeSettings.IOSRuntimeSettings", "bEnableCloudKitSupport", out bCloudKitSupported);
+			Directory.CreateDirectory(Path.GetDirectoryName(OutputFilename));
+			// we need to have something so Xcode will compile, so we just set the get-task-allow, since we know the value,
+			// which is based on distribution or not (true means debuggable)
+			StringBuilder Text = new StringBuilder();
+			Text.AppendLine("<?xml version=\"1.0\" encoding=\"UTF-8\"?>");
+			Text.AppendLine("<!DOCTYPE plist PUBLIC \"-//Apple//DTD PLIST 1.0//EN\" \"http://www.apple.com/DTDs/PropertyList-1.0.dtd\">");
+			Text.AppendLine("<plist version=\"1.0\">");
+			Text.AppendLine("<dict>");
+			Text.AppendLine("\t<key>get-task-allow</key>");
+			Text.AppendLine(string.Format("\t<{0}/>", bForDistribution ? "false" : "true"));
+			if (bCloudKitSupported)
+			{
+				Text.AppendLine("\t<key>com.apple.developer.icloud-container-identifiers</key>");
+				if (iCloudContainerIdentifiersXML == "")
+				{
+					Text.AppendLine("\t<array>");
+					Text.AppendLine("\t\t<string>iCloud.$(CFBundleIdentifier)</string>");
+					Text.AppendLine("\t</array>");
+				}
+				else
+				{
+					Text.AppendLine(iCloudContainerIdentifiersXML);
+				}
+				Text.AppendLine("\t<key>com.apple.developer.icloud-services</key>");
+				Text.AppendLine("\t<array>");
+				Text.AppendLine("\t\t<string>CloudKit</string>");
+				Text.AppendLine("\t\t<string>CloudDocuments</string>");
+				Text.AppendLine("\t</array>");
+				Text.AppendLine("\t<key>com.apple.developer.ubiquity-container-identifiers</key>");
+				if (UbiquityContainerIdentifiersXML == "")
+				{
+					Text.AppendLine("\t<array>");
+					Text.AppendLine("\t\t<string>iCloud.$(CFBundleIdentifier)</string>");
+					Text.AppendLine("\t</array>");
+				}
+				else
+				{
+					Text.AppendLine(UbiquityContainerIdentifiersXML);
+				}
+				Text.AppendLine("\t<key>com.apple.developer.ubiquity-kvstore-identifier</key>");
+				Text.AppendLine("\t<string>$(TeamIdentifierPrefix)$(CFBundleIdentifier)</string>");
+			}
+
+			bool bRemoteNotificationsSupported = false;
+			Ini.GetBool("/Script/IOSRuntimeSettings.IOSRuntimeSettings", "bEnableRemoteNotificationsSupport", out bRemoteNotificationsSupported);
+			if (bRemoteNotificationsSupported)
+			{
+				Text.AppendLine("\t<key>aps-environment</key>");
+				Text.AppendLine(string.Format("\t<string>{0}</string>", bForDistribution ? "production" : "development"));
+			}
+			Text.AppendLine("</dict>");
+			Text.AppendLine("</plist>");
+
+			if (File.Exists(OutputFilename))
+			{
+				// read existing file
+				string ExisitingFileContents = File.ReadAllText(OutputFilename);
+				bool bFileChanged = !ExisitingFileContents.Equals(Text.ToString(), StringComparison.Ordinal);
+				// overwrite file if there are content changes
+				if (bFileChanged)
+				{
+					File.WriteAllText(OutputFilename, Text.ToString());
+				}
+			}
+			else
+			{
+				File.WriteAllText(OutputFilename, Text.ToString());
+			}
+		}
+
         public static void PostBuildSync(IOSPostBuildSyncTarget Target)
 		{
 			IOSProjectSettings ProjectSettings = ((IOSPlatform)UEBuildPlatform.GetBuildPlatform(Target.Platform)).ReadProjectSettings(Target.ProjectFile);
@@ -1751,6 +1827,22 @@ namespace UnrealBuildTool
 						CmdLine += (!string.IsNullOrEmpty(MobileProvisionUUID) ? (" PROVISIONING_PROFILE_SPECIFIER=" + MobileProvisionUUID) : "");
 					}
 					Writer.WriteLine("/usr/bin/xcrun {0}", CmdLine);
+					
+					// get some info from the mobileprovisioning file
+					// the iCloud identifier and the bundle id may differ
+					string iCloudContainerIdentifiersXML = "";
+					string UbiquityContainerIdentifiersXML = "";
+					if (File.Exists(MobileProvisionFile.FullName))
+					{
+						MobileProvisionContents MobileProvisionContent = MobileProvisionContents.Read(MobileProvisionFile);
+
+						iCloudContainerIdentifiersXML = MobileProvisionContent.GetNodeValueByName("com.apple.developer.icloud-container-identifiers");
+						UbiquityContainerIdentifiersXML = MobileProvisionContent.GetNodeValueByName("com.apple.developer.ubiquity-container-identifiers");
+					}
+					// create the entitlements file
+					string IntermediateDir = (((Target.ProjectFile != null) ? Target.ProjectFile.Directory.ToString() : 
+						UnrealBuildTool.EngineDirectory.ToString())) + "/Intermediate/" + (Target.Platform == UnrealTargetPlatform.IOS ? "IOS" : "TVOS");
+					WriteEntitlementsFile(Path.Combine(IntermediateDir, AppName + ".entitlements"), Target.ProjectFile, Target.bForDistribution, iCloudContainerIdentifiersXML, UbiquityContainerIdentifiersXML);
 				}
 
 				Log.TraceInformation("Executing {0}", SignProjectScript);
diff --git a/Engine/Source/Programs/UnrealBuildTool/Platform/IOS/MobileProvisionContents.cs b/Engine/Source/Programs/UnrealBuildTool/Platform/IOS/MobileProvisionContents.cs
--- a/Engine/Source/Programs/UnrealBuildTool/Platform/IOS/MobileProvisionContents.cs
+++ b/Engine/Source/Programs/UnrealBuildTool/Platform/IOS/MobileProvisionContents.cs
@@ -144,5 +144,22 @@ namespace UnrealBuildTool
 				throw new BuildException("No PKCS7-Data section found in {0}", Location);
 			}
 		}
+
+		// return the outerXML of the node's value
+		public string GetNodeValueByName(string InValue)
+		{
+			XmlNodeList elemList = this.Document.GetElementsByTagName("key");
+			for (int i = 0; i < elemList.Count; i++)
+			{
+				if (elemList[i].InnerXml.Equals(InValue))
+				{
+					if (elemList[i].NextSibling != null)
+					{
+						return elemList[i].NextSibling.OuterXml;
+					}
+				}
+			}
+			return "";
+		}
 	}
 }
diff --git a/Engine/Source/Programs/UnrealBuildTool/Platform/IOS/UEDeployIOS.cs b/Engine/Source/Programs/UnrealBuildTool/Platform/IOS/UEDeployIOS.cs
--- a/Engine/Source/Programs/UnrealBuildTool/Platform/IOS/UEDeployIOS.cs
+++ b/Engine/Source/Programs/UnrealBuildTool/Platform/IOS/UEDeployIOS.cs
@@ -955,7 +955,6 @@ namespace UnrealBuildTool
 			Directory.CreateDirectory(BuildDirectory);
 
 			// create the entitlements file
-			WriteEntitlementsFile(Path.Combine(IntermediateDirectory, GameName + ".entitlements"), ProjectFile, bForDistribution);
 
 			// delete some old files if they exist
 			if (Directory.Exists(AppDirectory + "/_CodeSignature"))
@@ -1199,76 +1198,6 @@ namespace UnrealBuildTool
 			return PluginExtras;
 		}
 
-		private void WriteEntitlementsFile(string OutputFilename, FileReference ProjectFile, bool bForDistribution)
-		{
-			// get the settings from the ini file
-			// plist replacements
-			// @todo tvos: Separate TVOS version?
-			ConfigHierarchy Ini = ConfigCache.ReadHierarchy(ConfigHierarchyType.Engine, DirectoryReference.FromFile(ProjectFile), UnrealTargetPlatform.IOS);
-			bool bCloudKitSupported = false;
-			Ini.GetBool("/Script/IOSRuntimeSettings.IOSRuntimeSettings", "bEnableCloudKitSupport", out bCloudKitSupported);
-
-			Directory.CreateDirectory(Path.GetDirectoryName(OutputFilename));
-			// we need to have something so Xcode will compile, so we just set the get-task-allow, since we know the value,
-			// which is based on distribution or not (true means debuggable)
-			StringBuilder Text = new StringBuilder();
-			Text.AppendLine("<?xml version=\"1.0\" encoding=\"UTF-8\"?>");
-			Text.AppendLine("<!DOCTYPE plist PUBLIC \"-//Apple//DTD PLIST 1.0//EN\" \"http://www.apple.com/DTDs/PropertyList-1.0.dtd\">");
-			Text.AppendLine("<plist version=\"1.0\">");
-			Text.AppendLine("<dict>");
-			Text.AppendLine("\t<key>get-task-allow</key>");
-			Text.AppendLine(string.Format("\t<{0}/>", bForDistribution ? "false" : "true"));
-
-			if (bCloudKitSupported)
-			{
-				Text.AppendLine("\t<key>com.apple.developer.icloud-container-identifiers</key>");
-				Text.AppendLine("\t<array>");
-				Text.AppendLine("\t\t<string>iCloud.$(CFBundleIdentifier)</string>");
-				Text.AppendLine("\t</array>");
-				Text.AppendLine("\t<key>com.apple.developer.icloud-services</key>");
-				Text.AppendLine("\t<array>");
-				Text.AppendLine("\t\t<string>CloudKit</string>");
-				Text.AppendLine("\t\t<string>CloudDocuments</string>");
-				Text.AppendLine("\t</array>");
-				Text.AppendLine("\t<key>com.apple.developer.ubiquity-container-identifiers</key>");
-				Text.AppendLine("\t<array>");
-				Text.AppendLine("\t\t<string>iCloud.$(CFBundleIdentifier)</string>");
-				Text.AppendLine("\t</array>");
-				Text.AppendLine("\t<key>com.apple.developer.ubiquity-kvstore-identifier</key>");
-				Text.AppendLine("\t<string>$(TeamIdentifierPrefix)$(CFBundleIdentifier)</string>");
-			}
-
-			bool bRemoteNotificationsSupported = false;
-			Ini.GetBool("/Script/IOSRuntimeSettings.IOSRuntimeSettings", "bEnableRemoteNotificationsSupport", out bRemoteNotificationsSupported);
-
-			if (bRemoteNotificationsSupported)
-			{
-				Text.AppendLine("\t<key>aps-environment</key>");
-				Text.AppendLine(string.Format("\t<string>{0}</string>", bForDistribution ? "production" : "development"));
-			}
-
-			Text.AppendLine("</dict>");
-			Text.AppendLine("</plist>");
-
-			if (File.Exists(OutputFilename))
-			{
-				// read existing file
-				string ExisitingFileContents = File.ReadAllText(OutputFilename);
-
-				bool bFileChanged = !ExisitingFileContents.Equals(Text.ToString(), StringComparison.Ordinal);
-
-				// overwrite file if there are content changes
-				if (bFileChanged)
-				{
-					File.WriteAllText(OutputFilename, Text.ToString());
-				}
-			}
-			else
-			{
-				File.WriteAllText(OutputFilename, Text.ToString());
-			}
-		}
-
 		public static void SafeFileCopy(FileInfo SourceFile, string DestinationPath, bool bOverwrite)
 		{
 			FileInfo DI = new FileInfo(DestinationPath);
