Allow dynamic library as target for finalizing Android project

#ue4
#android

diff --git a/Engine/Source/Programs/UnrealBuildTool/Platform/Android/UEDeployAndroid.cs b/Engine/Source/Programs/UnrealBuildTool/Platform/Android/UEDeployAndroid.cs
--- a/Engine/Source/Programs/UnrealBuildTool/Platform/Android/UEDeployAndroid.cs
+++ b/Engine/Source/Programs/UnrealBuildTool/Platform/Android/UEDeployAndroid.cs
@@ -4071,14 +4071,19 @@ bSaveSymbols = true;
 			// get the receipt
 			SetAndroidPluginData(ToolChain.GetAllArchitectures(), CollectPluginDataPaths(Receipt));
 
+			bool bShouldCompileAsDll = Receipt.HasValueForAdditionalProperty("CompileAsDll", "true");
+
 			// Get the output paths
-			List<FileReference> OutputPaths = Receipt.BuildProducts.Where(x => x.Type == BuildProductType.Executable).Select(x => x.Path).ToList();
+			BuildProductType ProductType = bShouldCompileAsDll ? BuildProductType.DynamicLibrary : BuildProductType.Executable;
+			List<FileReference> OutputPaths = Receipt.BuildProducts.Where(x => x.Type == ProductType).Select(x => x.Path).ToList();
+			if (OutputPaths.Count < 1)
+			{
+				throw new BuildException("Target file does not contain either executable or dynamic library .so");
+			}
 
 			// we need to strip architecture from any of the output paths
 			string BaseSoName = ToolChain.RemoveArchName(OutputPaths[0].FullName);
 
-			bool bShouldCompileAsDll = Receipt.HasValueForAdditionalProperty("CompileAsDll", "true");
-
 			// make an apk at the end of compiling, so that we can run without packaging (debugger, cook on the fly, etc)
 			string RelativeEnginePath = UnrealBuildTool.EngineDirectory.MakeRelativeTo(DirectoryReference.GetCurrentDirectory());
 
