UE-72225

Merging in fixes to better support SNDBS in UBT with use of cl-filter.

#rnx

diff --git a/Engine/Build/SNDBSTemplates/cl-filter.exe.sn-dbs-tool.ini b/Engine/Build/SNDBSTemplates/cl-filter.exe.sn-dbs-tool.ini
new file mode 100644
--- /dev/null
+++ b/Engine/Build/SNDBSTemplates/cl-filter.exe.sn-dbs-tool.ini
@@ -0,0 +1,51 @@
+[tool]
+family=msvc
+vc_major_version=14
+use_surrogate=true
+force_synchronous_pdb_writes=true
+error_report_mode=prompt
+
+[group]
+server={VC_COMPILER_DIR}\mspdbsrv.exe
+
+[files]
+main=cl-filter.exe
+file01={VC_COMPILER_DIR}\c1.dll
+file01={VC_COMPILER_DIR}\c1ui.dll
+file02={VC_COMPILER_DIR}\c1xx.dll
+file03={VC_COMPILER_DIR}\c2.dll
+file04={VC_COMPILER_DIR}\mspdb140.dll
+file05={VC_COMPILER_DIR}\mspdbcore.dll
+file06={VC_COMPILER_DIR}\mspdbsrv.exe
+file07={VC_COMPILER_DIR}\mspft140.dll
+file08={VC_COMPILER_DIR}\vcmeta.dll
+file09={VC_COMPILER_DIR}\*\clui.dll
+file10={VC_COMPILER_DIR}\*\mspft140ui.dll
+file11={VC_COMPILER_DIR}\localespc.dll
+file12={VC_COMPILER_DIR}\cppcorecheck.dll
+file13={VC_COMPILER_DIR}\experimentalcppcorecheck.dll
+file14={VC_COMPILER_DIR}\espxengine.dll
+file15={VC_COMPILER_DIR}\c1.exe
+
+[output-file-patterns]
+outputfile01=\s*"([^ ",]+\.cpp\.txt)\"
+
+[output-file-rules]
+rule01=*\sqmcpp*.log|discard=true
+rule02=*\vctoolstelemetry*.dat|discard=true
+rule03=*\Microsoft\Windows\Temporary Internet Files\*|discard=true
+rule04=*\Microsoft\Windows\INetCache\*|discard=true
+
+[input-file-rules]
+rule01=*\sqmcpp*.log|ignore_transient_errors=true;ignore_unexpected_input=true
+rule02=*\vctoolstelemetry*.dat|ignore_transient_errors=true;ignore_unexpected_input=true
+rule03=*\Microsoft\Windows\Temporary Internet Files\*|ignore_transient_errors=true;ignore_unexpected_input=true
+rule04=*\Microsoft\Windows\INetCache\*|ignore_transient_errors=true;ignore_unexpected_input=true
+
+[system-file-filters]
+filter01=msvcr*.dll
+filter02=msvcp*.dll
+filter03=vcruntime140*.dll
+filter04=appcrt140*.dll
+filter05=desktopcrt140*.dll
+filter06=concrt140*.dll
diff --git a/Engine/Source/Programs/UnrealBuildTool/Configuration/UEBuildPlatform.cs b/Engine/Source/Programs/UnrealBuildTool/Configuration/UEBuildPlatform.cs
--- a/Engine/Source/Programs/UnrealBuildTool/Configuration/UEBuildPlatform.cs
+++ b/Engine/Source/Programs/UnrealBuildTool/Configuration/UEBuildPlatform.cs
@@ -542,6 +542,14 @@ namespace UnrealBuildTool
 			}
 		}
 
+		/// <summary>
+		/// Gets the platform name that should be used.
+		/// </summary>
+		public virtual string GetPlatformName()
+		{
+			return Platform.ToString();
+		}
+
 		/// <summary>
 		/// If this platform can be compiled with XGE
 		/// </summary>
diff --git a/Engine/Source/Programs/UnrealBuildTool/Executors/SNDBS.cs b/Engine/Source/Programs/UnrealBuildTool/Executors/SNDBS.cs
--- a/Engine/Source/Programs/UnrealBuildTool/Executors/SNDBS.cs
+++ b/Engine/Source/Programs/UnrealBuildTool/Executors/SNDBS.cs
@@ -1,21 +1,21 @@
 // Copyright 1998-2019 Epic Games, Inc. All Rights Reserved.
 
 using System;
+using System.Collections;
 using System.Collections.Generic;
-using System.IO;
 using System.Diagnostics;
-using System.Xml;
-using System.Text.RegularExpressions;
+using System.IO;
 using System.Linq;
-using System.Reflection;
-using System.Threading;
 using System.ServiceProcess;
+using System.Text;
 using Tools.DotNETCommon;
 
 namespace UnrealBuildTool
 {
 	class SNDBS : ActionExecutor
 	{
+		private FileReference IncludeRewriteRulesFile = FileReference.Combine(UnrealBuildTool.EngineDirectory, "Build", "SNDBS", "include-rewrite-rules.ini");
+
 		/// <summary>
 		/// Processor count multiplier for local execution. Can be below 1 to reserve CPU for other tasks.
 		/// </summary>
@@ -181,6 +181,7 @@ namespace UnrealBuildTool
 			int NumScriptedActions = 0;
 			List<Action> LocalActions = new List<Action>();
 			ActionThread DummyActionThread = new ActionThread(null, 1, 1);
+			bool PrintDebugInfo = false;
 			foreach (Action Action in InActions)
 			{
 				ActionThread ActionProcess = null;
@@ -242,6 +243,7 @@ namespace UnrealBuildTool
 							Log.TraceInformation("[{0}/{1}] {2} {3}", JobNumber, InActions.Count, Action.CommandDescription, Action.StatusDescription);
 							JobNumber++;
 							NumScriptedActions++;
+							PrintDebugInfo |= Action.bPrintDebugInfo;
 						}
 					}
 				}
@@ -257,12 +259,15 @@ namespace UnrealBuildTool
                 // Create the process
                 string SCERoot = Environment.GetEnvironmentVariable("SCE_ROOT_DIR");
                 string SNDBSExecutable = Path.Combine(SCERoot, "Common/SN-DBS/bin/dbsbuild.exe");
-                ProcessStartInfo PSI = new ProcessStartInfo(SNDBSExecutable, String.Format("-q -p UE4 -s \"{0}\"", FileReference.Combine(UnrealBuildTool.EngineDirectory, "Intermediate", "Build", "sndbs.bat").FullName));
+				DirectoryReference TemplatesDir = DirectoryReference.Combine(UnrealBuildTool.EngineDirectory, "Programs", "UnrealBuildTool", "SndbsTemplates");
+				string IncludeRewriteRulesArg = String.Format("--include-rewrite-rules \"{0}\"", IncludeRewriteRulesFile.FullName);
+				string VerbosityLevel = PrintDebugInfo ? "-v" : "-q";
+                ProcessStartInfo PSI = new ProcessStartInfo(SNDBSExecutable, String.Format("{0} -p UE4 -s \"{1}\" -templates \"{2}\" {3}", VerbosityLevel, FileReference.Combine(UnrealBuildTool.EngineDirectory, "Intermediate", "Build", "sndbs.bat").FullName, TemplatesDir.FullName, IncludeRewriteRulesArg));
 				PSI.RedirectStandardOutput = true;
 				PSI.RedirectStandardError = true;
 				PSI.UseShellExecute = false;
 				PSI.CreateNoWindow = true;
-				PSI.WorkingDirectory = Path.GetFullPath("."); ;
+				PSI.WorkingDirectory = Path.GetFullPath(".");
 				Process NewProcess = new Process();
 				NewProcess.StartInfo = PSI;
 				NewProcess.OutputDataReceived += new DataReceivedEventHandler(ActionDebugOutput);
@@ -330,6 +335,16 @@ namespace UnrealBuildTool
 			bool SNDBSResult = true;
 			if (Actions.Count > 0)
 			{
+				// Generate any needed templates. Can only generate one per executable, so just use the first Action for reference
+				IEnumerable<Action> CommandPaths = Actions.GroupBy(a => a.CommandPath).Select(g => g.FirstOrDefault()).ToList();
+				foreach (Action CommandPath in CommandPaths)
+				{
+					PrepareToolTemplate(CommandPath);
+				}
+
+				// Generate include-rewrite-rules.ini.
+				GenerateSNDBSIncludeRewriteRules();
+
 				// Use WMI to figure out physical cores, excluding hyper threading.
 				int NumCores = Utils.GetPhysicalProcessorCount();
 				
@@ -423,7 +438,70 @@ namespace UnrealBuildTool
 				Log.WriteLineIf(bLogDetailedActionStats,
 					LogEventType.Console, "Cumulative thread seconds ({0} processors): {1:0.00}", System.Environment.ProcessorCount, TotalThreadSeconds);
 			}
+
+			// Delete the include-rewrite-rules.ini if it was generated.
+			if (File.Exists(IncludeRewriteRulesFile.FullName))
+			{
+				File.Delete(IncludeRewriteRulesFile.FullName);
+			}
+
 			return SNDBSResult;
 		}
+
+		private void PrepareToolTemplate(Action Action)
+		{
+			string TemplateFileName = String.Format("{0}.sn-dbs-tool.ini", Action.CommandPath.GetFileName());
+			FileReference TemplateInput = FileReference.Combine(UnrealBuildTool.EngineDirectory, "Build", "SNDBSTemplates", TemplateFileName);
+			
+			// If no base template exists, don't try to generate one.
+			if (!File.Exists(TemplateInput.FullName))
+			{
+				return;
+			}
+
+			FileReference TemplateOutput = FileReference.Combine(UnrealBuildTool.EngineDirectory, "Programs", "UnrealBuildTool", "SNDBSTemplates", TemplateFileName);
+			if (!Directory.Exists(TemplateOutput.Directory.FullName))
+			{
+				Directory.CreateDirectory(TemplateOutput.Directory.FullName);
+			}
+
+			string TemplateText = File.ReadAllText(TemplateInput.FullName);
+			TemplateText = TemplateText.Replace("{COMMAND_PATH}", Action.CommandPath.Directory.FullName);
+			foreach (DictionaryEntry Variable in Environment.GetEnvironmentVariables(EnvironmentVariableTarget.Process))
+			{
+				string VariableName = String.Format("{{{0}}}", Variable.Key);
+				TemplateText = TemplateText.Replace(VariableName, Variable.Value.ToString());
+			}
+			File.WriteAllText(TemplateOutput.FullName, TemplateText);
+		}
+
+		private void GenerateSNDBSIncludeRewriteRules()
+		{
+			if (!Directory.Exists(IncludeRewriteRulesFile.Directory.FullName))
+			{
+				Directory.CreateDirectory(IncludeRewriteRulesFile.Directory.FullName);
+			}
+
+			List<string> IncludeRewriteRulesText = new List<string>();
+			IncludeRewriteRulesText.Add("[computed-include-rules]");
+			IncludeRewriteRulesText.Add(@"pattern1=^COMPILED_PLATFORM_HEADER\(\s*([^ ,]+)\)");
+
+			// Get all registered platforms. Most will just use the name as is, but some may have an override so
+			// add all distinct entries to the list.
+			IEnumerable<UnrealTargetPlatform> Platforms = UEBuildPlatform.GetRegisteredPlatforms();
+			List<string> PlatformNames = new List<string>();
+			foreach (UnrealTargetPlatform Platform in Platforms)
+			{
+				UEBuildPlatform PlatformData = UEBuildPlatform.GetBuildPlatform(Platform);
+				if (!PlatformNames.Contains(PlatformData.GetPlatformName()))
+				{
+					PlatformNames.Add(PlatformData.GetPlatformName());
+				}
+			}
+
+			IEnumerable<string> PlatformExpansions = PlatformNames.Select(p => String.Format("{0}/{0}$1", p));
+			IncludeRewriteRulesText.Add(String.Format("expansions1={0}", String.Join("|", PlatformExpansions)));
+			File.WriteAllText(IncludeRewriteRulesFile.FullName, String.Join(Environment.NewLine, IncludeRewriteRulesText));
+		}
 	}
 }
diff --git a/Engine/Source/Programs/UnrealBuildTool/Platform/Windows/UEBuildWindows.cs b/Engine/Source/Programs/UnrealBuildTool/Platform/Windows/UEBuildWindows.cs
--- a/Engine/Source/Programs/UnrealBuildTool/Platform/Windows/UEBuildWindows.cs
+++ b/Engine/Source/Programs/UnrealBuildTool/Platform/Windows/UEBuildWindows.cs
@@ -1487,6 +1487,14 @@ namespace UnrealBuildTool
 				   TryReadInstallDirRegistryKey32("Microsoft\\Microsoft SDKs\\NETFXSDK\\4.6.2", "KitsInstallationFolder", out OutInstallDir);
 		}
 
+		/// <summary>
+		/// Gets the platform name that should be used.
+		/// </summary>
+		public override string GetPlatformName()
+		{
+			return "Windows";
+		}
+
 		/// <summary>
 		/// If this platform can be compiled with SN-DBS
 		/// </summary>
@@ -1680,6 +1688,9 @@ namespace UnrealBuildTool
 			CompileEnvironment.Definitions.Add(String.Format("_WIN32_WINNT=0x{0:X4}", Target.WindowsPlatform.TargetWindowsVersion));
 			CompileEnvironment.Definitions.Add(String.Format("WINVER=0x{0:X4}", Target.WindowsPlatform.TargetWindowsVersion));
 			CompileEnvironment.Definitions.Add("PLATFORM_WINDOWS=1");
+			
+			// both Win32 and Win64 use Windows headers, so we enforce that here
+			CompileEnvironment.Definitions.Add(String.Format("OVERRIDE_PLATFORM_HEADER_NAME={0}", GetPlatformName()));
 
 			FileReference MorpheusShaderPath = FileReference.Combine(UnrealBuildTool.EngineDirectory, "Shaders", "Private", "Platform", "PS4", "PostProcessHMDMorpheus.usf");
 			if (FileReference.Exists(MorpheusShaderPath))
diff --git a/Engine/Source/Programs/UnrealBuildTool/Platform/Windows/VCEnvironment.cs b/Engine/Source/Programs/UnrealBuildTool/Platform/Windows/VCEnvironment.cs
--- a/Engine/Source/Programs/UnrealBuildTool/Platform/Windows/VCEnvironment.cs
+++ b/Engine/Source/Programs/UnrealBuildTool/Platform/Windows/VCEnvironment.cs
@@ -117,6 +117,10 @@ namespace UnrealBuildTool
             // Compile using 64 bit tools for 64 bit targets, and 32 for 32.
 			CompilerPath = GetCompilerToolPath(Platform, Compiler, CompilerDir);
 
+			// Add the compiler path and directory as environment variables for the process so they may be used elsewhere.
+			Environment.SetEnvironmentVariable("VC_COMPILER_PATH", CompilerPath.FullName, EnvironmentVariableTarget.Process);
+			Environment.SetEnvironmentVariable("VC_COMPILER_DIR", CompilerPath.Directory.FullName, EnvironmentVariableTarget.Process);
+
 			// Regardless of the target, if we're linking on a 64 bit machine, we want to use the 64 bit linker (it's faster than the 32 bit linker and can handle large linking jobs)
 			DirectoryReference DefaultLinkerDir = VCToolPath64;
 			LinkerPath = GetLinkerToolPath(Platform, Compiler, DefaultLinkerDir);
