- Support for LTO on iOS, using the bAllowLTCG flag (a la VCToolchain)

diff --git a/Engine/Source/Programs/UnrealBuildTool/Platform/IOS/IOSToolChain.cs b/Engine/Source/Programs/UnrealBuildTool/Platform/IOS/IOSToolChain.cs
--- a/Engine/Source/Programs/UnrealBuildTool/Platform/IOS/IOSToolChain.cs
+++ b/Engine/Source/Programs/UnrealBuildTool/Platform/IOS/IOSToolChain.cs
@@ -236,10 +236,11 @@ namespace UnrealBuildTool
 			Result += GetRTTIFlag(CompileEnvironment);
 			Result += " -fvisibility=hidden"; // hides the linker warnings with PhysX
 
-			// 			if (CompileEnvironment.TargetConfiguration == CPPTargetConfiguration.Shipping)
-			// 			{
-			// 				Result += " -flto";
-			// 			}
+			// use LTO if desired (like VCToolchain does)
+			if (CompileEnvironment.bAllowLTCG)
+ 			{
+ 				Result += " -flto";
+ 			}
 
 			Result += " -Wall -Werror";
 			Result += " -Wdelete-non-virtual-dtor";
@@ -512,6 +513,12 @@ namespace UnrealBuildTool
 			Result += " -stdlib=libc++";
 			Result += " -ObjC";
 			//			Result += " -v";
+
+			// use LTO if desired (like VCToolchain does)
+			if (LinkEnvironment.bAllowLTCG)
+			{
+				Result += " -flto";
+			}
 			
 			string SanitizerMode = Environment.GetEnvironmentVariable("ENABLE_ADDRESS_SANITIZER");
 			if(SanitizerMode != null && SanitizerMode == "YES")
