Removing Steam library path being hardcoded into Linux's linker rpath as this is an old hack that is no longer needed.

diff --git a/Engine/Source/Programs/UnrealBuildTool/Platform/Linux/LinuxToolChain.cs b/Engine/Source/Programs/UnrealBuildTool/Platform/Linux/LinuxToolChain.cs
--- a/Engine/Source/Programs/UnrealBuildTool/Platform/Linux/LinuxToolChain.cs
+++ b/Engine/Source/Programs/UnrealBuildTool/Platform/Linux/LinuxToolChain.cs
@@ -916,8 +916,6 @@ namespace UnrealBuildTool
 			Result += " -Wl,-rpath-link=${ORIGIN}";
 			Result += " -Wl,-rpath=${ORIGIN}/../../../Engine/Binaries/Linux";
 			Result += " -Wl,-rpath=${ORIGIN}/..";	// for modules that are in sub-folders of the main Engine/Binary/Linux folder
-			// FIXME: really ugly temp solution. Modules need to be able to specify this
-			Result += " -Wl,-rpath=${ORIGIN}/../../../Engine/Binaries/ThirdParty/Steamworks/Steamv139/x86_64-unknown-linux-gnu";
 			if (LinkEnvironment.Architecture.StartsWith("x86_64"))
 			{
 				Result += " -Wl,-rpath=${ORIGIN}/../../../Engine/Binaries/ThirdParty/Qualcomm/Linux";
