Additional fix for UE-75714, this time for Windows-based workflow

diff --git a/Engine/Source/Programs/UnrealBuildTool/Platform/IOS/IOSToolChain.cs b/Engine/Source/Programs/UnrealBuildTool/Platform/IOS/IOSToolChain.cs
--- a/Engine/Source/Programs/UnrealBuildTool/Platform/IOS/IOSToolChain.cs
+++ b/Engine/Source/Programs/UnrealBuildTool/Platform/IOS/IOSToolChain.cs
@@ -1848,11 +1848,11 @@ namespace UnrealBuildTool
 
 				if (AppName == "UE4Game" || AppName == "UE4Client" || Target.ProjectFile == null || Target.ProjectFile.IsUnderDirectory(UnrealBuildTool.EngineDirectory))
 				{
-					GenerateProjectFiles(Target.ProjectFile, new string[] { "-platforms=" + (Target.Platform == UnrealTargetPlatform.IOS ? "IOS" : "TVOS"), "-NoIntellIsense", (Target.Platform == UnrealTargetPlatform.IOS ? "-iosdeployonly" : "-tvosdeployonly"), "-ignorejunk", (Target.bForDistribution ? "-distribution" : "-development"), "-bundleID=" + BundleID });
+					GenerateProjectFiles(Target.ProjectFile, new string[] { "-platforms=" + (Target.Platform == UnrealTargetPlatform.IOS ? "IOS" : "TVOS"), "-NoIntellIsense", (Target.Platform == UnrealTargetPlatform.IOS ? "-iosdeployonly" : "-tvosdeployonly"), "-ignorejunk", (Target.bForDistribution ? "-distribution" : "-development"), "-bundleID=" + BundleID, "-includetemptargets" });
 				}
 				else
 				{
-					GenerateProjectFiles(Target.ProjectFile, new string[] { "-platforms=" + (Target.Platform == UnrealTargetPlatform.IOS ? "IOS" : "TVOS"), "-NoIntellIsense", (Target.Platform == UnrealTargetPlatform.IOS ? "-iosdeployonly" : "-tvosdeployonly"), "-ignorejunk", (Target.bForDistribution ? "-distribution" : "-development"), String.Format("-project={0}", Target.ProjectFile), "-game", "-bundleID=" + BundleID });
+					GenerateProjectFiles(Target.ProjectFile, new string[] { "-platforms=" + (Target.Platform == UnrealTargetPlatform.IOS ? "IOS" : "TVOS"), "-NoIntellIsense", (Target.Platform == UnrealTargetPlatform.IOS ? "-iosdeployonly" : "-tvosdeployonly"), "-ignorejunk", (Target.bForDistribution ? "-distribution" : "-development"), String.Format("-project={0}", Target.ProjectFile), "-game", "-bundleID=" + BundleID, "-includetemptargets" });
 				}
 				// Make sure it exists
 				if (!DirectoryReference.Exists(XcodeWorkspaceDir))
