UBT: Track any reads from config files, so we can more accurately detect when they've been modified. Some config files are modified whenever the editor is started, causing the makefile to be invalidated unnecessarily.

diff --git a/Engine/Source/Programs/DotNETCommon/DotNETUtilities/BinaryArchiveWriter.cs b/Engine/Source/Programs/DotNETCommon/DotNETUtilities/BinaryArchiveWriter.cs
--- a/Engine/Source/Programs/DotNETCommon/DotNETUtilities/BinaryArchiveWriter.cs
+++ b/Engine/Source/Programs/DotNETCommon/DotNETUtilities/BinaryArchiveWriter.cs
@@ -378,7 +378,7 @@ namespace Tools.DotNETCommon
 		/// <typeparam name="T">Type of the element</typeparam>
 		/// <param name="Items">List of items</param>
 		/// <param name="WriteElement">Writes an individual element to the archive</param>
-		public void WriteList<T>(List<T> Items, Action<T> WriteElement)
+		public void WriteList<T>(IReadOnlyList<T> Items, Action<T> WriteElement)
 		{
 			if(Items == null)
 			{
diff --git a/Engine/Source/Programs/UnrealBuildTool/Configuration/TargetRules.cs b/Engine/Source/Programs/UnrealBuildTool/Configuration/TargetRules.cs
--- a/Engine/Source/Programs/UnrealBuildTool/Configuration/TargetRules.cs
+++ b/Engine/Source/Programs/UnrealBuildTool/Configuration/TargetRules.cs
@@ -1,14 +1,9 @@
 // Copyright 1998-2019 Epic Games, Inc. All Rights Reserved.
 
 using System;
-using System.Collections;
 using System.Collections.Generic;
-using System.IO;
 using System.Linq;
 using System.Reflection;
-using System.Text;
-using System.Threading;
-using System.Threading.Tasks;
 using Tools.DotNETCommon;
 
 namespace UnrealBuildTool
@@ -179,6 +174,11 @@ namespace UnrealBuildTool
 		/// </summary>
 		public global::UnrealBuildTool.TargetType Type = global::UnrealBuildTool.TargetType.Game;
 
+		/// <summary>
+		/// Tracks a list of config values read while constructing this target
+		/// </summary>
+		internal ConfigValueTracker ConfigValueTracker = new ConfigValueTracker();
+
 		/// <summary>
 		/// Whether the target uses Steam.
 		/// </summary>
@@ -1290,7 +1290,7 @@ namespace UnrealBuildTool
 			// Read settings from config files
 			foreach(object ConfigurableObject in GetConfigurableObjects())
 			{
-				ConfigCache.ReadSettings(DirectoryReference.FromFile(ProjectFile), Platform, ConfigurableObject);
+				ConfigCache.ReadSettings(DirectoryReference.FromFile(ProjectFile), Platform, ConfigurableObject, ConfigValueTracker);
 				XmlConfig.ApplyTo(ConfigurableObject);
 			}
 
@@ -1638,6 +1638,11 @@ namespace UnrealBuildTool
 			get { return Inner.Type; }
 		}
 
+		internal ConfigValueTracker ConfigValueTracker
+		{
+			get { return Inner.ConfigValueTracker; }
+		}
+
 		public bool bUsesSteam
 		{
 			get { return Inner.bUsesSteam; }
diff --git a/Engine/Source/Programs/UnrealBuildTool/Configuration/UEBuildTarget.cs b/Engine/Source/Programs/UnrealBuildTool/Configuration/UEBuildTarget.cs
--- a/Engine/Source/Programs/UnrealBuildTool/Configuration/UEBuildTarget.cs
+++ b/Engine/Source/Programs/UnrealBuildTool/Configuration/UEBuildTarget.cs
@@ -261,6 +261,32 @@ namespace UnrealBuildTool
 		public static UnrealTargetPlatform Lumin = FindOrAddByName("Lumin");
 	}
 
+	/// <summary>
+	/// Extension methods used for serializing UnrealTargetPlatform instances
+	/// </summary>
+	static class UnrealTargetPlatformExtensionMethods
+	{
+		/// <summary>
+		/// Read an UnrealTargetPlatform instance from a binary archive
+		/// </summary>
+		/// <param name="Reader">Archive to read from</param>
+		/// <returns>New UnrealTargetPlatform instance</returns>
+		public static UnrealTargetPlatform ReadUnrealTargetPlatform(this BinaryArchiveReader Reader)
+		{
+			return UnrealTargetPlatform.Parse(Reader.ReadString());
+		}
+
+		/// <summary>
+		/// Write an UnrealTargetPlatform instance to a binary archive
+		/// </summary>
+		/// <param name="Writer">The archive to write to</param>
+		/// <param name="Platform">The platform to write</param>
+		public static void WriteUnrealTargetPlatform(this BinaryArchiveWriter Writer, UnrealTargetPlatform Platform)
+		{
+			Writer.WriteString(Platform.ToString());
+		}
+	}
+
 	/// <summary>
 	/// Platform groups
 	/// </summary>
@@ -1474,7 +1500,7 @@ namespace UnrealBuildTool
 
 			// Create the makefile
 			string ExternalMetadata = UEBuildPlatform.GetBuildPlatform(Platform).GetExternalBuildMetadata(ProjectFile);
-			TargetMakefile Makefile = new TargetMakefile(TargetToolChain.GetVersionInfo(), ExternalMetadata, ReceiptFileName, ProjectIntermediateDirectory, TargetType, bDeployAfterCompile, bHasProjectScriptPlugin);
+			TargetMakefile Makefile = new TargetMakefile(TargetToolChain.GetVersionInfo(), ExternalMetadata, ReceiptFileName, ProjectIntermediateDirectory, TargetType, Rules.ConfigValueTracker, bDeployAfterCompile, bHasProjectScriptPlugin);
 
 			// Setup the hot reload module list
 			Makefile.HotReloadModuleNames = GetHotReloadModuleNames();
diff --git a/Engine/Source/Programs/UnrealBuildTool/Platform/Android/UEDeployAndroid.cs b/Engine/Source/Programs/UnrealBuildTool/Platform/Android/UEDeployAndroid.cs
--- a/Engine/Source/Programs/UnrealBuildTool/Platform/Android/UEDeployAndroid.cs
+++ b/Engine/Source/Programs/UnrealBuildTool/Platform/Android/UEDeployAndroid.cs
@@ -1662,7 +1662,7 @@ namespace UnrealBuildTool
 						continue;
 					}
 
-					IEnumerable<string> Values;
+					IReadOnlyList<string> Values;
 					Section.TryGetValues(Key, out Values);
 
 					foreach (string Value in Values)
@@ -1683,7 +1683,7 @@ namespace UnrealBuildTool
 						continue;
 					}
 
-					IEnumerable<string> Values;
+					IReadOnlyList<string> Values;
 					Section.TryGetValues(Key, out Values);
 					foreach (string Value in Values)
 					{
diff --git a/Engine/Source/Programs/UnrealBuildTool/System/ConfigCache.cs b/Engine/Source/Programs/UnrealBuildTool/System/ConfigCache.cs
--- a/Engine/Source/Programs/UnrealBuildTool/System/ConfigCache.cs
+++ b/Engine/Source/Programs/UnrealBuildTool/System/ConfigCache.cs
@@ -255,6 +255,18 @@ namespace UnrealBuildTool
 		/// <param name="Platform">The platform being built</param>
 		/// <param name="TargetObject">Object to receive the settings</param>
 		public static void ReadSettings(DirectoryReference ProjectDir, UnrealTargetPlatform Platform, object TargetObject)
+		{
+			ReadSettings(ProjectDir, Platform, TargetObject, null);
+		}
+
+		/// <summary>
+		/// Read config settings for the given object
+		/// </summary>
+		/// <param name="ProjectDir">Path to the project directory</param>
+		/// <param name="Platform">The platform being built</param>
+		/// <param name="TargetObject">Object to receive the settings</param>
+		/// <param name="Tracker">Tracks the set of config values that were retrieved. May be null.</param>
+		internal static void ReadSettings(DirectoryReference ProjectDir, UnrealTargetPlatform Platform, object TargetObject, ConfigValueTracker Tracker)
 		{
 			List<ConfigField> Fields = FindConfigFieldsForType(TargetObject.GetType());
 			foreach(ConfigField Field in Fields)
@@ -262,14 +274,20 @@ namespace UnrealBuildTool
 				// Read the hierarchy listed
 				ConfigHierarchy Hierarchy = ReadHierarchy(Field.Attribute.ConfigType, ProjectDir, Platform);
 
+				// Get the key name
+				string KeyName = Field.Attribute.KeyName ?? Field.FieldInfo.Name;
+
+				// Get the value(s) associated with this key
+				IReadOnlyList<string> Values;
+				Hierarchy.TryGetValues(Field.Attribute.SectionName, KeyName, out Values);
+
 				// Parse the values from the config files and update the target object
-				if(Field.AddElement == null)
+				if (Field.AddElement == null)
 				{
-					string Text;
-					if(Hierarchy.TryGetValue(Field.Attribute.SectionName, Field.Attribute.KeyName ?? Field.FieldInfo.Name, out Text))
+					if(Values != null && Values.Count == 1)
 					{
 						object Value;
-						if(TryParseValue(Text, Field.FieldInfo.FieldType, out Value))
+						if(TryParseValue(Values[0], Field.FieldInfo.FieldType, out Value))
 						{
 							Field.FieldInfo.SetValue(TargetObject, Value);
 						}
@@ -277,10 +295,9 @@ namespace UnrealBuildTool
 				}
 				else
 				{
-					IEnumerable<string> Items;
-					if(Hierarchy.TryGetValues(Field.Attribute.SectionName, Field.Attribute.KeyName ?? Field.FieldInfo.Name, out Items))
+					if(Values != null)
 					{
-						foreach(string Item in Items)
+						foreach(string Item in Values)
 						{
 							object Value;
 							if(TryParseValue(Item, Field.ElementType, out Value))
@@ -290,6 +307,12 @@ namespace UnrealBuildTool
 						}
 					}
 				}
+
+				// Save the dependency
+				if (Tracker != null)
+				{
+					Tracker.Add(Field.Attribute.ConfigType, ProjectDir, Platform, Field.Attribute.SectionName, KeyName, Values);
+				}
 			}
 		}
 
diff --git a/Engine/Source/Programs/UnrealBuildTool/System/ConfigHierarchy.cs b/Engine/Source/Programs/UnrealBuildTool/System/ConfigHierarchy.cs
--- a/Engine/Source/Programs/UnrealBuildTool/System/ConfigHierarchy.cs
+++ b/Engine/Source/Programs/UnrealBuildTool/System/ConfigHierarchy.cs
@@ -141,7 +141,7 @@ namespace UnrealBuildTool
 		/// <param name="KeyName">The key name to search for</param>
 		/// <param name="Values">On success, receives a list of the corresponding values</param>
 		/// <returns>True if the key was found, false otherwise</returns>
-		public bool TryGetValues(string KeyName, out IEnumerable<string> Values)
+		public bool TryGetValues(string KeyName, out IReadOnlyList<string> Values)
 		{
 			List<string> ValuesList;
 			if(KeyToValue.TryGetValue(KeyName, out ValuesList))
@@ -279,7 +279,7 @@ namespace UnrealBuildTool
 		/// <returns>True if the key exists</returns>
 		public bool GetArray(string SectionName, string KeyName, out List<string> Values)
 		{
-			IEnumerable<string> ValuesEnumerable;
+			IReadOnlyList<string> ValuesEnumerable;
 			if(TryGetValues(SectionName, KeyName, out ValuesEnumerable))
 			{
 				Values = ValuesEnumerable.ToList();
@@ -433,7 +433,7 @@ namespace UnrealBuildTool
 		/// <param name="KeyName">Key name</param>
 		/// <param name="Values">Copy of the list containing all values associated with the specified key</param>
 		/// <returns>True if the key exists</returns>
-		public bool TryGetValues(string SectionName, string KeyName, out IEnumerable<string> Values)
+		public bool TryGetValues(string SectionName, string KeyName, out IReadOnlyList<string> Values)
 		{
 			return FindSection(SectionName).TryGetValues(KeyName, out Values);
 		}
@@ -913,8 +913,23 @@ namespace UnrealBuildTool
 				}
 			}
 
+			// Find all the generated config files
+			foreach(FileReference GeneratedConfigFile in EnumerateGeneratedConfigFileLocations(Type, ProjectDir, Platform))
+			{
+				yield return GeneratedConfigFile;
+			}
+		}
+
+		/// <summary>
+		/// Returns a list of INI filenames for the given project
+		/// </summary>
+		public static IEnumerable<FileReference> EnumerateGeneratedConfigFileLocations(ConfigHierarchyType Type, DirectoryReference ProjectDir, UnrealTargetPlatform Platform)
+		{
+			string BaseIniName = Enum.GetName(typeof(ConfigHierarchyType), Type);
+			string PlatformName = GetIniPlatformName(Platform);
+
 			// Get the generated config file too. EditorSettings overrides this from 
-			if(Type == ConfigHierarchyType.EditorSettings)
+			if (Type == ConfigHierarchyType.EditorSettings)
 			{
 				yield return FileReference.Combine(GetGameAgnosticSavedDir(), "Config", PlatformName, BaseIniName + ".ini");
 			}
diff --git a/Engine/Source/Programs/UnrealBuildTool/System/ConfigValueTracker.cs b/Engine/Source/Programs/UnrealBuildTool/System/ConfigValueTracker.cs
new file mode 100644
--- /dev/null
+++ b/Engine/Source/Programs/UnrealBuildTool/System/ConfigValueTracker.cs
@@ -0,0 +1,207 @@
+using System;
+using System.Collections.Generic;
+using System.Diagnostics;
+using System.Linq;
+using System.Text;
+using System.Threading.Tasks;
+using Tools.DotNETCommon;
+
+namespace UnrealBuildTool
+{
+	/// <summary>
+	/// Identifier for a config file key, including information about the hierarchy used to read it
+	/// </summary>
+	[DebuggerDisplay("{Name}")]
+	class ConfigDependencyKey : IEquatable<ConfigDependencyKey>
+	{
+		/// <summary>
+		/// The config hierarchy type
+		/// </summary>
+		public ConfigHierarchyType Type;
+
+		/// <summary>
+		/// Project directory to read config files from
+		/// </summary>
+		public DirectoryReference ProjectDir;
+		
+		/// <summary>
+		/// The platform being built
+		/// </summary>
+		public UnrealTargetPlatform Platform;
+
+		/// <summary>
+		/// The section name
+		/// </summary>
+		public string SectionName;
+
+		/// <summary>
+		/// The key name
+		/// </summary>
+		public string KeyName;
+
+		/// <summary>
+		/// Constructor
+		/// </summary>
+		/// <param name="Type">The config hierarchy type</param>
+		/// <param name="ProjectDir">Project directory to read config files from</param>
+		/// <param name="Platform">The platform being built</param>
+		/// <param name="SectionName">The section name</param>
+		/// <param name="KeyName">The key name</param>
+		public ConfigDependencyKey(ConfigHierarchyType Type, DirectoryReference ProjectDir, UnrealTargetPlatform Platform, string SectionName, string KeyName)
+		{
+			this.Type = Type;
+			this.ProjectDir = ProjectDir;
+			this.Platform = Platform;
+			this.SectionName = SectionName;
+			this.KeyName = KeyName;
+		}
+
+		/// <summary>
+		/// Construct a key from an archive
+		/// </summary>
+		/// <param name="Reader">Archive to read from</param>
+		public ConfigDependencyKey(BinaryArchiveReader Reader)
+		{
+			Type = (ConfigHierarchyType)Reader.ReadInt();
+			ProjectDir = Reader.ReadDirectoryReference();
+			Platform = Reader.ReadUnrealTargetPlatform();
+			SectionName = Reader.ReadString();
+			KeyName = Reader.ReadString();
+		}
+
+		/// <summary>
+		/// Writes this key to an archive
+		/// </summary>
+		/// <param name="Writer">Archive to write to</param>
+		public void Write(BinaryArchiveWriter Writer)
+		{
+			Writer.WriteInt((int)Type);
+			Writer.WriteDirectoryReference(ProjectDir);
+			Writer.WriteUnrealTargetPlatform(Platform);
+			Writer.WriteString(SectionName);
+			Writer.WriteString(KeyName);
+		}
+
+		/// <summary>
+		/// Tests whether this key is equal to another object
+		/// </summary>
+		/// <param name="Other">The object to compare to</param>
+		/// <returns>True if the keys are equal, false otherwise</returns>
+		public override bool Equals(object Other)
+		{
+			return (Other is ConfigDependencyKey) && Equals((ConfigDependencyKey)Other);
+		}
+
+		/// <summary>
+		/// Tests whether this key is equal to another key
+		/// </summary>
+		/// <param name="Other">The key to compare to</param>
+		/// <returns>True if the keys are equal, false otherwise</returns>
+		public bool Equals(ConfigDependencyKey Other)
+		{
+			return Type == Other.Type && ProjectDir == Other.ProjectDir && Platform == Other.Platform && SectionName == Other.SectionName && KeyName == Other.KeyName;
+		}
+
+		/// <summary>
+		/// Gets a hash code for this object
+		/// </summary>
+		/// <returns>Hash code for the object</returns>
+		public override int GetHashCode()
+		{
+			int Hash = 17;
+			Hash = (Hash * 31) + Type.GetHashCode();
+			Hash = (Hash * 31) + ((ProjectDir == null) ? 0 : ProjectDir.GetHashCode());
+			Hash = (Hash * 31) + Platform.GetHashCode();
+			Hash = (Hash * 31) + SectionName.GetHashCode();
+			Hash = (Hash * 31) + KeyName.GetHashCode();
+			return Hash;
+		}
+	}
+
+	/// <summary>
+	/// Stores a list of config key/value pairs that have been read
+	/// </summary>
+	class ConfigValueTracker
+	{
+		/// <summary>
+		/// The dependencies list
+		/// </summary>
+		Dictionary<ConfigDependencyKey, IReadOnlyList<string>> Dependencies;
+
+		/// <summary>
+		/// Constructor
+		/// </summary>
+		public ConfigValueTracker()
+		{
+			Dependencies = new Dictionary<ConfigDependencyKey, IReadOnlyList<string>>();
+		}
+
+		/// <summary>
+		/// Construct an object from an archive on disk
+		/// </summary>
+		/// <param name="Reader">Archive to read from</param>
+		public ConfigValueTracker(BinaryArchiveReader Reader)
+		{
+			Dependencies = Reader.ReadDictionary(() => new ConfigDependencyKey(Reader), () => (IReadOnlyList<string>)Reader.ReadList(() => Reader.ReadString()));
+		}
+
+		/// <summary>
+		/// Write the dependencies object to disk
+		/// </summary>
+		/// <param name="Writer">Archive to write to</param>
+		public void Write(BinaryArchiveWriter Writer)
+		{
+			Writer.WriteDictionary(Dependencies, Key => Key.Write(Writer), Value => Writer.WriteList(Value, x => Writer.WriteString(x)));
+		}
+
+		/// <summary>
+		/// Adds a new configuration value
+		/// </summary>
+		/// <param name="Type">The config hierarchy type</param>
+		/// <param name="ProjectDir">The project directory</param>
+		/// <param name="Platform">The platform being built</param>
+		/// <param name="SectionName">Name of the config file section</param>
+		/// <param name="KeyName">Name of the config file key</param>
+		/// <param name="Values">Current values for this key</param>
+		public void Add(ConfigHierarchyType Type, DirectoryReference ProjectDir, UnrealTargetPlatform Platform, string SectionName, string KeyName, IReadOnlyList<string> Values)
+		{
+			ConfigDependencyKey Key = new ConfigDependencyKey(Type, ProjectDir, Platform, SectionName, KeyName);
+			Dependencies[Key] = Values;
+		}
+
+		/// <summary>
+		/// Checks whether the list of dependencies is still valid
+		/// </summary>
+		/// <returns></returns>
+		public bool IsValid()
+		{
+			foreach(KeyValuePair<ConfigDependencyKey, IReadOnlyList<string>> Pair in Dependencies)
+			{
+				// Read the appropriate hierarchy
+				ConfigHierarchy Hierarchy = ConfigCache.ReadHierarchy(Pair.Key.Type, Pair.Key.ProjectDir, Pair.Key.Platform);
+
+				// Get the value(s) associated with this key
+				IReadOnlyList<string> NewValues;
+				Hierarchy.TryGetValues(Pair.Key.SectionName, Pair.Key.KeyName, out NewValues);
+
+				// Check if they're different
+				if(Pair.Value == null)
+				{
+					if(NewValues != null)
+					{
+						return false;
+					}
+				}
+				else
+				{
+					if(NewValues == null || !Enumerable.SequenceEqual(Pair.Value, NewValues, StringComparer.Ordinal))
+					{
+						return false;
+					}
+				}
+			}
+			return true;
+		}
+	}
+}
+
diff --git a/Engine/Source/Programs/UnrealBuildTool/System/TargetMakefile.cs b/Engine/Source/Programs/UnrealBuildTool/System/TargetMakefile.cs
--- a/Engine/Source/Programs/UnrealBuildTool/System/TargetMakefile.cs
+++ b/Engine/Source/Programs/UnrealBuildTool/System/TargetMakefile.cs
@@ -21,7 +21,7 @@ namespace UnrealBuildTool
 		/// <summary>
 		/// The version number to write
 		/// </summary>
-		public const int CurrentVersion = 14;
+		public const int CurrentVersion = 16;
 
 		/// <summary>
 		/// The time at which the makefile was created
@@ -53,6 +53,11 @@ namespace UnrealBuildTool
 		/// </summary>
 		public TargetType TargetType;
 
+		/// <summary>
+		/// Map of config file keys to values. Makefile will be invalidated if any of these change.
+		/// </summary>
+		public ConfigValueTracker ConfigValueTracker;
+
 		/// <summary>
 		/// Whether the target should be deployed after being built
 		/// </summary>
@@ -129,6 +134,11 @@ namespace UnrealBuildTool
 		/// </summary>
 		public List<UHTModuleHeaderInfo> UObjectModuleHeaders = new List<UHTModuleHeaderInfo>();
 
+		/// <summary>
+		/// List of config settings in generated config files
+		/// </summary>
+		public Dictionary<string, string> ConfigSettings = new Dictionary<string, string>();
+
 		/// <summary>
 		/// List of all plugin names. The makefile will be considered invalid if any of these changes, or new plugins are added.
 		/// </summary>
@@ -137,7 +147,7 @@ namespace UnrealBuildTool
 		/// <summary>
 		/// Additional files which are required 
 		/// </summary>
-		public HashSet<FileItem> AdditionalDependencies = new HashSet<FileItem>(); 
+		public HashSet<FileItem> AdditionalDependencies = new HashSet<FileItem>();
 
 		/// <summary>
 		/// Constructor
@@ -147,9 +157,10 @@ namespace UnrealBuildTool
 		/// <param name="ReceiptFile">Path to the receipt file</param>
 		/// <param name="ProjectIntermediateDirectory">Path to the project intermediate directory</param>
 		/// <param name="TargetType">The type of target</param>
+		/// <param name="ConfigValueTracker">Set of dependencies on config files</param>
 		/// <param name="bDeployAfterCompile">Whether to deploy the target after compiling</param>
 		/// <param name="bHasProjectScriptPlugin">Whether the target has a project script plugin</param>
-		public TargetMakefile(string ToolchainInfo, string ExternalMetadata, FileReference ReceiptFile, DirectoryReference ProjectIntermediateDirectory, TargetType TargetType, bool bDeployAfterCompile, bool bHasProjectScriptPlugin)
+		public TargetMakefile(string ToolchainInfo, string ExternalMetadata, FileReference ReceiptFile, DirectoryReference ProjectIntermediateDirectory, TargetType TargetType, ConfigValueTracker ConfigValueTracker, bool bDeployAfterCompile, bool bHasProjectScriptPlugin)
 		{
 			this.CreateTimeUtc = DateTime.UtcNow;
 			this.ToolchainInfo = ToolchainInfo;
@@ -157,6 +168,7 @@ namespace UnrealBuildTool
 			this.ReceiptFile = ReceiptFile;
 			this.ProjectIntermediateDirectory = ProjectIntermediateDirectory;
 			this.TargetType = TargetType;
+			this.ConfigValueTracker = ConfigValueTracker;
 			this.bDeployAfterCompile = bDeployAfterCompile;
 			this.bHasProjectScriptPlugin = bHasProjectScriptPlugin;
 			this.Actions = new List<Action>();
@@ -185,6 +197,7 @@ namespace UnrealBuildTool
 			ReceiptFile = Reader.ReadFileReference();
 			ProjectIntermediateDirectory = Reader.ReadDirectoryReference();
 			TargetType = (TargetType)Reader.ReadInt();
+			ConfigValueTracker = new ConfigValueTracker(Reader);
 			bDeployAfterCompile = Reader.ReadBool();
 			bHasProjectScriptPlugin = Reader.ReadBool();
 			AdditionalArguments = Reader.ReadArray(() => Reader.ReadString());
@@ -216,6 +229,7 @@ namespace UnrealBuildTool
 			Writer.WriteFileReference(ReceiptFile);
 			Writer.WriteDirectoryReference(ProjectIntermediateDirectory);
 			Writer.WriteInt((int)TargetType);
+			ConfigValueTracker.Write(Writer);
 			Writer.WriteBool(bDeployAfterCompile);
 			Writer.WriteBool(bHasProjectScriptPlugin);
 			Writer.WriteArray(AdditionalArguments, Item => Writer.WriteString(Item));
@@ -386,20 +400,11 @@ namespace UnrealBuildTool
 					return null;
 				}
 
-				// Check if ini files are newer. Ini files contain build settings too.
-				DirectoryReference ProjectDirectory = DirectoryReference.FromFile(ProjectFile);
-				foreach (ConfigHierarchyType IniType in (ConfigHierarchyType[])Enum.GetValues(typeof(ConfigHierarchyType)))
+				// Check if any config settings have changed. Ini files contain build settings too.
+				if(!Makefile.ConfigValueTracker.IsValid())
 				{
-					foreach (FileReference IniFilename in ConfigHierarchy.EnumerateConfigFileLocations(IniType, ProjectDirectory, Platform))
-					{
-						FileInfo IniFileInfo = new FileInfo(IniFilename.FullName);
-						if(IniFileInfo.LastWriteTimeUtc > Makefile.CreateTimeUtc)
-						{
-							// Ini files are newer than makefile
-							ReasonNotLoaded = "ini files are newer than makefile";
-							return null;
-						}
-					}
+					ReasonNotLoaded = "config setting changed";
+					return null;
 				}
 
 				// Get the current build metadata from the platform
diff --git a/Engine/Source/Programs/UnrealBuildTool/UnrealBuildTool.csproj b/Engine/Source/Programs/UnrealBuildTool/UnrealBuildTool.csproj
--- a/Engine/Source/Programs/UnrealBuildTool/UnrealBuildTool.csproj
+++ b/Engine/Source/Programs/UnrealBuildTool/UnrealBuildTool.csproj
@@ -190,6 +190,7 @@
     <Compile Include="Modes\AggregateParsedTimingInfo.cs" />
     <Compile Include="System\Asn.cs" />
     <Compile Include="System\BuildVersion.cs" />
+    <Compile Include="System\ConfigValueTracker.cs" />
     <Compile Include="System\CppCompileEnvironment.cs" />
     <Compile Include="System\CppDependencyCache.cs" />
     <Compile Include="System\CppIncludeParser.cs" />
