UBT: Always treat *.modules files as being up to date for an installed engine. Prevents BuildId being regenerated if timestamps for DLLs are newer.

diff --git a/Engine/Source/Programs/UnrealBuildTool/Modes/WriteMetadataMode.cs b/Engine/Source/Programs/UnrealBuildTool/Modes/WriteMetadataMode.cs
--- a/Engine/Source/Programs/UnrealBuildTool/Modes/WriteMetadataMode.cs
+++ b/Engine/Source/Programs/UnrealBuildTool/Modes/WriteMetadataMode.cs
@@ -296,13 +296,16 @@ namespace UnrealBuildTool
 		/// <returns>True if the manifest is out of date</returns>
 		bool IsOutOfDate(FileReference ManifestFileName, ModuleManifest Manifest)
 		{
-			DateTime ManifestTime = FileReference.GetLastWriteTimeUtc(ManifestFileName);
-			foreach(string FileName in Manifest.ModuleNameToFileName.Values)
+			if(!UnrealBuildTool.IsFileInstalled(ManifestFileName))
 			{
-				FileInfo ModuleInfo = new FileInfo(FileReference.Combine(ManifestFileName.Directory, FileName).FullName);
-				if(!ModuleInfo.Exists || ModuleInfo.LastWriteTimeUtc > ManifestTime)
+				DateTime ManifestTime = FileReference.GetLastWriteTimeUtc(ManifestFileName);
+				foreach(string FileName in Manifest.ModuleNameToFileName.Values)
 				{
-					return true;
+					FileInfo ModuleInfo = new FileInfo(FileReference.Combine(ManifestFileName.Directory, FileName).FullName);
+					if(!ModuleInfo.Exists || ModuleInfo.LastWriteTimeUtc > ManifestTime)
+					{
+						return true;
+					}
 				}
 			}
 			return false;
