Default build configuration is with legacy MeshEditor disabled, as it is incomplete and no longer supported.

diff --git a/Engine/Source/Programs/UnrealBuildTool/Configuration/TargetRules.cs b/Engine/Source/Programs/UnrealBuildTool/Configuration/TargetRules.cs
--- a/Engine/Source/Programs/UnrealBuildTool/Configuration/TargetRules.cs
+++ b/Engine/Source/Programs/UnrealBuildTool/Configuration/TargetRules.cs
@@ -303,7 +303,7 @@ namespace UnrealBuildTool
 		/// Whether to enable the mesh editor.
 		/// </summary>
 		[RequiresUniqueBuildEnvironment]
-		public bool bEnableMeshEditor = true;
+		public bool bEnableMeshEditor = false;
 
 		/// <summary>
 		/// Whether to compile the Chaos physics plugin.
