Re-enable XGE for cross building for Linux on Windows

diff --git a/Engine/Source/Programs/UnrealBuildTool/Platform/Linux/UEBuildLinux.cs b/Engine/Source/Programs/UnrealBuildTool/Platform/Linux/UEBuildLinux.cs
--- a/Engine/Source/Programs/UnrealBuildTool/Platform/Linux/UEBuildLinux.cs
+++ b/Engine/Source/Programs/UnrealBuildTool/Platform/Linux/UEBuildLinux.cs
@@ -297,7 +297,8 @@ namespace UnrealBuildTool
 			// [bschaefer] 2018-08-24: disabling XGE due to a bug where XGE seems to be lower casing folders names that are headers ie. misc/Header.h vs Misc/Header.h
 			// [bschaefer] 2018-10-04: enabling XGE as an update in xgConsole seems to have fixed it for me
 			// [bschaefer] 2018-12-17: disable XGE again, as the same issue before seems to still be happening but intermittently
-			return false; // BuildHostPlatform.Current.Platform == UnrealTargetPlatform.Win64;
+			// [bschaefer] 2019-6-13: enable XGE, as the bug from before is now fixed
+			return BuildHostPlatform.Current.Platform == UnrealTargetPlatform.Win64;
 		}
 
 		public override bool CanUseParallelExecutor()
