Merging using //UE4/Partner-Microsoft-XR/_to_//UE4/Release-4.23/

Package for HoloLens when project root directory is different from uproject name

diff --git a/Engine/Source/Programs/UnrealBuildTool/Platform/HoloLens/ManifestGenerator.cs b/Engine/Source/Programs/UnrealBuildTool/Platform/HoloLens/ManifestGenerator.cs
--- a/Engine/Source/Programs/UnrealBuildTool/Platform/HoloLens/ManifestGenerator.cs
+++ b/Engine/Source/Programs/UnrealBuildTool/Platform/HoloLens/ManifestGenerator.cs
@@ -1915,8 +1915,10 @@ namespace UnrealBuildTool
 				ConfigPostfix = TargetConfig.ToString();
 			}
 
-			string MakeRelativeTo = IsGameSpecificExe ? ProjectPath : UnrealBuildTool.EngineDirectory.FullName;
-			string RelativeExePath = Path.Combine(RelativeProjectRootForStage, Utils.MakePathRelativeTo(ExecutablePath, MakeRelativeTo));
+			string MakeRelativeTo = IsGameSpecificExe ? Path.Combine(ProjectPath, "..") : UnrealBuildTool.EngineDirectory.FullName;
+			string RelativeExePath = IsGameSpecificExe ?
+				Utils.MakePathRelativeTo(ExecutablePath, MakeRelativeTo) :
+				Path.Combine(RelativeProjectRootForStage, Utils.MakePathRelativeTo(ExecutablePath, MakeRelativeTo));
 
 			XmlAttribute Id = AppxManifestXmlDocument.CreateAttribute("Id");
 			Id.Value = "App" + PackageBaseName + ConfigPostfix;
