UBT: Fixed UBT not cleanign up some Linux build products (UE-76563).

diff --git a/Engine/Source/Programs/UnrealBuildTool/Platform/Linux/UEBuildLinux.cs b/Engine/Source/Programs/UnrealBuildTool/Platform/Linux/UEBuildLinux.cs
--- a/Engine/Source/Programs/UnrealBuildTool/Platform/Linux/UEBuildLinux.cs
+++ b/Engine/Source/Programs/UnrealBuildTool/Platform/Linux/UEBuildLinux.cs
@@ -319,13 +319,17 @@ namespace UnrealBuildTool
 			if (FileName.StartsWith("lib"))
 			{
 				return IsBuildProductName(FileName, 3, FileName.Length - 3, NamePrefixes, NameSuffixes, ".a")
-					|| IsBuildProductName(FileName, 3, FileName.Length - 3, NamePrefixes, NameSuffixes, ".so");
+					|| IsBuildProductName(FileName, 3, FileName.Length - 3, NamePrefixes, NameSuffixes, ".so")
+					|| IsBuildProductName(FileName, 3, FileName.Length - 3, NamePrefixes, NameSuffixes, ".sym")
+					|| IsBuildProductName(FileName, 3, FileName.Length - 3, NamePrefixes, NameSuffixes, ".debug");
 			}
 			else
 			{
 				return IsBuildProductName(FileName, NamePrefixes, NameSuffixes, "")
 					|| IsBuildProductName(FileName, NamePrefixes, NameSuffixes, ".so")
-					|| IsBuildProductName(FileName, NamePrefixes, NameSuffixes, ".a");
+					|| IsBuildProductName(FileName, NamePrefixes, NameSuffixes, ".a")
+					|| IsBuildProductName(FileName, NamePrefixes, NameSuffixes, ".sym")
+					|| IsBuildProductName(FileName, NamePrefixes, NameSuffixes, ".debug");
 			}
 		}
 
