Disable Oculus plugins on Android X64

#ue4
#android

diff --git a/Engine/Source/Programs/UnrealBuildTool/Platform/Android/AndroidToolChain.cs b/Engine/Source/Programs/UnrealBuildTool/Platform/Android/AndroidToolChain.cs
--- a/Engine/Source/Programs/UnrealBuildTool/Platform/Android/AndroidToolChain.cs
+++ b/Engine/Source/Programs/UnrealBuildTool/Platform/Android/AndroidToolChain.cs
@@ -82,8 +82,8 @@ namespace UnrealBuildTool
 		static private Dictionary<string, string[]> ModulesToSkip = new Dictionary<string, string[]> {
 			{ "-armv7", new string[] {  } },
 			{ "-arm64", new string[] {  } },
-			{ "-x86",   new string[] { "OnlineSubsystemOculus", } },
-			{ "-x64",   new string[] { "OnlineSubsystemOculus", "OnlineSubsystemGooglePlay", } },
+			{ "-x86",   new string[] { "OnlineSubsystemOculus", "OculusHMD" } },
+			{ "-x64",   new string[] { "OnlineSubsystemOculus", "OculusHMD", "OnlineSubsystemGooglePlay" } },
 		};
 
 		static private Dictionary<string, string[]> GeneratedModulesToSkip = new Dictionary<string, string[]> {
