Fix staging of libUE4.so for Ant builds on Android

#ue4
#android

diff --git a/Engine/Source/Programs/UnrealBuildTool/Platform/Android/UEDeployAndroid.cs b/Engine/Source/Programs/UnrealBuildTool/Platform/Android/UEDeployAndroid.cs
--- a/Engine/Source/Programs/UnrealBuildTool/Platform/Android/UEDeployAndroid.cs
+++ b/Engine/Source/Programs/UnrealBuildTool/Platform/Android/UEDeployAndroid.cs
@@ -3871,6 +3871,35 @@ namespace UnrealBuildTool
 								// copy the binary to the standard .so location
 								File.Copy(SourceSOName, FinalSOName, true);
 							}
+
+							// remove any read only flags
+							FileInfo DestFileInfo2 = new FileInfo(FinalSOName);
+							DestFileInfo2.Attributes = DestFileInfo2.Attributes & ~FileAttributes.ReadOnly;
+							File.SetLastWriteTimeUtc(FinalSOName, File.GetLastWriteTimeUtc(SourceSOName));
+
+							// run ndk-build for Ant (will stage libUE4.so into libs)
+							string LibSOName = UE4BuildPath + "/libs/" + NDKArch + "/libUE4.so";
+
+							// always delete libs up to this point so fat binaries and incremental builds work together (otherwise we might end up with multiple
+							// so files in an apk that doesn't want them)
+							// note that we don't want to delete all libs, just the ones we copied
+							TimeSpan Diff = File.GetLastWriteTimeUtc(LibSOName) - File.GetLastWriteTimeUtc(FinalSOName);
+							if (!File.Exists(LibSOName) || Diff.TotalSeconds < -1 || Diff.TotalSeconds > 1)
+							{
+								foreach (string Lib in Directory.EnumerateFiles(UE4BuildPath + "/libs", "libUE4*.so", SearchOption.AllDirectories))
+								{
+									File.Delete(Lib);
+								}
+
+								string CommandLine = "APP_ABI=\"" + NDKArch + " " + "\"";
+								if (!bForDistribution)
+								{
+									CommandLine += " NDK_DEBUG=1";
+								}
+								RunCommandLineProgramWithException(UE4BuildPath, NDKBuildPath, CommandLine, "Preparing native code for debugging...", true);
+
+								File.SetLastWriteTimeUtc(LibSOName, File.GetLastWriteTimeUtc(FinalSOName));
+							}
 						}
 						else
 						{
