Expose isEmbedded to UPL (Android-only)

#android
#ue4

diff --git a/Engine/Source/Programs/UnrealBuildTool/Platform/Android/UEDeployAndroid.cs b/Engine/Source/Programs/UnrealBuildTool/Platform/Android/UEDeployAndroid.cs
--- a/Engine/Source/Programs/UnrealBuildTool/Platform/Android/UEDeployAndroid.cs
+++ b/Engine/Source/Programs/UnrealBuildTool/Platform/Android/UEDeployAndroid.cs
@@ -3586,7 +3586,7 @@ namespace UnrealBuildTool
 				}
 			}
 
-			UPL.Init(NDKArches, bForDistribution, EngineDirectory, UE4BuildPath, ProjectDirectory, Configuration.ToString());
+			UPL.Init(NDKArches, bForDistribution, EngineDirectory, UE4BuildPath, ProjectDirectory, Configuration.ToString(), bSkipGradleBuild);
 
 			IEnumerable<Tuple<string, string, string>> BuildList = null;
 
diff --git a/Engine/Source/Programs/UnrealBuildTool/Platform/IOS/UEDeployIOS.cs b/Engine/Source/Programs/UnrealBuildTool/Platform/IOS/UEDeployIOS.cs
--- a/Engine/Source/Programs/UnrealBuildTool/Platform/IOS/UEDeployIOS.cs
+++ b/Engine/Source/Programs/UnrealBuildTool/Platform/IOS/UEDeployIOS.cs
@@ -838,7 +838,7 @@ namespace UnrealBuildTool
 			UnrealPluginLanguage UPL = new UnrealPluginLanguage(ProjectFile, UPLScripts, ProjectArches, "", "", UnrealTargetPlatform.IOS);
 
 			// Passing in true for distribution is not ideal here but given the way that ios packaging happens and this call chain it seems unavoidable for now, maybe there is a way to correctly pass it in that I can't find?
-			UPL.Init(ProjectArches, true, RelativeEnginePath, BundlePath, ProjectDirectory, Config.ToString());
+			UPL.Init(ProjectArches, true, RelativeEnginePath, BundlePath, ProjectDirectory, Config.ToString(), false);
 
 			return GenerateIOSPList(ProjectFile, Config, ProjectDirectory, bIsUE4Game, GameName, ProjectName, InEngineDir, AppDirectory, SdkVersion, UPL, BundleID, out bSupportsPortrait, out bSupportsLandscape, out bSkipIcons);
 		}
diff --git a/Engine/Source/Programs/UnrealBuildTool/Platform/Lumin/UEDeployLumin.cs b/Engine/Source/Programs/UnrealBuildTool/Platform/Lumin/UEDeployLumin.cs
--- a/Engine/Source/Programs/UnrealBuildTool/Platform/Lumin/UEDeployLumin.cs
+++ b/Engine/Source/Programs/UnrealBuildTool/Platform/Lumin/UEDeployLumin.cs
@@ -276,8 +276,10 @@ namespace UnrealBuildTool
 
 			SetLuminPluginData(MLSDKArches, CollectPluginDataPaths(Receipt));
 
+			bool bIsEmbedded = Receipt.HasValueForAdditionalProperty("CompileAsDll", "true");
+
 			//gather all of the xml
-			UPL.Init(MLSDKArches, true, RelativeEnginePath, UE4BuildPath, RelativeProjectPath, ConfigurationString);
+			UPL.Init(MLSDKArches, true, RelativeEnginePath, UE4BuildPath, RelativeProjectPath, ConfigurationString, bIsEmbedded);
 		}
 
 		public string StageFiles()
diff --git a/Engine/Source/Programs/UnrealBuildTool/System/UnrealPluginLanguage.cs b/Engine/Source/Programs/UnrealBuildTool/System/UnrealPluginLanguage.cs
--- a/Engine/Source/Programs/UnrealBuildTool/System/UnrealPluginLanguage.cs
+++ b/Engine/Source/Programs/UnrealBuildTool/System/UnrealPluginLanguage.cs
@@ -51,6 +51,7 @@ namespace UnrealBuildTool
 	 *	$S(BuildDir) = project's platform appropriate build directory (within the Intermediate folder)
 	 *	$S(Configuration) = configuration type (Debug, DebugGame, Development, Test, Shipping)
 	 *	$B(Distribution) = true if distribution build
+	 *	$B(IsEmbedded) = true if project compiled for embedded use [Android-only]
 	 *	$I(EngineMajorVersion) = major version of engine (ex. 4)
 	 *	$I(EngineMinorVersion) = minor version of engine (ex. 21)
 	 *	$I(EnginePatchVersion) = patch version of engine (ex. 0)
@@ -2401,9 +2402,10 @@ namespace UnrealBuildTool
 			return GlobalContext.StringVariables["Output"];
 		}
 
-		public void Init(List<string> Architectures, bool bDistribution, string EngineDirectory, string BuildDirectory, string ProjectDirectory, string Configuration)
+		public void Init(List<string> Architectures, bool bDistribution, string EngineDirectory, string BuildDirectory, string ProjectDirectory, string Configuration, bool bIsEmbedded)
 		{
 			GlobalContext.BoolVariables["Distribution"] = bDistribution;
+			GlobalContext.BoolVariables["IsEmbedded"] = bIsEmbedded;
 			GlobalContext.StringVariables["Configuration"] = Configuration;
 
 			GlobalContext.StringVariables["EngineDir"] = EngineDirectory.Replace("\\", "/");
