UBT: Rename the bOutputPublicallyDistributable setting to reduce confusion. Now called ModuleRules.bLegalToDistributeObjectCode and TargetRules.bLegalToDistributeBinary.

#rnx

diff --git a/Engine/Plugins/Media/MfMedia/Source/MfMedia/MfMedia.Build.cs b/Engine/Plugins/Media/MfMedia/Source/MfMedia/MfMedia.Build.cs
--- a/Engine/Plugins/Media/MfMedia/Source/MfMedia/MfMedia.Build.cs
+++ b/Engine/Plugins/Media/MfMedia/Source/MfMedia/MfMedia.Build.cs
@@ -6,8 +6,6 @@ namespace UnrealBuildTool.Rules
 	{
 		public MfMedia(ReadOnlyTargetRules Target) : base(Target)
 		{
-			bOutputPubliclyDistributable = true;
-
 			DynamicallyLoadedModuleNames.AddRange(
 				new string[] {
 					"Media",
diff --git a/Engine/Plugins/Media/MfMedia/Source/MfMediaFactory/MfMediaFactory.Build.cs b/Engine/Plugins/Media/MfMedia/Source/MfMediaFactory/MfMediaFactory.Build.cs
--- a/Engine/Plugins/Media/MfMedia/Source/MfMediaFactory/MfMediaFactory.Build.cs
+++ b/Engine/Plugins/Media/MfMedia/Source/MfMediaFactory/MfMediaFactory.Build.cs
@@ -6,9 +6,6 @@ namespace UnrealBuildTool.Rules
 	{
 		public MfMediaFactory(ReadOnlyTargetRules Target) : base(Target)
 		{
-			// this is for Xbox and Windows, so it's using public APIs, so we can distribute it in binary
-			bOutputPubliclyDistributable = true;
-
 			DynamicallyLoadedModuleNames.AddRange(
 				new string[] {
 					"Media",
diff --git a/Engine/Source/Programs/BuildPatchTool/BuildPatchTool.Target.cs b/Engine/Source/Programs/BuildPatchTool/BuildPatchTool.Target.cs
--- a/Engine/Source/Programs/BuildPatchTool/BuildPatchTool.Target.cs
+++ b/Engine/Source/Programs/BuildPatchTool/BuildPatchTool.Target.cs
@@ -12,7 +12,7 @@ public class BuildPatchToolTarget : TargetRules
 		Type = TargetType.Program;
 		LinkType = TargetLinkType.Monolithic;
 		LaunchModuleName = "BuildPatchTool";
-		bOutputPubliclyDistributable = true;
+        bLegalToDistributeBinary = true;
 		UndecoratedConfiguration = UnrealTargetConfiguration.Shipping;
 
 		bCompileAgainstEngine = false;
diff --git a/Engine/Source/Programs/CrashReportClient/CrashReportClient.Target.cs b/Engine/Source/Programs/CrashReportClient/CrashReportClient.Target.cs
--- a/Engine/Source/Programs/CrashReportClient/CrashReportClient.Target.cs
+++ b/Engine/Source/Programs/CrashReportClient/CrashReportClient.Target.cs
@@ -20,7 +20,7 @@ public class CrashReportClientTarget : TargetRules
 			ExtraModuleNames.Add("EditorStyle");
 		}
 
-        bOutputPubliclyDistributable = true;
+        bLegalToDistributeBinary = true;
 
 		bBuildDeveloperTools = false;
 
diff --git a/Engine/Source/Programs/UnrealBuildTool/Configuration/ModuleRules.cs b/Engine/Source/Programs/UnrealBuildTool/Configuration/ModuleRules.cs
--- a/Engine/Source/Programs/UnrealBuildTool/Configuration/ModuleRules.cs
+++ b/Engine/Source/Programs/UnrealBuildTool/Configuration/ModuleRules.cs
@@ -587,7 +587,17 @@ namespace UnrealBuildTool
 		/// dependencies on modules that are not (i.e. CarefullyRedist, NotForLicensees, NoRedist).
 		/// This should be used when you plan to release binaries but not source.
 		/// </summary>
-		public bool bOutputPubliclyDistributable = false;
+		public bool bLegalToDistributeObjectCode = false;
+
+		/// <summary>
+		/// Obsolete. Use bLegalToDistributeObjectCode instead.
+		/// </summary>
+		[Obsolete("bOutputPubliclyDistributable has been deprecated in 4.24. Use bLegalToDistributeObjectCode instead.")]
+		public bool bOutputPubliclyDistributable
+		{
+			get { return bLegalToDistributeObjectCode; }
+			set { bLegalToDistributeObjectCode = value; }
+		}
 
 		/// <summary>
 		/// List of folders which are whitelisted to be referenced when compiling this binary, without propagating restricted folder names
diff --git a/Engine/Source/Programs/UnrealBuildTool/Configuration/TargetRules.cs b/Engine/Source/Programs/UnrealBuildTool/Configuration/TargetRules.cs
--- a/Engine/Source/Programs/UnrealBuildTool/Configuration/TargetRules.cs
+++ b/Engine/Source/Programs/UnrealBuildTool/Configuration/TargetRules.cs
@@ -216,7 +216,17 @@ namespace UnrealBuildTool
 		/// Whether the output from this target can be publicly distributed, even if it has dependencies on modules that are in folders 
 		/// with special restrictions (eg. CarefullyRedist, NotForLicensees, NoRedist).
 		/// </summary>
-		public bool bOutputPubliclyDistributable = false;
+		public bool bLegalToDistributeBinary = false;
+
+		/// <summary>
+		/// Obsolete. Use bLegalToDistributeBinary instead.
+		/// </summary>
+		[Obsolete("bOutputPubliclyDistributable has been deprecated in 4.24. Use bLegalToDistributeBinary instead.")]
+		public bool bOutputPubliclyDistributable
+		{
+			get { return bLegalToDistributeBinary; }
+			set { bLegalToDistributeBinary = value; }
+		}
 
 		/// <summary>
 		/// Specifies the configuration whose binaries do not require a "-Platform-Configuration" suffix.
@@ -1705,9 +1715,9 @@ namespace UnrealBuildTool
 			get { return Inner.bDebugBuildsActuallyUseDebugCRT; }
 		}
 
-		public bool bOutputPubliclyDistributable
+		public bool bLegalToDistributeBinary
 		{
-			get { return Inner.bOutputPubliclyDistributable; }
+			get { return Inner.bLegalToDistributeBinary; }
 		}
 
 		public UnrealTargetConfiguration UndecoratedConfiguration
diff --git a/Engine/Source/Programs/UnrealBuildTool/Configuration/UEBuildModule.cs b/Engine/Source/Programs/UnrealBuildTool/Configuration/UEBuildModule.cs
--- a/Engine/Source/Programs/UnrealBuildTool/Configuration/UEBuildModule.cs
+++ b/Engine/Source/Programs/UnrealBuildTool/Configuration/UEBuildModule.cs
@@ -323,7 +323,7 @@ namespace UnrealBuildTool
 		public Dictionary<RestrictedFolder, DirectoryReference> FindRestrictedFolderReferences(DirectoryReference ProjectDir)
 		{
 			Dictionary<RestrictedFolder, DirectoryReference> References = new Dictionary<RestrictedFolder, DirectoryReference>();
-			if (!Rules.bOutputPubliclyDistributable)
+			if (!Rules.bLegalToDistributeObjectCode)
 			{
 				// Find all the directories that this module references
 				HashSet<DirectoryReference> ReferencedDirs = new HashSet<DirectoryReference>();
diff --git a/Engine/Source/Programs/UnrealBuildTool/Configuration/UEBuildTarget.cs b/Engine/Source/Programs/UnrealBuildTool/Configuration/UEBuildTarget.cs
--- a/Engine/Source/Programs/UnrealBuildTool/Configuration/UEBuildTarget.cs
+++ b/Engine/Source/Programs/UnrealBuildTool/Configuration/UEBuildTarget.cs
@@ -1786,7 +1786,7 @@ namespace UnrealBuildTool
 			if(!Rules.bDisableLinking)
 			{
 				// Check the distribution level of all binaries based on the dependencies they have
-				if(ProjectFile == null && !Rules.bOutputPubliclyDistributable)
+				if(ProjectFile == null && !Rules.bLegalToDistributeBinary)
 				{
 					Dictionary<UEBuildModule, Dictionary<RestrictedFolder, DirectoryReference>> ModuleRestrictedFolderCache = new Dictionary<UEBuildModule, Dictionary<RestrictedFolder, DirectoryReference>>();
 
diff --git a/Engine/Source/Runtime/VulkanRHI/VulkanRHI.Build.cs b/Engine/Source/Runtime/VulkanRHI/VulkanRHI.Build.cs
--- a/Engine/Source/Runtime/VulkanRHI/VulkanRHI.Build.cs
+++ b/Engine/Source/Runtime/VulkanRHI/VulkanRHI.Build.cs
@@ -8,7 +8,7 @@ public class VulkanRHI : ModuleRules
 {
 	public VulkanRHI(ReadOnlyTargetRules Target) : base(Target)
 	{
-		bOutputPubliclyDistributable = true;
+		bLegalToDistributeObjectCode = true;
 
 		PrivateIncludePaths.Add("Runtime/VulkanRHI/Private");
 		if (Target.Platform == UnrealTargetPlatform.Win32 || Target.Platform == UnrealTargetPlatform.Win64)
