Removing log line

diff --git a/Engine/Source/Programs/UnrealBuildTool/Platform/HoloLens/UEBuildHoloLens.cs b/Engine/Source/Programs/UnrealBuildTool/Platform/HoloLens/UEBuildHoloLens.cs
--- a/Engine/Source/Programs/UnrealBuildTool/Platform/HoloLens/UEBuildHoloLens.cs
+++ b/Engine/Source/Programs/UnrealBuildTool/Platform/HoloLens/UEBuildHoloLens.cs
@@ -212,7 +212,6 @@ namespace UnrealBuildTool
 			DirectoryReference SelectedWindowsSdkDir;
 			if (!WindowsPlatform.TryGetWindowsSdkDir(Target.HoloLensPlatform.Win10SDKVersionString, out SelectedWindowsSdkVersion, out SelectedWindowsSdkDir))
 			{
-				Log.TraceWarning("Specified Windows SDK version {0} is no longer installed using latest.", Target.HoloLensPlatform.Win10SDKVersion);
 				Target.HoloLensPlatform.Win10SDKVersionString = "Latest";
 			}
 
