Removed GLES_SILENCE_DEPRECATION since CL 5226277 now automatically removes GLES code when targeting iOS 12.

diff --git a/Engine/Source/Programs/UnrealBuildTool/Platform/IOS/UEBuildIOS.cs b/Engine/Source/Programs/UnrealBuildTool/Platform/IOS/UEBuildIOS.cs
--- a/Engine/Source/Programs/UnrealBuildTool/Platform/IOS/UEBuildIOS.cs
+++ b/Engine/Source/Programs/UnrealBuildTool/Platform/IOS/UEBuildIOS.cs
@@ -1012,7 +1012,6 @@ namespace UnrealBuildTool
 
 			CompileEnvironment.Definitions.Add("PLATFORM_IOS=1");
 			CompileEnvironment.Definitions.Add("PLATFORM_APPLE=1");
-			CompileEnvironment.Definitions.Add("GLES_SILENCE_DEPRECATION=1");  // suppress GLES "deprecated" warnings until a proper solution is implemented (see UE-65643)
 
 			CompileEnvironment.Definitions.Add("WITH_TTS=0");
 			CompileEnvironment.Definitions.Add("WITH_SPEECH_RECOGNITION=0");
