Added support for ClientOnlyNoCommandlet to module type.
Robert.Manuszewski

diff --git a/Engine/Source/Programs/UnrealBuildTool/System/ExternalExecution.cs b/Engine/Source/Programs/UnrealBuildTool/System/ExternalExecution.cs
--- a/Engine/Source/Programs/UnrealBuildTool/System/ExternalExecution.cs
+++ b/Engine/Source/Programs/UnrealBuildTool/System/ExternalExecution.cs
@@ -122,7 +122,8 @@ namespace UnrealBuildTool
                 case ModuleHostType.CookedOnly:
                 case ModuleHostType.ServerOnly:
                 case ModuleHostType.ClientOnly:
-                    return UHTModuleType.EngineRuntime;
+				case ModuleHostType.ClientOnlyNoCommandlet:
+					return UHTModuleType.EngineRuntime;
 				case ModuleHostType.Developer:
 					return UHTModuleType.EngineDeveloper;
 				case ModuleHostType.Editor:
@@ -144,6 +145,7 @@ namespace UnrealBuildTool
                 case ModuleHostType.CookedOnly:
                 case ModuleHostType.ServerOnly:
                 case ModuleHostType.ClientOnly:
+				case ModuleHostType.ClientOnlyNoCommandlet:
                     return UHTModuleType.GameRuntime;
 				case ModuleHostType.Developer:
 					return UHTModuleType.GameDeveloper;
diff --git a/Engine/Source/Programs/UnrealBuildTool/System/ModuleDescriptor.cs b/Engine/Source/Programs/UnrealBuildTool/System/ModuleDescriptor.cs
--- a/Engine/Source/Programs/UnrealBuildTool/System/ModuleDescriptor.cs
+++ b/Engine/Source/Programs/UnrealBuildTool/System/ModuleDescriptor.cs
@@ -65,9 +65,14 @@ namespace UnrealBuildTool
         ServerOnly,
 
 		/// <summary>
-		/// Loaded only by clients
+		/// Loaded only by clients, and commandlets, and editor....
 		/// </summary>
         ClientOnly,
+
+		/// <summary>
+		/// Loaded only by clients and editor (editor can run PIE which is kinda a commandlet)
+		/// </summary>
+		ClientOnlyNoCommandlet,
 	}
 
 	/// <summary>
diff --git a/Engine/Source/Runtime/Projects/Private/ModuleDescriptor.cpp b/Engine/Source/Runtime/Projects/Private/ModuleDescriptor.cpp
--- a/Engine/Source/Runtime/Projects/Private/ModuleDescriptor.cpp
+++ b/Engine/Source/Runtime/Projects/Private/ModuleDescriptor.cpp
@@ -106,7 +106,8 @@ const TCHAR* EHostType::ToString( const EHostType::Type Value )
 
 		case ClientOnly:
 			return TEXT("ClientOnly");
-
+		case ClientOnlyNoCommandlet:
+			return TEXT("ClientOnlyNoCommandlet");
 		default:
 			ensureMsgf( false, TEXT( "Unrecognized EModuleType value: %i" ), Value );
 			return NULL;
@@ -566,11 +567,15 @@ bool FModuleDescriptor::IsLoadedInCurrentConfiguration() const
 	case EHostType::ServerOnly:
 		return !FPlatformProperties::IsClientOnly();
 
+	case EHostType::ClientOnlyNoCommandlet:
+#if (WITH_ENGINE || WITH_PLUGIN_SUPPORT)  && !IS_PROGRAM
+		return (!IsRunningDedicatedServer()) && (!IsRunningCommandlet());
+#endif
+		// the fall in the case of not having defines listed above is intentional
 	case EHostType::ClientOnly:
 		return !IsRunningDedicatedServer();
-
+	
 	}
-
 	return false;
 }
 
diff --git a/Engine/Source/Runtime/Projects/Public/ModuleDescriptor.h b/Engine/Source/Runtime/Projects/Public/ModuleDescriptor.h
--- a/Engine/Source/Runtime/Projects/Public/ModuleDescriptor.h
+++ b/Engine/Source/Runtime/Projects/Public/ModuleDescriptor.h
@@ -100,6 +100,9 @@ namespace EHostType
 		// Loads on all targets except dedicated servers.
 		ClientOnly,
 
+		// Loads in editor and client but not in commandlets.
+		ClientOnlyNoCommandlet,
+		
 		//~ NOTE: If you add a new value, make sure to update the ToString() method below!
 		Max
 	};
