Added option to VSCode generator to include all files (disabled by default).

diff --git a/Engine/Source/Programs/UnrealBuildTool/ProjectFiles/VisualStudioCode/VSCodeProjectFileGenerator.cs b/Engine/Source/Programs/UnrealBuildTool/ProjectFiles/VisualStudioCode/VSCodeProjectFileGenerator.cs
--- a/Engine/Source/Programs/UnrealBuildTool/ProjectFiles/VisualStudioCode/VSCodeProjectFileGenerator.cs
+++ b/Engine/Source/Programs/UnrealBuildTool/ProjectFiles/VisualStudioCode/VSCodeProjectFileGenerator.cs
@@ -40,6 +40,9 @@ namespace UnrealBuildTool
 		private string FrameworkExecutableExtension;
 		private string FrameworkLibraryExtension = ".dll";
 
+		[XmlConfigFile(Name = "IncludeAllFiles")]
+		private bool IncludeAllFiles = false;
+
 		private enum EPathType
 		{
 			Absolute,
@@ -1200,15 +1203,18 @@ namespace UnrealBuildTool
 
 			OutFile.BeginRootObject();
 			{
-				OutFile.BeginObject("files.exclude");
+				if (!IncludeAllFiles)
 				{
-					string WorkspaceRoot = UnrealBuildTool.RootDirectory.ToString().Replace('\\', '/') + "/";
-					foreach (string PathToExclude in PathsToExclude)
+					OutFile.BeginObject("files.exclude");
 					{
-						OutFile.AddField(PathToExclude.Replace('\\', '/').Replace(WorkspaceRoot, ""), true);
+						string WorkspaceRoot = UnrealBuildTool.RootDirectory.ToString().Replace('\\', '/') + "/";
+						foreach (string PathToExclude in PathsToExclude)
+						{
+							OutFile.AddField(PathToExclude.Replace('\\', '/').Replace(WorkspaceRoot, ""), true);
+						}
 					}
+					OutFile.EndObject();
 				}
-				OutFile.EndObject();
 			}
 			OutFile.EndRootObject();
 
