change DefaultWindowsSdkVersion back to 10.0.16299.0

diff --git a/Engine/Source/Programs/UnrealBuildTool/Platform/Windows/UEBuildWindows.cs b/Engine/Source/Programs/UnrealBuildTool/Platform/Windows/UEBuildWindows.cs
--- a/Engine/Source/Programs/UnrealBuildTool/Platform/Windows/UEBuildWindows.cs
+++ b/Engine/Source/Programs/UnrealBuildTool/Platform/Windows/UEBuildWindows.cs
@@ -565,7 +565,7 @@ namespace UnrealBuildTool
 		/// <summary>
 		/// The default Windows SDK version to be used, if installed.
 		/// </summary>
-		static readonly VersionNumber DefaultWindowsSdkVersion = VersionNumber.Parse("10.0.17763.0");
+		static readonly VersionNumber DefaultWindowsSdkVersion = VersionNumber.Parse("10.0.16299.0");
 
 		/// <summary>
 		/// Cache of Visual Studio installation directories
