UBT: Include additional context when a managed process fails to start.

#rnx

diff --git a/Engine/Source/Programs/DotNETCommon/DotNETUtilities/ManagedProcess.cs b/Engine/Source/Programs/DotNETCommon/DotNETUtilities/ManagedProcess.cs
--- a/Engine/Source/Programs/DotNETCommon/DotNETUtilities/ManagedProcess.cs
+++ b/Engine/Source/Programs/DotNETCommon/DotNETUtilities/ManagedProcess.cs
@@ -453,6 +453,11 @@ namespace Tools.DotNETCommon
 						}
 					}
 				}
+				catch (Exception Ex)
+				{
+					ExceptionUtils.AddContext(Ex, "while launching {0} {1}", FileName, CommandLine);
+					throw;
+				}
 				finally
 				{
 					if(EnvironmentBlock != IntPtr.Zero)
diff --git a/Engine/Source/Programs/UnrealBuildTool/Executors/ParallelExecutor.cs b/Engine/Source/Programs/UnrealBuildTool/Executors/ParallelExecutor.cs
--- a/Engine/Source/Programs/UnrealBuildTool/Executors/ParallelExecutor.cs
+++ b/Engine/Source/Programs/UnrealBuildTool/Executors/ParallelExecutor.cs
@@ -249,10 +249,18 @@ namespace UnrealBuildTool
 				Action.LogLines.Add(Action.Inner.StatusDescription);
 			}
 
-			using(ManagedProcess Process = new ManagedProcess(ProcessGroup, Action.Inner.CommandPath.FullName, Action.Inner.CommandArguments, Action.Inner.WorkingDirectory.FullName, null, null, ProcessPriorityClass.BelowNormal))
+			try
 			{
-				Action.LogLines.AddRange(Process.ReadAllLines());
-				Action.ExitCode = Process.ExitCode;
+				using (ManagedProcess Process = new ManagedProcess(ProcessGroup, Action.Inner.CommandPath.FullName, Action.Inner.CommandArguments, Action.Inner.WorkingDirectory.FullName, null, null, ProcessPriorityClass.BelowNormal))
+				{
+					Action.LogLines.AddRange(Process.ReadAllLines());
+					Action.ExitCode = Process.ExitCode;
+				}
+			}
+			catch(Exception Ex)
+			{
+				Log.WriteException(Ex, null);
+				Action.ExitCode = 1;
 			}
 
 			lock(CompletedActions)
