HTML5 packaging from UE4 on GitHub on Linux

- fix dos2unix line ending issues
- note: this jira still needs `Engine/Binaries/ThirdParty/Python/Linux/bin/python2.7` executable bit set (asked Ben.Marsh to do this...)

#rnx

diff --git a/Engine/Source/Programs/UnrealBuildTool/Platform/HTML5/HTML5SDKInfo.cs b/Engine/Source/Programs/UnrealBuildTool/Platform/HTML5/HTML5SDKInfo.cs
--- a/Engine/Source/Programs/UnrealBuildTool/Platform/HTML5/HTML5SDKInfo.cs
+++ b/Engine/Source/Programs/UnrealBuildTool/Platform/HTML5/HTML5SDKInfo.cs
@@ -458,6 +458,28 @@ namespace UnrealBuildTool
 			Environment.SetEnvironmentVariable("NODE", NODE_JS);
 			Environment.SetEnvironmentVariable("LLVM", LLVM_ROOT);
 
+			// --------------------------------------------------
+			// the following is needed when UE4 from GitHub on Linux is used
+			if (BuildHostPlatform.Current.Platform == UnrealTargetPlatform.Linux)
+			{
+				string findpath = "\"" + Path.Combine(EMSCRIPTEN_ROOT, "system", "lib") + "\"";
+				ProcessStartInfo findproc = new ProcessStartInfo("find", findpath + " -type f -name \"*symbols\" -exec dos2unix {} \\;");
+				findproc.CreateNoWindow = true;
+				findproc.UseShellExecute = false;
+				try
+				{
+					Process process = Process.Start(findproc);
+					process.WaitForExit();
+					Log.TraceInformation("find symbols dos2unix conversions ExitCode: {0}", process.ExitCode);
+					process.Close();
+				}
+				catch (System.ComponentModel.Win32Exception ex)
+				{
+					// Process.Start() as terminated quick enough betore control has returned to process.WaitForExit()
+					Log.TraceInformation("Win32Exception ex.NativeErrorCode: {0}", ex.NativeErrorCode);
+				}
+			}
+
 			return DOT_EMSCRIPTEN;
 		}
 
