UBT: Improve messaging for trying to build a target type or configuration which is not supported from the current engine distribution.

diff --git a/Engine/Source/Programs/UnrealBuildTool/Configuration/UEBuildTarget.cs b/Engine/Source/Programs/UnrealBuildTool/Configuration/UEBuildTarget.cs
--- a/Engine/Source/Programs/UnrealBuildTool/Configuration/UEBuildTarget.cs
+++ b/Engine/Source/Programs/UnrealBuildTool/Configuration/UEBuildTarget.cs
@@ -611,8 +611,29 @@ namespace UnrealBuildTool
 				throw new BuildException("{0} does not support the {1} platform.", Descriptor.Name, Descriptor.Platform.ToString());
 			}
 
+			// Make sure this target type is supported
+			if (!InstalledPlatformInfo.IsValid(RulesObject.Type, Descriptor.Platform, Descriptor.Configuration, EProjectType.Code, InstalledPlatformState.Downloaded))
+			{
+				if (InstalledPlatformInfo.IsValid(RulesObject.Type, Descriptor.Platform, Descriptor.Configuration, EProjectType.Code, InstalledPlatformState.Supported))
+				{
+					throw new BuildException("Download support for building {0} {1} targets from the launcher.", Descriptor.Platform, RulesObject.Type);
+				}
+				else if(!InstalledPlatformInfo.IsValid(RulesObject.Type, null, null, EProjectType.Code, InstalledPlatformState.Supported))
+				{
+					throw new BuildException("{0} targets are not currently supported from this engine distribution.", RulesObject.Type);
+				}
+				else if (!InstalledPlatformInfo.IsValid(RulesObject.Type, RulesObject.Platform, null, EProjectType.Code, InstalledPlatformState.Supported))
+				{
+					throw new BuildException("The {0} platform is not supported from this engine distribution.", RulesObject.Platform);
+				}
+				else
+				{
+					throw new BuildException("Targets cannot be built in the {0} configuration with this engine distribution.", RulesObject.Configuration);
+				}
+			}
+
 			// If we're using the shared build environment, make sure all the settings are valid
-			if(RulesObject.BuildEnvironment == TargetBuildEnvironment.Shared)
+			if (RulesObject.BuildEnvironment == TargetBuildEnvironment.Shared)
 			{
 				ValidateSharedEnvironment(RulesAssembly, Descriptor.Name, RulesObject);
 			}
diff --git a/Engine/Source/Programs/UnrealBuildTool/System/InstalledPlatformInfo.cs b/Engine/Source/Programs/UnrealBuildTool/System/InstalledPlatformInfo.cs
--- a/Engine/Source/Programs/UnrealBuildTool/System/InstalledPlatformInfo.cs
+++ b/Engine/Source/Programs/UnrealBuildTool/System/InstalledPlatformInfo.cs
@@ -36,6 +36,23 @@ namespace UnrealBuildTool
 		Content,
 	};
 
+	/// <summary>
+	/// The state of a downloaded platform
+	/// </summary>
+	[Flags]
+	public enum InstalledPlatformState
+	{
+		/// <summary>
+		/// Query whether the platform is supported
+		/// </summary>
+		Supported,
+
+		/// <summary>
+		/// Query whether the platform has been downloaded
+		/// </summary>
+		Downloaded,
+	}
+
 	/// <summary>
 	/// Contains methods to allow querying the available installed platforms
 	/// </summary>
@@ -305,6 +322,53 @@ namespace UnrealBuildTool
 			);
 		}
 
+		/// <summary>
+		/// Determines whether the given target type is supported
+		/// </summary>
+		/// <param name="TargetType">The target type being built</param>
+		/// <param name="Platform">The platform being built</param>
+		/// <param name="Configuration">The configuration being built</param>
+		/// <param name="ProjectType">The project type required</param>
+		/// <param name="State">State of the given platform support</param>
+		/// <returns>True if the target can be built</returns>
+		public static bool IsValid(TargetType? TargetType, UnrealTargetPlatform? Platform, UnrealTargetConfiguration? Configuration, EProjectType ProjectType, InstalledPlatformState State)
+		{
+			if(!UnrealBuildTool.IsEngineInstalled() || InstalledPlatformConfigurations == null)
+			{
+				return true;
+			}
+
+			foreach(InstalledPlatformConfiguration Config in InstalledPlatformConfigurations)
+			{
+				// Check whether this configuration matches all the criteria
+				if(TargetType.HasValue && Config.PlatformType != TargetType.Value)
+				{
+					continue;
+				}
+				if(Platform.HasValue && Config.Platform != Platform.Value)
+				{
+					continue;
+				}
+				if(Configuration.HasValue && Config.Configuration != Configuration.Value)
+				{
+					continue;
+				}
+				if(ProjectType != EProjectType.Any && Config.ProjectType != EProjectType.Any && Config.ProjectType != ProjectType)
+				{
+					continue;
+				}
+				if(State == InstalledPlatformState.Downloaded && !String.IsNullOrEmpty(Config.RequiredFile) && !File.Exists(Config.RequiredFile))
+				{
+					continue;
+				}
+
+				// Success!
+				return true;
+			}
+
+			return false;
+		}
+
 		private static bool ContainsValidConfiguration(Predicate<InstalledPlatformConfiguration> ConfigFilter)
 		{
 			if (UnrealBuildTool.IsEngineInstalled() && InstalledPlatformConfigurations != null)
