UBT: Add better messages for attempting to build for a platform that hasn't been downloaded from the launcher.

diff --git a/Engine/Source/Programs/UnrealBuildTool/Configuration/TargetDescriptor.cs b/Engine/Source/Programs/UnrealBuildTool/Configuration/TargetDescriptor.cs
--- a/Engine/Source/Programs/UnrealBuildTool/Configuration/TargetDescriptor.cs
+++ b/Engine/Source/Programs/UnrealBuildTool/Configuration/TargetDescriptor.cs
@@ -280,6 +280,19 @@ namespace UnrealBuildTool
 			// Expand all the platforms, architectures and configurations
 			foreach(UnrealTargetPlatform Platform in Platforms)
 			{
+				// Make sure the platform is valid
+				if (!InstalledPlatformInfo.IsValid(null, Platform, null, EProjectType.Code, InstalledPlatformState.Downloaded))
+				{
+					if (!InstalledPlatformInfo.IsValid(null, Platform, null, EProjectType.Code, InstalledPlatformState.Supported))
+					{
+						throw new BuildException("The {0} platform is not supported from this engine distribution.", Platform);
+					}
+					else
+					{
+						throw new BuildException("Missing files required to build {0} targets. Enable {0} as an optional download component in the Epic Games Launcher.", Platform);
+					}
+				}
+
 				// Parse the architecture parameter, or get the default for the platform
 				List<string> Architectures = new List<string>(Arguments.GetValues("-Architecture=", '+'));
 				if(Architectures.Count == 0)
