Add /arm64 and /x64 dirs, and .xml and .pri files to hololens artifacts

diff --git a/Engine/Source/Programs/UnrealBuildTool/Platform/HoloLens/HoloLensToolChain.cs b/Engine/Source/Programs/UnrealBuildTool/Platform/HoloLens/HoloLensToolChain.cs
--- a/Engine/Source/Programs/UnrealBuildTool/Platform/HoloLens/HoloLensToolChain.cs
+++ b/Engine/Source/Programs/UnrealBuildTool/Platform/HoloLens/HoloLensToolChain.cs
@@ -1112,5 +1112,38 @@ namespace UnrealBuildTool
 			return CurrentWindowsSdkVersion.ToString();
 		}
 
+		public override void ModifyBuildProducts(ReadOnlyTargetRules Target, UEBuildBinary Binary, List<string> Libraries, List<UEBuildBundleResource> BundleResources, Dictionary<FileReference, BuildProductType> BuildProducts)
+		{
+			DirectoryReference HoloLensBinaryDirectory = Binary.OutputFilePath.Directory;
+
+			AddBuildProductSafe(BuildProducts, FileReference.Combine(HoloLensBinaryDirectory, "AppxManifest_"+ Target.Architecture + ".xml"), BuildProductType.BuildResource);
+			AddBuildProductSafe(BuildProducts, FileReference.Combine(HoloLensBinaryDirectory, "resources_"+ Target.Architecture + ".pri"), BuildProductType.BuildResource);
+			if (Target.Configuration == UnrealTargetConfiguration.Development)
+			{
+				AddBuildProductSafe(BuildProducts, FileReference.Combine(HoloLensBinaryDirectory, "UE4Game" + Target.Architecture + ".exe"), BuildProductType.Executable);
+				AddBuildProductSafe(BuildProducts, FileReference.Combine(HoloLensBinaryDirectory, "UE4Game" + Target.Architecture + ".pdb"), BuildProductType.SymbolFile);
+			}
+			else
+			{
+				AddBuildProductSafe(BuildProducts, FileReference.Combine(HoloLensBinaryDirectory, "UE4Game-HoloLens-" + Target.Configuration + Target.Architecture + ".exe"), BuildProductType.Executable);
+				AddBuildProductSafe(BuildProducts, FileReference.Combine(HoloLensBinaryDirectory, "UE4Game-HoloLens-" + Target.Configuration + Target.Architecture + ".pdb"), BuildProductType.SymbolFile);
+			}
+			AddBuildProductSafe(BuildProducts, FileReference.Combine(HoloLensBinaryDirectory, "UE4Game-HoloLens-" + Target.Configuration + Target.Architecture + ".target"), BuildProductType.BuildResource);
+			AddBuildProductSafe(BuildProducts, FileReference.Combine(HoloLensBinaryDirectory, Target.Architecture + "\\Resources\\Logo.png"), BuildProductType.BuildResource);
+			AddBuildProductSafe(BuildProducts, FileReference.Combine(HoloLensBinaryDirectory, Target.Architecture + "\\Resources\\resources.resw"), BuildProductType.BuildResource);
+			AddBuildProductSafe(BuildProducts, FileReference.Combine(HoloLensBinaryDirectory, Target.Architecture + "\\Resources\\SmallLogo.png"), BuildProductType.BuildResource);
+			AddBuildProductSafe(BuildProducts, FileReference.Combine(HoloLensBinaryDirectory, Target.Architecture + "\\Resources\\SplashScreen.png"), BuildProductType.BuildResource);
+			AddBuildProductSafe(BuildProducts, FileReference.Combine(HoloLensBinaryDirectory, Target.Architecture + "\\Resources\\WideLogo.png"), BuildProductType.BuildResource);
+			AddBuildProductSafe(BuildProducts, FileReference.Combine(HoloLensBinaryDirectory, Target.Architecture + "\\Resources\\en\\resources.resw"), BuildProductType.BuildResource);
+		}
+
+		private void AddBuildProductSafe(Dictionary<FileReference, BuildProductType> BuildProducts, FileReference FileToAdd, BuildProductType ProductType)
+		{
+			if (!BuildProducts.ContainsKey(FileToAdd))
+			{
+				BuildProducts.Add(FileToAdd, ProductType);
+			}
+		}
+
 	};
 }
