UBT: Fix object files not being built when using the -plugin=... argument and a monolithic target.

#rnx

diff --git a/Engine/Source/Programs/UnrealBuildTool/Configuration/UEBuildBinary.cs b/Engine/Source/Programs/UnrealBuildTool/Configuration/UEBuildBinary.cs
--- a/Engine/Source/Programs/UnrealBuildTool/Configuration/UEBuildBinary.cs
+++ b/Engine/Source/Programs/UnrealBuildTool/Configuration/UEBuildBinary.cs
@@ -698,6 +698,12 @@ namespace UnrealBuildTool
 					Log.TraceVerbose("Compile module: " + Module.Name);
 					LinkInputFiles = Module.Compile(Target, ToolChain, BinaryCompileEnvironment, WorkingSet, Makefile);
 
+					// Save the module outputs. In monolithic builds, this is just the object files.
+					if (Target.LinkType == TargetLinkType.Monolithic)
+					{
+						Makefile.ModuleNameToOutputItems[Module.Name] = LinkInputFiles.ToArray();
+					}
+
 					// NOTE: Because of 'Shared PCHs', in monolithic builds the same PCH file may appear as a link input
 					// multiple times for a single binary.  We'll check for that here, and only add it once.  This avoids
 					// a linker warning about redundant .obj files. 
