Fix for 7708707, shouldn't hardcode QAGame, should use target name.

diff --git a/Engine/Source/Programs/UnrealBuildTool/Platform/HoloLens/HoloLensToolChain.cs b/Engine/Source/Programs/UnrealBuildTool/Platform/HoloLens/HoloLensToolChain.cs
--- a/Engine/Source/Programs/UnrealBuildTool/Platform/HoloLens/HoloLensToolChain.cs
+++ b/Engine/Source/Programs/UnrealBuildTool/Platform/HoloLens/HoloLensToolChain.cs
@@ -1120,15 +1120,15 @@ namespace UnrealBuildTool
 			AddBuildProductSafe(BuildProducts, FileReference.Combine(HoloLensBinaryDirectory, "resources_"+ Target.Architecture + ".pri"), BuildProductType.BuildResource);
 			if (Target.Configuration == UnrealTargetConfiguration.Development)
 			{
-				AddBuildProductSafe(BuildProducts, FileReference.Combine(HoloLensBinaryDirectory, "UE4Game" + Target.Architecture + ".exe"), BuildProductType.Executable);
-				AddBuildProductSafe(BuildProducts, FileReference.Combine(HoloLensBinaryDirectory, "UE4Game" + Target.Architecture + ".pdb"), BuildProductType.SymbolFile);
+				AddBuildProductSafe(BuildProducts, FileReference.Combine(HoloLensBinaryDirectory, Target.Name + Target.Architecture + ".exe"), BuildProductType.Executable);
+				AddBuildProductSafe(BuildProducts, FileReference.Combine(HoloLensBinaryDirectory, Target.Name + Target.Architecture + ".pdb"), BuildProductType.SymbolFile);
 			}
 			else
 			{
-				AddBuildProductSafe(BuildProducts, FileReference.Combine(HoloLensBinaryDirectory, "UE4Game-HoloLens-" + Target.Configuration + Target.Architecture + ".exe"), BuildProductType.Executable);
-				AddBuildProductSafe(BuildProducts, FileReference.Combine(HoloLensBinaryDirectory, "UE4Game-HoloLens-" + Target.Configuration + Target.Architecture + ".pdb"), BuildProductType.SymbolFile);
+				AddBuildProductSafe(BuildProducts, FileReference.Combine(HoloLensBinaryDirectory, Target.Name + "-HoloLens-" + Target.Configuration + Target.Architecture + ".exe"), BuildProductType.Executable);
+				AddBuildProductSafe(BuildProducts, FileReference.Combine(HoloLensBinaryDirectory, Target.Name + "-HoloLens-" + Target.Configuration + Target.Architecture + ".pdb"), BuildProductType.SymbolFile);
 			}
-			AddBuildProductSafe(BuildProducts, FileReference.Combine(HoloLensBinaryDirectory, "UE4Game-HoloLens-" + Target.Configuration + Target.Architecture + ".target"), BuildProductType.BuildResource);
+			AddBuildProductSafe(BuildProducts, FileReference.Combine(HoloLensBinaryDirectory, Target.Name + "-HoloLens-" + Target.Configuration + Target.Architecture + ".target"), BuildProductType.BuildResource);
 			AddBuildProductSafe(BuildProducts, FileReference.Combine(HoloLensBinaryDirectory, Target.Architecture + "\\Resources\\Logo.png"), BuildProductType.BuildResource);
 			AddBuildProductSafe(BuildProducts, FileReference.Combine(HoloLensBinaryDirectory, Target.Architecture + "\\Resources\\resources.resw"), BuildProductType.BuildResource);
 			AddBuildProductSafe(BuildProducts, FileReference.Combine(HoloLensBinaryDirectory, Target.Architecture + "\\Resources\\SmallLogo.png"), BuildProductType.BuildResource);
