UBT: Fix the path to the solution options file for VS2019, so we can collapse all the projects and select the right configuration by default.

diff --git a/Engine/Source/Programs/UnrealBuildTool/ProjectFiles/VisualStudio/VCProjectFileGenerator.cs b/Engine/Source/Programs/UnrealBuildTool/ProjectFiles/VisualStudio/VCProjectFileGenerator.cs
--- a/Engine/Source/Programs/UnrealBuildTool/ProjectFiles/VisualStudio/VCProjectFileGenerator.cs
+++ b/Engine/Source/Programs/UnrealBuildTool/ProjectFiles/VisualStudio/VCProjectFileGenerator.cs
@@ -777,7 +777,7 @@ namespace UnrealBuildTool
 						SolutionOptionsFileName = FileReference.Combine(MasterProjectPath, ".vs", Path.GetFileNameWithoutExtension(SolutionFileName), "v15", ".suo");
 						break;
 					case VCProjectFileFormat.VisualStudio2019:
-						SolutionOptionsFileName = FileReference.Combine(MasterProjectPath, ".vs", Path.GetFileNameWithoutExtension(SolutionFileName), "v16", ".suo");
+						SolutionOptionsFileName = FileReference.Combine(MasterProjectPath, ".vs", Path.GetFileNameWithoutExtension(SolutionFileName), "v15", ".suo"); // Still uses v15
 						break;
 					default:
 						throw new BuildException("Unsupported Visual Studio version");
