Changed order in which modules are mapped to visual studio projects
It is now Game -> Mod -> Program -> Engine -> Enterprise

This resolves issues were a Program in the game source root would take all the source code away from the Game project.
In this case now though the Program source is put into the Game project, it is best to have program target in a Programs folder under the game source rather then at the same level as the game targets.

diff --git a/Engine/Source/Programs/UnrealBuildTool/ProjectFiles/ProjectFileGenerator.cs b/Engine/Source/Programs/UnrealBuildTool/ProjectFiles/ProjectFileGenerator.cs
--- a/Engine/Source/Programs/UnrealBuildTool/ProjectFiles/ProjectFileGenerator.cs
+++ b/Engine/Source/Programs/UnrealBuildTool/ProjectFiles/ProjectFileGenerator.cs
@@ -1986,8 +1986,38 @@ namespace UnrealBuildTool
             }
         }
 
+
 		private ProjectFile FindProjectForModule(FileReference CurModuleFile, List<FileReference> AllGames, Dictionary<FileReference, ProjectFile> ProgramProjects, List<ProjectFile> ModProjects, out DirectoryReference BaseFolder)
 		{
+			ProjectFile FindOrAddProjectHelper(string InProjectFileNameBase, DirectoryReference InBaseFolder)
+			{
+				// Setup a project file entry for this module's project.  Remember, some projects may host multiple modules!
+				FileReference ProjectFileName = FileReference.Combine(IntermediateProjectFilesPath, InProjectFileNameBase + ProjectFileExtension);
+				bool bProjectAlreadyExisted;
+				return FindOrAddProject(ProjectFileName, InBaseFolder, IncludeInGeneratedProjects: true, bAlreadyExisted: out bProjectAlreadyExisted);
+			}
+
+			string ProjectFileNameBase = null;
+			BaseFolder = null;
+
+			// Figure out which game project this target belongs to
+			FileReference ProjectInfo = FindGameContainingFile(AllGames, CurModuleFile);
+			if (ProjectInfo != null)
+			{
+				BaseFolder = ProjectInfo.Directory;
+				return FindOrAddProjectHelper(ProjectInfo.GetFileNameWithoutExtension(), BaseFolder);
+			}
+
+			// Check if it's a mod
+			foreach (ProjectFile ModProject in ModProjects)
+			{
+				if (CurModuleFile.IsUnderDirectory(ModProject.BaseDir))
+				{
+					BaseFolder = ModProject.BaseDir;
+					return ModProject;
+				}
+			}
+
 			// Check if this module is under any program project base folder
 			if (ProgramProjects != null)
 			{
@@ -2001,12 +2031,8 @@ namespace UnrealBuildTool
 				}
 			}
 
-			string ProjectFileNameBase = null;
-
-			string PossibleProgramTargetName = CurModuleFile.GetFileNameWithoutAnyExtensions();
-
 			// check for engine, or platform extension engine folders
-			if( CurModuleFile.IsUnderDirectory(UnrealBuildTool.EngineDirectory))
+			if ( CurModuleFile.IsUnderDirectory(UnrealBuildTool.EngineDirectory))
 			{
 				ProjectFileNameBase = EngineProjectFileNameBase;
 				BaseFolder = UnrealBuildTool.EngineDirectory;
@@ -2017,32 +2043,13 @@ namespace UnrealBuildTool
 				ProjectFileNameBase = EnterpriseProjectFileNameBase;
 				BaseFolder = UnrealBuildTool.EnterpriseDirectory;
 			}
-			else
-			{
-				// Check if it's a mod
-				foreach(ProjectFile ModProject in ModProjects)
-				{
-					if(CurModuleFile.IsUnderDirectory(ModProject.BaseDir))
-					{
-						BaseFolder = ModProject.BaseDir;
-						return ModProject;
-					}
-				}
 
-				// Figure out which game project this target belongs to
-				FileReference ProjectInfo = FindGameContainingFile(AllGames, CurModuleFile);
-				if(ProjectInfo == null)
-				{
-					throw new BuildException( "Found a non-engine module file (" + CurModuleFile + ") that did not exist within any of the known game folders" );
-				}
-				BaseFolder = ProjectInfo.Directory;
-				ProjectFileNameBase = ProjectInfo.GetFileNameWithoutExtension();
+			if (ProjectInfo == null)
+			{
+				throw new BuildException("Found a non-engine module file (" + CurModuleFile + ") that did not exist within any of the known game folders");
 			}
 
-			// Setup a project file entry for this module's project.  Remember, some projects may host multiple modules!
-			FileReference ProjectFileName = FileReference.Combine( IntermediateProjectFilesPath, ProjectFileNameBase + ProjectFileExtension );
-			bool bProjectAlreadyExisted;
-			return FindOrAddProject( ProjectFileName, BaseFolder, IncludeInGeneratedProjects:true, bAlreadyExisted:out bProjectAlreadyExisted );
+			return FindOrAddProjectHelper(ProjectFileNameBase, BaseFolder);
 		}
 
 		/// <summary>
