Fixed incorrect check for non enigne modules which broke in 7839052.
Also changed local function to a member to fix build issue.

diff --git a/Engine/Source/Programs/UnrealBuildTool/ProjectFiles/ProjectFileGenerator.cs b/Engine/Source/Programs/UnrealBuildTool/ProjectFiles/ProjectFileGenerator.cs
--- a/Engine/Source/Programs/UnrealBuildTool/ProjectFiles/ProjectFileGenerator.cs
+++ b/Engine/Source/Programs/UnrealBuildTool/ProjectFiles/ProjectFileGenerator.cs
@@ -1987,15 +1987,16 @@ namespace UnrealBuildTool
         }
 
 
+		private ProjectFile FindOrAddProjectHelper(string InProjectFileNameBase, DirectoryReference InBaseFolder)
+		{
+			// Setup a project file entry for this module's project.  Remember, some projects may host multiple modules!
+			FileReference ProjectFileName = FileReference.Combine(IntermediateProjectFilesPath, InProjectFileNameBase + ProjectFileExtension);
+			bool bProjectAlreadyExisted;
+			return FindOrAddProject(ProjectFileName, InBaseFolder, IncludeInGeneratedProjects: true, bAlreadyExisted: out bProjectAlreadyExisted);
+		}
+
 		private ProjectFile FindProjectForModule(FileReference CurModuleFile, List<FileReference> AllGames, Dictionary<FileReference, ProjectFile> ProgramProjects, List<ProjectFile> ModProjects, out DirectoryReference BaseFolder)
 		{
-			ProjectFile FindOrAddProjectHelper(string InProjectFileNameBase, DirectoryReference InBaseFolder)
-			{
-				// Setup a project file entry for this module's project.  Remember, some projects may host multiple modules!
-				FileReference ProjectFileName = FileReference.Combine(IntermediateProjectFilesPath, InProjectFileNameBase + ProjectFileExtension);
-				bool bProjectAlreadyExisted;
-				return FindOrAddProject(ProjectFileName, InBaseFolder, IncludeInGeneratedProjects: true, bAlreadyExisted: out bProjectAlreadyExisted);
-			}
 
 			string ProjectFileNameBase = null;
 			BaseFolder = null;
@@ -2044,7 +2045,7 @@ namespace UnrealBuildTool
 				BaseFolder = UnrealBuildTool.EnterpriseDirectory;
 			}
 
-			if (ProjectInfo == null)
+			if (ProjectInfo == null && ProjectFileNameBase == null)
 			{
 				throw new BuildException("Found a non-engine module file (" + CurModuleFile + ") that did not exist within any of the known game folders");
 			}
