Detect pak encryption being enabled when no keys have been generated. Issue a warning and disable encryption.

diff --git a/Engine/Source/Programs/UnrealBuildTool/System/EncryptionAndSigning.cs b/Engine/Source/Programs/UnrealBuildTool/System/EncryptionAndSigning.cs
--- a/Engine/Source/Programs/UnrealBuildTool/System/EncryptionAndSigning.cs
+++ b/Engine/Source/Programs/UnrealBuildTool/System/EncryptionAndSigning.cs
@@ -449,7 +449,9 @@ namespace UnrealBuildTool
 				Settings.bEnablePakSigning = false;
 			}
 
-			if (Settings.bDataCryptoRequired && Settings.IsAnyEncryptionEnabled() && (Settings.EncryptionKey == null || !Settings.EncryptionKey.IsValid()))
+			bool bEncryptionKeyValid = (Settings.EncryptionKey != null && Settings.EncryptionKey.IsValid());
+			bool bAnyEncryptionRequested = Settings.bEnablePakFullAssetEncryption || Settings.bEnablePakIndexEncryption || Settings.bEnablePakIniEncryption || Settings.bEnablePakUAssetEncryption;
+			if (Settings.bDataCryptoRequired && bAnyEncryptionRequested && !bEncryptionKeyValid)
 			{
 				Log.TraceWarningOnce("Pak encryption is enabled, but no valid encryption key was found. Please generate a key in the editor project crypto settings. Encryption will be disabled");
 				Settings.bEnablePakUAssetEncryption = false;
