Live Coding: Allow building foreign plugins while Live Coding is active.

diff --git a/Engine/Source/Programs/UnrealBuildTool/Modes/BuildMode.cs b/Engine/Source/Programs/UnrealBuildTool/Modes/BuildMode.cs
--- a/Engine/Source/Programs/UnrealBuildTool/Modes/BuildMode.cs
+++ b/Engine/Source/Programs/UnrealBuildTool/Modes/BuildMode.cs
@@ -517,7 +517,7 @@ namespace UnrealBuildTool
 			}
 
 			// Guard against a live coding session for this target being active
-			if(HotReloadMode != HotReloadMode.LiveCoding && HotReload.IsLiveCodingSessionActive(Makefile))
+			if(HotReloadMode != HotReloadMode.LiveCoding && TargetDescriptor.ForeignPlugin == null && HotReload.IsLiveCodingSessionActive(Makefile))
 			{
 				throw new BuildException("Unable to start regular build while Live Coding is active. Press Ctrl+Alt+F11 to trigger a Live Coding compile.");
 			}
