Add warning when referencing a module without a fully qualifed path, doing so will break the incremental building of it.
Added SystemLibrary list you can use if you actually need to add a module without a fully qualified path and dont want to see this warning (should not be needed in most cases)

Also updated the android toolchain to filter paths to a library in the same way a additional library search path is filtered enabling us to just specify the full path to multiple libraries for different architectures.

diff --git a/Engine/Source/Programs/UnrealBuildTool/Configuration/ModuleRules.cs b/Engine/Source/Programs/UnrealBuildTool/Configuration/ModuleRules.cs
--- a/Engine/Source/Programs/UnrealBuildTool/Configuration/ModuleRules.cs
+++ b/Engine/Source/Programs/UnrealBuildTool/Configuration/ModuleRules.cs
@@ -682,7 +682,14 @@ namespace UnrealBuildTool
 		/// <summary>
 		/// List of system/library paths (directory of .lib files) - typically used for External (third party) modules
 		/// </summary>
-		public List<string> PublicLibraryPaths = new List<string>();
+		[Obsolete(
+			"For external libraries use the full path in PublicAdditionalLibraries, if its a system library then use PublicSystemLibraries/PublicSystemLibraryPaths")]
+		public List<string> PublicLibraryPaths => PublicSystemLibraryPaths;
+
+		/// <summary>
+		/// List of system library paths (directory of .lib files) - for External (third party) modules please use the PublicAdditionalLibaries instead
+		/// </summary>
+		public List<string> PublicSystemLibraryPaths = new List<string>();
 
 		/// <summary>
 		/// List of search paths for libraries at runtime (eg. .so files)
@@ -699,6 +706,11 @@ namespace UnrealBuildTool
 		/// </summary>
 		public List<string> PublicAdditionalLibraries = new List<string>();
 
+		/// <summary>
+		/// List of system libraries to use - these are typically referenced via name and then found via the system paths. If you need to reference a .lib file use the PublicAdditionalLibraries instead
+		/// </summary>
+		public List<string> PublicSystemLibraries = new List<string>();
+
 		/// <summary>
 		/// List of XCode frameworks (iOS and MacOS)
 		/// </summary>
diff --git a/Engine/Source/Programs/UnrealBuildTool/Configuration/UEBuildModule.cs b/Engine/Source/Programs/UnrealBuildTool/Configuration/UEBuildModule.cs
--- a/Engine/Source/Programs/UnrealBuildTool/Configuration/UEBuildModule.cs
+++ b/Engine/Source/Programs/UnrealBuildTool/Configuration/UEBuildModule.cs
@@ -97,15 +97,20 @@ namespace UnrealBuildTool
 		protected readonly HashSet<DirectoryReference> PublicSystemIncludePaths;
 
 		/// <summary>
-		/// Set of all public library paths
+		/// Set of all public system library paths
 		/// </summary>
-		protected readonly HashSet<DirectoryReference> PublicLibraryPaths;
+		protected readonly HashSet<DirectoryReference> PublicSystemLibraryPaths;
 
 		/// <summary>
 		/// Set of all additional libraries
 		/// </summary>
 		protected readonly HashSet<string> PublicAdditionalLibraries;
 
+		/// <summary>
+		/// Set of all system libraries
+		/// </summary>
+		protected readonly HashSet<string> PublicSystemLibraries;
+
 		/// <summary>
 		/// Set of additional frameworks
 		/// </summary>
@@ -175,11 +180,24 @@ namespace UnrealBuildTool
 			PublicDefinitions = HashSetFromOptionalEnumerableStringParameter(Rules.PublicDefinitions);
 			PublicIncludePaths = CreateDirectoryHashSet(Rules.PublicIncludePaths);
 			PublicSystemIncludePaths = CreateDirectoryHashSet(Rules.PublicSystemIncludePaths);
-			PublicLibraryPaths = CreateDirectoryHashSet(Rules.PublicLibraryPaths);
+			PublicSystemLibraryPaths = CreateDirectoryHashSet(Rules.PublicSystemLibraryPaths);
 			PublicAdditionalLibraries = HashSetFromOptionalEnumerableStringParameter(Rules.PublicAdditionalLibraries);
+			PublicSystemLibraries = HashSetFromOptionalEnumerableStringParameter(Rules.PublicSystemLibraries);
 			PublicFrameworks = HashSetFromOptionalEnumerableStringParameter(Rules.PublicFrameworks);
 			PublicWeakFrameworks = HashSetFromOptionalEnumerableStringParameter(Rules.PublicWeakFrameworks);
 
+			foreach (string LibraryName in PublicAdditionalLibraries)
+			{
+				// if the library path is fully qualified we just add it, this is the preferred method of adding a library
+				if (File.Exists(LibraryName))
+				{
+					continue;
+				}
+
+				// the library path does not seem to be resolvable as is, lets warn about it as dependency checking will not work for it
+				Log.TraceWarning("Library '{0}' was not resolvable to a file when used in Module '{1}', assuming it is a filename and will search library paths for it. This is slow and dependency checking will not work for it. Please update reference to be fully qualified alternatively use PublicSystemLibraryPaths if you do intended to use this slow path to suppress this warning. ", LibraryName, Name);
+			}
+
 			PublicAdditionalFrameworks = new HashSet<UEBuildFramework>();
 			if(Rules.PublicAdditionalFrameworks != null)
 			{
@@ -312,6 +330,7 @@ namespace UnrealBuildTool
 		public void GatherAdditionalResources(List<string> Libraries, List<UEBuildBundleResource> BundleResources)
 		{
 			Libraries.AddRange(PublicAdditionalLibraries);
+			Libraries.AddRange(PublicSystemLibraries);
 			BundleResources.AddRange(PublicAdditionalBundleResources);
 		}
 
@@ -387,9 +406,9 @@ namespace UnrealBuildTool
 			{
 				Directories.Add(PublicSystemIncludePath);
 			}
-			foreach(DirectoryReference PublicLibraryPath in PublicLibraryPaths)
+			foreach (DirectoryReference PublicSystemLibraryPath in PublicSystemLibraryPaths)
 			{
-				Directories.Add(PublicLibraryPath);
+				Directories.Add(PublicSystemLibraryPath);
 			}
 		}
 
@@ -706,8 +725,9 @@ namespace UnrealBuildTool
 				}
 
 				// Add this module's public include library paths and additional libraries.
-				LibraryPaths.AddRange(PublicLibraryPaths);
+				LibraryPaths.AddRange(PublicSystemLibraryPaths);
 				AdditionalLibraries.AddRange(PublicAdditionalLibraries);
+				AdditionalLibraries.AddRange(PublicSystemLibraries);
 				RuntimeLibraryPaths.AddRange(ExpandPathVariables(Rules.PublicRuntimeLibraryPaths, SourceBinary.OutputDir, ExeDir));
 				Frameworks.AddRange(PublicFrameworks);
 				WeakFrameworks.AddRange(PublicWeakFrameworks);
@@ -883,8 +903,9 @@ namespace UnrealBuildTool
 			ExportJsonStringArray(Writer, "PublicSystemIncludePaths", PublicSystemIncludePaths.Select(x => x.FullName));
 			ExportJsonStringArray(Writer, "PublicIncludePaths", PublicIncludePaths.Select(x => x.FullName));
 			ExportJsonStringArray(Writer, "PrivateIncludePaths", PrivateIncludePaths.Select(x => x.FullName));
-			ExportJsonStringArray(Writer, "PublicLibraryPaths", PublicLibraryPaths.Select(x => x.FullName));
+			ExportJsonStringArray(Writer, "PublicSystemLibraryPaths", PublicSystemLibraryPaths.Select(x => x.FullName));
 			ExportJsonStringArray(Writer, "PublicAdditionalLibraries", PublicAdditionalLibraries);
+			ExportJsonStringArray(Writer, "PublicSystemLibraries", PublicSystemLibraries);
 			ExportJsonStringArray(Writer, "PublicFrameworks", PublicFrameworks);
 			ExportJsonStringArray(Writer, "PublicWeakFrameworks", PublicWeakFrameworks);
 			ExportJsonStringArray(Writer, "PublicDelayLoadDLLs", PublicDelayLoadDLLs);
diff --git a/Engine/Source/Programs/UnrealBuildTool/Configuration/UEBuildTarget.cs b/Engine/Source/Programs/UnrealBuildTool/Configuration/UEBuildTarget.cs
--- a/Engine/Source/Programs/UnrealBuildTool/Configuration/UEBuildTarget.cs
+++ b/Engine/Source/Programs/UnrealBuildTool/Configuration/UEBuildTarget.cs
@@ -1157,7 +1157,10 @@ namespace UnrealBuildTool
 				// Get a list of all the library paths
 				List<string> LibraryPaths = new List<string>();
 				LibraryPaths.Add(Directory.GetCurrentDirectory());
-				LibraryPaths.AddRange(Rules.PublicLibraryPaths.Where(x => !x.StartsWith("$(")).Select(x => Path.GetFullPath(x.Replace('/', Path.DirectorySeparatorChar))));
+
+				List<string> SystemLibraryPaths = new List<string>();
+				SystemLibraryPaths.Add(Directory.GetCurrentDirectory());
+				SystemLibraryPaths.AddRange(Rules.PublicSystemLibraryPaths.Where(x => !x.StartsWith("$(")).Select(x => Path.GetFullPath(x.Replace('/', Path.DirectorySeparatorChar))));
 
 				// Get all the extensions to look for
 				List<string> LibraryExtensions = new List<string>();
@@ -1171,20 +1174,15 @@ namespace UnrealBuildTool
 					{
 						foreach (string LibraryPath in LibraryPaths)
 						{
-							string LibraryFileName = Path.Combine(LibraryPath, LibraryName);
-							if (File.Exists(LibraryFileName))
-							{
-								Files.Add(new FileReference(LibraryFileName));
-							}
+							ResolveLibraryName(LibraryPath, LibraryName, LibraryExtension, Files);
+						}
+					}
 
-							if(LibraryName.IndexOfAny(new char[] { Path.DirectorySeparatorChar, Path.AltDirectorySeparatorChar }) == -1)
-							{
-								string UnixLibraryFileName = Path.Combine(LibraryPath, "lib" + LibraryName + LibraryExtension);
-								if (File.Exists(UnixLibraryFileName))
-								{
-									Files.Add(new FileReference(UnixLibraryFileName));
-								}
-							}
+					foreach (string LibraryName in Rules.PublicSystemLibraryPaths)
+					{
+						foreach (string LibraryPath in SystemLibraryPaths)
+						{
+							ResolveLibraryName(LibraryPath, LibraryName, LibraryExtension, Files);
 						}
 					}
 				}
@@ -1217,6 +1215,23 @@ namespace UnrealBuildTool
 			FileReference.WriteAllLines(Location, Files.Where(x => x.IsUnderDirectory(UnrealBuildTool.RootDirectory)).Select(x => x.MakeRelativeTo(UnrealBuildTool.RootDirectory).Replace(Path.DirectorySeparatorChar, '/')).OrderBy(x => x));
 		}
 
+		private void ResolveLibraryName(string LibraryPath, string LibraryName, string LibraryExtension, HashSet<FileReference> Files)
+		{
+			string LibraryFileName = Path.Combine(LibraryPath, LibraryName);
+			if (File.Exists(LibraryFileName))
+			{
+				Files.Add(new FileReference(LibraryFileName));
+			}
+
+			if (LibraryName.IndexOfAny(new char[] { Path.DirectorySeparatorChar, Path.AltDirectorySeparatorChar }) == -1)
+			{
+				string UnixLibraryFileName = Path.Combine(LibraryPath, "lib" + LibraryName + LibraryExtension);
+				if (File.Exists(UnixLibraryFileName))
+				{
+					Files.Add(new FileReference(UnixLibraryFileName));
+				}
+			}
+		}
 		/// <summary>
 		/// Generates a public manifest file for writing out
 		/// </summary>
@@ -3597,7 +3612,7 @@ namespace UnrealBuildTool
 			RemoveTrailingSlashes(RulesObject.PublicIncludePaths);
 			RemoveTrailingSlashes(RulesObject.PublicSystemIncludePaths);
 			RemoveTrailingSlashes(RulesObject.PrivateIncludePaths);
-			RemoveTrailingSlashes(RulesObject.PublicLibraryPaths);
+			RemoveTrailingSlashes(RulesObject.PublicSystemLibraryPaths);
 
 			// Validate rules object
 			if (RulesObject.Type == ModuleRules.ModuleType.CPlusPlus)
@@ -3714,7 +3729,7 @@ namespace UnrealBuildTool
 					{
 						RulesObject.PublicIncludePaths = CombinePathList(ProjectSourceDirectoryName, RulesObject.PublicIncludePaths);
 						RulesObject.PrivateIncludePaths = CombinePathList(ProjectSourceDirectoryName, RulesObject.PrivateIncludePaths);
-						RulesObject.PublicLibraryPaths = CombinePathList(ProjectSourceDirectoryName, RulesObject.PublicLibraryPaths);
+						RulesObject.PublicSystemLibraryPaths = CombinePathList(ProjectSourceDirectoryName, RulesObject.PublicSystemLibraryPaths);
 					}
 				}
 
diff --git a/Engine/Source/Programs/UnrealBuildTool/Platform/Android/AndroidToolChain.cs b/Engine/Source/Programs/UnrealBuildTool/Platform/Android/AndroidToolChain.cs
--- a/Engine/Source/Programs/UnrealBuildTool/Platform/Android/AndroidToolChain.cs
+++ b/Engine/Source/Programs/UnrealBuildTool/Platform/Android/AndroidToolChain.cs
@@ -1008,6 +1008,12 @@ namespace UnrealBuildTool
 				}
 			}
 
+			// apply the same directory filtering to libraries as we do to additional library paths
+			if (!IsDirectoryForArch(Path.GetDirectoryName(Lib), Arch))
+			{
+				return true;
+			}
+
 			// if another architecture is in the filename, reject it
 			foreach (string ComboName in AllComboNames)
 			{
diff --git a/Engine/Source/Programs/UnrealBuildTool/Platform/HoloLens/UEBuildHoloLens.cs b/Engine/Source/Programs/UnrealBuildTool/Platform/HoloLens/UEBuildHoloLens.cs
--- a/Engine/Source/Programs/UnrealBuildTool/Platform/HoloLens/UEBuildHoloLens.cs
+++ b/Engine/Source/Programs/UnrealBuildTool/Platform/HoloLens/UEBuildHoloLens.cs
@@ -478,7 +478,8 @@ namespace UnrealBuildTool
 				// The HoloLensSDK module handles proper paths and libs for HoloLens.
 				// However, the D3D11RHI module will include the DX11 module.
 				Rules.PublicIncludePaths.Clear();
-				Rules.PublicLibraryPaths.Clear();
+				Rules.PublicSystemLibraryPaths.Clear();
+				Rules.PublicSystemLibraries.Clear();
 				Rules.PublicAdditionalLibraries.Clear();
 				Rules.PublicDefinitions.Remove("WITH_D3DX_LIBS=1");
 				Rules.PublicDefinitions.Add("WITH_D3DX_LIBS=0");
