- Revamped the Platform Extension plugin support to go back to just passing one directory to ReadPluginsFromDirectory now that the directory swap allows it

diff --git a/Engine/Platforms/XXX/Plugins/CableComponent/CableComponent_XXX.uplugin b/Engine/Platforms/XXX/Plugins/CableComponent/CableComponent_XXX.uplugin
new file mode 100644
--- /dev/null
+++ b/Engine/Platforms/XXX/Plugins/CableComponent/CableComponent_XXX.uplugin
@@ -0,0 +1,4 @@
+{
+	"FileVersion": 3,
+	"bIsPluginExtension": true
+}
\ No newline at end of file
diff --git a/Engine/Platforms/XXX/Plugins/CableComponent/Source/CableComponent/CableComponent_XXX.Build.cs b/Engine/Platforms/XXX/Plugins/CableComponent/Source/CableComponent/CableComponent_XXX.Build.cs
new file mode 100644
--- /dev/null
+++ b/Engine/Platforms/XXX/Plugins/CableComponent/Source/CableComponent/CableComponent_XXX.Build.cs
@@ -0,0 +1,12 @@
+// Copyright 1998-2019 Epic Games, Inc. All Rights Reserved.
+
+namespace UnrealBuildTool.Rules
+{
+	public class CableComponent_XXX : CableComponent
+	{
+        public CableComponent_XXX(ReadOnlyTargetRules Target) : base(Target)
+		{
+
+		}
+	}
+}
diff --git a/Engine/Source/Programs/AutomationTool/AutomationUtils/CommandUtils.cs b/Engine/Source/Programs/AutomationTool/AutomationUtils/CommandUtils.cs
--- a/Engine/Source/Programs/AutomationTool/AutomationUtils/CommandUtils.cs
+++ b/Engine/Source/Programs/AutomationTool/AutomationUtils/CommandUtils.cs
@@ -1975,6 +1975,36 @@ namespace AutomationTool
 		/// </summary>
 		public static readonly DirectoryReference EngineDirectory = DirectoryReference.Combine(RootDirectory, "Engine");
 
+		/// <summary>
+		/// Return the main engine directory and any platform extension engine directories
+		/// </summary>
+ 		public static DirectoryReference[] GetAllEngineDirectories()
+		{
+			List<DirectoryReference> EngineDirectories = new List<DirectoryReference>() { EngineDirectory };
+			DirectoryReference EnginePlatformsDirectory = DirectoryReference.Combine(EngineDirectory, "Platforms");
+			if (DirectoryReference.Exists(EnginePlatformsDirectory))
+			{
+				EngineDirectories.AddRange(DirectoryReference.EnumerateDirectories(EnginePlatformsDirectory).ToList());
+			}
+
+			return EngineDirectories.ToArray();
+		}
+
+		/// <summary>
+		/// Return the main project directory and any platform extension project directories
+		/// </summary>
+		public static DirectoryReference[] GetAllProjectDirectories(FileReference ProjectFile)
+		{
+			List<DirectoryReference> ProjectDirectories = new List<DirectoryReference>() { ProjectFile.Directory };
+			DirectoryReference ProjectPlatformsDirectory = DirectoryReference.Combine(ProjectFile.Directory, "Platforms");
+			if (DirectoryReference.Exists(ProjectPlatformsDirectory))
+			{
+				ProjectDirectories.AddRange(DirectoryReference.EnumerateDirectories(ProjectPlatformsDirectory).ToList());
+			}
+
+			return ProjectDirectories.ToArray();
+		}
+
 		/// <summary>
 		/// Telemetry data for the current run. Add -WriteTelemetry=<Path> to the command line to export to disk.
 		/// </summary>
diff --git a/Engine/Source/Programs/AutomationTool/Scripts/Localisation.Automation.cs b/Engine/Source/Programs/AutomationTool/Scripts/Localisation.Automation.cs
--- a/Engine/Source/Programs/AutomationTool/Scripts/Localisation.Automation.cs
+++ b/Engine/Source/Programs/AutomationTool/Scripts/Localisation.Automation.cs
@@ -136,8 +136,8 @@ class Localize : BuildCommand
 		// Build up any additional batches needed for plugins
 		if (ShouldGatherPlugins)
 		{
-			var PluginsRootDirectory = CombinePaths(UEProjectRoot, UEProjectDirectory, "Plugins");
-			IReadOnlyList<PluginInfo> AllPlugins = Plugins.ReadPluginsFromDirectory(new DirectoryReference(PluginsRootDirectory), UEProjectName.Length == 0 ? PluginType.Engine : PluginType.Project);
+			var PluginsRootDirectory = CombinePaths(UEProjectRoot, UEProjectDirectory);
+			IReadOnlyList<PluginInfo> AllPlugins = Plugins.ReadPluginsFromDirectory(new DirectoryReference(PluginsRootDirectory), "Plugins", UEProjectName.Length == 0 ? PluginType.Engine : PluginType.Project);
 
 			// Add a batch for each plugin that meets our criteria
 			var AvailablePluginNames = new HashSet<string>();
diff --git a/Engine/Source/Programs/UnrealBuildTool/Configuration/UEBuildPlugin.cs b/Engine/Source/Programs/UnrealBuildTool/Configuration/UEBuildPlugin.cs
--- a/Engine/Source/Programs/UnrealBuildTool/Configuration/UEBuildPlugin.cs
+++ b/Engine/Source/Programs/UnrealBuildTool/Configuration/UEBuildPlugin.cs
@@ -1,4 +1,4 @@
-﻿// Copyright 1998-2019 Epic Games, Inc. All Rights Reserved.
+// Copyright 1998-2019 Epic Games, Inc. All Rights Reserved.
 
 using System;
 using System.Collections.Generic;
@@ -64,6 +64,14 @@ namespace UnrealBuildTool
 			get { return Info.File; }
 		}
 
+		/// <summary>
+		/// Accessor for the child files for this plugin
+		/// </summary>
+		public List<FileReference> ChildFiles
+		{
+			get { return Info.ChildFiles; }
+		}
+
 		/// <summary>
 		/// Accessor for the type of the plugin
 		/// </summary>
diff --git a/Engine/Source/Programs/UnrealBuildTool/Configuration/UEBuildTarget.cs b/Engine/Source/Programs/UnrealBuildTool/Configuration/UEBuildTarget.cs
--- a/Engine/Source/Programs/UnrealBuildTool/Configuration/UEBuildTarget.cs
+++ b/Engine/Source/Programs/UnrealBuildTool/Configuration/UEBuildTarget.cs
@@ -1378,6 +1378,10 @@ namespace UnrealBuildTool
 				if(EnabledPlugin.bDescriptorNeededAtRuntime || EnabledPlugin.bDescriptorReferencedExplicitly)
 				{
 					Receipt.RuntimeDependencies.Add(EnabledPlugin.File, StagedFileType.UFS);
+					foreach (FileReference ChildFile in EnabledPlugin.ChildFiles)
+					{
+						Receipt.RuntimeDependencies.Add(ChildFile, StagedFileType.UFS);
+					}
 				}
 			}
 
diff --git a/Engine/Source/Programs/UnrealBuildTool/ProjectFiles/ProjectFileGenerator.cs b/Engine/Source/Programs/UnrealBuildTool/ProjectFiles/ProjectFileGenerator.cs
--- a/Engine/Source/Programs/UnrealBuildTool/ProjectFiles/ProjectFileGenerator.cs
+++ b/Engine/Source/Programs/UnrealBuildTool/ProjectFiles/ProjectFileGenerator.cs
@@ -2330,26 +2330,23 @@ namespace UnrealBuildTool
 			if(GameProjects.Count == 1)
 			{
 				ProjectFile GameProject = GameProjects.First();
-				foreach (DirectoryReference ProjectDirectory in UnrealBuildTool.GetAllProjectDirectories(GameProject.BaseDir))
+				foreach(PluginInfo PluginInfo in Plugins.ReadProjectPlugins(GameProject.BaseDir))
 				{
-					foreach (PluginInfo PluginInfo in Plugins.ReadProjectPlugins(ProjectDirectory))
+					if (PluginInfo.Descriptor.Modules != null && PluginInfo.Descriptor.Modules.Length > 0 && PluginInfo.Type == PluginType.Mod)
 					{
-						if (PluginInfo.Descriptor.Modules != null && PluginInfo.Descriptor.Modules.Length > 0 && PluginInfo.Type == PluginType.Mod)
-						{
-							FileReference ModProjectFilePath = FileReference.Combine(PluginInfo.Directory, "Mods", PluginInfo.Name + ProjectFileExtension);
+						FileReference ModProjectFilePath = FileReference.Combine(PluginInfo.Directory, "Mods", PluginInfo.Name + ProjectFileExtension);
 
-							bool bProjectAlreadyExisted;
-							ProjectFile ModProjectFile = FindOrAddProject(ModProjectFilePath, PluginInfo.Directory, IncludeInGeneratedProjects: true, bAlreadyExisted: out bProjectAlreadyExisted);
-							ModProjectFile.IsForeignProject = GameProject.IsForeignProject;
-							ModProjectFile.IsGeneratedProject = true;
-							ModProjectFile.IsStubProject = false;
-							ModProjectFile.PluginFilePath = PluginInfo.File;
-							ModProjectFile.ProjectTargets.AddRange(GameProject.ProjectTargets);
+						bool bProjectAlreadyExisted;
+						ProjectFile ModProjectFile = FindOrAddProject(ModProjectFilePath, PluginInfo.Directory, IncludeInGeneratedProjects: true, bAlreadyExisted: out bProjectAlreadyExisted);
+						ModProjectFile.IsForeignProject = GameProject.IsForeignProject;
+						ModProjectFile.IsGeneratedProject = true;
+						ModProjectFile.IsStubProject = false;
+						ModProjectFile.PluginFilePath = PluginInfo.File;
+						ModProjectFile.ProjectTargets.AddRange(GameProject.ProjectTargets);
 
-							AddPluginFilesToProject(PluginInfo.File, PluginInfo.Directory, ModProjectFile);
+						AddPluginFilesToProject(PluginInfo.File, PluginInfo.Directory, ModProjectFile);
 
-							ModProjects.Add(ModProjectFile);
-						}
+						ModProjects.Add(ModProjectFile);
 					}
 				}
 			}
diff --git a/Engine/Source/Programs/UnrealBuildTool/System/PluginDescriptor.cs b/Engine/Source/Programs/UnrealBuildTool/System/PluginDescriptor.cs
--- a/Engine/Source/Programs/UnrealBuildTool/System/PluginDescriptor.cs
+++ b/Engine/Source/Programs/UnrealBuildTool/System/PluginDescriptor.cs
@@ -114,6 +114,11 @@ namespace UnrealBuildTool
 		/// </summary>
 		public string EngineVersion;
 
+		/// <summary>4
+		/// If true, this plugin from a platform extension extending another plugin */
+		/// </summary>
+		public bool bIsPluginExtension;
+
 		/// <summary>
 		/// List of platforms supported by this plugin. This list will be copied to any plugin reference from a project file, to allow filtering entire plugins from staged builds.
 		/// </summary>
@@ -230,6 +235,7 @@ namespace UnrealBuildTool
 			RawObject.TryGetStringField("SupportURL", out SupportURL);
 			RawObject.TryGetStringField("EngineVersion", out EngineVersion);
 			RawObject.TryGetStringArrayField("SupportedPrograms", out SupportedPrograms);
+			RawObject.TryGetBoolField("bIsPluginExtension", out bIsPluginExtension);
 
 			try
 			{
@@ -368,6 +374,10 @@ namespace UnrealBuildTool
 			{
 				Writer.WriteStringArrayField("SupportedPrograms", SupportedPrograms);
 			}
+			if (bIsPluginExtension)
+			{
+				Writer.WriteValue("bIsPluginExtension", bIsPluginExtension);
+			}
 
 			ModuleDescriptor.WriteArray(Writer, "Modules", Modules);
 
diff --git a/Engine/Source/Programs/UnrealBuildTool/System/Plugins.cs b/Engine/Source/Programs/UnrealBuildTool/System/Plugins.cs
--- a/Engine/Source/Programs/UnrealBuildTool/System/Plugins.cs
+++ b/Engine/Source/Programs/UnrealBuildTool/System/Plugins.cs
@@ -77,6 +77,11 @@ namespace UnrealBuildTool
 		/// </summary>
 		public readonly DirectoryReference Directory;
 
+		/// <summary>
+		/// Children plugin files that can be added to this plugin (platform extensions)
+		/// </summary>
+		public List<FileReference> ChildFiles = new List<FileReference>();
+
 		/// <summary>
 		/// The plugin descriptor
 		/// </summary>
@@ -207,11 +212,11 @@ namespace UnrealBuildTool
 				foreach (string AdditionalDirectory in AdditionalDirectories)
 				{
 					DirectoryReference DirRef = DirectoryReference.Combine(ProjectFileName.Directory, AdditionalDirectory);
-					Plugins.AddRange(ReadPluginsFromDirectory(DirRef, PluginType.External));
+					Plugins.AddRange(ReadPluginsFromDirectory(DirRef, "", PluginType.External));
 				}
 			}
 
-            return Plugins;
+			return Plugins;
 		}
 
 		/// <summary>
@@ -256,8 +261,7 @@ namespace UnrealBuildTool
 		/// <returns>Sequence of the found PluginInfo object.</returns>
 		public static IReadOnlyList<PluginInfo> ReadEnginePlugins(DirectoryReference EngineDirectory)
 		{
-			DirectoryReference PluginsDir = DirectoryReference.Combine(EngineDirectory, "Plugins");
-			return ReadPluginsFromDirectory(PluginsDir, PluginType.Engine);
+			return ReadPluginsFromDirectory(EngineDirectory, "Plugins", PluginType.Engine);
 		}
 
 		/// <summary>
@@ -267,8 +271,7 @@ namespace UnrealBuildTool
 		/// <returns>Sequence of the found PluginInfo object.</returns>
 		public static IReadOnlyList<PluginInfo> ReadEnterprisePlugins(DirectoryReference EnterpriseDirectory)
 		{
-			DirectoryReference PluginsDir = DirectoryReference.Combine(EnterpriseDirectory, "Plugins");
-			return ReadPluginsFromDirectory(PluginsDir, PluginType.Enterprise);
+			return ReadPluginsFromDirectory(EnterpriseDirectory, "Plugins", PluginType.Enterprise);
 		}
 
 		/// <summary>
@@ -279,41 +282,171 @@ namespace UnrealBuildTool
 		public static IReadOnlyList<PluginInfo> ReadProjectPlugins(DirectoryReference ProjectDirectory)
 		{
 			List<PluginInfo> Plugins = new List<PluginInfo>();
-			Plugins.AddRange(ReadPluginsFromDirectory(DirectoryReference.Combine(ProjectDirectory, "Plugins"), PluginType.Project));
-			Plugins.AddRange(ReadPluginsFromDirectory(DirectoryReference.Combine(ProjectDirectory, "Mods"), PluginType.Mod));
+			Plugins.AddRange(ReadPluginsFromDirectory(ProjectDirectory, "Plugins", PluginType.Project));
+			Plugins.AddRange(ReadPluginsFromDirectory(ProjectDirectory, "Mods", PluginType.Mod));
 			return Plugins.AsReadOnly();
 		}
 
-        /// <summary>
-        /// Read all of the plugins found in the project specified additional plugin directories
-        /// </summary>
-        /// <param name="AdditionalDirectory">The list of additional directories to scan</param>
-        /// <returns>List of the found PluginInfo objects</returns>
-        public static IReadOnlyList<PluginInfo> ReadAdditionalPlugins(DirectoryReference AdditionalDirectory)
-        {
-			return ReadPluginsFromDirectory(AdditionalDirectory, PluginType.External);
-        }
-
-        /// <summary>
-        /// Read all the plugin descriptors under the given directory
-        /// </summary>
-        /// <param name="ParentDirectory">The parent directory to look in.</param>
-        /// <param name="Type">The plugin type</param>
-        /// <returns>Sequence of the found PluginInfo object.</returns>
-        public static IReadOnlyList<PluginInfo> ReadPluginsFromDirectory(DirectoryReference ParentDirectory, PluginType Type)
+		/// <summary>
+		/// Read all of the plugins found in the project specified additional plugin directories
+		/// </summary>
+		/// <param name="AdditionalDirectory">The additional directory to scan</param>
+		/// <param name="Subdirectory">A subdirectory to look under for AdditionalDirectory </param>
+		/// <returns>List of the found PluginInfo objects</returns>
+		public static IReadOnlyList<PluginInfo> ReadAdditionalPlugins(DirectoryReference AdditionalDirectory, string Subdirectory)
 		{
-			List<PluginInfo> Plugins;
-			if (!PluginInfoCache.TryGetValue(ParentDirectory, out Plugins))
+			return ReadPluginsFromDirectory(AdditionalDirectory, Subdirectory, PluginType.External);
+		}
+
+		/// <summary>
+		///  Attempt to merge a child plugin up into a parent plugin (via file naming scheme). Very little merging happens
+		///  but it does allow for platform extensions to extend a plugin with module files
+		/// </summary>
+		/// <param name="Child">Child plugin that needs to merge to a main, parent plugin</param>
+		/// <param name="Filename">Child plugin's filename, used to determine the parent's name</param>
+		private static void TryMergeWithParent(PluginInfo Child, FileReference Filename)
+		{
+			// find the parent
+			PluginInfo Parent = null;
+
+			string[] Tokens = Filename.GetFileNameWithoutAnyExtensions().Split("_".ToCharArray());
+			if (Tokens.Length == 2)
 			{
-				Plugins = new List<PluginInfo>();
-				foreach (FileReference PluginFileName in EnumeratePlugins(ParentDirectory))
+				string ParentPluginName = Tokens[0];
+				foreach (KeyValuePair<DirectoryReference, List<PluginInfo>> Pair in PluginInfoCache)
 				{
-					PluginInfo Plugin = new PluginInfo(PluginFileName, Type);
-					Plugins.Add(Plugin);
+					Parent = Pair.Value.FirstOrDefault(x => x.Name.Equals(ParentPluginName, StringComparison.InvariantCultureIgnoreCase));
+					if (Parent != null)
+					{
+						break;
+					}
 				}
-				PluginInfoCache.Add(ParentDirectory, Plugins);
 			}
-			return Plugins;
+
+
+			// did we find a parent plugin?
+			if (Parent == null)
+			{
+				throw new BuildException("Child plugin {0} was not named properly. It should be in the form <ParentPlugin>_<Platform>.uplugin", Filename);
+			}
+
+			// add our uplugin file to the existing plugin to be used to search for modules later
+			Parent.ChildFiles.Add(Child.File);
+
+			// make sure we are whitelisted for any modules we list, if the parent had a whitelist
+			if (Child.Descriptor.Modules != null)
+			{
+				// get the part after last underscore, which is the platform name
+				string DirectoryName = Child.File.GetFileNameWithoutExtension();
+				string PlatformName = DirectoryName.Split("_".ToCharArray()).LastOrDefault();
+
+				// this should cause an error if it's invalid platform name
+				UnrealTargetPlatform Platform = UnrealTargetPlatform.Parse(PlatformName);
+
+				foreach (ModuleDescriptor ChildModule in Child.Descriptor.Modules)
+				{
+					ModuleDescriptor ParentModule = Parent.Descriptor.Modules.FirstOrDefault(x => x.Name.Equals(ChildModule.Name) && x.Type == ChildModule.Type);
+					if (ParentModule != null)
+					{
+						// merge white/blacklists (if the parent had a list, and child didn't specify a list, just add the child platform to the parent list - for white and black!)
+						if (ParentModule.WhitelistPlatforms != null && ParentModule.WhitelistPlatforms.Length > 0)
+						{
+							List<UnrealTargetPlatform> Whitelist = ParentModule.WhitelistPlatforms.ToList();
+							if (ChildModule.WhitelistPlatforms != null && ChildModule.WhitelistPlatforms.Length > 0)
+							{
+								Whitelist.AddRange(ChildModule.WhitelistPlatforms);
+							}
+							else
+							{
+								Whitelist.Add(Platform);
+							}
+							ParentModule.WhitelistPlatforms = Whitelist.ToArray();
+						}
+						if (ParentModule.BlacklistPlatforms != null && ParentModule.BlacklistPlatforms.Length > 0)
+						{
+							List<UnrealTargetPlatform> Blacklist = ParentModule.BlacklistPlatforms.ToList();
+							if (ChildModule.BlacklistPlatforms != null && ChildModule.BlacklistPlatforms.Length > 0)
+							{
+								Blacklist.AddRange(ChildModule.BlacklistPlatforms);
+							}
+							else
+							{
+								Blacklist.Add(Platform);
+							}
+							ParentModule.BlacklistPlatforms = Blacklist.ToArray();
+						}
+					}
+				}
+			}
+			// @todo platplug: what else do we want to support merging?!?
+		}
+
+
+		/// <summary>
+		/// Read all the plugin descriptors under the given directory
+		/// </summary>
+		/// <param name="RootDirectory">The directory to look in.</param>
+		/// <param name="Subdirectory">A subdirectory to look in in RootDirectory and any other Platform directories under Root</param>
+		/// <param name="Type">The plugin type</param>
+		/// <returns>Sequence of the found PluginInfo object.</returns>
+		public static IReadOnlyList<PluginInfo> ReadPluginsFromDirectory(DirectoryReference RootDirectory, string Subdirectory, PluginType Type)
+		{
+			// look for directories in RootDirectory and and Platform directories under RootDirectory
+			List<DirectoryReference> RootDirectories = new List<DirectoryReference>() { DirectoryReference.Combine(RootDirectory, Subdirectory) };
+
+			// now look for platform subdirectories with the Subdirectory
+			DirectoryReference PlatformDirectory = DirectoryReference.Combine(RootDirectory, "Platforms");
+			if (DirectoryReference.Exists(PlatformDirectory))
+			{
+				foreach (DirectoryReference Dir in DirectoryReference.EnumerateDirectories(PlatformDirectory))
+				{
+					RootDirectories.Add(DirectoryReference.Combine(Dir, Subdirectory));
+				}
+			}
+
+			Dictionary<PluginInfo, FileReference> ChildPlugins = new Dictionary<PluginInfo, FileReference>();
+			List<PluginInfo> AllParentPlugins = new List<PluginInfo>();
+
+			foreach (DirectoryReference Dir in RootDirectories)
+			{
+				if (!DirectoryReference.Exists(Dir))
+				{
+					continue;
+				}
+
+				List<PluginInfo> Plugins;
+				if (!PluginInfoCache.TryGetValue(Dir, out Plugins))
+				{
+					Plugins = new List<PluginInfo>();
+					foreach (FileReference PluginFileName in EnumeratePlugins(Dir))
+					{
+						PluginInfo Plugin = new PluginInfo(PluginFileName, Type);
+
+						Console.WriteLine("Plugin: {0}", PluginFileName);
+						// is there a parent to merge up into?
+						if (Plugin.Descriptor.bIsPluginExtension)
+						{
+							ChildPlugins.Add(Plugin, PluginFileName);
+						}
+						else
+						{
+							Plugins.Add(Plugin);
+						}
+					}
+					PluginInfoCache.Add(Dir, Plugins);
+				}
+
+				// gather all of the plugins into one list
+				AllParentPlugins.AddRange(Plugins);
+			}
+
+			// now that all parent plugins are read in, we can let the children look up the parents
+			foreach (KeyValuePair<PluginInfo, FileReference> Pair in ChildPlugins)
+			{
+				TryMergeWithParent(Pair.Key, Pair.Value);
+			}
+
+			return AllParentPlugins;
 		}
 
 		/// <summary>
diff --git a/Engine/Source/Programs/UnrealBuildTool/System/RulesCompiler.cs b/Engine/Source/Programs/UnrealBuildTool/System/RulesCompiler.cs
--- a/Engine/Source/Programs/UnrealBuildTool/System/RulesCompiler.cs
+++ b/Engine/Source/Programs/UnrealBuildTool/System/RulesCompiler.cs
@@ -386,19 +386,12 @@ namespace UnrealBuildTool
 			{
 				List<PluginInfo> IncludedPlugins = new List<PluginInfo>();
 
-				// look in the platform extensions engine directories for plugins
-				List<DirectoryReference> RootDirectories = new List<DirectoryReference>();
-				RootDirectories.AddRange(UnrealBuildTool.GetAllEngineDirectories());
-
 				// search for all engine plugins
-				foreach (DirectoryReference RootDir in RootDirectories) 
-				{
-					IncludedPlugins.AddRange(Plugins.ReadEnginePlugins(RootDir));
-				}
+				IncludedPlugins.AddRange(Plugins.ReadEnginePlugins(UnrealBuildTool.EngineDirectory));
 
 				RulesScope EngineScope = new RulesScope("Engine", null);
 
-				EngineRulesAssembly = CreateEngineOrEnterpriseRulesAssembly(EngineScope, RootDirectories, ProjectFileGenerator.EngineProjectFileNameBase, IncludedPlugins, UnrealBuildTool.IsEngineInstalled() || bUsePrecompiled, bSkipCompile, null);
+				EngineRulesAssembly = CreateEngineOrEnterpriseRulesAssembly(EngineScope, UnrealBuildTool.GetAllEngineDirectories().ToList(), ProjectFileGenerator.EngineProjectFileNameBase, IncludedPlugins, UnrealBuildTool.IsEngineInstalled() || bUsePrecompiled, bSkipCompile, null);
 			}
 			return EngineRulesAssembly;
 		}
@@ -418,10 +411,11 @@ namespace UnrealBuildTool
 				{
 					RulesScope EnterpriseScope = new RulesScope("Enterprise", EngineAssembly.Scope);
 
-					List<DirectoryReference> RootDirectories = new List<DirectoryReference>() { UnrealBuildTool.EnterpriseDirectory };
+					List<DirectoryReference> EnterpriseDirectories = new List<DirectoryReference>() { UnrealBuildTool.EnterpriseDirectory };
 
 					IReadOnlyList<PluginInfo> IncludedPlugins = Plugins.ReadEnterprisePlugins(UnrealBuildTool.EnterpriseDirectory);
-					EnterpriseRulesAssembly = CreateEngineOrEnterpriseRulesAssembly(EnterpriseScope, RootDirectories, ProjectFileGenerator.EnterpriseProjectFileNameBase, IncludedPlugins, UnrealBuildTool.IsEnterpriseInstalled() || bUsePrecompiled, bSkipCompile, EngineAssembly);
+					EnterpriseRulesAssembly = CreateEngineOrEnterpriseRulesAssembly(EnterpriseScope, new List<DirectoryReference>() { UnrealBuildTool.EnterpriseDirectory }, 
+						ProjectFileGenerator.EnterpriseProjectFileNameBase, IncludedPlugins, UnrealBuildTool.IsEnterpriseInstalled() || bUsePrecompiled, bSkipCompile, EngineAssembly);
 				}
 				else
 				{
@@ -553,9 +547,6 @@ namespace UnrealBuildTool
 				Dictionary<FileReference, ModuleRulesContext> ModuleFiles = new Dictionary<FileReference, ModuleRulesContext>();
 				List<FileReference> TargetFiles = new List<FileReference>();
 
-				// Find all the project plugins
-				List<PluginInfo> ProjectPlugins = new List<PluginInfo>();
-
 				// Find all the rules/plugins under the project source directories
 				foreach (DirectoryReference ProjectDirectory in UnrealBuildTool.GetAllProjectDirectories(ProjectFileName))
 				{
@@ -563,16 +554,18 @@ namespace UnrealBuildTool
 
 					AddModuleRulesWithContext(ProjectSourceDirectory, DefaultModuleContext, ModuleFiles);
 					TargetFiles.AddRange(FindAllRulesFiles(ProjectSourceDirectory, RulesFileType.Target));
+				}
 
-					ProjectPlugins.AddRange(Plugins.ReadProjectPlugins(ProjectDirectory));
+				// Find all the project plugins
+				List<PluginInfo> ProjectPlugins = new List<PluginInfo>();
+				ProjectPlugins.AddRange(Plugins.ReadProjectPlugins(MainProjectDirectory));
 
-					// Add the project's additional plugin directories plugins too
-					if (Project.AdditionalPluginDirectories != null)
+				// Add the project's additional plugin directories plugins too
+				if (Project.AdditionalPluginDirectories != null)
+				{
+					foreach (string AdditionalPluginDirectory in Project.AdditionalPluginDirectories)
 					{
-						foreach (string AdditionalPluginDirectory in Project.AdditionalPluginDirectories)
-						{
-							ProjectPlugins.AddRange(Plugins.ReadAdditionalPlugins(DirectoryReference.Combine(ProjectDirectory, AdditionalPluginDirectory)));
-						}
+						ProjectPlugins.AddRange(Plugins.ReadAdditionalPlugins(MainProjectDirectory, AdditionalPluginDirectory));
 					}
 				}
 
@@ -689,7 +682,12 @@ namespace UnrealBuildTool
 
 			foreach (PluginInfo Plugin in Plugins)
 			{
-				IReadOnlyList<FileReference> PluginModuleFiles = FindAllRulesFiles(DirectoryReference.Combine(Plugin.Directory, "Source"), RulesFileType.Module);
+				List<FileReference> PluginModuleFiles = FindAllRulesFiles(DirectoryReference.Combine(Plugin.Directory, "Source"), RulesFileType.Module).ToList();
+				foreach (FileReference ChildFile in Plugin.ChildFiles)
+				{
+					PluginModuleFiles.AddRange(FindAllRulesFiles(DirectoryReference.Combine(ChildFile.Directory, "Source"), RulesFileType.Module));
+				}
+
 				foreach (FileReference ModuleFile in PluginModuleFiles)
 				{
 					ModuleRulesContext PluginContext = new ModuleRulesContext(DefaultContext);
diff --git a/Engine/Source/Runtime/Projects/Private/PluginDescriptor.cpp b/Engine/Source/Runtime/Projects/Private/PluginDescriptor.cpp
--- a/Engine/Source/Runtime/Projects/Private/PluginDescriptor.cpp
+++ b/Engine/Source/Runtime/Projects/Private/PluginDescriptor.cpp
@@ -33,7 +33,8 @@ FPluginDescriptor::FPluginDescriptor()
 	, bInstalled(false)
 	, bRequiresBuildPlatform(false)
 	, bIsHidden(false)
-{ 
+	, bIsPluginExtension(false)
+{
 }
 
 
@@ -116,6 +117,7 @@ bool FPluginDescriptor::Read(const FJsonObject& Object, FText& OutFailReason)
 	Object.TryGetStringField(TEXT("EngineVersion"), EngineVersion);
 	Object.TryGetStringArrayField(TEXT("SupportedTargetPlatforms"), SupportedTargetPlatforms);
 	Object.TryGetStringArrayField(TEXT("SupportedPrograms"), SupportedPrograms);
+	Object.TryGetBoolField(TEXT("bIsPluginExtension"), bIsPluginExtension);
 
 	if (!FModuleDescriptor::ReadArray(Object, TEXT("Modules"), Modules, OutFailReason))
 	{
@@ -215,7 +217,10 @@ void FPluginDescriptor::Write(TJsonWriter<>& Writer) const
 	{
 		Writer.WriteValue(TEXT("SupportedPrograms"), SupportedPrograms);
 	}
-
+	if (bIsPluginExtension)
+	{
+		Writer.WriteValue(TEXT("bIsPluginExtension"), bIsPluginExtension);
+	}
 	FModuleDescriptor::WriteArray(Writer, TEXT("Modules"), Modules);
 
 	FLocalizationTargetDescriptor::WriteArray(Writer, TEXT("LocalizationTargets"), LocalizationTargets);
diff --git a/Engine/Source/Runtime/Projects/Private/PluginManager.cpp b/Engine/Source/Runtime/Projects/Private/PluginManager.cpp
--- a/Engine/Source/Runtime/Projects/Private/PluginManager.cpp
+++ b/Engine/Source/Runtime/Projects/Private/PluginManager.cpp
@@ -219,20 +219,25 @@ void FPluginManager::ReadAllPlugins(TMap<FString, TSharedRef<FPlugin>>& Plugins,
 		FindPluginManifestsInDirectory(*FPaths::ProjectPluginsDir(), ManifestFileNames);
 	}
 
+	// track child plugins that don't want to go into main plugin set
+	TArray<TSharedRef<FPlugin>> ChildPlugins;
+
 	// If we didn't find any manifests, do a recursive search for plugins
 	if (ManifestFileNames.Num() == 0)
 	{
 		// Find "built-in" plugins.  That is, plugins situated right within the Engine directory.
-		ReadPluginsInDirectory(FPaths::EnginePluginsDir(), EPluginType::Engine, Plugins);
+		ReadPluginsInDirectory(FPaths::EnginePluginsDir(), EPluginType::Engine, Plugins, ChildPlugins);
+		ReadPluginsInDirectory(FPaths::EnginePlatformExtensionsDir() + TEXT("Plugins"), EPluginType::Engine, Plugins, ChildPlugins);
 
 		// Find plugins in the game project directory (<MyGameProject>/Plugins). If there are any engine plugins matching the name of a game plugin,
 		// assume that the game plugin version is preferred.
 		if (Project != nullptr)
 		{
 #if IS_PROGRAM
-			ReadPluginsInDirectory(FPaths::GetPath(FPaths::GetProjectFilePath()) / TEXT("Plugins"), EPluginType::Project, Plugins);
+			ReadPluginsInDirectory(FPaths::GetPath(FPaths::GetProjectFilePath()) / TEXT("Plugins"), EPluginType::Project, Plugins, ChildPlugins);
 #else
-			ReadPluginsInDirectory(FPaths::ProjectPluginsDir(), EPluginType::Project, Plugins);
+			ReadPluginsInDirectory(FPaths::ProjectPluginsDir(), EPluginType::Project, Plugins, ChildPlugins);
+			ReadPluginsInDirectory(FPaths::ProjectPlatformExtensionsDir() + TEXT("Plugins"), EPluginType::Project, Plugins, ChildPlugins);
 #endif
 		}
 	}
@@ -253,6 +258,7 @@ void FPluginManager::ReadAllPlugins(TMap<FString, TSharedRef<FPlugin>>& Plugins,
 
 			// Get all the standard plugin directories
 			const FString EngineDir = FPaths::EngineDir();
+			const FString PlatformExtensionEngineDir = FPaths::EnginePlatformExtensionsDir();
 			const FString EnterpriseDir = FPaths::EnterpriseDir();
 			const FString ProjectModsDir = FPaths::ProjectModsDir();
 
@@ -260,7 +266,7 @@ void FPluginManager::ReadAllPlugins(TMap<FString, TSharedRef<FPlugin>>& Plugins,
 			for (const FPluginManifestEntry& Entry : Manifest.Contents)
 			{
 				EPluginType Type;
-				if (Entry.File.StartsWith(EngineDir))
+				if (Entry.File.StartsWith(EngineDir) || Entry.File.StartsWith(PlatformExtensionEngineDir))
 				{
 					Type = EPluginType::Engine;
 				}
@@ -276,7 +282,7 @@ void FPluginManager::ReadAllPlugins(TMap<FString, TSharedRef<FPlugin>>& Plugins,
 				{
 					Type = EPluginType::Project;
 				}
-				CreatePluginObject(Entry.File, Entry.Descriptor, Type, Plugins);
+				CreatePluginObject(Entry.File, Entry.Descriptor, Type, Plugins, ChildPlugins);
 			}
 		}
 	}
@@ -284,29 +290,84 @@ void FPluginManager::ReadAllPlugins(TMap<FString, TSharedRef<FPlugin>>& Plugins,
 	if (Project != nullptr)
 	{
 		// Always add the mods from the loose directory without using manifests, because they're not packaged together.
-		ReadPluginsInDirectory(FPaths::ProjectModsDir(), EPluginType::Mod, Plugins);
+		ReadPluginsInDirectory(FPaths::ProjectModsDir(), EPluginType::Mod, Plugins, ChildPlugins);
 
 		// If they have a list of additional directories to check, add those plugins too
 		for (const FString& Dir : Project->GetAdditionalPluginDirectories())
 		{
-			ReadPluginsInDirectory(Dir, EPluginType::External, Plugins);
+			ReadPluginsInDirectory(Dir, EPluginType::External, Plugins, ChildPlugins);
 		}
 
 		// Add plugins from FPaths::EnterprisePluginsDir if it exists
 		if (FPaths::DirectoryExists(FPaths::EnterprisePluginsDir()))
 		{
-			ReadPluginsInDirectory(FPaths::EnterprisePluginsDir(), EPluginType::Enterprise, Plugins);
+			ReadPluginsInDirectory(FPaths::EnterprisePluginsDir(), EPluginType::Enterprise, Plugins, ChildPlugins);
 		}
 	}
 
 	for (const FString& ExtraSearchPath : ExtraSearchPaths)
 	{
-		ReadPluginsInDirectory(ExtraSearchPath, EPluginType::External, Plugins);
+		ReadPluginsInDirectory(ExtraSearchPath, EPluginType::External, Plugins, ChildPlugins);
+	}
+
+	// now that we have all the plugins, merge child plugins
+	for (TSharedRef<FPlugin> Child : ChildPlugins)
+	{
+		// find the parent
+		TArray<FString> Tokens;
+		FPaths::GetCleanFilename(Child->GetDescriptorFileName()).ParseIntoArray(Tokens, TEXT("_"), true);
+		TSharedRef<FPlugin>* ParentPtr = nullptr;
+		if (Tokens.Num() == 2)
+		{
+			FString ParentPluginName = Tokens[0];
+			ParentPtr = Plugins.Find(ParentPluginName);
+		}
+		if (ParentPtr != nullptr)
+		{
+			TSharedRef<FPlugin>& Parent = *ParentPtr;
+			for (const FModuleDescriptor& ChildModule : Child->GetDescriptor().Modules)
+			{
+				// look for a matching parent
+				for (FModuleDescriptor& ParentModule : Parent->Descriptor.Modules)
+				{
+					if (ParentModule.Name == ChildModule.Name && ParentModule.Type == ChildModule.Type)
+					{
+						if (ParentModule.WhitelistPlatforms.Num() > 0)
+						{
+							if (ChildModule.WhitelistPlatforms.Num() > 0)
+							{
+								ParentModule.WhitelistPlatforms.Append(ChildModule.WhitelistPlatforms);
+							}
+							else
+							{
+								ParentModule.WhitelistPlatforms.Add(FPlatformProperties::PlatformName());
+							}
+						}
+						if (ParentModule.BlacklistPlatforms.Num() > 0)
+						{
+							if (ChildModule.BlacklistPlatforms.Num() > 0)
+							{
+								ParentModule.BlacklistPlatforms.Append(ChildModule.BlacklistPlatforms);
+							}
+							else
+							{
+								ParentModule.BlacklistPlatforms.Add(FPlatformProperties::PlatformName());
+							}
+						}
+					}
+				}
+			}
+		}
+		else
+		{
+			UE_LOG(LogPluginManager, Error, TEXT("Child plugin %s was not named properly. It should be in the form <ParentPlugin>_<Platform>.uplugin."), *Child->GetDescriptorFileName());
+		}
 	}
+
 #endif
 }
 
-void FPluginManager::ReadPluginsInDirectory(const FString& PluginsDirectory, const EPluginType Type, TMap<FString, TSharedRef<FPlugin>>& Plugins)
+void FPluginManager::ReadPluginsInDirectory(const FString& PluginsDirectory, const EPluginType Type, TMap<FString, TSharedRef<FPlugin>>& Plugins, TArray<TSharedRef<FPlugin>>& ChildPlugins)
 {
 	// Make sure the directory even exists
 	if(FPlatformFileManager::Get().GetPlatformFile().DirectoryExists(*PluginsDirectory))
@@ -320,7 +381,7 @@ void FPluginManager::ReadPluginsInDirectory(const FString& PluginsDirectory, con
 			FText FailureReason;
 			if ( Descriptor.Load(FileName, FailureReason) )
 			{
-				CreatePluginObject(FileName, Descriptor, Type, Plugins);
+				CreatePluginObject(FileName, Descriptor, Type, Plugins, ChildPlugins);
 			}
 			else
 			{
@@ -369,10 +430,17 @@ void FPluginManager::FindPluginManifestsInDirectory(const FString& PluginManifes
 	IFileManager::Get().IterateDirectory(*PluginManifestDirectory, Visitor);
 }
 
-void FPluginManager::CreatePluginObject(const FString& FileName, const FPluginDescriptor& Descriptor, const EPluginType Type, TMap<FString, TSharedRef<FPlugin>>& Plugins)
+void FPluginManager::CreatePluginObject(const FString& FileName, const FPluginDescriptor& Descriptor, const EPluginType Type, TMap<FString, TSharedRef<FPlugin>>& Plugins, TArray<TSharedRef<FPlugin>>& ChildPlugins)
 {
 	TSharedRef<FPlugin> Plugin = MakeShareable(new FPlugin(FileName, Descriptor, Type));
 
+	// children plugins are gathered and used later
+	if (Plugin->GetDescriptor().bIsPluginExtension)
+	{
+		ChildPlugins.Add(Plugin);
+		return;
+	}
+
 	FString FullPath = FPaths::ConvertRelativePathToFull(FileName);
 	UE_LOG(LogPluginManager, Verbose, TEXT("Read plugin descriptor for %s, from %s"), *Plugin->GetName(), *FullPath);
 
diff --git a/Engine/Source/Runtime/Projects/Private/PluginManager.h b/Engine/Source/Runtime/Projects/Private/PluginManager.h
--- a/Engine/Source/Runtime/Projects/Private/PluginManager.h
+++ b/Engine/Source/Runtime/Projects/Private/PluginManager.h
@@ -124,10 +124,10 @@ private:
 	static void ReadAllPlugins(TMap<FString, TSharedRef<FPlugin>>& Plugins, const TSet<FString>& ExtraSearchPaths);
 
 	/** Reads all the plugin descriptors from disk */
-	static void ReadPluginsInDirectory(const FString& PluginsDirectory, const EPluginType Type, TMap<FString, TSharedRef<FPlugin>>& Plugins);
+	static void ReadPluginsInDirectory(const FString& PluginsDirectory, const EPluginType Type, TMap<FString, TSharedRef<FPlugin>>& Plugins, TArray<TSharedRef<FPlugin>>& ChildPlugins);
 
 	/** Creates a FPlugin object and adds it to the given map */
-	static void CreatePluginObject(const FString& FileName, const FPluginDescriptor& Descriptor, const EPluginType Type, TMap<FString, TSharedRef<FPlugin>>& Plugins);
+	static void CreatePluginObject(const FString& FileName, const FPluginDescriptor& Descriptor, const EPluginType Type, TMap<FString, TSharedRef<FPlugin>>& Plugins, TArray<TSharedRef<FPlugin>>& ChildPlugins);
 
 	/** Finds all the plugin descriptors underneath a given directory */
 	static void FindPluginsInDirectory(const FString& PluginsDirectory, TArray<FString>& FileNames);
diff --git a/Engine/Source/Runtime/Projects/Public/PluginDescriptor.h b/Engine/Source/Runtime/Projects/Public/PluginDescriptor.h
--- a/Engine/Source/Runtime/Projects/Public/PluginDescriptor.h
+++ b/Engine/Source/Runtime/Projects/Public/PluginDescriptor.h
@@ -91,6 +91,9 @@ struct PROJECTS_API FPluginDescriptor
 	/** For auto-generated plugins that should not be listed in the plugin browser for users to disable freely. */
 	bool bIsHidden;
 
+	/** If true, this plugin from a platform extension extending another plugin */
+	bool bIsPluginExtension;
+
 	/** Pre-build steps for each host platform */
 	FCustomBuildSteps PreBuildSteps;
 
