Remove log spam.

#rnx

diff --git a/Engine/Source/Programs/UnrealBuildTool/System/Plugins.cs b/Engine/Source/Programs/UnrealBuildTool/System/Plugins.cs
--- a/Engine/Source/Programs/UnrealBuildTool/System/Plugins.cs
+++ b/Engine/Source/Programs/UnrealBuildTool/System/Plugins.cs
@@ -422,7 +422,6 @@ namespace UnrealBuildTool
 					{
 						PluginInfo Plugin = new PluginInfo(PluginFileName, Type);
 
-						Console.WriteLine("Plugin: {0}", PluginFileName);
 						// is there a parent to merge up into?
 						if (Plugin.Descriptor.bIsPluginExtension)
 						{
