- Fixed plugin extensions to only stage for the runtime platform, even if they have editor modules, etc.

diff --git a/Engine/Source/Programs/UnrealBuildTool/Configuration/UEBuildTarget.cs b/Engine/Source/Programs/UnrealBuildTool/Configuration/UEBuildTarget.cs
--- a/Engine/Source/Programs/UnrealBuildTool/Configuration/UEBuildTarget.cs
+++ b/Engine/Source/Programs/UnrealBuildTool/Configuration/UEBuildTarget.cs
@@ -1380,7 +1380,12 @@ namespace UnrealBuildTool
 					Receipt.RuntimeDependencies.Add(EnabledPlugin.File, StagedFileType.UFS);
 					foreach (FileReference ChildFile in EnabledPlugin.ChildFiles)
 					{
-						Receipt.RuntimeDependencies.Add(ChildFile, StagedFileType.UFS);
+						// only want to stage plugin extensions for the running platform, not host type platforms as well
+						// we have a naming convention that can assume this is the suffix of the filename
+						if (ChildFile.GetFileNameWithoutAnyExtensions().EndsWith("_" + Receipt.Platform.ToString()))
+						{
+							Receipt.RuntimeDependencies.Add(ChildFile, StagedFileType.UFS);
+						}
 					}
 				}
 			}
