Excluding target.cs files from the module source, fixes issue were target source gets assigned to a project multiple times (which in some cases results in it being assigned to a otherwise empty project)
Resolves issue with games having a program with the same overlapping name.

(ushell-p4-cherrypick of 8012043 by [REDACTED_PERSON])

diff --git a/Engine/Source/Programs/UnrealBuildTool/ProjectFiles/ProjectFileGenerator.cs b/Engine/Source/Programs/UnrealBuildTool/ProjectFiles/ProjectFileGenerator.cs
--- a/Engine/Source/Programs/UnrealBuildTool/ProjectFiles/ProjectFileGenerator.cs
+++ b/Engine/Source/Programs/UnrealBuildTool/ProjectFiles/ProjectFileGenerator.cs
@@ -1940,6 +1940,8 @@ namespace UnrealBuildTool
 
 					// Find all of the source files (and other files) and add them to the project
 					List<FileReference> FoundFiles = SourceFileSearch.FindModuleSourceFiles( CurModuleFile, SearchSubdirectories:SearchSubdirectories );
+					// remove any target files, they are technically not part of the module and are explicitly added when the project is created
+					FoundFiles.RemoveAll(f => f.FullName.EndsWith(".Target.cs"));
 					ProjectFile.AddFilesToProject( FoundFiles, BaseFolder );
 
 					// Check if there's a plugin directory here
