Reworked module assignment to visual project files to simplify the rule, now the first matching target starting at the module file and searching upwards in the directory structure will be used as base for which project to use. If multiple targets are on the same level the old rules of Game -> Mod -> Program - > Engine -> Platform -> Enterprise applies.

(ushell-p4-cherrypick of 8034070 by [REDACTED_PERSON])

diff --git a/Engine/Source/Programs/UnrealBuildTool/ProjectFiles/ProjectFileGenerator.cs b/Engine/Source/Programs/UnrealBuildTool/ProjectFiles/ProjectFileGenerator.cs
--- a/Engine/Source/Programs/UnrealBuildTool/ProjectFiles/ProjectFileGenerator.cs
+++ b/Engine/Source/Programs/UnrealBuildTool/ProjectFiles/ProjectFileGenerator.cs
@@ -2006,60 +2006,68 @@ namespace UnrealBuildTool
 
 		private ProjectFile FindProjectForModule(FileReference CurModuleFile, List<FileReference> AllGames, Dictionary<FileReference, ProjectFile> ProgramProjects, List<ProjectFile> ModProjects, out DirectoryReference BaseFolder)
 		{
-
-			string ProjectFileNameBase = null;
+			// Starting at the base directory of the module find a project which has the same directory as base, walking up the directory hierarchy until a match is found
 			BaseFolder = null;
 
-			// Figure out which game project this target belongs to
-			FileReference ProjectInfo = FindGameContainingFile(AllGames, CurModuleFile);
-			if (ProjectInfo != null)
-			{
-				BaseFolder = ProjectInfo.Directory;
-				return FindOrAddProjectHelper(ProjectInfo.GetFileNameWithoutExtension(), BaseFolder);
-			}
-
-			// Check if it's a mod
-			foreach (ProjectFile ModProject in ModProjects)
+			DirectoryReference Path = CurModuleFile.Directory;
+			while (!Path.IsRootDirectory())
 			{
-				if (CurModuleFile.IsUnderDirectory(ModProject.BaseDir))
+				// Figure out which game project this target belongs to
+				foreach (FileReference Game in AllGames)
 				{
-					BaseFolder = ModProject.BaseDir;
-					return ModProject;
+					// the source and the actual game directory are conceptually the same
+					if (Path == Game.Directory || Path == DirectoryReference.Combine(Game.Directory, "Source"))
+					{
+						FileReference ProjectInfo = Game;
+						BaseFolder = ProjectInfo.Directory;
+						return FindOrAddProjectHelper(ProjectInfo.GetFileNameWithoutExtension(), BaseFolder);
+					}
+				}
+				// Check if it's a mod
+				foreach (ProjectFile ModProject in ModProjects)
+				{
+					if (Path == ModProject.BaseDir)
+					{
+						BaseFolder = ModProject.BaseDir;
+						return ModProject;
+					}
 				}
-			}
 
-			// Check if this module is under any program project base folder
-			if (ProgramProjects != null)
-			{
-				foreach (ProjectFile ProgramProject in ProgramProjects.Values)
+				// Check if this module is under any program project base folder
+				if (ProgramProjects != null)
 				{
-					if (CurModuleFile.IsUnderDirectory(ProgramProject.BaseDir))
+					foreach (ProjectFile ProgramProject in ProgramProjects.Values)
 					{
-						BaseFolder = ProgramProject.BaseDir;
-						return ProgramProject;
+						if (Path == ProgramProject.BaseDir)
+						{
+							BaseFolder = ProgramProject.BaseDir;
+							return ProgramProject;
+						}
 					}
 				}
-			}
 
-			// check for engine, or platform extension engine folders
-			if ( CurModuleFile.IsUnderDirectory(UnrealBuildTool.EngineDirectory))
-			{
-				ProjectFileNameBase = EngineProjectFileNameBase;
-				BaseFolder = UnrealBuildTool.EngineDirectory;
-			}
-			else if( CurModuleFile.IsUnderDirectory(UnrealBuildTool.EnterpriseSourceDirectory) ||
-				CurModuleFile.IsUnderDirectory(DirectoryReference.Combine(UnrealBuildTool.EnterpriseDirectory, "Plugins")) )
-			{
-				ProjectFileNameBase = EnterpriseProjectFileNameBase;
-				BaseFolder = UnrealBuildTool.EnterpriseDirectory;
-			}
+				// check for engine, or platform extension engine folders
+				if (Path == UnrealBuildTool.EngineDirectory)
+				{
+					BaseFolder = UnrealBuildTool.EngineDirectory;
+					return FindOrAddProjectHelper(EngineProjectFileNameBase, BaseFolder);
+				}
+				if (Path == UnrealBuildTool.PlatformExtensionsDirectory)
+				{
+					BaseFolder = UnrealBuildTool.RootDirectory;
+					return FindOrAddProjectHelper(EngineProjectFileNameBase, BaseFolder);
+				}
+				if (Path == UnrealBuildTool.EnterpriseDirectory)
+				{
+					BaseFolder = UnrealBuildTool.EnterpriseDirectory;
+					return FindOrAddProjectHelper(EnterpriseProjectFileNameBase, BaseFolder);
+				}
 
-			if (ProjectInfo == null && ProjectFileNameBase == null)
-			{
-				throw new BuildException("Found a non-engine module file (" + CurModuleFile + ") that did not exist within any of the known game folders");
+				// no match found, lets search the parent directory
+				Path = Path.ParentDirectory;
 			}
 
-			return FindOrAddProjectHelper(ProjectFileNameBase, BaseFolder);
+			throw new BuildException("Found a module file (" + CurModuleFile + ") that did not exist within any of the known game folders or other source locations");
 		}
 
 		/// <summary>
