UBT: Change bLegacyPublicIncludePaths to a property that derives the correct value depending on other available settings.

#rnx

diff --git a/Engine/Source/Programs/UnrealBuildTool/Configuration/ModuleRules.cs b/Engine/Source/Programs/UnrealBuildTool/Configuration/ModuleRules.cs
--- a/Engine/Source/Programs/UnrealBuildTool/Configuration/ModuleRules.cs
+++ b/Engine/Source/Programs/UnrealBuildTool/Configuration/ModuleRules.cs
@@ -828,7 +828,12 @@ namespace UnrealBuildTool
 		/// Whether this module qualifies included headers from other modules relative to the root of their 'Public' folder. This reduces the number
 		/// of search paths that have to be passed to the compiler, improving performance and reducing the length of the compiler command line.
 		/// </summary>
-		public bool? bLegacyPublicIncludePaths;
+		public bool bLegacyPublicIncludePaths
+		{
+			set { bLegacyPublicIncludePathsPrivate = value; }
+			get { return bLegacyPublicIncludePathsPrivate ?? (bUseBackwardsCompatibleDefaults? Target.bLegacyPublicIncludePaths : false); }
+		}
+		private bool? bLegacyPublicIncludePathsPrivate;
 
 		/// <summary>
 		/// Which stanard to use for compiling this module
diff --git a/Engine/Source/Programs/UnrealBuildTool/Configuration/UEBuildModuleCPP.cs b/Engine/Source/Programs/UnrealBuildTool/Configuration/UEBuildModuleCPP.cs
--- a/Engine/Source/Programs/UnrealBuildTool/Configuration/UEBuildModuleCPP.cs
+++ b/Engine/Source/Programs/UnrealBuildTool/Configuration/UEBuildModuleCPP.cs
@@ -1178,7 +1178,7 @@ namespace UnrealBuildTool
 			}
 
 			// Setup the compile environment for the module.
-			SetupPrivateCompileEnvironment(Result.UserIncludePaths, Result.SystemIncludePaths, Result.Definitions, Result.AdditionalFrameworks, (Rules != null)? Rules.bLegacyPublicIncludePaths.Value : true);
+			SetupPrivateCompileEnvironment(Result.UserIncludePaths, Result.SystemIncludePaths, Result.Definitions, Result.AdditionalFrameworks, Rules.bLegacyPublicIncludePaths);
 
 			return Result;
 		}
@@ -1220,7 +1220,7 @@ namespace UnrealBuildTool
 			// Now set up the compile environment for the modules in the original order that we encountered them
 			foreach (UEBuildModule Module in ModuleToIncludePathsOnlyFlag.Keys)
 			{
-				Module.AddModuleToCompileEnvironment(null, CompileEnvironment.UserIncludePaths, CompileEnvironment.SystemIncludePaths, CompileEnvironment.Definitions, CompileEnvironment.AdditionalFrameworks, (Rules != null)? Rules.bLegacyPublicIncludePaths.Value : true);
+				Module.AddModuleToCompileEnvironment(null, CompileEnvironment.UserIncludePaths, CompileEnvironment.SystemIncludePaths, CompileEnvironment.Definitions, CompileEnvironment.AdditionalFrameworks, Rules.bLegacyPublicIncludePaths);
 			}
 			return CompileEnvironment;
 		}
diff --git a/Engine/Source/Programs/UnrealBuildTool/Configuration/UEBuildTarget.cs b/Engine/Source/Programs/UnrealBuildTool/Configuration/UEBuildTarget.cs
--- a/Engine/Source/Programs/UnrealBuildTool/Configuration/UEBuildTarget.cs
+++ b/Engine/Source/Programs/UnrealBuildTool/Configuration/UEBuildTarget.cs
@@ -3805,19 +3805,6 @@ namespace UnrealBuildTool
 					}
 				}
 
-				// Override the default for whether the module requires nested include paths
-				if(RulesObject.bLegacyPublicIncludePaths == null)
-				{
-					if(RulesObject.bUseBackwardsCompatibleDefaults)
-					{
-						RulesObject.bLegacyPublicIncludePaths = Rules.bLegacyPublicIncludePaths;
-					}
-					else
-					{
-						RulesObject.bLegacyPublicIncludePaths = false;
-					}
-				}
-
 				// Allow the current platform to modify the module rules
 				UEBuildPlatform.GetBuildPlatform(Platform).ModifyModuleRulesForActivePlatform(ModuleName, RulesObject, Rules);
 
