Fix modules in binary-only plugins being compiled for hot reload, which are distributed with ModuleRules.bUsePrecompiled = true.

diff --git a/Engine/Source/Programs/UnrealBuildTool/Configuration/UEBuildTarget.cs b/Engine/Source/Programs/UnrealBuildTool/Configuration/UEBuildTarget.cs
--- a/Engine/Source/Programs/UnrealBuildTool/Configuration/UEBuildTarget.cs
+++ b/Engine/Source/Programs/UnrealBuildTool/Configuration/UEBuildTarget.cs
@@ -2445,7 +2445,7 @@ namespace UnrealBuildTool
 					{
 						if (!UnrealBuildTool.IsProjectInstalled() || EnabledPlugins.Where(x => x.Type == PluginType.Mod).Any(x => Binary.OutputFilePaths[0].IsUnderDirectory(x.Directory)))
 						{
-							HotReloadModuleNames.UnionWith(GameModules.OfType<UEBuildModuleCPP>().Select(x => x.Name));
+							HotReloadModuleNames.UnionWith(GameModules.OfType<UEBuildModuleCPP>().Where(x => !x.Rules.bUsePrecompiled).Select(x => x.Name));
 						}
 					}
 				}
