Removed usage of platform extensions directory which does not exist anymore, resolves issue introduced in change 8034150

diff --git a/Engine/Source/Programs/UnrealBuildTool/ProjectFiles/ProjectFileGenerator.cs b/Engine/Source/Programs/UnrealBuildTool/ProjectFiles/ProjectFileGenerator.cs
--- a/Engine/Source/Programs/UnrealBuildTool/ProjectFiles/ProjectFileGenerator.cs
+++ b/Engine/Source/Programs/UnrealBuildTool/ProjectFiles/ProjectFileGenerator.cs
@@ -2052,11 +2052,7 @@ namespace UnrealBuildTool
 					BaseFolder = UnrealBuildTool.EngineDirectory;
 					return FindOrAddProjectHelper(EngineProjectFileNameBase, BaseFolder);
 				}
-				if (Path == UnrealBuildTool.PlatformExtensionsDirectory)
-				{
-					BaseFolder = UnrealBuildTool.RootDirectory;
-					return FindOrAddProjectHelper(EngineProjectFileNameBase, BaseFolder);
-				}
+
 				if (Path == UnrealBuildTool.EnterpriseDirectory)
 				{
 					BaseFolder = UnrealBuildTool.EnterpriseDirectory;
