Bump to v14 clang 8.0.1 toolchain for Linux

diff --git a/Engine/Build/BatchFiles/Linux/SetupToolchain.sh b/Engine/Build/BatchFiles/Linux/SetupToolchain.sh
--- a/Engine/Build/BatchFiles/Linux/SetupToolchain.sh
+++ b/Engine/Build/BatchFiles/Linux/SetupToolchain.sh
@@ -4,7 +4,7 @@
 # put ourselves into Engine directory (two up from location of this script)
 pushd "`dirname "$0"`/../../.." > /dev/null
 
-TOOLCHAIN_VERSION=v13_clang-7.0.1-centos7
+TOOLCHAIN_VERSION=v14_clang-8.0.1-centos7
 TOOLCHAIN_ARCHIVE=$TOOLCHAIN_VERSION.tar.gz
 
 TOOLCHAIN_URL=http://cdn.unrealengine.com/Toolchain_Linux/native-linux-$TOOLCHAIN_ARCHIVE
diff --git a/Engine/Source/Programs/UnrealBuildTool/Platform/Linux/UEBuildLinux.cs b/Engine/Source/Programs/UnrealBuildTool/Platform/Linux/UEBuildLinux.cs
--- a/Engine/Source/Programs/UnrealBuildTool/Platform/Linux/UEBuildLinux.cs
+++ b/Engine/Source/Programs/UnrealBuildTool/Platform/Linux/UEBuildLinux.cs
@@ -591,7 +591,7 @@ namespace UnrealBuildTool
 		/// <summary>
 		/// This is the SDK version we support
 		/// </summary>
-		static string ExpectedSDKVersion = "v13_clang-7.0.1-centos7";	// now unified for all the architectures
+		static string ExpectedSDKVersion = "v14_clang-8.0.1-centos7";	// now unified for all the architectures
 
 		/// <summary>
 		/// Platform name (embeds architecture for now)
