Allow 8.0.1 clang

Fix build issue

diff --git a/Engine/Source/Programs/UnrealBuildTool/Platform/Linux/LinuxToolChain.cs b/Engine/Source/Programs/UnrealBuildTool/Platform/Linux/LinuxToolChain.cs
--- a/Engine/Source/Programs/UnrealBuildTool/Platform/Linux/LinuxToolChain.cs
+++ b/Engine/Source/Programs/UnrealBuildTool/Platform/Linux/LinuxToolChain.cs
@@ -176,10 +176,10 @@ namespace UnrealBuildTool
 				throw new BuildException("Unable to build: no compatible clang version found. Please run Setup.sh");
 			}
 			// prevent unknown clangs since the build is likely to fail on too old or too new compilers
-			else if ((CompilerVersionMajor * 10 + CompilerVersionMinor) > 70 || (CompilerVersionMajor * 10 + CompilerVersionMinor) < 60)
+			else if ((CompilerVersionMajor * 10 + CompilerVersionMinor) > 80 || (CompilerVersionMajor * 10 + CompilerVersionMinor) > 70 || (CompilerVersionMajor * 10 + CompilerVersionMinor) < 60)
 			{
 				throw new BuildException(
-					string.Format("This version of the Unreal Engine can only be compiled with clang 7.0 and 6.0. clang {0} may not build it - please use a different version.",
+					string.Format("This version of the Unreal Engine can only be compiled with clang 8.0, 7.0 and 6.0. clang {0} may not build it - please use a different version.",
 						CompilerVersionString)
 					);
 			}
