Linux lld: Use -fno-pie flag when compiling with clang 8.0 or higher

diff --git a/Engine/Source/Programs/UnrealBuildTool/Platform/Linux/LinuxToolChain.cs b/Engine/Source/Programs/UnrealBuildTool/Platform/Linux/LinuxToolChain.cs
--- a/Engine/Source/Programs/UnrealBuildTool/Platform/Linux/LinuxToolChain.cs
+++ b/Engine/Source/Programs/UnrealBuildTool/Platform/Linux/LinuxToolChain.cs
@@ -958,7 +958,14 @@ namespace UnrealBuildTool
 
 				if (bSuppressPIE)
 				{
-					Result += " -Wl,-nopie";
+					if (CompilerVersionGreaterOrEqual(7, 0, 0))
+					{
+						Result += " -Wl,-no-pie";
+					}
+					else
+					{
+						Result += " -Wl,-nopie";
+					}
 				}
 			}
 
