PR #6063: Rider Project Model generation v 1.0 (Contributed by [REDACTED_PERSON])

diff --git a/Engine/Source/Programs/UnrealBuildTool/Configuration/UEBuildModule.cs b/Engine/Source/Programs/UnrealBuildTool/Configuration/UEBuildModule.cs
--- a/Engine/Source/Programs/UnrealBuildTool/Configuration/UEBuildModule.cs
+++ b/Engine/Source/Programs/UnrealBuildTool/Configuration/UEBuildModule.cs
@@ -74,37 +74,37 @@ namespace UnrealBuildTool
 		/// <summary>
 		/// Set of all the public definitions
 		/// </summary>
-		protected readonly HashSet<string> PublicDefinitions;
+		public readonly HashSet<string> PublicDefinitions;
 
 		/// <summary>
 		/// Set of all public include paths
 		/// </summary>
-		protected readonly HashSet<DirectoryReference> PublicIncludePaths;
+		public readonly HashSet<DirectoryReference> PublicIncludePaths;
 
 		/// <summary>
 		/// Nested public include paths which used to be added automatically, but are now only added for modules with bNestedPublicIncludePaths set.
 		/// </summary>
-		protected readonly HashSet<DirectoryReference> LegacyPublicIncludePaths = new HashSet<DirectoryReference>();
+		public readonly HashSet<DirectoryReference> LegacyPublicIncludePaths = new HashSet<DirectoryReference>();
 
 		/// <summary>
 		/// Set of all private include paths
 		/// </summary>
-		protected readonly HashSet<DirectoryReference> PrivateIncludePaths;
+		public readonly HashSet<DirectoryReference> PrivateIncludePaths;
 
 		/// <summary>
 		/// Set of all system include paths
 		/// </summary>
-		protected readonly HashSet<DirectoryReference> PublicSystemIncludePaths;
+		public readonly HashSet<DirectoryReference> PublicSystemIncludePaths;
 
 		/// <summary>
 		/// Set of all public system library paths
 		/// </summary>
-		protected readonly HashSet<DirectoryReference> PublicSystemLibraryPaths;
+		public readonly HashSet<DirectoryReference> PublicSystemLibraryPaths;
 
 		/// <summary>
 		/// Set of all additional libraries
 		/// </summary>
-		protected readonly HashSet<string> PublicAdditionalLibraries;
+		public readonly HashSet<string> PublicAdditionalLibraries;
 
 		/// <summary>
 		/// Set of all system libraries
@@ -114,12 +114,12 @@ namespace UnrealBuildTool
 		/// <summary>
 		/// Set of additional frameworks
 		/// </summary>
-		protected readonly HashSet<string> PublicFrameworks;
+		public readonly HashSet<string> PublicFrameworks;
 
 		/// <summary>
 		/// 
 		/// </summary>
-		protected readonly HashSet<string> PublicWeakFrameworks;
+		public readonly HashSet<string> PublicWeakFrameworks;
 
 		/// <summary>
 		/// 
@@ -134,32 +134,32 @@ namespace UnrealBuildTool
 		/// <summary>
 		/// Names of modules with header files that this module's public interface needs access to.
 		/// </summary>
-		protected List<UEBuildModule> PublicIncludePathModules;
+		public List<UEBuildModule> PublicIncludePathModules;
 
 		/// <summary>
 		/// Names of modules that this module's public interface depends on.
 		/// </summary>
-		protected List<UEBuildModule> PublicDependencyModules;
+		public List<UEBuildModule> PublicDependencyModules;
 
 		/// <summary>
 		/// Names of DLLs that this module should delay load
 		/// </summary>
-		protected HashSet<string> PublicDelayLoadDLLs;
+		public HashSet<string> PublicDelayLoadDLLs;
 
 		/// <summary>
 		/// Names of modules with header files that this module's private implementation needs access to.
 		/// </summary>
-		protected List<UEBuildModule> PrivateIncludePathModules;
+		public List<UEBuildModule> PrivateIncludePathModules;
 
 		/// <summary>
 		/// Names of modules that this module's private implementation depends on.
 		/// </summary>
-		protected List<UEBuildModule> PrivateDependencyModules;
+		public List<UEBuildModule> PrivateDependencyModules;
 
 		/// <summary>
 		/// Extra modules this module may require at run time
 		/// </summary>
-		protected List<UEBuildModule> DynamicallyLoadedModules;
+		public List<UEBuildModule> DynamicallyLoadedModules;
 
 		/// <summary>
 		/// Set of all whitelisted restricted folder references
@@ -871,6 +871,19 @@ namespace UnrealBuildTool
 			}
 		}
 
+		/// <summary>
+		/// Returns valueless API defines (like MODULE_API or MODULE_VTABLE)
+		/// </summary>
+		public IEnumerable<string> GetEmptyApiMacros()
+		{
+			if (Rules.Type == ModuleRules.ModuleType.CPlusPlus)
+			{
+				return new[] {ModuleVTableDefine + "=", ModuleApiDefine + "="};
+			}
+
+			return new string[0];
+		}
+
 		/// <summary>
 		/// Write information about this binary to a JSON file
 		/// </summary>
diff --git a/Engine/Source/Programs/UnrealBuildTool/Configuration/UEBuildTarget.cs b/Engine/Source/Programs/UnrealBuildTool/Configuration/UEBuildTarget.cs
--- a/Engine/Source/Programs/UnrealBuildTool/Configuration/UEBuildTarget.cs
+++ b/Engine/Source/Programs/UnrealBuildTool/Configuration/UEBuildTarget.cs
@@ -941,7 +941,7 @@ namespace UnrealBuildTool
 		/// <summary>
 		/// The name of the .Target.cs file, if the target was created with one
 		/// </summary>
-		readonly FileReference TargetRulesFile;
+		public readonly FileReference TargetRulesFile;
 
 		/// <summary>
 		/// Whether to deploy this target after compilation
@@ -1898,7 +1898,7 @@ namespace UnrealBuildTool
 		/// Gets the output directory for the main executable
 		/// </summary>
 		/// <returns>The executable directory</returns>
-		DirectoryReference GetExecutableDir()
+		public DirectoryReference GetExecutableDir()
 		{
 			DirectoryReference ExeDir = Binaries[0].OutputDir;
 			if (Platform == UnrealTargetPlatform.Mac && ExeDir.FullName.EndsWith(".app/Contents/MacOS"))
@@ -2675,7 +2675,7 @@ namespace UnrealBuildTool
 					}
 				}
 			}
-
+			
 			// Create rules for each remaining module, and check that it's set to be compiled
 			foreach(string FilteredModuleName in FilteredModuleNames)
 			{
@@ -2906,7 +2906,7 @@ namespace UnrealBuildTool
 		/// <summary>
 		/// Sets up the plugins for this target
 		/// </summary>
-		protected virtual void SetupPlugins()
+		public void SetupPlugins()
 		{
 			// Find all the valid plugins
 			Dictionary<string, PluginInfo> NameToInfo = RulesAssembly.EnumeratePlugins().ToDictionary(x => x.Name, x => x, StringComparer.InvariantCultureIgnoreCase);
diff --git a/Engine/Source/Programs/UnrealBuildTool/Modes/GenerateProjectFilesMode.cs b/Engine/Source/Programs/UnrealBuildTool/Modes/GenerateProjectFilesMode.cs
--- a/Engine/Source/Programs/UnrealBuildTool/Modes/GenerateProjectFilesMode.cs
+++ b/Engine/Source/Programs/UnrealBuildTool/Modes/GenerateProjectFilesMode.cs
@@ -36,6 +36,7 @@ namespace UnrealBuildTool
 		[CommandLine("-VSCode", Value = nameof(ProjectFileFormat.VisualStudioCode))]
 		[CommandLine("-VSMac", Value = nameof(ProjectFileFormat.VisualStudioMac))]
 		[CommandLine("-CLion", Value = nameof(ProjectFileFormat.CLion))]
+		[CommandLine("-Rider", Value = nameof(ProjectFileFormat.Rider))]
 		HashSet<ProjectFileFormat> ProjectFileFormats = new HashSet<ProjectFileFormat>();
 
 		/// <summary>
@@ -157,6 +158,9 @@ namespace UnrealBuildTool
 					case ProjectFileFormat.VisualStudioMac:
 						Generator = new VCMacProjectFileGenerator(ProjectFile, Arguments);
 						break;
+					case ProjectFileFormat.Rider:
+						Generator = new RiderProjectFileGenerator(ProjectFile, Arguments);
+						break;
 					default:
 						throw new BuildException("Unhandled project file type '{0}", ProjectFileFormat);
 				}
diff --git a/Engine/Source/Programs/UnrealBuildTool/ProjectFiles/ProjectFileGenerator.cs b/Engine/Source/Programs/UnrealBuildTool/ProjectFiles/ProjectFileGenerator.cs
--- a/Engine/Source/Programs/UnrealBuildTool/ProjectFiles/ProjectFileGenerator.cs
+++ b/Engine/Source/Programs/UnrealBuildTool/ProjectFiles/ProjectFileGenerator.cs
@@ -138,7 +138,8 @@ namespace UnrealBuildTool
 		Eddie,
 		VisualStudioCode,
 		VisualStudioMac,
-		CLion
+		CLion,
+		Rider
 	}
 
 	/// <summary>
@@ -1167,7 +1168,7 @@ namespace UnrealBuildTool
 		/// <summary>
 		/// Adds all game project files, including target projects and config files
 		/// </summary>
-		private void AddAllGameProjects(List<ProjectFile> GameProjects, string SupportedPlatformNames, MasterProjectFolder ProjectsFolder)
+		protected void AddAllGameProjects(List<ProjectFile> GameProjects, string SupportedPlatformNames, MasterProjectFolder ProjectsFolder)
 		{
 			HashSet<DirectoryReference> UniqueGameProjectDirectories = new HashSet<DirectoryReference>();
 			foreach( ProjectFile GameProject in GameProjects )
@@ -2329,7 +2330,7 @@ namespace UnrealBuildTool
 		/// </summary>
 		/// <param name="GameProjects">List of game project files</param>
 		/// <param name="ModProjects">Receives the list of mod projects on success</param>
-		void AddProjectsForMods(List<ProjectFile> GameProjects, out List<ProjectFile> ModProjects)
+		protected void AddProjectsForMods(List<ProjectFile> GameProjects, out List<ProjectFile> ModProjects)
 		{
 			// Find all the mods for game projects
 			ModProjects = new List<ProjectFile>();
@@ -2363,7 +2364,7 @@ namespace UnrealBuildTool
 		/// </summary>
 		/// <param name="BaseName">The base name for the project file</param>
 		/// <returns>Full path to the project file</returns>
-		private FileReference GetProjectLocation(string BaseName)
+		protected FileReference GetProjectLocation(string BaseName)
 		{
 			return FileReference.Combine(IntermediateProjectFilesPath, BaseName + ProjectFileExtension);
 		}
diff --git a/Engine/Source/Programs/UnrealBuildTool/ProjectFiles/Rider/RiderProjectFile.cs b/Engine/Source/Programs/UnrealBuildTool/ProjectFiles/Rider/RiderProjectFile.cs
new file mode 100644
--- /dev/null
+++ b/Engine/Source/Programs/UnrealBuildTool/ProjectFiles/Rider/RiderProjectFile.cs
@@ -0,0 +1,354 @@
+using System.Collections.Generic;
+using System.Linq;
+using Tools.DotNETCommon;
+
+namespace UnrealBuildTool
+{
+	internal class RiderProjectFile : ProjectFile
+	{
+		public DirectoryReference RootPath;
+		public HashSet<TargetType> TargetTypes;
+		public CommandLineArguments Arguments;
+
+		public RiderProjectFile(FileReference InProjectFilePath) : base(InProjectFilePath)
+		{
+		}
+
+		/// <summary>
+		/// Write project file info in JSON file.
+		/// For every combination of <c>UnrealTargetPlatform</c>, <c>UnrealTargetConfiguration</c> and <c>TargetType</c>
+		/// will be generated separate JSON file.
+		/// Project file will be stored:
+		/// For UE4:  {UE4Root}/Engine/Intermediate/ProjectFiles/.Rider/{Platform}/{Configuration}/{TargetType}/{ProjectName}.json
+		/// For game: {GameRoot}/Intermediate/ProjectFiles/.Rider/{Platform}/{Configuration}/{TargetType}/{ProjectName}.json
+		/// </summary>
+		/// <remarks>
+		/// * <c>UnrealTargetPlatform.Win32</c> will be always ignored.
+		/// * <c>TargetType.Editor</c> will be generated for current platform only and will ignore <c>UnrealTargetConfiguration.Test</c> and <c>UnrealTargetConfiguration.Shipping</c> configurations
+		/// * <c>TargetType.Program</c>  will be generated for current platform only and <c>UnrealTargetConfiguration.Development</c> configuration only 
+		/// </remarks>
+		/// <param name="InPlatforms"></param>
+		/// <param name="InConfigurations"></param>
+		/// <param name="PlatformProjectGenerators"></param>
+		/// <returns></returns>
+		public override bool WriteProjectFile(List<UnrealTargetPlatform> InPlatforms,
+			List<UnrealTargetConfiguration> InConfigurations,
+			PlatformProjectGeneratorCollection PlatformProjectGenerators)
+		{
+			string ProjectName = ProjectFilePath.GetFileNameWithoutAnyExtensions();
+			DirectoryReference projectRootFolder = DirectoryReference.Combine(RootPath, ".Rider");
+			foreach (UnrealTargetPlatform Platform in InPlatforms.Where(it => it != UnrealTargetPlatform.Win32))
+			{
+				DirectoryReference PlatformFolder = DirectoryReference.Combine(projectRootFolder, Platform.ToString());
+				foreach (UnrealTargetConfiguration Configuration in InConfigurations)
+				{
+					DirectoryReference configurationFolder = DirectoryReference.Combine(PlatformFolder, Configuration.ToString());
+					foreach (ProjectTarget ProjectTarget in ProjectTargets)
+					{
+						if (TargetTypes.Any() && !TargetTypes.Contains(ProjectTarget.TargetRules.Type)) continue;
+
+						switch (ProjectTarget.TargetRules.Type)
+						{
+							// Skip Programs for all configs except for current platform + Development configuration
+							case TargetType.Program
+								when (BuildHostPlatform.Current.Platform != Platform ||
+								      Configuration != UnrealTargetConfiguration.Development):
+							// Skip Editor for all platforms except for current platform
+							case TargetType.Editor
+								when (BuildHostPlatform.Current.Platform != Platform ||
+								      (Configuration == UnrealTargetConfiguration.Test ||
+								       Configuration == UnrealTargetConfiguration.Shipping)):
+								continue;
+						}
+
+						DirectoryReference TargetFolder =
+							DirectoryReference.Combine(configurationFolder, ProjectTarget.TargetRules.Type.ToString());
+
+						string DefaultArchitecture = UEBuildPlatform
+							.GetBuildPlatform(BuildHostPlatform.Current.Platform)
+							.GetDefaultArchitecture(ProjectTarget.UnrealProjectFilePath);
+						TargetDescriptor TargetDesc = new TargetDescriptor(ProjectTarget.UnrealProjectFilePath, ProjectTarget.Name,
+							BuildHostPlatform.Current.Platform, UnrealTargetConfiguration.Development,
+							DefaultArchitecture, Arguments);
+						UEBuildTarget BuildTarget = UEBuildTarget.Create(TargetDesc, false, false);
+						FileReference OutputFile = FileReference.Combine(TargetFolder, $"{ProjectName}.json");
+						DirectoryReference.CreateDirectory(OutputFile.Directory);
+						using (JsonWriter Writer = new JsonWriter(OutputFile))
+						{
+							ExportTarget(BuildTarget, Writer);	
+						}
+					}
+				}
+			}
+
+			return true;
+		}
+
+		/// <summary>
+		/// Write a Target to a JSON writer. Is array is empty, don't write anything
+		/// </summary>
+		/// <param name="Target"></param>
+		/// <param name="Writer">Writer for the array data</param>
+		private static void ExportTarget(UEBuildTarget Target, JsonWriter Writer)
+		{
+			Writer.WriteObjectStart();
+
+			Writer.WriteValue("Name", Target.TargetName);
+			Writer.WriteValue("Configuration", Target.Configuration.ToString());
+			Writer.WriteValue("Platform", Target.Platform.ToString());
+			Writer.WriteValue("TargetFile", Target.TargetRulesFile.FullName);
+			if (Target.ProjectFile != null)
+			{
+				Writer.WriteValue("ProjectFile", Target.ProjectFile.FullName);
+			}
+			
+			ExportEnvironmentToJson(Target, Writer);
+			
+			if(Target.Binaries.Any())
+			{
+				Writer.WriteArrayStart("Binaries");
+				foreach (UEBuildBinary Binary in Target.Binaries)
+				{
+					Writer.WriteObjectStart();
+					ExportBinary(Binary, Writer);
+					Writer.WriteObjectEnd();
+				}
+				Writer.WriteArrayEnd();
+			}
+
+			List<UEBuildModule> TargetModules = Target.Binaries.SelectMany(x => x.Modules).ToList();
+			if(TargetModules.Any())
+			{
+				Writer.WriteObjectStart("Modules");
+				foreach (UEBuildModule Module in TargetModules)
+				{
+					Writer.WriteObjectStart(Module.Name);
+					ExportModule(Module, Target.GetExecutableDir(), Writer);
+					Writer.WriteObjectEnd();
+				}
+				Writer.WriteObjectEnd();
+			}
+			
+			ExportPluginsFromTarget(Target, Writer);
+			
+			Writer.WriteObjectEnd();
+		}
+
+		/// <summary>
+		/// Write a Module to a JSON writer. If array is empty, don't write anything
+		/// </summary>
+		/// <param name="TargetOutputDir"></param>
+		/// <param name="Writer">Writer for the array data</param>
+		/// <param name="Module"></param>
+		private static void ExportModule(UEBuildModule Module, DirectoryReference TargetOutputDir, JsonWriter Writer)
+		{
+			DirectoryReference BinaryOutputDir = Module.Binary?.OutputDir;
+			Writer.WriteValue("Name", Module.Name);
+			Writer.WriteValue("Directory", Module.ModuleDirectory.FullName);
+			Writer.WriteValue("Rules", Module.RulesFile.FullName);
+			Writer.WriteValue("PCHUsage", Module.Rules.PCHUsage.ToString());
+			if (Module is UEBuildModuleCPP ModuleCPP)
+			{
+				Writer.WriteValue("GeneratedCodeDirectory", ModuleCPP.GeneratedCodeDirectory != null ? ModuleCPP.GeneratedCodeDirectory.FullName : string.Empty);
+			}
+
+			if (Module.Rules.PrivatePCHHeaderFile != null)
+			{
+				Writer.WriteValue("PrivatePCH", FileReference.Combine(Module.ModuleDirectory, Module.Rules.PrivatePCHHeaderFile).FullName);
+			}
+
+			if (Module.Rules.SharedPCHHeaderFile != null)
+			{
+				Writer.WriteValue("SharedPCH", FileReference.Combine(Module.ModuleDirectory, Module.Rules.SharedPCHHeaderFile).FullName);
+			}
+
+			ExportJsonModuleArray(Writer, "PublicDependencyModules", Module.PublicDependencyModules);
+			ExportJsonModuleArray(Writer, "PublicIncludePathModules", Module.PublicIncludePathModules);
+			ExportJsonModuleArray(Writer, "PrivateDependencyModules", Module.PrivateDependencyModules);
+			ExportJsonModuleArray(Writer, "PrivateIncludePathModules", Module.PrivateIncludePathModules);
+			ExportJsonModuleArray(Writer, "DynamicallyLoadedModules", Module.DynamicallyLoadedModules);
+
+			ExportJsonStringArray(Writer, "PublicSystemIncludePaths", Module.PublicSystemIncludePaths.Select(x => x.FullName));
+			ExportJsonStringArray(Writer, "PublicIncludePaths", Module.PublicIncludePaths.Select(x => x.FullName));
+			
+			ExportJsonStringArray(Writer, "LegacyPublicIncludePaths", Module.LegacyPublicIncludePaths.Select(x => x.FullName));
+			
+			ExportJsonStringArray(Writer, "PrivateIncludePaths", Module.PrivateIncludePaths.Select(x => x.FullName));
+			ExportJsonStringArray(Writer, "PublicLibraryPaths", Module.PublicSystemLibraryPaths.Select(x => x.FullName));
+			ExportJsonStringArray(Writer, "PublicAdditionalLibraries", Module.PublicAdditionalLibraries);
+			ExportJsonStringArray(Writer, "PublicFrameworks", Module.PublicFrameworks);
+			ExportJsonStringArray(Writer, "PublicWeakFrameworks", Module.PublicWeakFrameworks);
+			ExportJsonStringArray(Writer, "PublicDelayLoadDLLs", Module.PublicDelayLoadDLLs);
+			ExportJsonStringArray(Writer, "PublicDefinitions", Module.PublicDefinitions);
+			
+			ExportJsonStringArray(Writer, "PrivateDefinitions", Module.Rules.PrivateDefinitions);
+			ExportJsonStringArray(Writer, "ProjectDefinitions", /* TODO: Add method ShouldAddProjectDefinitions */ !Module.Rules.bTreatAsEngineModule ? Module.Rules.Target.ProjectDefinitions : new string[0]);
+			ExportJsonStringArray(Writer, "ApiDefinitions", Module.GetEmptyApiMacros());
+			Writer.WriteValue("ShouldAddLegacyPublicIncludePaths", Module.Rules.bLegacyPublicIncludePaths);
+
+			if(Module.Rules.CircularlyReferencedDependentModules.Any())
+			{
+				Writer.WriteArrayStart("CircularlyReferencedModules");
+				foreach (string ModuleName in Module.Rules.CircularlyReferencedDependentModules)
+				{
+					Writer.WriteValue(ModuleName);
+				}
+				Writer.WriteArrayEnd();
+			}
+			
+			if(Module.Rules.RuntimeDependencies.Inner.Any())
+			{
+				Writer.WriteArrayStart("RuntimeDependencies");
+				foreach (ModuleRules.RuntimeDependency RuntimeDependency in Module.Rules.RuntimeDependencies.Inner)
+				{
+					Writer.WriteObjectStart();
+					
+					Writer.WriteValue("Path",
+						Module.ExpandPathVariables(RuntimeDependency.Path, BinaryOutputDir, TargetOutputDir));
+					if (RuntimeDependency.SourcePath != null)
+					{
+						Writer.WriteValue("SourcePath",
+							Module.ExpandPathVariables(RuntimeDependency.SourcePath, BinaryOutputDir, TargetOutputDir));
+					}
+
+					Writer.WriteValue("Type", RuntimeDependency.Type.ToString());
+					
+					Writer.WriteObjectEnd();
+				}
+				Writer.WriteArrayEnd();
+			}
+		}
+		
+		/// <summary>
+		/// Write an array of Modules to a JSON writer. If array is empty, don't write anything
+		/// </summary>
+		/// <param name="Writer">Writer for the array data</param>
+		/// <param name="ArrayName">Name of the array property</param>
+		/// <param name="Modules">Sequence of Modules to write. May be null.</param>
+		private static void ExportJsonModuleArray(JsonWriter Writer, string ArrayName, IEnumerable<UEBuildModule> Modules)
+		{
+			if (Modules == null || !Modules.Any()) return;
+			
+			Writer.WriteArrayStart(ArrayName);
+			foreach (UEBuildModule Module in Modules)
+			{
+				Writer.WriteValue(Module.Name);
+			}
+			Writer.WriteArrayEnd();
+		}
+		
+		/// <summary>
+		/// Write an array of strings to a JSON writer. Ifl array is empty, don't write anything
+		/// </summary>
+		/// <param name="Writer">Writer for the array data</param>
+		/// <param name="ArrayName">Name of the array property</param>
+		/// <param name="Strings">Sequence of strings to write. May be null.</param>
+		static void ExportJsonStringArray(JsonWriter Writer, string ArrayName, IEnumerable<string> Strings)
+		{
+			if (Strings == null || !Strings.Any()) return;
+			
+			Writer.WriteArrayStart(ArrayName);
+			foreach (string String in Strings)
+			{
+				Writer.WriteValue(String);
+			}
+			Writer.WriteArrayEnd();
+		}
+		
+		/// <summary>
+		/// Write uplugin content to a JSON writer
+		/// </summary>
+		/// <param name="Plugin">Uplugin description</param>
+		/// <param name="Writer">JSON writer</param>
+		private static void ExportPlugin(UEBuildPlugin Plugin, JsonWriter Writer)
+		{
+			Writer.WriteObjectStart(Plugin.Name);
+			
+			Writer.WriteValue("File", Plugin.File.FullName);
+			Writer.WriteValue("Type", Plugin.Type.ToString());
+			if(Plugin.Dependencies.Any())
+			{
+				Writer.WriteStringArrayField("Dependencies", Plugin.Dependencies.Select(it => it.Name));
+			}
+			if(Plugin.Modules.Any())
+			{
+				Writer.WriteStringArrayField("Modules", Plugin.Modules.Select(it => it.Name));
+			}
+			
+			Writer.WriteObjectEnd();
+		}
+		
+		/// <summary>
+		/// Setup plugins for Target and write plugins to JSON writer. Don't write anything if there are no plugins 
+		/// </summary>
+		/// <param name="Target"></param>
+		/// <param name="Writer"></param>
+		private static void ExportPluginsFromTarget(UEBuildTarget Target, JsonWriter Writer)
+		{
+			Target.SetupPlugins();
+			if (!Target.BuildPlugins.Any()) return;
+			
+			Writer.WriteObjectStart("Plugins");
+			foreach (UEBuildPlugin plugin in Target.BuildPlugins)
+			{
+				ExportPlugin(plugin, Writer);
+			}
+			Writer.WriteObjectEnd();
+		}
+
+		/// <summary>
+		/// Write information about this binary to a JSON file
+		/// </summary>
+		/// <param name="Binary"></param>
+		/// <param name="Writer">Writer for this binary's data</param>
+		private static void ExportBinary(UEBuildBinary Binary, JsonWriter Writer)
+		{
+			Writer.WriteValue("File", Binary.OutputFilePath.FullName);
+			Writer.WriteValue("Type", Binary.Type.ToString());
+
+			Writer.WriteArrayStart("Modules");
+			foreach(UEBuildModule Module in Binary.Modules)
+			{
+				Writer.WriteValue(Module.Name);
+			}
+			Writer.WriteArrayEnd();
+		}
+		
+		/// <summary>
+		/// Write C++ toolchain information to JSON writer
+		/// </summary>
+		/// <param name="Target"></param>
+		/// <param name="Writer"></param>
+		private static void ExportEnvironmentToJson(UEBuildTarget Target, JsonWriter Writer)
+		{
+			CppCompileEnvironment GlobalCompileEnvironment = Target.CreateCompileEnvironmentForProjectFiles();
+			
+			Writer.WriteArrayStart("EnvironmentIncludePaths");
+			foreach (DirectoryReference Path in GlobalCompileEnvironment.UserIncludePaths)
+			{
+				Writer.WriteValue(Path.FullName);
+			}
+			foreach (DirectoryReference Path in GlobalCompileEnvironment.SystemIncludePaths)
+			{
+				Writer.WriteValue(Path.FullName);
+			}
+			
+			// TODO: get corresponding includes for specific platforms
+			if (UEBuildPlatform.IsPlatformInGroup(Target.Platform, UnrealPlatformGroup.Windows))
+			{
+				foreach (DirectoryReference Path in Target.Rules.WindowsPlatform.Environment.IncludePaths)
+				{
+					Writer.WriteValue(Path.FullName);
+				}
+			}
+			Writer.WriteArrayEnd();
+	
+			Writer.WriteArrayStart("EnvironmentDefinitions");
+			foreach (string Definition in GlobalCompileEnvironment.Definitions)
+			{
+				Writer.WriteValue(Definition);
+			}
+			Writer.WriteArrayEnd();
+		}
+	}
+}
\ No newline at end of file
diff --git a/Engine/Source/Programs/UnrealBuildTool/ProjectFiles/Rider/RiderProjectFileGenerator.cs b/Engine/Source/Programs/UnrealBuildTool/ProjectFiles/Rider/RiderProjectFileGenerator.cs
new file mode 100644
--- /dev/null
+++ b/Engine/Source/Programs/UnrealBuildTool/ProjectFiles/Rider/RiderProjectFileGenerator.cs
@@ -0,0 +1,499 @@
+using System;
+using System.Collections.Generic;
+using System.Linq;
+using Tools.DotNETCommon;
+
+namespace UnrealBuildTool
+{
+	class RiderMasterProjectFolder : MasterProjectFolder
+	{
+		public RiderMasterProjectFolder(ProjectFileGenerator InitOwnerProjectFileGenerator, string InitFolderName) :
+			base(InitOwnerProjectFileGenerator, InitFolderName)
+		{
+		}
+	}
+
+	class RiderProjectFileGenerator : ProjectFileGenerator
+	{
+		public override string ProjectFileExtension => ".json";
+
+		private readonly DirectoryReference EngineSourceProgramsDirectory =
+			DirectoryReference.Combine(UnrealBuildTool.EngineSourceDirectory, "Programs");
+
+		private readonly DirectoryReference EnterpriseSourceProgramsDirectory =
+			DirectoryReference.Combine(UnrealBuildTool.EnterpriseSourceDirectory, "Programs");
+
+		private readonly CommandLineArguments Arguments;
+
+
+		/// <summary>
+		/// Platforms to generate project files for
+		/// </summary>
+		[CommandLine("-Platforms=", ListSeparator = '+')]
+		HashSet<UnrealTargetPlatform> Platforms = new HashSet<UnrealTargetPlatform>();
+
+		/// <summary>
+		/// Target types to generate project files for
+		/// </summary>
+		[CommandLine("-TargetTypes=", ListSeparator = '+')]
+		HashSet<TargetType> TargetTypes = new HashSet<TargetType>();
+
+		/// <summary>
+		/// Target configurations to generate project files for
+		/// </summary>
+		[CommandLine("-TargetConfigurations=", ListSeparator = '+')]
+		HashSet<UnrealTargetConfiguration> TargetConfigurations = new HashSet<UnrealTargetConfiguration>();
+
+		/// <summary>
+		/// Projects to generate project filess for
+		/// </summary>
+		[CommandLine("-ProjectNames=", ListSeparator = '+')]
+		HashSet<string> ProjectNames = new HashSet<string>();
+
+		public RiderProjectFileGenerator(FileReference InOnlyGameProject,
+			CommandLineArguments InArguments)
+			: base(InOnlyGameProject)
+		{
+			Arguments = InArguments;
+			Arguments.ApplyTo(this);
+		}
+
+		public override bool ShouldGenerateIntelliSenseData() => true;
+
+		public override void CleanProjectFiles(DirectoryReference InMasterProjectDirectory, string InMasterProjectName,
+			DirectoryReference InIntermediateProjectFilesDirectory)
+		{
+			DirectoryReference.Delete(InMasterProjectDirectory);
+		}
+
+		private void ConfigureProjectFileGeneration()
+		{
+			bIncludeConfigFiles = true;
+			bIncludeEngineSource = true;
+			bIncludeDocumentation = false;
+			bIncludeShaderSource = true;
+			bGenerateProjectFiles = true;
+			bIncludeTemplateFiles = false;
+			bIncludeTestAndShippingConfigs = true;
+			IncludeEnginePrograms = true; // It's true by default, but I like to have things explicit
+		}
+
+		protected override ProjectFile AllocateProjectFile(FileReference InitFilePath)
+		{
+			RiderProjectFile projectFile = new RiderProjectFile(InitFilePath)
+				{RootPath = InitFilePath.Directory, Arguments = Arguments, TargetTypes = TargetTypes};
+			return projectFile;
+		}
+
+		public override MasterProjectFolder AllocateMasterProjectFolder(ProjectFileGenerator OwnerProjectFileGenerator,
+			string FolderName)
+		{
+			return new RiderMasterProjectFolder(OwnerProjectFileGenerator, FolderName);
+		}
+
+		protected override bool WriteProjectFiles(PlatformProjectGeneratorCollection PlatformProjectGenerators)
+		{
+			using (ProgressWriter Progress = new ProgressWriter("Writing project files...", true))
+			{
+				List<ProjectFile> projectsToGenerate = new List<ProjectFile>(GeneratedProjectFiles);
+				if (ProjectNames.Any())
+				{
+					projectsToGenerate = projectsToGenerate.Where(it =>
+						ProjectNames.Contains(it.ProjectFilePath.GetFileNameWithoutAnyExtensions())).ToList();
+				}
+
+				int TotalProjectFileCount = projectsToGenerate.Count;
+
+				HashSet<UnrealTargetPlatform> platformsToGenerate = new HashSet<UnrealTargetPlatform>(SupportedPlatforms);
+				if (Platforms.Any())
+				{
+					platformsToGenerate.IntersectWith(Platforms);
+				}
+
+				HashSet<UnrealTargetConfiguration> configurationsToGenerate = new HashSet<UnrealTargetConfiguration>(SupportedConfigurations);
+				if (TargetConfigurations.Any())
+				{
+					configurationsToGenerate.IntersectWith(TargetConfigurations);
+				}
+
+				for (int ProjectFileIndex = 0; ProjectFileIndex < projectsToGenerate.Count; ++ProjectFileIndex)
+				{
+					ProjectFile CurProject = projectsToGenerate[ProjectFileIndex];
+					if (!CurProject.WriteProjectFile(platformsToGenerate.ToList(), configurationsToGenerate.ToList(),
+						PlatformProjectGenerators))
+					{
+						return false;
+					}
+
+					Progress.Write(ProjectFileIndex + 1, TotalProjectFileCount);
+				}
+
+				Progress.Write(TotalProjectFileCount, TotalProjectFileCount);
+			}
+
+			return true;
+		}
+
+		private void AddProjectsForAllTargets(
+			PlatformProjectGeneratorCollection PlatformProjectGenerators,
+			List<FileReference> AllGames,
+			out ProjectFile EngineProject,
+			out List<ProjectFile> GameProjects)
+		{
+			// As we're creating project files, we'll also keep track of whether we created an "engine" project and return that if we have one
+			EngineProject = null;
+			GameProjects = new List<ProjectFile>();
+
+			// Find all of the target files.  This will filter out any modules or targets that don't
+			// belong to platforms we're generating project files for.
+			List<FileReference> AllTargetFiles = DiscoverTargets(AllGames);
+
+			// Sort the targets by name. When we have multiple targets of a given type for a project, we'll use the order to determine which goes in the primary project file (so that client names with a suffix will go into their own project).
+			AllTargetFiles = AllTargetFiles.OrderBy(x => x.FullName, StringComparer.OrdinalIgnoreCase).ToList();
+
+			foreach (FileReference TargetFilePath in AllTargetFiles)
+			{
+				string TargetName = TargetFilePath.GetFileNameWithoutAnyExtensions();
+
+				// Check to see if this is an Engine target.  That is, the target is located under the "Engine" folder
+				bool IsEngineTarget = false;
+				bool IsEnterpriseTarget = false;
+				bool WantProjectFileForTarget = true;
+				if (TargetFilePath.IsUnderDirectory(UnrealBuildTool.EngineDirectory))
+				{
+					// This is an engine target
+					IsEngineTarget = true;
+
+					if (TargetFilePath.IsUnderDirectory(EngineSourceProgramsDirectory))
+					{
+						WantProjectFileForTarget = IncludeEnginePrograms;
+					}
+					else if (TargetFilePath.IsUnderDirectory(UnrealBuildTool.EngineSourceDirectory))
+					{
+						WantProjectFileForTarget = bIncludeEngineSource;
+					}
+				}
+				else if (TargetFilePath.IsUnderDirectory(UnrealBuildTool.EnterpriseSourceDirectory))
+				{
+					// This is an enterprise target
+					IsEnterpriseTarget = true;
+
+					if (TargetFilePath.IsUnderDirectory(EnterpriseSourceProgramsDirectory))
+					{
+						WantProjectFileForTarget = bIncludeEnterpriseSource && IncludeEnginePrograms;
+					}
+					else
+					{
+						WantProjectFileForTarget = bIncludeEnterpriseSource;
+					}
+				}
+
+				if (WantProjectFileForTarget)
+				{
+					RulesAssembly RulesAssembly;
+
+					FileReference CheckProjectFile =
+						AllGames.FirstOrDefault(x => TargetFilePath.IsUnderDirectory(x.Directory));
+					if (CheckProjectFile == null)
+					{
+						if (TargetFilePath.IsUnderDirectory(UnrealBuildTool.EnterpriseDirectory))
+						{
+							RulesAssembly = RulesCompiler.CreateEnterpriseRulesAssembly(false, false);
+						}
+						else
+						{
+							RulesAssembly = RulesCompiler.CreateEngineRulesAssembly(false, false);
+						}
+					}
+					else
+					{
+						RulesAssembly = RulesCompiler.CreateProjectRulesAssembly(CheckProjectFile, false, false);
+					}
+
+					// Create target rules for all of the platforms and configuration combinations that we want to enable support for.
+					// Just use the current platform as we only need to recover the target type and both should be supported for all targets...
+					TargetRules TargetRulesObject = RulesAssembly.CreateTargetRules(TargetName,
+						BuildHostPlatform.Current.Platform, UnrealTargetConfiguration.Development, "", CheckProjectFile,
+						null);
+
+					bool IsProgramTarget = false;
+
+					DirectoryReference GameFolder = null;
+					string ProjectFileNameBase;
+					if (TargetRulesObject.Type == TargetType.Program)
+					{
+						IsProgramTarget = true;
+						ProjectFileNameBase = TargetName;
+					}
+					else if (IsEngineTarget)
+					{
+						ProjectFileNameBase = EngineProjectFileNameBase;
+					}
+					else if (IsEnterpriseTarget)
+					{
+						ProjectFileNameBase = EnterpriseProjectFileNameBase;
+					}
+					else
+					{
+						// Figure out which game project this target belongs to
+						FileReference ProjectInfo = FindGameContainingFile(AllGames, TargetFilePath);
+						if (ProjectInfo == null)
+						{
+							throw new BuildException("Found a non-engine target file (" + TargetFilePath +
+							                         ") that did not exist within any of the known game folders");
+						}
+
+						GameFolder = ProjectInfo.Directory;
+						ProjectFileNameBase = ProjectInfo.GetFileNameWithoutExtension();
+					}
+
+					// Get the suffix to use for this project file. If we have multiple targets of the same type, we'll have to split them out into separate IDE project files.
+					string GeneratedProjectName = TargetRulesObject.GeneratedProjectName;
+					if (GeneratedProjectName == null)
+					{
+						ProjectFile ExistingProjectFile;
+						if (ProjectFileMap.TryGetValue(GetProjectLocation(ProjectFileNameBase),
+							    out ExistingProjectFile) &&
+						    ExistingProjectFile.ProjectTargets.Any(x => x.TargetRules.Type == TargetRulesObject.Type))
+						{
+							GeneratedProjectName = TargetRulesObject.Name;
+						}
+						else
+						{
+							GeneratedProjectName = ProjectFileNameBase;
+						}
+					}
+
+					FileReference ProjectFilePath;
+					if (TargetRulesObject.Type != TargetType.Program && GameFolder != null)
+					{
+						ProjectFilePath = FileReference.Combine(GameFolder, "Intermediate", "ProjectFiles",
+							GeneratedProjectName + ProjectFileExtension);
+					}
+					else
+					{
+						ProjectFilePath = GetProjectLocation(GeneratedProjectName);
+					}
+
+					if (TargetRulesObject.Type == TargetType.Game || TargetRulesObject.Type == TargetType.Client ||
+					    TargetRulesObject.Type == TargetType.Server)
+					{
+						// Allow platforms to generate stub projects here...
+						PlatformProjectGenerators.GenerateGameProjectStubs(
+							InGenerator: this,
+							InTargetName: TargetName,
+							InTargetFilepath: TargetFilePath.FullName,
+							InTargetRules: TargetRulesObject,
+							InPlatforms: SupportedPlatforms,
+							InConfigurations: SupportedConfigurations);
+					}
+
+					DirectoryReference BaseFolder;
+					if (IsProgramTarget)
+					{
+						BaseFolder = TargetFilePath.Directory;
+					}
+					else if (IsEngineTarget)
+					{
+						BaseFolder = UnrealBuildTool.EngineDirectory;
+					}
+					else if (IsEnterpriseTarget)
+					{
+						BaseFolder = UnrealBuildTool.EnterpriseDirectory;
+					}
+					else
+					{
+						BaseFolder = GameFolder;
+					}
+
+					bool bProjectAlreadyExisted;
+					ProjectFile ProjectFile = FindOrAddProject(ProjectFilePath, BaseFolder,
+						true, out bProjectAlreadyExisted);
+					ProjectFile.IsForeignProject =
+						CheckProjectFile != null && !NativeProjects.IsNativeProject(CheckProjectFile);
+					ProjectFile.IsGeneratedProject = true;
+					ProjectFile.IsStubProject = UnrealBuildTool.IsProjectInstalled();
+					if (TargetRulesObject.bBuildInSolutionByDefault.HasValue)
+					{
+						ProjectFile.ShouldBuildByDefaultForSolutionTargets =
+							TargetRulesObject.bBuildInSolutionByDefault.Value;
+					}
+
+					// Add the project to the right output list
+					if (IsEngineTarget)
+					{
+						EngineProject = ProjectFile;
+						if (UnrealBuildTool.IsEngineInstalled())
+						{
+							// Allow engine projects to be created but not built for Installed Engine builds
+							EngineProject.IsForeignProject = false;
+							EngineProject.IsGeneratedProject = true;
+							EngineProject.IsStubProject = true;
+						}
+					}
+					else if (IsEnterpriseTarget)
+					{
+						ProjectFile EnterpriseProject = ProjectFile;
+						if (UnrealBuildTool.IsEnterpriseInstalled())
+						{
+							// Allow enterprise projects to be created but not built for Installed Engine builds
+							EnterpriseProject.IsForeignProject = false;
+							EnterpriseProject.IsGeneratedProject = true;
+							EnterpriseProject.IsStubProject = true;
+						}
+					}
+					else
+					{
+						if (!bProjectAlreadyExisted)
+						{
+							GameProjects.Add(ProjectFile);
+
+							// Add the .uproject file for this game/template
+							FileReference UProjectFilePath =
+								FileReference.Combine(BaseFolder, ProjectFileNameBase + ".uproject");
+							if (FileReference.Exists(UProjectFilePath))
+							{
+								ProjectFile.AddFileToProject(UProjectFilePath, BaseFolder);
+							}
+							else
+							{
+								throw new BuildException(
+									"Not expecting to find a game with no .uproject file.  File '{0}' doesn't exist",
+									UProjectFilePath);
+							}
+						}
+					}
+
+					foreach (ProjectTarget ExistingProjectTarget in ProjectFile.ProjectTargets)
+					{
+						if (ExistingProjectTarget.TargetRules.Type == TargetRulesObject.Type)
+						{
+							throw new BuildException(
+								"Not expecting project {0} to already have a target rules of with configuration name {1} ({2}) while trying to add: {3}",
+								ProjectFilePath, TargetRulesObject.Type.ToString(),
+								ExistingProjectTarget.TargetRules.ToString(), TargetRulesObject.ToString());
+						}
+
+						// Not expecting to have both a game and a program in the same project.  These would alias because we share the project and solution configuration names (just because it makes sense to)
+						if ((ExistingProjectTarget.TargetRules.Type == TargetType.Game &&
+						     TargetRulesObject.Type == TargetType.Program) ||
+						    (ExistingProjectTarget.TargetRules.Type == TargetType.Program &&
+						     TargetRulesObject.Type == TargetType.Game))
+						{
+							throw new BuildException(
+								"Not expecting project {0} to already have a Game/Program target ({1}) associated with it while trying to add: {2}",
+								ProjectFilePath, ExistingProjectTarget.TargetRules.ToString(),
+								TargetRulesObject.ToString());
+						}
+					}
+
+					ProjectTarget ProjectTarget = new ProjectTarget()
+					{
+						TargetRules = TargetRulesObject,
+						TargetFilePath = TargetFilePath,
+						ProjectFilePath = ProjectFilePath,
+						UnrealProjectFilePath = CheckProjectFile,
+						SupportedPlatforms = TargetRulesObject.GetSupportedPlatforms()
+							.Where(x => UEBuildPlatform.GetBuildPlatform(x, true) != null).ToArray(),
+						CreateRulesDelegate = (Platform, Configuration) =>
+							RulesAssembly.CreateTargetRules(TargetName, Platform, Configuration, "", CheckProjectFile,
+								null)
+					};
+
+					ProjectFile.ProjectTargets.Add(ProjectTarget);
+
+					// Make sure the *.Target.cs file is in the project.
+					ProjectFile.AddFileToProject(TargetFilePath, BaseFolder);
+
+					Log.TraceVerbose("Generating target {0} for {1}", TargetRulesObject.Type.ToString(),
+						ProjectFilePath);
+				}
+			}
+		}
+		
+		private void SetupSupportedPlatformsAndConfigurations()
+		{
+			SetupSupportedPlatformsAndConfigurations(true, out _);
+		}
+		
+		public override bool GenerateProjectFiles(PlatformProjectGeneratorCollection PlatformProjectGenerators,
+			String[] arguments)
+		{
+			ConfigureProjectFileGeneration();
+
+			if (bGeneratingGameProjectFiles || UnrealBuildTool.IsEngineInstalled())
+			{
+				MasterProjectPath = OnlyGameProject.Directory;
+				MasterProjectName = OnlyGameProject.GetFileNameWithoutExtension();
+
+				if (!DirectoryReference.Exists(DirectoryReference.Combine(MasterProjectPath, "Source")))
+				{
+					if (!DirectoryReference.Exists(DirectoryReference.Combine(MasterProjectPath, "Intermediate",
+						"Source")))
+					{
+						if (BuildHostPlatform.Current.Platform == UnrealTargetPlatform.Mac)
+						{
+							MasterProjectPath = UnrealBuildTool.EngineDirectory;
+							GameProjectName = "UE4Game";
+						}
+
+						if (!DirectoryReference.Exists(DirectoryReference.Combine(MasterProjectPath, "Source")))
+						{
+							throw new BuildException("Directory '{0}' is missing 'Source' folder.", MasterProjectPath);
+						}
+					}
+				}
+
+				IntermediateProjectFilesPath =
+					DirectoryReference.Combine(MasterProjectPath, "Intermediate", "ProjectFiles");
+			}
+
+			SetupSupportedPlatformsAndConfigurations();
+
+			Log.TraceVerbose("Detected supported platforms: " + SupportedPlatforms);
+
+			List<FileReference> AllGameProjects = FindGameProjects();
+
+			GatherProjects(PlatformProjectGenerators, AllGameProjects);
+
+			WriteProjectFiles(PlatformProjectGenerators);
+			Log.TraceVerbose("Project generation complete ({0} generated, {1} imported)", GeneratedProjectFiles.Count,
+				OtherProjectFiles.Count);
+
+			return true;
+		}
+
+		private void GatherProjects(PlatformProjectGeneratorCollection PlatformProjectGenerators,
+			List<FileReference> AllGameProjects)
+		{ 
+			ProjectFile EngineProject = null;
+			List<ProjectFile> GameProjects = null;
+			List<ProjectFile> ModProjects = null;
+			// Setup buildable projects for all targets
+			AddProjectsForAllTargets(PlatformProjectGenerators, AllGameProjects, out EngineProject,
+				out GameProjects);
+
+			AddProjectsForMods(GameProjects, out ModProjects);
+			AddAllGameProjects(GameProjects,null,  null);
+
+			// If we're still missing an engine project because we don't have any targets for it, make one up.
+			if (EngineProject == null)
+			{
+				FileReference ProjectFilePath = FileReference.Combine(IntermediateProjectFilesPath,
+					EngineProjectFileNameBase + ProjectFileExtension);
+
+				EngineProject = FindOrAddProject(ProjectFilePath, UnrealBuildTool.EngineDirectory, true, out _);
+
+				EngineProject.IsForeignProject = false;
+				EngineProject.IsGeneratedProject = true;
+				EngineProject.IsStubProject = true;
+			}
+		}
+
+		protected override bool WriteMasterProjectFile(ProjectFile ProjectFile,
+			PlatformProjectGeneratorCollection PlatformProjectGenerators)
+		{
+			return true;
+		}
+	}
+}
\ No newline at end of file
diff --git a/Engine/Source/Programs/UnrealBuildTool/UnrealBuildTool.csproj b/Engine/Source/Programs/UnrealBuildTool/UnrealBuildTool.csproj
--- a/Engine/Source/Programs/UnrealBuildTool/UnrealBuildTool.csproj
+++ b/Engine/Source/Programs/UnrealBuildTool/UnrealBuildTool.csproj
@@ -180,6 +180,8 @@
     <Compile Include="ProjectFiles\Project.cs" />
     <Compile Include="ProjectFiles\ProjectFileGenerator.cs" />
     <Compile Include="ProjectFiles\QMake\QMakefileGenerator.cs" />
+    <Compile Include="ProjectFiles\Rider\RiderProjectFile.cs" />
+    <Compile Include="ProjectFiles\Rider\RiderProjectFileGenerator.cs" />
     <Compile Include="ProjectFiles\VisualStudioCode\VSCodeProjectFileGenerator.cs" />
     <Compile Include="ProjectFiles\VisualStudio\VCProject.cs" />
     <Compile Include="ProjectFiles\VisualStudio\VCProjectFileGenerator.cs" />
