Fix C# 7 anachronisms.

#rnx

diff --git a/Engine/Source/Programs/UnrealBuildTool/ProjectFiles/Rider/RiderProjectFile.cs b/Engine/Source/Programs/UnrealBuildTool/ProjectFiles/Rider/RiderProjectFile.cs
--- a/Engine/Source/Programs/UnrealBuildTool/ProjectFiles/Rider/RiderProjectFile.cs
+++ b/Engine/Source/Programs/UnrealBuildTool/ProjectFiles/Rider/RiderProjectFile.cs
@@ -47,18 +47,16 @@ namespace UnrealBuildTool
 					{
 						if (TargetTypes.Any() && !TargetTypes.Contains(ProjectTarget.TargetRules.Type)) continue;
 
-						switch (ProjectTarget.TargetRules.Type)
+						// Skip Programs for all configs except for current platform + Development configuration
+						if (ProjectTarget.TargetRules.Type == TargetType.Program && (BuildHostPlatform.Current.Platform != Platform || Configuration != UnrealTargetConfiguration.Development))
 						{
-							// Skip Programs for all configs except for current platform + Development configuration
-							case TargetType.Program
-								when (BuildHostPlatform.Current.Platform != Platform ||
-								      Configuration != UnrealTargetConfiguration.Development):
-							// Skip Editor for all platforms except for current platform
-							case TargetType.Editor
-								when (BuildHostPlatform.Current.Platform != Platform ||
-								      (Configuration == UnrealTargetConfiguration.Test ||
-								       Configuration == UnrealTargetConfiguration.Shipping)):
-								continue;
+							continue;
+						}
+
+						// Skip Editor for all platforms except for current platform
+						if (ProjectTarget.TargetRules.Type == TargetType.Editor && (BuildHostPlatform.Current.Platform != Platform || (Configuration == UnrealTargetConfiguration.Test || Configuration == UnrealTargetConfiguration.Shipping)))
+						{
+							continue;
 						}
 
 						DirectoryReference TargetFolder =
@@ -147,7 +145,9 @@ namespace UnrealBuildTool
 			Writer.WriteValue("Directory", Module.ModuleDirectory.FullName);
 			Writer.WriteValue("Rules", Module.RulesFile.FullName);
 			Writer.WriteValue("PCHUsage", Module.Rules.PCHUsage.ToString());
-			if (Module is UEBuildModuleCPP ModuleCPP)
+
+			UEBuildModuleCPP ModuleCPP = Module as UEBuildModuleCPP;
+			if (ModuleCPP != null)
 			{
 				Writer.WriteValue("GeneratedCodeDirectory", ModuleCPP.GeneratedCodeDirectory != null ? ModuleCPP.GeneratedCodeDirectory.FullName : string.Empty);
 			}
diff --git a/Engine/Source/Programs/UnrealBuildTool/ProjectFiles/Rider/RiderProjectFileGenerator.cs b/Engine/Source/Programs/UnrealBuildTool/ProjectFiles/Rider/RiderProjectFileGenerator.cs
--- a/Engine/Source/Programs/UnrealBuildTool/ProjectFiles/Rider/RiderProjectFileGenerator.cs
+++ b/Engine/Source/Programs/UnrealBuildTool/ProjectFiles/Rider/RiderProjectFileGenerator.cs
@@ -413,7 +413,8 @@ namespace UnrealBuildTool
 		
 		private void SetupSupportedPlatformsAndConfigurations()
 		{
-			SetupSupportedPlatformsAndConfigurations(true, out _);
+			string SupportedPlatformNames;
+			SetupSupportedPlatformsAndConfigurations(true, out SupportedPlatformNames);
 		}
 		
 		public override bool GenerateProjectFiles(PlatformProjectGeneratorCollection PlatformProjectGenerators,
@@ -482,7 +483,8 @@ namespace UnrealBuildTool
 				FileReference ProjectFilePath = FileReference.Combine(IntermediateProjectFilesPath,
 					EngineProjectFileNameBase + ProjectFileExtension);
 
-				EngineProject = FindOrAddProject(ProjectFilePath, UnrealBuildTool.EngineDirectory, true, out _);
+				bool bAlreadyExisted;
+				EngineProject = FindOrAddProject(ProjectFilePath, UnrealBuildTool.EngineDirectory, true, out bAlreadyExisted);
 
 				EngineProject.IsForeignProject = false;
 				EngineProject.IsGeneratedProject = true;
diff --git a/Engine/Source/Programs/UnrealBuildTool/UnrealBuildTool.csproj b/Engine/Source/Programs/UnrealBuildTool/UnrealBuildTool.csproj
--- a/Engine/Source/Programs/UnrealBuildTool/UnrealBuildTool.csproj
+++ b/Engine/Source/Programs/UnrealBuildTool/UnrealBuildTool.csproj
@@ -48,6 +48,7 @@
     <DocumentationFile>..\..\..\Binaries\DotNET\UnrealBuildTool.xml</DocumentationFile>
     <NoWarn>
     </NoWarn>
+    <LangVersion>6</LangVersion>
   </PropertyGroup>
   <PropertyGroup Condition=" '$(Configuration)|$(Platform)' == 'Debug|AnyCPU' ">
     <OutputPath>..\..\..\Binaries\DotNET\</OutputPath>
