- Support for UHT to generate headers in platform extension modules (they are put in the project's main Intermediate directory along with the other headers)

diff --git a/Engine/Source/Programs/UnrealBuildTool/Configuration/UEBuildModule.cs b/Engine/Source/Programs/UnrealBuildTool/Configuration/UEBuildModule.cs
--- a/Engine/Source/Programs/UnrealBuildTool/Configuration/UEBuildModule.cs
+++ b/Engine/Source/Programs/UnrealBuildTool/Configuration/UEBuildModule.cs
@@ -41,7 +41,7 @@ namespace UnrealBuildTool
 		/// <summary>
 		/// Paths to all potential module source directories (with platform extension directories added in)
 		/// </summary>
-		protected DirectoryReference[] ModuleDirectories;
+		public DirectoryReference[] ModuleDirectories;
 
 		/// <summary>
 		/// The name of the .Build.cs file this module was created from, if any
diff --git a/Engine/Source/Programs/UnrealBuildTool/System/ExternalExecution.cs b/Engine/Source/Programs/UnrealBuildTool/System/ExternalExecution.cs
--- a/Engine/Source/Programs/UnrealBuildTool/System/ExternalExecution.cs
+++ b/Engine/Source/Programs/UnrealBuildTool/System/ExternalExecution.cs
@@ -118,10 +118,10 @@ namespace UnrealBuildTool
 					return UHTModuleType.Program;
 				case ModuleHostType.Runtime:
 				case ModuleHostType.RuntimeNoCommandlet:
-                case ModuleHostType.RuntimeAndProgram:
-                case ModuleHostType.CookedOnly:
-                case ModuleHostType.ServerOnly:
-                case ModuleHostType.ClientOnly:
+				case ModuleHostType.RuntimeAndProgram:
+				case ModuleHostType.CookedOnly:
+				case ModuleHostType.ServerOnly:
+				case ModuleHostType.ClientOnly:
 				case ModuleHostType.ClientOnlyNoCommandlet:
 					return UHTModuleType.EngineRuntime;
 				case ModuleHostType.Developer:
@@ -141,12 +141,12 @@ namespace UnrealBuildTool
 					return UHTModuleType.Program;
 				case ModuleHostType.Runtime:
 				case ModuleHostType.RuntimeNoCommandlet:
-                case ModuleHostType.RuntimeAndProgram:
-                case ModuleHostType.CookedOnly:
-                case ModuleHostType.ServerOnly:
-                case ModuleHostType.ClientOnly:
+				case ModuleHostType.RuntimeAndProgram:
+				case ModuleHostType.CookedOnly:
+				case ModuleHostType.ServerOnly:
+				case ModuleHostType.ClientOnly:
 				case ModuleHostType.ClientOnlyNoCommandlet:
-                    return UHTModuleType.GameRuntime;
+					return UHTModuleType.GameRuntime;
 				case ModuleHostType.Developer:
 					return UHTModuleType.GameDeveloper;
 				case ModuleHostType.Editor:
@@ -174,9 +174,9 @@ namespace UnrealBuildTool
 		public FileReference ModuleRulesFile;
 
 		/// <summary>
-		/// Module base directory
+		/// Paths to all potential module source directories (with platform extension directories added in)
 		/// </summary>
-		public DirectoryReference ModuleDirectory;
+		public DirectoryReference[] ModuleDirectories;
 
 		/// <summary>
 		/// Module type
@@ -218,11 +218,11 @@ namespace UnrealBuildTool
 		/// </summary>
 		public bool bIsReadOnly;
 
-		public UHTModuleInfo(string ModuleName, FileReference ModuleRulesFile, DirectoryReference ModuleDirectory, UHTModuleType ModuleType, DirectoryItem GeneratedCodeDirectory, EGeneratedCodeVersion GeneratedCodeVersion, bool bIsReadOnly)
+		public UHTModuleInfo(string ModuleName, FileReference ModuleRulesFile, DirectoryReference[] ModuleDirectories, UHTModuleType ModuleType, DirectoryItem GeneratedCodeDirectory, EGeneratedCodeVersion GeneratedCodeVersion, bool bIsReadOnly)
 		{
 			this.ModuleName = ModuleName;
 			this.ModuleRulesFile = ModuleRulesFile;
-			this.ModuleDirectory = ModuleDirectory;
+			this.ModuleDirectories = ModuleDirectories;
 			this.ModuleType = ModuleType.ToString();
 			this.PublicUObjectClassesHeaders = new List<FileItem>();
 			this.PublicUObjectHeaders = new List<FileItem>();
@@ -236,7 +236,7 @@ namespace UnrealBuildTool
 		{
 			ModuleName = Reader.ReadString();
 			ModuleRulesFile = Reader.ReadFileReference();
-			ModuleDirectory = Reader.ReadDirectoryReference();
+			ModuleDirectories = Reader.ReadArray<DirectoryReference>(Reader.ReadDirectoryReference);
 			ModuleType = Reader.ReadString();
 			PublicUObjectClassesHeaders = Reader.ReadList(() => Reader.ReadFileItem());
 			PublicUObjectHeaders = Reader.ReadList(() => Reader.ReadFileItem());
@@ -251,7 +251,7 @@ namespace UnrealBuildTool
 		{
 			Writer.WriteString(ModuleName);
 			Writer.WriteFileReference(ModuleRulesFile);
-			Writer.WriteDirectoryReference(ModuleDirectory);
+			Writer.WriteArray<DirectoryReference>(ModuleDirectories, Writer.WriteDirectoryReference);
 			Writer.WriteString(ModuleType);
 			Writer.WriteList(PublicUObjectClassesHeaders, Item => Writer.WriteFileItem(Item));
 			Writer.WriteList(PublicUObjectHeaders, Item => Writer.WriteFileItem(Item));
@@ -315,8 +315,8 @@ namespace UnrealBuildTool
 			{
 				Name = Info.ModuleName;
 				ModuleType = Info.ModuleType;
-				BaseDirectory = Info.ModuleDirectory.FullName;
-				IncludeBase = Info.ModuleDirectory.ParentDirectory.FullName;
+				BaseDirectory = Info.ModuleDirectories[0].FullName;
+				IncludeBase = Info.ModuleDirectories[0].ParentDirectory.FullName;
 				OutputDirectory = Path.GetDirectoryName(Info.GeneratedCPPFilenameBase);
 				ClassesHeaders = Info.PublicUObjectClassesHeaders.Select((Header) => Header.AbsolutePath).ToList();
 				PublicHeaders = Info.PublicUObjectHeaders.Select((Header) => Header.AbsolutePath).ToList();
@@ -383,23 +383,23 @@ namespace UnrealBuildTool
 	{
 		static UHTModuleType GetEngineModuleTypeFromDescriptor(ModuleDescriptor Module)
 		{
-            UHTModuleType? Type = UHTModuleTypeExtensions.EngineModuleTypeFromHostType(Module.Type);
-            if (Type == null)
-            {
-                throw new BuildException("Unhandled engine module type {0} for {1}", Module.Type.ToString(), Module.Name);
-            }
-            return Type.GetValueOrDefault();
-        }
+			UHTModuleType? Type = UHTModuleTypeExtensions.EngineModuleTypeFromHostType(Module.Type);
+			if (Type == null)
+			{
+				throw new BuildException("Unhandled engine module type {0} for {1}", Module.Type.ToString(), Module.Name);
+			}
+			return Type.GetValueOrDefault();
+		}
 
 		static UHTModuleType GetGameModuleTypeFromDescriptor(ModuleDescriptor Module)
 		{
-            UHTModuleType? Type = UHTModuleTypeExtensions.GameModuleTypeFromHostType(Module.Type);
-            if (Type == null)
-            {
-                throw new BuildException("Unhandled game module type {0}", Module.Type.ToString());
-            }
-            return Type.GetValueOrDefault();
-        }
+			UHTModuleType? Type = UHTModuleTypeExtensions.GameModuleTypeFromHostType(Module.Type);
+			if (Type == null)
+			{
+				throw new BuildException("Unhandled game module type {0}", Module.Type.ToString());
+			}
+			return Type.GetValueOrDefault();
+		}
 
 		/// <summary>
 		/// Returns a copy of Nodes sorted by dependency.  Independent or circularly-dependent nodes should
@@ -408,7 +408,7 @@ namespace UnrealBuildTool
 		/// <param name="NodeList">The list of nodes to sort.</param>
 		static void StableTopologicalSort(List<UEBuildModuleCPP> NodeList)
 		{
-			int            NodeCount = NodeList.Count;
+			int NodeCount = NodeList.Count;
 
 			Dictionary<UEBuildModule, HashSet<UEBuildModule>> Cache = new Dictionary<UEBuildModule, HashSet<UEBuildModule>>();
 
@@ -562,7 +562,7 @@ namespace UnrealBuildTool
 				{
 					UEBuildModuleCPP Module = ModulesSortedByType[Idx];
 
-					UHTModuleInfo Info = new UHTModuleInfo(Module.Name, Module.RulesFile, Module.ModuleDirectory, ModuleToType[Module], DirectoryItem.GetItemByDirectoryReference(Module.GeneratedCodeDirectory), GeneratedCodeVersion, Module.Rules.bUsePrecompiled);
+					UHTModuleInfo Info = new UHTModuleInfo(Module.Name, Module.RulesFile, Module.ModuleDirectories, ModuleToType[Module], DirectoryItem.GetItemByDirectoryReference(Module.GeneratedCodeDirectory), GeneratedCodeVersion, Module.Rules.bUsePrecompiled);
 					ModuleInfoArray[Idx] = Info;
 
 					Queue.Enqueue(() => SetupUObjectModule(Info, ExcludedFolders, MetadataCache, Queue));
@@ -613,14 +613,17 @@ namespace UnrealBuildTool
 
 		static void SetupUObjectModule(UHTModuleInfo ModuleInfo, ReadOnlyHashSet<string> ExcludedFolders, SourceFileMetadataCache MetadataCache, ThreadPoolWorkQueue Queue)
 		{
-			DirectoryItem ModuleDirectoryItem = DirectoryItem.GetItemByDirectoryReference(ModuleInfo.ModuleDirectory);
+			foreach (DirectoryReference ModuleDirectory in ModuleInfo.ModuleDirectories)
+			{
+				DirectoryItem ModuleDirectoryItem = DirectoryItem.GetItemByDirectoryReference(ModuleDirectory);
 
-			List<FileItem> HeaderFiles = new List<FileItem>();
-			FindHeaders(ModuleDirectoryItem, ExcludedFolders, HeaderFiles);
+				List<FileItem> HeaderFiles = new List<FileItem>();
+				FindHeaders(ModuleDirectoryItem, ExcludedFolders, HeaderFiles);
 
-			foreach (FileItem HeaderFile in HeaderFiles)
-			{
-				Queue.Enqueue(() => SetupUObjectModuleHeader(ModuleInfo, HeaderFile, MetadataCache));
+				foreach (FileItem HeaderFile in HeaderFiles)
+				{
+					Queue.Enqueue(() => SetupUObjectModuleHeader(ModuleInfo, HeaderFile, MetadataCache));
+				}
 			}
 		}
 
@@ -632,15 +635,21 @@ namespace UnrealBuildTool
 			{
 				lock(ModuleInfo)
 				{
-					if (HeaderFile.Location.IsUnderDirectory(DirectoryReference.Combine(ModuleInfo.ModuleDirectory, "Classes")))
-					{
-						ModuleInfo.PublicUObjectClassesHeaders.Add(HeaderFile);
-					}
-					else if (HeaderFile.Location.IsUnderDirectory(DirectoryReference.Combine(ModuleInfo.ModuleDirectory, "Public")))
+					bool bFoundHeaderLocation = false;
+					foreach (DirectoryReference ModuleDirectory in ModuleInfo.ModuleDirectories)
 					{
-						ModuleInfo.PublicUObjectHeaders.Add(HeaderFile);
+						if (HeaderFile.Location.IsUnderDirectory(DirectoryReference.Combine(ModuleDirectory, "Classes")))
+						{
+							ModuleInfo.PublicUObjectClassesHeaders.Add(HeaderFile);
+							bFoundHeaderLocation = true;
+						}
+						else if (HeaderFile.Location.IsUnderDirectory(DirectoryReference.Combine(ModuleDirectory, "Public")))
+						{
+							ModuleInfo.PublicUObjectHeaders.Add(HeaderFile);
+							bFoundHeaderLocation = true;
+						}
 					}
-					else
+					if (!bFoundHeaderLocation)
 					{
 						ModuleInfo.PrivateUObjectHeaders.Add(HeaderFile);
 					}
