Add a "GenerateClangDatabase" mode to UBT, which will generate a compile_commands.json file in the workspace root directory for Clang tooling. Filter which files to include using the -Filter=... argument (eg. -Filter=QAGame/Source/...).

Example usage:

UnrealBuildTool.exe UnrealPak Win64 Development -Mode=GenerateClangDatabase -Filter=Engine/Source/Programs/...

diff --git a/Engine/Source/Programs/DotNETCommon/DotNETUtilities/JsonWriter.cs b/Engine/Source/Programs/DotNETCommon/DotNETUtilities/JsonWriter.cs
--- a/Engine/Source/Programs/DotNETCommon/DotNETUtilities/JsonWriter.cs
+++ b/Engine/Source/Programs/DotNETCommon/DotNETUtilities/JsonWriter.cs
@@ -104,6 +104,20 @@ namespace Tools.DotNETCommon
 			bRequiresComma = true;
 		}
 
+		/// <summary>
+		/// Write the opening bracket for an unnamed array
+		/// </summary>
+		/// <param name="ArrayName">Name of the field</param>
+		public void WriteArrayStart()
+		{
+			WriteCommaNewline();
+
+			Writer.Write("{0}[", Indent);
+
+			Indent += "\t";
+			bRequiresComma = false;
+		}
+
 		/// <summary>
 		/// Write the name and opening bracket for an array
 		/// </summary>
diff --git a/Engine/Source/Programs/UnrealBuildTool/Modes/GenerateClangDatabase.cs b/Engine/Source/Programs/UnrealBuildTool/Modes/GenerateClangDatabase.cs
new file mode 100644
--- /dev/null
+++ b/Engine/Source/Programs/UnrealBuildTool/Modes/GenerateClangDatabase.cs
@@ -0,0 +1,198 @@
+using System;
+using System.Collections.Generic;
+using System.Linq;
+using System.Text;
+using System.Threading.Tasks;
+using Tools.DotNETCommon;
+
+namespace UnrealBuildTool
+{
+	/// <summary>
+	/// Generate a clang compile_commands file for a target
+	/// </summary>
+	[ToolMode("GenerateClangDatabase", ToolModeOptions.XmlConfig | ToolModeOptions.BuildPlatforms | ToolModeOptions.SingleInstance | ToolModeOptions.StartPrefetchingEngine | ToolModeOptions.ShowExecutionTime)]
+	class GenerateClangDatabase : ToolMode
+	{
+		/// <summary>
+		/// Set of filters for files to include in the database. Relative to the root directory, or to the project file.
+		/// </summary>
+		[CommandLine("-Filter=")]
+		List<string> FilterRules = new List<string>();
+
+		/// <summary>
+		/// Execute the command
+		/// </summary>
+		/// <param name="Arguments">Command line arguments</param>
+		/// <returns>Exit code</returns>
+		public override int Execute(CommandLineArguments Arguments)
+		{
+			Arguments.ApplyTo(this);
+
+			// Create the build configuration object, and read the settings
+			BuildConfiguration BuildConfiguration = new BuildConfiguration();
+			XmlConfig.ApplyTo(BuildConfiguration);
+			Arguments.ApplyTo(BuildConfiguration);
+
+			// Parse the filter argument
+			FileFilter FileFilter = null;
+			if(FilterRules.Count > 0)
+			{
+				FileFilter = new FileFilter(FileFilterType.Exclude);
+				foreach(string FilterRule in FilterRules)
+				{
+					FileFilter.AddRules(FilterRule.Split(';'));
+				}
+			}
+
+			// Parse all the target descriptors
+			List<TargetDescriptor> TargetDescriptors = TargetDescriptor.ParseCommandLine(Arguments, BuildConfiguration.bUsePrecompiled, BuildConfiguration.bSkipRulesCompile);
+
+			// Generate the compile DB for each target
+			using (ISourceFileWorkingSet WorkingSet = new EmptySourceFileWorkingSet())
+			{
+				// Find the compile commands for each file in the target
+				Dictionary<FileReference, string> FileToCommand = new Dictionary<FileReference, string>();
+				foreach (TargetDescriptor TargetDescriptor in TargetDescriptors)
+				{
+					// Disable PCHs and unity builds for the target
+					TargetDescriptor.AdditionalArguments = TargetDescriptor.AdditionalArguments.Append(new string[] { "-NoPCH", "-DisableUnity", "-ToolChain=ClangDatabaseProxyToolChain" });
+
+					// Create a makefile for the target
+					UEBuildTarget Target = UEBuildTarget.Create(TargetDescriptor, BuildConfiguration.bSkipRulesCompile, BuildConfiguration.bUsePrecompiled);
+
+					// Find the location of the compiler
+					VCEnvironment Environment = VCEnvironment.Create(WindowsCompiler.Clang, Target.Platform, Target.Rules.WindowsPlatform.Architecture, null, Target.Rules.WindowsPlatform.WindowsSdkVersion);
+					FileReference ClangPath = FileReference.Combine(Environment.CompilerDir, "clang.exe");
+
+					// Create all the binaries and modules
+					CppCompileEnvironment GlobalCompileEnvironment = Target.CreateCompileEnvironmentForProjectFiles();
+					foreach (UEBuildBinary Binary in Target.Binaries)
+					{
+						CppCompileEnvironment BinaryCompileEnvironment = Binary.CreateBinaryCompileEnvironment(GlobalCompileEnvironment);
+						foreach (UEBuildModuleCPP Module in Binary.Modules.OfType<UEBuildModuleCPP>())
+						{
+							if(!Module.Rules.bUsePrecompiled)
+							{
+								UEBuildModuleCPP.InputFileCollection InputFileCollection = Module.FindInputFiles(Target.Platform, new Dictionary<DirectoryItem, FileItem[]>());
+
+								List<FileItem> InputFiles = new List<FileItem>();
+								InputFiles.AddRange(InputFileCollection.CPPFiles);
+								InputFiles.AddRange(InputFileCollection.CCFiles);
+
+								CppCompileEnvironment ModuleCompileEnvironment = Module.CreateModuleCompileEnvironment(Target.Rules, BinaryCompileEnvironment);
+
+								StringBuilder CommandBuilder = new StringBuilder();
+								CommandBuilder.AppendFormat("\"{0}\"", ClangPath.FullName);
+								foreach (FileItem ForceIncludeFile in ModuleCompileEnvironment.ForceIncludeFiles)
+								{
+									CommandBuilder.AppendFormat(" -include \"{0}\"", ForceIncludeFile.FullName);
+								}
+								foreach (string Definition in ModuleCompileEnvironment.Definitions)
+								{
+									CommandBuilder.AppendFormat(" -D\"{0}\"", Definition);
+								}
+								foreach (DirectoryReference IncludePath in ModuleCompileEnvironment.UserIncludePaths)
+								{
+									CommandBuilder.AppendFormat(" -I\"{0}\"", IncludePath);
+								}
+								foreach (DirectoryReference IncludePath in ModuleCompileEnvironment.SystemIncludePaths)
+								{
+									CommandBuilder.AppendFormat(" -I\"{0}\"", IncludePath);
+								}
+
+								string Command = CommandBuilder.ToString();
+								foreach (FileItem InputFile in InputFiles)
+								{
+									if(FileFilter == null || FileFilter.Matches(InputFile.Location.MakeRelativeTo(UnrealBuildTool.RootDirectory)))
+									{
+										FileToCommand[InputFile.Location] = Command;
+									}
+								}
+							}
+						}
+					}
+				}
+
+				// Write the compile database
+				FileReference DatabaseFile = FileReference.Combine(UnrealBuildTool.RootDirectory, "compile_commands.json");
+				using (JsonWriter Writer = new JsonWriter(DatabaseFile))
+				{
+					Writer.WriteArrayStart();
+					foreach(KeyValuePair<FileReference, string> FileCommandPair in FileToCommand.OrderBy(x => x.Key.FullName))
+					{
+						Writer.WriteObjectStart();
+						Writer.WriteValue("file", FileCommandPair.Key.FullName);
+						Writer.WriteValue("command", FileCommandPair.Value);
+						Writer.WriteValue("directory", UnrealBuildTool.EngineSourceDirectory.ToString());
+						Writer.WriteObjectEnd();
+					}
+					Writer.WriteArrayEnd();
+				}
+			}
+
+			return 0;
+		}
+	}
+
+	/// <summary>
+	/// Toolchain implementation which just dumps arguments to the compiler
+	/// </summary>
+	class ClangDatabaseProxyToolChain : UEToolChain
+	{
+		FileReference ClangPath;
+
+		public ClangDatabaseProxyToolChain(ReadOnlyTargetRules Target)
+		{
+			VCEnvironment Environment = VCEnvironment.Create(WindowsCompiler.Clang, Target.Platform, Target.WindowsPlatform.Architecture, null, Target.WindowsPlatform.WindowsSdkVersion);
+			ClangPath = FileReference.Combine(Environment.CompilerDir, "clang.exe");
+		}
+
+		public override CPPOutput CompileCPPFiles(CppCompileEnvironment CompileEnvironment, List<FileItem> InputFiles, DirectoryReference OutputDir, string ModuleName, List<Action> Actions)
+		{
+			StringBuilder Arguments = new StringBuilder();
+
+			foreach(FileItem ForceIncludeFile in CompileEnvironment.ForceIncludeFiles)
+			{
+				Arguments.AppendFormat(" -include \"{0}\"", ForceIncludeFile.FullName);
+			}
+			foreach (string Definition in CompileEnvironment.Definitions)
+			{
+				Arguments.AppendFormat(" -D\"{0}\"", Definition);
+			}
+			foreach (DirectoryReference IncludePath in CompileEnvironment.UserIncludePaths)
+			{
+				Arguments.AppendFormat(" -I\"{0}\"", IncludePath);
+			}
+			foreach (DirectoryReference IncludePath in CompileEnvironment.SystemIncludePaths)
+			{
+				Arguments.AppendFormat(" -I\"{0}\"", IncludePath);
+			}
+
+			List<FileItem> ObjectFiles = new List<FileItem>();
+			foreach (FileItem InputFile in InputFiles)
+			{
+				Action NewAction = new Action(ActionType.Compile);
+				NewAction.CommandPath = ClangPath;
+				NewAction.CommandArguments = String.Format("{0} \"{1}\"", Arguments.ToString(), InputFile.AbsolutePath);
+				NewAction.WorkingDirectory = UnrealBuildTool.EngineSourceDirectory;
+				NewAction.PrerequisiteItems.Add(InputFile);
+				NewAction.ProducedItems.Add(FileItem.GetItemByFileReference(InputFile.Location.ChangeExtension(".dummy.obj")));
+				Actions.Add(NewAction);
+
+				ObjectFiles.AddRange(NewAction.ProducedItems);
+			}
+
+			return new CPPOutput { ObjectFiles = ObjectFiles };
+		}
+
+		public override FileItem LinkFiles(LinkEnvironment LinkEnvironment, bool bBuildImportLibraryOnly, List<Action> Actions)
+		{
+			throw new NotImplementedException();
+		}
+
+		public override FileItem[] LinkAllFiles(LinkEnvironment LinkEnvironment, bool bBuildImportLibraryOnly, List<Action> Actions)
+		{
+			return new FileItem[0];
+		}
+	}
+}
diff --git a/Engine/Source/Programs/UnrealBuildTool/UnrealBuildTool.csproj b/Engine/Source/Programs/UnrealBuildTool/UnrealBuildTool.csproj
--- a/Engine/Source/Programs/UnrealBuildTool/UnrealBuildTool.csproj
+++ b/Engine/Source/Programs/UnrealBuildTool/UnrealBuildTool.csproj
@@ -145,6 +145,7 @@
     <Compile Include="Modes\BuildMode.cs" />
     <Compile Include="Modes\CleanMode.cs" />
     <Compile Include="Modes\DeployMode.cs" />
+    <Compile Include="Modes\GenerateClangDatabase.cs" />
     <Compile Include="Modes\QueryTargetsMode.cs" />
     <Compile Include="Modes\ExecuteMode.cs" />
     <Compile Include="Modes\JsonExportMode.cs" />
