Missing file.

#rnx

diff --git a/Engine/Source/Programs/UnrealBuildTool/Configuration/UEBuildModuleCPP.cs b/Engine/Source/Programs/UnrealBuildTool/Configuration/UEBuildModuleCPP.cs
--- a/Engine/Source/Programs/UnrealBuildTool/Configuration/UEBuildModuleCPP.cs
+++ b/Engine/Source/Programs/UnrealBuildTool/Configuration/UEBuildModuleCPP.cs
@@ -19,7 +19,7 @@ namespace UnrealBuildTool
 		/// <summary>
 		/// Stores a list of all source files, of different types
 		/// </summary>
-		class InputFileCollection
+		public class InputFileCollection
 		{
 			public readonly List<FileItem> HeaderFiles = new List<FileItem>();
 			public readonly List<FileItem> ISPCHeaderFiles = new List<FileItem>();
@@ -281,8 +281,15 @@ namespace UnrealBuildTool
 				return LinkInputFiles;
 			}
 
+			// Add all the module source directories to the makefile
+			foreach (DirectoryReference ModuleDirectory in ModuleDirectories)
+			{
+				DirectoryItem ModuleDirectoryItem = DirectoryItem.GetItemByDirectoryReference(ModuleDirectory);
+				Makefile.SourceDirectories.Add(ModuleDirectoryItem);
+			}
+
 			// Find all the input files
-			InputFileCollection InputFiles = FindInputFiles(Target.Platform, Makefile);
+			InputFileCollection InputFiles = FindInputFiles(Target.Platform, Makefile.DirectoryToSourceFiles);
 
 			// Process all of the header file dependencies for this module
 			CheckFirstIncludeMatchesEachCppFile(Target, ModuleCompileEnvironment, InputFiles.HeaderFiles, InputFiles.CPPFiles);
@@ -1261,9 +1268,9 @@ namespace UnrealBuildTool
 		/// Finds all the source files that should be built for this module
 		/// </summary>
 		/// <param name="Platform">The platform the module is being built for</param>
-		/// <param name="Makefile">Makefile for the target being built</param>
+		/// <param name="DirectoryToSourceFiles">Map of directory to source files inside it</param>
 		/// <returns>Set of source files that should be built</returns>
-		InputFileCollection FindInputFiles(UnrealTargetPlatform Platform, TargetMakefile Makefile)
+		public InputFileCollection FindInputFiles(UnrealTargetPlatform Platform, Dictionary<DirectoryItem, FileItem[]> DirectoryToSourceFiles)
 		{
 			ReadOnlyHashSet<string> ExcludedNames = UEBuildPlatform.GetBuildPlatform(Platform).GetExcludedFolderNames();
 
@@ -1273,8 +1280,7 @@ namespace UnrealBuildTool
 			foreach (DirectoryReference Dir in ModuleDirectories)
 			{
 				DirectoryItem ModuleDirectoryItem = DirectoryItem.GetItemByDirectoryReference(Dir);
-				FindInputFilesFromDirectoryRecursive(ModuleDirectoryItem, ExcludedNames, SourceDirectories, Makefile.DirectoryToSourceFiles, InputFiles);
-				Makefile.SourceDirectories.Add(ModuleDirectoryItem);
+				FindInputFilesFromDirectoryRecursive(ModuleDirectoryItem, ExcludedNames, SourceDirectories, DirectoryToSourceFiles, InputFiles);
 			}
 
 			return InputFiles;
diff --git a/Engine/Source/Programs/UnrealBuildTool/Modes/GenerateClangDatabase.cs b/Engine/Source/Programs/UnrealBuildTool/Modes/GenerateClangDatabase.cs
--- a/Engine/Source/Programs/UnrealBuildTool/Modes/GenerateClangDatabase.cs
+++ b/Engine/Source/Programs/UnrealBuildTool/Modes/GenerateClangDatabase.cs
@@ -55,7 +55,7 @@ namespace UnrealBuildTool
 				foreach (TargetDescriptor TargetDescriptor in TargetDescriptors)
 				{
 					// Disable PCHs and unity builds for the target
-					TargetDescriptor.AdditionalArguments = TargetDescriptor.AdditionalArguments.Append(new string[] { "-NoPCH", "-DisableUnity", "-ToolChain=ClangDatabaseProxyToolChain" });
+					TargetDescriptor.AdditionalArguments = TargetDescriptor.AdditionalArguments.Append(new string[] { "-NoPCH", "-DisableUnity" });
 
 					// Create a makefile for the target
 					UEBuildTarget Target = UEBuildTarget.Create(TargetDescriptor, BuildConfiguration.bSkipRulesCompile, BuildConfiguration.bUsePrecompiled);
@@ -133,66 +133,4 @@ namespace UnrealBuildTool
 			return 0;
 		}
 	}
-
-	/// <summary>
-	/// Toolchain implementation which just dumps arguments to the compiler
-	/// </summary>
-	class ClangDatabaseProxyToolChain : UEToolChain
-	{
-		FileReference ClangPath;
-
-		public ClangDatabaseProxyToolChain(ReadOnlyTargetRules Target)
-		{
-			VCEnvironment Environment = VCEnvironment.Create(WindowsCompiler.Clang, Target.Platform, Target.WindowsPlatform.Architecture, null, Target.WindowsPlatform.WindowsSdkVersion);
-			ClangPath = FileReference.Combine(Environment.CompilerDir, "clang.exe");
-		}
-
-		public override CPPOutput CompileCPPFiles(CppCompileEnvironment CompileEnvironment, List<FileItem> InputFiles, DirectoryReference OutputDir, string ModuleName, List<Action> Actions)
-		{
-			StringBuilder Arguments = new StringBuilder();
-
-			foreach(FileItem ForceIncludeFile in CompileEnvironment.ForceIncludeFiles)
-			{
-				Arguments.AppendFormat(" -include \"{0}\"", ForceIncludeFile.FullName);
-			}
-			foreach (string Definition in CompileEnvironment.Definitions)
-			{
-				Arguments.AppendFormat(" -D\"{0}\"", Definition);
-			}
-			foreach (DirectoryReference IncludePath in CompileEnvironment.UserIncludePaths)
-			{
-				Arguments.AppendFormat(" -I\"{0}\"", IncludePath);
-			}
-			foreach (DirectoryReference IncludePath in CompileEnvironment.SystemIncludePaths)
-			{
-				Arguments.AppendFormat(" -I\"{0}\"", IncludePath);
-			}
-
-			List<FileItem> ObjectFiles = new List<FileItem>();
-			foreach (FileItem InputFile in InputFiles)
-			{
-				Action NewAction = new Action(ActionType.Compile);
-				NewAction.CommandPath = ClangPath;
-				NewAction.CommandArguments = String.Format("{0} \"{1}\"", Arguments.ToString(), InputFile.AbsolutePath);
-				NewAction.WorkingDirectory = UnrealBuildTool.EngineSourceDirectory;
-				NewAction.PrerequisiteItems.Add(InputFile);
-				NewAction.ProducedItems.Add(FileItem.GetItemByFileReference(InputFile.Location.ChangeExtension(".dummy.obj")));
-				Actions.Add(NewAction);
-
-				ObjectFiles.AddRange(NewAction.ProducedItems);
-			}
-
-			return new CPPOutput { ObjectFiles = ObjectFiles };
-		}
-
-		public override FileItem LinkFiles(LinkEnvironment LinkEnvironment, bool bBuildImportLibraryOnly, List<Action> Actions)
-		{
-			throw new NotImplementedException();
-		}
-
-		public override FileItem[] LinkAllFiles(LinkEnvironment LinkEnvironment, bool bBuildImportLibraryOnly, List<Action> Actions)
-		{
-			return new FileItem[0];
-		}
-	}
 }
