UBT: Add an error when trying to build the editor in an unsupported configuration.

diff --git a/Engine/Source/Programs/UnrealBuildTool/Configuration/TargetRules.cs b/Engine/Source/Programs/UnrealBuildTool/Configuration/TargetRules.cs
--- a/Engine/Source/Programs/UnrealBuildTool/Configuration/TargetRules.cs
+++ b/Engine/Source/Programs/UnrealBuildTool/Configuration/TargetRules.cs
@@ -1542,7 +1542,7 @@ namespace UnrealBuildTool
 			// Otherwise take the SupportedPlatformsAttribute from the first type in the inheritance chain that supports it
 			for (Type CurrentType = GetType(); CurrentType != null; CurrentType = CurrentType.BaseType)
 			{
-				object[] Attributes = GetType().GetCustomAttributes(typeof(SupportedPlatformsAttribute), false);
+				object[] Attributes = CurrentType.GetCustomAttributes(typeof(SupportedPlatformsAttribute), false);
 				if (Attributes.Length > 0)
 				{
 					return Attributes.OfType<SupportedPlatformsAttribute>().SelectMany(x => x.Platforms).Distinct().ToArray();
@@ -1564,6 +1564,33 @@ namespace UnrealBuildTool
 			}
 		}
 
+		/// <summary>
+		/// Gets a list of configurations that this target supports
+		/// </summary>
+		/// <returns>Array of configurations that the target supports</returns>
+		internal UnrealTargetConfiguration[] GetSupportedConfigurations()
+		{
+			// Otherwise take the SupportedConfigurationsAttribute from the first type in the inheritance chain that supports it
+			for (Type CurrentType = GetType(); CurrentType != null; CurrentType = CurrentType.BaseType)
+			{
+				object[] Attributes = CurrentType.GetCustomAttributes(typeof(SupportedConfigurationsAttribute), false);
+				if (Attributes.Length > 0)
+				{
+					return Attributes.OfType<SupportedConfigurationsAttribute>().SelectMany(x => x.Configurations).Distinct().ToArray();
+				}
+			}
+
+			// Otherwise, get the default for the target type
+			if (Type == TargetType.Editor)
+			{
+				return new[] { UnrealTargetConfiguration.Debug, UnrealTargetConfiguration.DebugGame, UnrealTargetConfiguration.Development };
+			}
+			else
+			{
+				return ((UnrealTargetConfiguration[])Enum.GetValues(typeof(UnrealTargetConfiguration))).Where(x => x != UnrealTargetConfiguration.Unknown).ToArray();
+			}
+		}
+
 		/// <summary>
 		/// Finds all the subobjects which can be configured by command line options and config files
 		/// </summary>
diff --git a/Engine/Source/Programs/UnrealBuildTool/Configuration/UEBuildTarget.cs b/Engine/Source/Programs/UnrealBuildTool/Configuration/UEBuildTarget.cs
--- a/Engine/Source/Programs/UnrealBuildTool/Configuration/UEBuildTarget.cs
+++ b/Engine/Source/Programs/UnrealBuildTool/Configuration/UEBuildTarget.cs
@@ -611,6 +611,12 @@ namespace UnrealBuildTool
 				throw new BuildException("{0} does not support the {1} platform.", Descriptor.Name, Descriptor.Platform.ToString());
 			}
 
+			// Make sure this configuration is supports
+			if(!RulesObject.GetSupportedConfigurations().Contains(Descriptor.Configuration))
+			{
+				throw new BuildException("{0} does not support the {1} configuration", Descriptor.Name, Descriptor.Configuration);
+			}
+
 			// Make sure this target type is supported
 			if (!InstalledPlatformInfo.IsValid(RulesObject.Type, Descriptor.Platform, Descriptor.Configuration, EProjectType.Code, InstalledPlatformState.Downloaded))
 			{
diff --git a/Engine/Source/Programs/UnrealBuildTool/ProjectFiles/VisualStudio/VCProject.cs b/Engine/Source/Programs/UnrealBuildTool/ProjectFiles/VisualStudio/VCProject.cs
--- a/Engine/Source/Programs/UnrealBuildTool/ProjectFiles/VisualStudio/VCProject.cs
+++ b/Engine/Source/Programs/UnrealBuildTool/ProjectFiles/VisualStudio/VCProject.cs
@@ -160,29 +160,6 @@ namespace UnrealBuildTool
 		/// <returns>Project context matching the given solution context</returns>
 		public abstract MSBuildProjectContext GetMatchingProjectContext(TargetType SolutionTarget, UnrealTargetConfiguration SolutionConfiguration, UnrealTargetPlatform SolutionPlatform, PlatformProjectGeneratorCollection PlatformProjectGenerators);
 
-		static UnrealTargetConfiguration[] GetSupportedConfigurations(TargetRules Rules)
-		{
-			// Otherwise take the SupportedConfigurationsAttribute from the first type in the inheritance chain that supports it
-			for (Type CurrentType = Rules.GetType(); CurrentType != null; CurrentType = CurrentType.BaseType)
-			{
-				object[] Attributes = Rules.GetType().GetCustomAttributes(typeof(SupportedConfigurationsAttribute), false);
-				if (Attributes.Length > 0)
-				{
-					return Attributes.OfType<SupportedConfigurationsAttribute>().SelectMany(x => x.Configurations).Distinct().ToArray();
-				}
-			}
-
-			// Otherwise, get the default for the target type
-			if(Rules.Type == TargetType.Editor)
-			{
-				return new[] { UnrealTargetConfiguration.Debug, UnrealTargetConfiguration.DebugGame, UnrealTargetConfiguration.Development };
-			}
-			else
-			{
-				return ((UnrealTargetConfiguration[])Enum.GetValues(typeof(UnrealTargetConfiguration))).Where(x => x != UnrealTargetConfiguration.Unknown).ToArray();
-			}
-		}
-
 		/// <summary>
 		/// Checks to see if the specified solution platform and configuration is able to map to this project
 		/// </summary>
@@ -218,7 +195,7 @@ namespace UnrealBuildTool
 			if (ProjectTarget.TargetRules != null)
 			{
 				SupportedPlatforms.AddRange(ProjectTarget.SupportedPlatforms);
-				SupportedConfigurations.AddRange(GetSupportedConfigurations(ProjectTarget.TargetRules));
+				SupportedConfigurations.AddRange(ProjectTarget.TargetRules.GetSupportedConfigurations());
 			}
 
 			// Add all of the extra platforms/configurations for this target
