Fix Clang executable path, missing source file in output compile_commands.json.

#rnx

diff --git a/Engine/Source/Programs/UnrealBuildTool/Modes/GenerateClangDatabase.cs b/Engine/Source/Programs/UnrealBuildTool/Modes/GenerateClangDatabase.cs
--- a/Engine/Source/Programs/UnrealBuildTool/Modes/GenerateClangDatabase.cs
+++ b/Engine/Source/Programs/UnrealBuildTool/Modes/GenerateClangDatabase.cs
@@ -64,7 +64,7 @@ namespace UnrealBuildTool
 
 					// Find the location of the compiler
 					VCEnvironment Environment = VCEnvironment.Create(WindowsCompiler.Clang, Target.Platform, Target.Rules.WindowsPlatform.Architecture, null, Target.Rules.WindowsPlatform.WindowsSdkVersion);
-					FileReference ClangPath = FileReference.Combine(Environment.CompilerDir, "clang.exe");
+					FileReference ClangPath = FileReference.Combine(Environment.CompilerDir, "bin", "clang++.exe");
 
 					// Create all the binaries and modules
 					CppCompileEnvironment GlobalCompileEnvironment = Target.CreateCompileEnvironmentForProjectFiles();
@@ -102,12 +102,11 @@ namespace UnrealBuildTool
 									CommandBuilder.AppendFormat(" -I\"{0}\"", IncludePath);
 								}
 
-								string Command = CommandBuilder.ToString();
 								foreach (FileItem InputFile in InputFiles)
 								{
 									if(FileFilter == null || FileFilter.Matches(InputFile.Location.MakeRelativeTo(UnrealBuildTool.RootDirectory)))
 									{
-										FileToCommand[InputFile.Location] = Command;
+										FileToCommand[InputFile.Location] = String.Format("{0} \"{1}\"", CommandBuilder, InputFile.FullName);
 									}
 								}
 							}
