Add subclass ModuleRules file references to ExternalDependencies to properly generate new makefiles when they are changed.

diff --git a/Engine/Source/Programs/UnrealBuildTool/Configuration/ModuleRules.cs b/Engine/Source/Programs/UnrealBuildTool/Configuration/ModuleRules.cs
--- a/Engine/Source/Programs/UnrealBuildTool/Configuration/ModuleRules.cs
+++ b/Engine/Source/Programs/UnrealBuildTool/Configuration/ModuleRules.cs
@@ -818,6 +818,11 @@ namespace UnrealBuildTool
 		/// </summary>
 		public List<string> ExternalDependencies = new List<string>();
 
+		/// <summary>
+		/// Subclass rules files which invalidate the makefile if modified.
+		/// </summary>
+		public List<string> SubclassRules;
+
 		/// <summary>
 		/// Whether this module requires the IMPLEMENT_MODULE macro to be implemented. Most UE4 modules require this, since we use the IMPLEMENT_MODULE macro
 		/// to do other global overloads (eg. operator new/delete forwarding to GMalloc).
diff --git a/Engine/Source/Programs/UnrealBuildTool/Configuration/UEBuildTarget.cs b/Engine/Source/Programs/UnrealBuildTool/Configuration/UEBuildTarget.cs
--- a/Engine/Source/Programs/UnrealBuildTool/Configuration/UEBuildTarget.cs
+++ b/Engine/Source/Programs/UnrealBuildTool/Configuration/UEBuildTarget.cs
@@ -1883,6 +1883,14 @@ namespace UnrealBuildTool
 					FileReference Location = FileReference.Combine(Module.RulesFile.Directory, ExternalDependency);
 					Makefile.AdditionalDependencies.Add(FileItem.GetItemByFileReference(Location));
 				}
+				if (Module.Rules.SubclassRules != null)
+				{
+					foreach (string SubclassRule in Module.Rules.SubclassRules)
+					{
+						FileItem SubclassRuleFileItem = FileItem.GetItemByFileReference(new FileReference(SubclassRule));
+						Makefile.AdditionalDependencies.Add(SubclassRuleFileItem);
+					}
+				}
 			}
 			Makefile.AdditionalDependencies.UnionWith(Makefile.PluginFiles);
 
diff --git a/Engine/Source/Programs/UnrealBuildTool/System/RulesAssembly.cs b/Engine/Source/Programs/UnrealBuildTool/System/RulesAssembly.cs
--- a/Engine/Source/Programs/UnrealBuildTool/System/RulesAssembly.cs
+++ b/Engine/Source/Programs/UnrealBuildTool/System/RulesAssembly.cs
@@ -431,12 +431,14 @@ namespace UnrealBuildTool
 					Type SubType = RulesObjectType;
 
 					RulesObject.DirectoriesForModuleSubClasses = new Dictionary<Type, DirectoryReference>();
+					RulesObject.SubclassRules = new List<string>();
 					while (SubType != BaseRulesObjectType)
 					{
 						FileReference SubTypeFileName;
 						if (TryGetFileNameFromType(SubType, out SubTypeFileName))
 						{
 							RulesObject.DirectoriesForModuleSubClasses.Add(SubType, SubTypeFileName.Directory);
+							RulesObject.SubclassRules.Add(SubTypeFileName.FullName);
 						}
 						SubType = SubType.BaseType;
 					}
