HTML5 Mac shipping linker error fix

[RN]

diff --git a/Engine/Source/Programs/UnrealBuildTool/Platform/HTML5/HTML5ToolChain.cs b/Engine/Source/Programs/UnrealBuildTool/Platform/HTML5/HTML5ToolChain.cs
--- a/Engine/Source/Programs/UnrealBuildTool/Platform/HTML5/HTML5ToolChain.cs
+++ b/Engine/Source/Programs/UnrealBuildTool/Platform/HTML5/HTML5ToolChain.cs
@@ -91,7 +91,7 @@ namespace UnrealBuildTool
 			Log.TraceInformation("*** Emscripten Config File: " + Environment.GetEnvironmentVariable("EM_CONFIG"));
 		}
 
-		string GetSharedArguments_Global(CppConfiguration Configuration, bool bOptimizeForSize, string Architecture, bool bEnableShadowVariableWarnings, bool bShadowVariableWarningsAsErrors, bool bEnableUndefinedIdentifierWarnings, bool bUndefinedIdentifierWarningsAsErrors, bool bUseInlining)
+		string GetSharedArguments_Global(CppConfiguration Configuration, bool bOptimizeForSize, string Architecture, bool bEnableShadowVariableWarnings, bool bShadowVariableWarningsAsErrors, bool bEnableUndefinedIdentifierWarnings, bool bUndefinedIdentifierWarningsAsErrors, bool bUseInlining, bool bUE423_OSX_LinkerFix=false)
 		{
 			string Result = " ";
 //			string Result = " -Werror";
@@ -138,7 +138,14 @@ namespace UnrealBuildTool
 
 			else if (Configuration == CppConfiguration.Shipping)
 			{
-				Result += " -O3"; // favor speed over size
+				if (bUE423_OSX_LinkerFix && (BuildHostPlatform.Current.Platform == UnrealTargetPlatform.Mac))
+				{	// UE-78966 - TEMP HACK - remove this on next EMSCRIPTEN_TOOLCHAIN_UPGRADE_CHECK
+					Result += " -O2";
+				}
+				else
+				{
+					Result += " -O3"; // favor speed over size
+				}
 			}
 
 			if (!bUseInlining)
@@ -251,7 +258,7 @@ namespace UnrealBuildTool
 
 		string GetLinkArguments(LinkEnvironment LinkEnvironment)
 		{
-			string Result = GetSharedArguments_Global(LinkEnvironment.Configuration, LinkEnvironment.bOptimizeForSize, LinkEnvironment.Architecture, false, false, false, false, false);
+			string Result = GetSharedArguments_Global(LinkEnvironment.Configuration, LinkEnvironment.bOptimizeForSize, LinkEnvironment.Architecture, false, false, false, false, false, true);
 
 			/* N.B. When editing link flags in this function, UnrealBuildTool does not seem to automatically pick them up and do an incremental
 			 *	relink only of UE4Game.js (at least when building blueprints projects). Therefore after editing, delete the old build
