Prevent iOS Xcode project generation from adding a default bundle identifier to the UE4.xcodeproj file which caused issues with Blueprint projects on Mac.

diff --git a/Engine/Source/Programs/UnrealBuildTool/ProjectFiles/Xcode/XcodeProject.cs b/Engine/Source/Programs/UnrealBuildTool/ProjectFiles/Xcode/XcodeProject.cs
--- a/Engine/Source/Programs/UnrealBuildTool/ProjectFiles/Xcode/XcodeProject.cs
+++ b/Engine/Source/Programs/UnrealBuildTool/ProjectFiles/Xcode/XcodeProject.cs
@@ -1051,7 +1051,10 @@ namespace UnrealBuildTool
 					{
 						Content.Append("\t\t\t\t\"PROVISIONING_PROFILE_SPECIFIER[sdk=iphoneos*]\" = \"" + UUID_IOS + "\";" + ProjectFileGenerator.NewLine);
 					}
-					Content.Append("\t\t\t\t\"PRODUCT_BUNDLE_IDENTIFIER[sdk=iphoneos*]\" = " + IOS_BUNDLE + ";" + ProjectFileGenerator.NewLine);
+					if (ProjectFile != null)
+					{
+						Content.Append("\t\t\t\t\"PRODUCT_BUNDLE_IDENTIFIER[sdk=iphoneos*]\" = " + IOS_BUNDLE + ";" + ProjectFileGenerator.NewLine);
+					}
 				}
                 if (TVOSRunTimeVersion != null)
 				{
@@ -1068,7 +1071,10 @@ namespace UnrealBuildTool
 					{
 						Content.Append("\t\t\t\t\"PROVISIONING_PROFILE_SPECIFIER[sdk=appletvos*]\" = \"" + UUID_TVOS + "\";" + ProjectFileGenerator.NewLine);
 					}
-					Content.Append("\t\t\t\t\"PRODUCT_BUNDLE_IDENTIFIER[sdk=appletvos*]\" = " + TVOS_BUNDLE + ";" + ProjectFileGenerator.NewLine);
+					if (ProjectFile != null)
+					{
+						Content.Append("\t\t\t\t\"PRODUCT_BUNDLE_IDENTIFIER[sdk=appletvos*]\" = " + TVOS_BUNDLE + ";" + ProjectFileGenerator.NewLine);
+					}
 				}
 				if (bSupportMac)
 				{
