Use VS2017 15.9.11 toolchain.

#rnx

diff --git a/Engine/Source/Programs/UnrealBuildTool/Platform/Windows/UEBuildWindows.cs b/Engine/Source/Programs/UnrealBuildTool/Platform/Windows/UEBuildWindows.cs
--- a/Engine/Source/Programs/UnrealBuildTool/Platform/Windows/UEBuildWindows.cs
+++ b/Engine/Source/Programs/UnrealBuildTool/Platform/Windows/UEBuildWindows.cs
@@ -560,7 +560,7 @@ namespace UnrealBuildTool
 		/// <summary>
 		/// The default compiler version to be used, if installed. 
 		/// </summary>
-		static readonly VersionNumber DefaultVisualStudioToolChainVersion = VersionNumber.Parse("14.16.27023.1");
+		static readonly VersionNumber DefaultVisualStudioToolChainVersion = VersionNumber.Parse("14.16.27023.2"); // VS2017 v15.9.15
 
 		/// <summary>
 		/// The default Windows SDK version to be used, if installed.
