Fix IOS min version hardcoded in toolchain

diff --git a/Engine/Source/Programs/UnrealBuildTool/Platform/IOS/UEBuildIOS.cs b/Engine/Source/Programs/UnrealBuildTool/Platform/IOS/UEBuildIOS.cs
--- a/Engine/Source/Programs/UnrealBuildTool/Platform/IOS/UEBuildIOS.cs
+++ b/Engine/Source/Programs/UnrealBuildTool/Platform/IOS/UEBuildIOS.cs
@@ -375,7 +375,7 @@ namespace UnrealBuildTool
 					case "IOS_12":
 						return "12.0";
 					default:
-						return "10.0";
+						return "11.0";
 				}
 			}
 		}
