UBT: Prevent creation of a UE4Editor target for validating the shared build environment from modifying environment variables for the main process. This causes the exported XGE scripts to have the wrong environment variables, and a separate environment to be created. When multiple targets have overlapping artifacts, this results in illegal cross-environment action references.

diff --git a/Engine/Source/Programs/UnrealBuildTool/Configuration/UEBuildTarget.cs b/Engine/Source/Programs/UnrealBuildTool/Configuration/UEBuildTarget.cs
--- a/Engine/Source/Programs/UnrealBuildTool/Configuration/UEBuildTarget.cs
+++ b/Engine/Source/Programs/UnrealBuildTool/Configuration/UEBuildTarget.cs
@@ -1481,6 +1481,7 @@ namespace UnrealBuildTool
 			LinkEnvironment GlobalLinkEnvironment = new LinkEnvironment(GlobalCompileEnvironment.Platform, GlobalCompileEnvironment.Configuration, GlobalCompileEnvironment.Architecture);
 
 			UEToolChain TargetToolChain = CreateToolchain(Platform);
+			TargetToolChain.SetEnvironmentVariables();
 			SetupGlobalEnvironment(TargetToolChain, GlobalCompileEnvironment, GlobalLinkEnvironment);
 
 			// Save off the original list of binaries. We'll use this to figure out which PCHs to create later, to avoid switching PCHs when compiling single modules.
diff --git a/Engine/Source/Programs/UnrealBuildTool/Platform/Windows/VCEnvironment.cs b/Engine/Source/Programs/UnrealBuildTool/Platform/Windows/VCEnvironment.cs
--- a/Engine/Source/Programs/UnrealBuildTool/Platform/Windows/VCEnvironment.cs
+++ b/Engine/Source/Programs/UnrealBuildTool/Platform/Windows/VCEnvironment.cs
@@ -128,10 +128,6 @@ namespace UnrealBuildTool
             // Compile using 64 bit tools for 64 bit targets, and 32 for 32.
 			CompilerPath = GetCompilerToolPath(Platform, Compiler, Architecture, CompilerDir);
 
-			// Add the compiler path and directory as environment variables for the process so they may be used elsewhere.
-			Environment.SetEnvironmentVariable("VC_COMPILER_PATH", CompilerPath.FullName, EnvironmentVariableTarget.Process);
-			Environment.SetEnvironmentVariable("VC_COMPILER_DIR", CompilerPath.Directory.FullName, EnvironmentVariableTarget.Process);
-
 			// Regardless of the target, if we're linking on a 64 bit machine, we want to use the 64 bit linker (it's faster than the 32 bit linker and can handle large linking jobs)
 			DirectoryReference DefaultLinkerDir = VCToolPath;
 			LinkerPath = GetLinkerToolPath(Platform, Compiler, DefaultLinkerDir);
@@ -140,8 +136,22 @@ namespace UnrealBuildTool
 			// Get the resource compiler path from the Windows SDK
 			ResourceCompilerPath = GetResourceCompilerToolPath(Platform, WindowsSdkDir, WindowsSdkVersion);
 
+			// Get the ISPC compiler bundled with the engine
 			ISPCCompilerPath = GetISPCCompilerToolPath(Platform);
 
+			// Get all the system include paths
+			SetupEnvironment(Platform);
+		}
+
+		/// <summary>
+		/// Updates environment variables needed for running with this toolchain
+		/// </summary>
+		public void SetEnvironmentVariables()
+		{
+			// Add the compiler path and directory as environment variables for the process so they may be used elsewhere.
+			Environment.SetEnvironmentVariable("VC_COMPILER_PATH", CompilerPath.FullName, EnvironmentVariableTarget.Process);
+			Environment.SetEnvironmentVariable("VC_COMPILER_DIR", CompilerPath.Directory.FullName, EnvironmentVariableTarget.Process);
+
 			// Add both toolchain paths to the PATH environment variable. There are some support DLLs which are only added to one of the paths, but which the toolchain in the other directory
 			// needs to run (eg. mspdbcore.dll).
 			if (Architecture == WindowsArchitecture.x64)
@@ -160,10 +170,6 @@ namespace UnrealBuildTool
 				AddDirectoryToPath(GetVCToolPath(ToolChain, ToolChainDir, WindowsArchitecture.x86));
 				AddDirectoryToPath(GetVCToolPath(ToolChain, ToolChainDir, WindowsArchitecture.x64));
 			}
-
-
-			// Get all the system include paths
-			SetupEnvironment(Platform);
 		}
 
 		/// <summary>
diff --git a/Engine/Source/Programs/UnrealBuildTool/Platform/Windows/VCToolChain.cs b/Engine/Source/Programs/UnrealBuildTool/Platform/Windows/VCToolChain.cs
--- a/Engine/Source/Programs/UnrealBuildTool/Platform/Windows/VCToolChain.cs
+++ b/Engine/Source/Programs/UnrealBuildTool/Platform/Windows/VCToolChain.cs
@@ -46,6 +46,14 @@ namespace UnrealBuildTool
 			}
 		}
 
+		/// <summary>
+		/// Prepares the environment for building
+		/// </summary>
+		public override void SetEnvironmentVariables()
+		{
+			EnvVars.SetEnvironmentVariables();
+		}
+
 		/// <summary>
 		/// Returns the version info for the toolchain. This will be output before building.
 		/// </summary>
diff --git a/Engine/Source/Programs/UnrealBuildTool/ToolChain/UEToolChain.cs b/Engine/Source/Programs/UnrealBuildTool/ToolChain/UEToolChain.cs
--- a/Engine/Source/Programs/UnrealBuildTool/ToolChain/UEToolChain.cs
+++ b/Engine/Source/Programs/UnrealBuildTool/ToolChain/UEToolChain.cs
@@ -16,6 +16,10 @@ namespace UnrealBuildTool
 		{
 		}
 
+		public virtual void SetEnvironmentVariables()
+		{
+		}
+
 		public virtual string GetVersionInfo()
 		{
 			return null;
