Fix cross compiling for Linux with XGE

diff --git a/Engine/Source/Programs/UnrealBuildTool/Platform/Linux/LinuxToolChain.cs b/Engine/Source/Programs/UnrealBuildTool/Platform/Linux/LinuxToolChain.cs
--- a/Engine/Source/Programs/UnrealBuildTool/Platform/Linux/LinuxToolChain.cs
+++ b/Engine/Source/Programs/UnrealBuildTool/Platform/Linux/LinuxToolChain.cs
@@ -104,13 +104,12 @@ namespace UnrealBuildTool
 
 				// set up the path to our toolchain
 				GCCPath = "";
-				// we rely on the fact that appending ".exe" is optional when invoking a binary on Windows
-				ClangPath = Path.Combine(BaseLinuxPath, @"bin", "clang++");
-				ArPath = Path.Combine(Path.Combine(BaseLinuxPath, String.Format("bin/{0}-{1}", Architecture, "ar")));
-				LlvmArPath = Path.Combine(Path.Combine(BaseLinuxPath, String.Format("bin/{0}", "llvm-ar")));
-				RanlibPath = Path.Combine(Path.Combine(BaseLinuxPath, String.Format("bin/{0}-{1}", Architecture, "ranlib")));
-				StripPath = Path.Combine(Path.Combine(BaseLinuxPath, String.Format("bin/{0}-{1}", Architecture, "strip")));
-				ObjcopyPath = Path.Combine(Path.Combine(BaseLinuxPath, String.Format("bin/{0}-{1}", Architecture, "objcopy")));
+				ClangPath = Path.Combine(BaseLinuxPath, @"bin", "clang++" + GetHostPlatformBinarySuffix());
+				ArPath = Path.Combine(Path.Combine(BaseLinuxPath, String.Format("bin/{0}-{1}", Architecture, "ar" + GetHostPlatformBinarySuffix())));
+				LlvmArPath = Path.Combine(Path.Combine(BaseLinuxPath, String.Format("bin/{0}", "llvm-ar" + GetHostPlatformBinarySuffix())));
+				RanlibPath = Path.Combine(Path.Combine(BaseLinuxPath, String.Format("bin/{0}-{1}", Architecture, "ranlib" + GetHostPlatformBinarySuffix())));
+				StripPath = Path.Combine(Path.Combine(BaseLinuxPath, String.Format("bin/{0}-{1}", Architecture, "strip" + GetHostPlatformBinarySuffix())));
+				ObjcopyPath = Path.Combine(Path.Combine(BaseLinuxPath, String.Format("bin/{0}-{1}", Architecture, "objcopy" + GetHostPlatformBinarySuffix())));
 
 				// When cross-compiling on Windows, use old FixDeps. It is slow, but it does not have timing issues
 				bUseFixdeps = (BuildHostPlatform.Current.Platform == UnrealTargetPlatform.Win64 || BuildHostPlatform.Current.Platform == UnrealTargetPlatform.Win32);
@@ -198,6 +197,16 @@ namespace UnrealBuildTool
 			bPreservePSYM = InPreservePSYM;
 		}
 
+		private string GetHostPlatformBinarySuffix()
+		{
+			if (BuildHostPlatform.Current.Platform == UnrealTargetPlatform.Win64 || BuildHostPlatform.Current.Platform == UnrealTargetPlatform.Win32)
+			{
+				return ".exe";
+			}
+
+			return "";
+		}
+
 		protected virtual bool CrossCompiling()
 		{
 			return bIsCrossCompiling;
