UBT: Add a new "UncookedOnly" module type, which only loads a module in uncooked builds.

The Developer module type is currently generally used for this purpose, because it allows modules to load from an editor executable running with -game, but this is not the documented or intended behavior.

diff --git a/Engine/Source/Programs/UnrealBuildTool/System/ExternalExecution.cs b/Engine/Source/Programs/UnrealBuildTool/System/ExternalExecution.cs
--- a/Engine/Source/Programs/UnrealBuildTool/System/ExternalExecution.cs
+++ b/Engine/Source/Programs/UnrealBuildTool/System/ExternalExecution.cs
@@ -120,6 +120,7 @@ namespace UnrealBuildTool
 				case ModuleHostType.RuntimeNoCommandlet:
 				case ModuleHostType.RuntimeAndProgram:
 				case ModuleHostType.CookedOnly:
+				case ModuleHostType.UncookedOnly:
 				case ModuleHostType.ServerOnly:
 				case ModuleHostType.ClientOnly:
 				case ModuleHostType.ClientOnlyNoCommandlet:
@@ -143,6 +144,7 @@ namespace UnrealBuildTool
 				case ModuleHostType.RuntimeNoCommandlet:
 				case ModuleHostType.RuntimeAndProgram:
 				case ModuleHostType.CookedOnly:
+				case ModuleHostType.UncookedOnly:
 				case ModuleHostType.ServerOnly:
 				case ModuleHostType.ClientOnly:
 				case ModuleHostType.ClientOnlyNoCommandlet:
diff --git a/Engine/Source/Programs/UnrealBuildTool/System/ModuleDescriptor.cs b/Engine/Source/Programs/UnrealBuildTool/System/ModuleDescriptor.cs
--- a/Engine/Source/Programs/UnrealBuildTool/System/ModuleDescriptor.cs
+++ b/Engine/Source/Programs/UnrealBuildTool/System/ModuleDescriptor.cs
@@ -34,15 +34,20 @@ namespace UnrealBuildTool
         /// </summary>
         RuntimeAndProgram,
 
-        /// <summary>
+		/// <summary>
 		/// Loaded only in cooked builds
 		/// </summary>
 		CookedOnly,
 
-        /// <summary>
-        /// Loaded only when the engine has support for developer tools enabled
-        /// </summary>
-        Developer,
+		/// <summary>
+		/// Loaded only in uncooked builds
+		/// </summary>
+		UncookedOnly,
+
+		/// <summary>
+		/// Loaded only when the engine has support for developer tools enabled
+		/// </summary>
+		Developer,
 
 		/// <summary>
 		/// Loaded only by the editor
@@ -458,7 +463,9 @@ namespace UnrealBuildTool
                     return TargetType != TargetType.Program;
                 case ModuleHostType.CookedOnly:
                     return bBuildRequiresCookedData;
-                case ModuleHostType.RuntimeAndProgram:
+				case ModuleHostType.UncookedOnly:
+					return !bBuildRequiresCookedData;
+				case ModuleHostType.RuntimeAndProgram:
 					return true;
                 case ModuleHostType.Developer:
 					return bBuildDeveloperTools;
diff --git a/Engine/Source/Runtime/Projects/Private/ModuleDescriptor.cpp b/Engine/Source/Runtime/Projects/Private/ModuleDescriptor.cpp
--- a/Engine/Source/Runtime/Projects/Private/ModuleDescriptor.cpp
+++ b/Engine/Source/Runtime/Projects/Private/ModuleDescriptor.cpp
@@ -89,6 +89,9 @@ const TCHAR* EHostType::ToString( const EHostType::Type Value )
 		case CookedOnly:
 			return TEXT("CookedOnly");
 
+		case UncookedOnly:
+			return TEXT("UncookedOnly");
+
 		case Developer:
 			return TEXT( "Developer" );
 
@@ -470,6 +473,9 @@ bool FModuleDescriptor::IsCompiledInCurrentConfiguration() const
 	case EHostType::CookedOnly:
 		return FPlatformProperties::RequiresCookedData();
 
+	case EHostType::UncookedOnly:
+		return !FPlatformProperties::RequiresCookedData();
+
 	case EHostType::Developer:
 		#if WITH_UNREAL_DEVELOPER_TOOLS
 			return true;
@@ -540,6 +546,9 @@ bool FModuleDescriptor::IsLoadedInCurrentConfiguration() const
 	case EHostType::CookedOnly:
 		return FPlatformProperties::RequiresCookedData();
 
+	case EHostType::UncookedOnly:
+		return !FPlatformProperties::RequiresCookedData();
+
 	case EHostType::Developer:
 		#if WITH_UNREAL_DEVELOPER_TOOLS
 			return true;
diff --git a/Engine/Source/Runtime/Projects/Public/ModuleDescriptor.h b/Engine/Source/Runtime/Projects/Public/ModuleDescriptor.h
--- a/Engine/Source/Runtime/Projects/Public/ModuleDescriptor.h
+++ b/Engine/Source/Runtime/Projects/Public/ModuleDescriptor.h
@@ -81,7 +81,10 @@ namespace EHostType
 		
 		// Loads only in cooked games.
 		CookedOnly,
-		
+
+		// Only loads in uncooked games.
+		UncookedOnly,
+
 		// Only loads in development runtime or editor builds. Does not load in shipping builds.
 		Developer,
 		
