Enable developer tools for game/client/server targets that aren't built in Test or Shipping.

diff --git a/Engine/Source/Programs/UnrealBuildTool/Configuration/TargetRules.cs b/Engine/Source/Programs/UnrealBuildTool/Configuration/TargetRules.cs
--- a/Engine/Source/Programs/UnrealBuildTool/Configuration/TargetRules.cs
+++ b/Engine/Source/Programs/UnrealBuildTool/Configuration/TargetRules.cs
@@ -406,7 +406,7 @@ namespace UnrealBuildTool
 		public bool bBuildDeveloperTools
 		{
 			set { bBuildDeveloperToolsOverride = value; }
-			get { return bBuildDeveloperToolsOverride ?? (bCompileAgainstEngine && (Type == TargetType.Editor || Type == TargetType.Program)); }
+			get { return bBuildDeveloperToolsOverride ?? (bCompileAgainstEngine && (Type == TargetType.Editor || Type == TargetType.Program || (Configuration != UnrealTargetConfiguration.Test && Configuration != UnrealTargetConfiguration.Shipping))); }
 		}
 
 		/// <summary>
diff --git a/Engine/Source/Programs/UnrealBuildTool/Configuration/UEBuildTarget.cs b/Engine/Source/Programs/UnrealBuildTool/Configuration/UEBuildTarget.cs
--- a/Engine/Source/Programs/UnrealBuildTool/Configuration/UEBuildTarget.cs
+++ b/Engine/Source/Programs/UnrealBuildTool/Configuration/UEBuildTarget.cs
@@ -2583,9 +2583,6 @@ namespace UnrealBuildTool
 				}
 			}
 
-			// Whether to allow developer modules
-			bool bAllowDeveloperModules = (Configuration != UnrealTargetConfiguration.Shipping);
-
 			// Find all the platform folders to exclude from the list of valid modules
 			ReadOnlyHashSet<string> ExcludeFolders = UEBuildPlatform.GetBuildPlatform(Platform).GetExcludedFolderNames();
 
@@ -2609,7 +2606,7 @@ namespace UnrealBuildTool
 
 				// Also allow anything in the developer directory in non-shipping configurations (though we blacklist by default unless the PrecompileForTargets
 				// setting indicates that it's actually useful at runtime).
-				if(bAllowDeveloperModules)
+				if(Rules.bBuildDeveloperTools)
 				{
 					Directories.Add(UnrealBuildTool.EngineSourceDeveloperDirectory);
 					Directories.Add(DirectoryReference.Combine(UnrealBuildTool.EnterpriseSourceDirectory, "Developer"));
@@ -2676,7 +2673,7 @@ namespace UnrealBuildTool
 				// Add all the modules
 				foreach (ModuleDescriptor ModuleDescriptor in Plugin.Descriptor.Modules)
 				{
-					if (ModuleDescriptor.IsCompiledInConfiguration(Platform, Configuration, TargetName, TargetType, bAllowDeveloperModules && Rules.bBuildDeveloperTools, Rules.bBuildEditor, Rules.bBuildRequiresCookedData))
+					if (ModuleDescriptor.IsCompiledInConfiguration(Platform, Configuration, TargetName, TargetType, Rules.bBuildDeveloperTools, Rules.bBuildEditor, Rules.bBuildRequiresCookedData))
 					{
 						FileReference ModuleFileName = RulesAssembly.GetModuleFileName(ModuleDescriptor.Name);
 						if(ModuleFileName == null)
