Deprecate the 'Developer' module type. The 'UncookedOnly' type has the same functionality but is more descriptive, and the 'DeveloperTool' module type correctly represents the state of the bBuildDeveloperTools flag.

diff --git a/Engine/Source/Programs/UnrealBuildTool/System/ExternalExecution.cs b/Engine/Source/Programs/UnrealBuildTool/System/ExternalExecution.cs
--- a/Engine/Source/Programs/UnrealBuildTool/System/ExternalExecution.cs
+++ b/Engine/Source/Programs/UnrealBuildTool/System/ExternalExecution.cs
@@ -126,9 +126,11 @@ namespace UnrealBuildTool
 				case ModuleHostType.ClientOnlyNoCommandlet:
 					return UHTModuleType.EngineRuntime;
 				case ModuleHostType.Developer:
+				case ModuleHostType.DeveloperTool:
 					return UHTModuleType.EngineDeveloper;
 				case ModuleHostType.Editor:
 				case ModuleHostType.EditorNoCommandlet:
+				case ModuleHostType.EditorAndProgram:
 					return UHTModuleType.EngineEditor;
 				default:
 					return null;
@@ -150,9 +152,11 @@ namespace UnrealBuildTool
 				case ModuleHostType.ClientOnlyNoCommandlet:
 					return UHTModuleType.GameRuntime;
 				case ModuleHostType.Developer:
+				case ModuleHostType.DeveloperTool:
 					return UHTModuleType.GameDeveloper;
 				case ModuleHostType.Editor:
 				case ModuleHostType.EditorNoCommandlet:
+				case ModuleHostType.EditorAndProgram:
 					return UHTModuleType.GameEditor;
 				default:
 					return null;
diff --git a/Engine/Source/Programs/UnrealBuildTool/System/ModuleDescriptor.cs b/Engine/Source/Programs/UnrealBuildTool/System/ModuleDescriptor.cs
--- a/Engine/Source/Programs/UnrealBuildTool/System/ModuleDescriptor.cs
+++ b/Engine/Source/Programs/UnrealBuildTool/System/ModuleDescriptor.cs
@@ -49,6 +49,11 @@ namespace UnrealBuildTool
 		/// </summary>
 		Developer,
 
+		/// <summary>
+		/// Loads on any targets where bBuildDeveloperTools is enabled
+		/// </summary>
+		DeveloperTool,
+
 		/// <summary>
 		/// Loaded only by the editor
 		/// </summary>
@@ -59,6 +64,11 @@ namespace UnrealBuildTool
 		/// </summary>
 		EditorNoCommandlet,
 
+		/// <summary>
+		/// Loaded by the editor or program targets
+		/// </summary>
+		EditorAndProgram,
+
 		/// <summary>
 		/// Loaded only by programs
 		/// </summary>
@@ -391,6 +401,19 @@ namespace UnrealBuildTool
 			}
 		}
 
+		/// <summary>
+		/// Produces any warnings and errors for the module settings
+		/// </summary>
+		/// <param name="File">File containing the module declaration</param>
+		public void Validate(FileReference File)
+		{
+			if(Type == ModuleHostType.Developer)
+			{
+				Log.TraceWarningOnce("The 'Developer' module type has been deprecated in 4.24. Use 'DeveloperTool' for modules that can be loaded by game/client/server targets in non-shipping configurations, or 'UncookedOnly' for modules that should only be loaded by uncooked editor and program targets (eg. modules containing blueprint nodes)");
+				Log.TraceWarningOnce(File, "The 'Developer' module type has been deprecated in 4.24.");
+			}
+		}
+
 		/// <summary>
 		/// Determines whether the given plugin module is part of the current build.
 		/// </summary>
@@ -461,17 +484,21 @@ namespace UnrealBuildTool
 				case ModuleHostType.Runtime:
 				case ModuleHostType.RuntimeNoCommandlet:
                     return TargetType != TargetType.Program;
-                case ModuleHostType.CookedOnly:
+				case ModuleHostType.RuntimeAndProgram:
+					return true;
+				case ModuleHostType.CookedOnly:
                     return bBuildRequiresCookedData;
 				case ModuleHostType.UncookedOnly:
 					return !bBuildRequiresCookedData;
-				case ModuleHostType.RuntimeAndProgram:
-					return true;
-                case ModuleHostType.Developer:
+				case ModuleHostType.Developer:
+					return TargetType == TargetType.Editor || bBuildEditor || TargetType == TargetType.Program;
+				case ModuleHostType.DeveloperTool:
 					return bBuildDeveloperTools;
 				case ModuleHostType.Editor:
 				case ModuleHostType.EditorNoCommandlet:
 					return TargetType == TargetType.Editor || bBuildEditor;
+				case ModuleHostType.EditorAndProgram:
+					return TargetType == TargetType.Editor || bBuildEditor || TargetType == TargetType.Program;
 				case ModuleHostType.Program:
 					return TargetType == TargetType.Program;
                 case ModuleHostType.ServerOnly:
diff --git a/Engine/Source/Programs/UnrealBuildTool/System/PluginDescriptor.cs b/Engine/Source/Programs/UnrealBuildTool/System/PluginDescriptor.cs
--- a/Engine/Source/Programs/UnrealBuildTool/System/PluginDescriptor.cs
+++ b/Engine/Source/Programs/UnrealBuildTool/System/PluginDescriptor.cs
@@ -303,7 +303,15 @@ namespace UnrealBuildTool
 			JsonObject RawObject = JsonObject.Read(FileName);
 			try
 			{
-				return new PluginDescriptor(RawObject);
+				PluginDescriptor Descriptor = new PluginDescriptor(RawObject);
+				if (Descriptor.Modules != null)
+				{
+					foreach (ModuleDescriptor Module in Descriptor.Modules)
+					{
+						Module.Validate(FileName);
+					}
+				}
+				return Descriptor;
 			}
 			catch (JsonParseException ParseException)
 			{
diff --git a/Engine/Source/Programs/UnrealBuildTool/System/ProjectDescriptor.cs b/Engine/Source/Programs/UnrealBuildTool/System/ProjectDescriptor.cs
--- a/Engine/Source/Programs/UnrealBuildTool/System/ProjectDescriptor.cs
+++ b/Engine/Source/Programs/UnrealBuildTool/System/ProjectDescriptor.cs
@@ -184,7 +184,15 @@ namespace UnrealBuildTool
 			JsonObject RawObject = JsonObject.Read(FileName);
 			try
 			{
-				return new ProjectDescriptor(RawObject);
+				ProjectDescriptor Descriptor = new ProjectDescriptor(RawObject);
+				if(Descriptor.Modules != null)
+				{
+					foreach (ModuleDescriptor Module in Descriptor.Modules)
+					{
+						Module.Validate(FileName);
+					}
+				}
+				return Descriptor;
 			}
 			catch (JsonParseException ParseException)
 			{
diff --git a/Engine/Source/Runtime/Projects/Private/ModuleDescriptor.cpp b/Engine/Source/Runtime/Projects/Private/ModuleDescriptor.cpp
--- a/Engine/Source/Runtime/Projects/Private/ModuleDescriptor.cpp
+++ b/Engine/Source/Runtime/Projects/Private/ModuleDescriptor.cpp
@@ -84,23 +84,29 @@ const TCHAR* EHostType::ToString( const EHostType::Type Value )
 			return TEXT( "RuntimeNoCommandlet" );
 
 		case RuntimeAndProgram:
-			return TEXT("RuntimeAndProgram");
+			return TEXT( "RuntimeAndProgram" );
 
 		case CookedOnly:
-			return TEXT("CookedOnly");
+			return TEXT( "CookedOnly" );
 
 		case UncookedOnly:
-			return TEXT("UncookedOnly");
+			return TEXT( "UncookedOnly" );
 
 		case Developer:
 			return TEXT( "Developer" );
 
+		case DeveloperTool:
+			return TEXT( "DeveloperTool" );
+
 		case Editor:
 			return TEXT( "Editor" );
 
 		case EditorNoCommandlet:
 			return TEXT( "EditorNoCommandlet" );
 
+		case EditorAndProgram:
+			return TEXT( "EditorAndProgram" );
+
 		case Program:
 			return TEXT("Program");
 
@@ -109,8 +115,10 @@ const TCHAR* EHostType::ToString( const EHostType::Type Value )
 
 		case ClientOnly:
 			return TEXT("ClientOnly");
+
 		case ClientOnlyNoCommandlet:
 			return TEXT("ClientOnlyNoCommandlet");
+
 		default:
 			ensureMsgf( false, TEXT( "Unrecognized EModuleType value: %i" ), Value );
 			return NULL;
@@ -477,10 +485,18 @@ bool FModuleDescriptor::IsCompiledInCurrentConfiguration() const
 		return !FPlatformProperties::RequiresCookedData();
 
 	case EHostType::Developer:
+		#if WITH_EDITOR || IS_PROGRAM
+			return true;
+		#else
+			return false;
+		#endif
+
+	case EHostType::DeveloperTool:
 		#if WITH_UNREAL_DEVELOPER_TOOLS
 			return true;
+		#else
+			return false;		
 		#endif
-		break;
 
 	case EHostType::Editor:
 	case EHostType::EditorNoCommandlet:
@@ -489,6 +505,13 @@ bool FModuleDescriptor::IsCompiledInCurrentConfiguration() const
 		#endif
 		break;
 
+	case EHostType::EditorAndProgram:
+		#if WITH_EDITOR || IS_PROGRAM
+			return true;
+		#else
+			return false;
+		#endif
+
 	case EHostType::Program:
 		#if IS_PROGRAM
 			return true;
@@ -550,10 +573,18 @@ bool FModuleDescriptor::IsLoadedInCurrentConfiguration() const
 		return !FPlatformProperties::RequiresCookedData();
 
 	case EHostType::Developer:
+		#if WITH_EDITOR || IS_PROGRAM
+			return true;
+		#else
+			return false;
+		#endif
+
+	case EHostType::DeveloperTool:
 		#if WITH_UNREAL_DEVELOPER_TOOLS
 			return true;
+		#else
+			return false;
 		#endif
-		break;
 
 	case EHostType::Editor:
 		#if WITH_EDITOR
@@ -567,6 +598,15 @@ bool FModuleDescriptor::IsLoadedInCurrentConfiguration() const
 		#endif
 		break;
 
+	case EHostType::EditorAndProgram:
+		#if WITH_EDITOR
+			return GIsEditor;
+		#elif IS_PROGRAM
+			return true;
+		#else
+			return false;
+		#endif
+
 	case EHostType::Program:
 		#if WITH_PLUGIN_SUPPORT && IS_PROGRAM
 			return true;
diff --git a/Engine/Source/Runtime/Projects/Public/ModuleDescriptor.h b/Engine/Source/Runtime/Projects/Public/ModuleDescriptor.h
--- a/Engine/Source/Runtime/Projects/Public/ModuleDescriptor.h
+++ b/Engine/Source/Runtime/Projects/Public/ModuleDescriptor.h
@@ -85,15 +85,23 @@ namespace EHostType
 		// Only loads in uncooked games.
 		UncookedOnly,
 
-		// Only loads in development runtime or editor builds. Does not load in shipping builds.
-		Developer,
-		
+		// Deprecated due to ambiguities. Only loads in editor and program targets, but loads in any editor mode (eg. -game, -server).
+		// Use UncookedOnly for the same behavior (eg. for editor blueprint nodes needed in uncooked games), or DeveloperTool for modules
+		// that can also be loaded in cooked games but should not be shipped (eg. debugging utilities).
+		Developer UE_DEPRECATED(4.24, "The Developer module type is deprecated. Use UncookedOnly for the same behavior, or DeveloperTool for modules that can also be loaded at runtime"),
+
+		// Loads on any targets where bBuildDeveloperTools is enabled.
+		DeveloperTool,
+
 		// Loads only when the editor is starting up.
 		Editor,
 		
 		// Loads only when the editor is starting up, but not in commandlet mode.
 		EditorNoCommandlet,
-		
+
+		// Loads only on editor and program targets
+		EditorAndProgram,
+
 		// Only loads on program targets.
 		Program,
 		
