Fixup directory references in WorkspaceTool, and allow projects under Engine/Extras to be added to project files.

#rnx

diff --git a/Engine/Source/Programs/UnrealBuildTool/ProjectFiles/ProjectFileGenerator.cs b/Engine/Source/Programs/UnrealBuildTool/ProjectFiles/ProjectFileGenerator.cs
--- a/Engine/Source/Programs/UnrealBuildTool/ProjectFiles/ProjectFileGenerator.cs
+++ b/Engine/Source/Programs/UnrealBuildTool/ProjectFiles/ProjectFileGenerator.cs
@@ -429,6 +429,10 @@ namespace UnrealBuildTool
 			string[] UnsupportedPlatformNames = Utils.MakeListOfUnsupportedPlatforms(SupportedPlatforms).ToArray();
 
 			List<FileReference> FoundProjects = new List<FileReference>();
+
+			DirectoryReference EngineExtras = DirectoryReference.Combine(UnrealBuildTool.EngineDirectory, "Extras");
+			DiscoverCSharpProgramProjectsRecursively(EngineExtras, FoundProjects);
+
 			DirectoryReference EngineProgramsSource = DirectoryReference.Combine(UnrealBuildTool.EngineDirectory, "Source", "Programs");
 			DiscoverCSharpProgramProjectsRecursively(EngineProgramsSource, FoundProjects);
 
