UBT/UAT: Restore support for generating .UDN documentation files.

#rnx

diff --git a/Engine/Source/Programs/AutomationTool/BuildGraph/BuildGraph.cs b/Engine/Source/Programs/AutomationTool/BuildGraph/BuildGraph.cs
--- a/Engine/Source/Programs/AutomationTool/BuildGraph/BuildGraph.cs
+++ b/Engine/Source/Programs/AutomationTool/BuildGraph/BuildGraph.cs
@@ -2,16 +2,13 @@
 
 using System;
 using System.Collections.Generic;
+using System.IO;
 using System.Linq;
-using System.Text;
-using System.Threading.Tasks;
-using UnrealBuildTool;
-using BuildGraph;
 using System.Reflection;
-using System.Collections;
-using System.IO;
+using System.Text;
 using System.Xml;
 using Tools.DotNETCommon;
+using UnrealBuildTool;
 
 namespace AutomationTool
 {
@@ -168,7 +165,7 @@ namespace AutomationTool
 			// Generate documentation
 			if(DocumentationFileName != null)
 			{
-				GenerateDocumentation(NameToTask, new FileReference(DocumentationFileName));
+				WriteDocumentation(NameToTask, new FileReference(DocumentationFileName));
 				return ExitCode.Success;
 			}
 
@@ -824,7 +821,7 @@ namespace AutomationTool
 		/// </summary>
 		/// <param name="NameToTask">Map of task name to implementation</param>
 		/// <param name="OutputFile">Output file</param>
-		static void GenerateDocumentation(Dictionary<string, ScriptTask> NameToTask, FileReference OutputFile)
+		static void WriteDocumentation(Dictionary<string, ScriptTask> NameToTask, FileReference OutputFile)
 		{
 			// Find all the assemblies containing tasks
 			Assembly[] TaskAssemblies = NameToTask.Values.Select(x => x.ParametersClass.Assembly).Distinct().ToArray();
@@ -860,6 +857,125 @@ namespace AutomationTool
 			}
 
 			// Write the output file
+			if (OutputFile.HasExtension(".udn"))
+			{
+				WriteDocumentationUDN(NameToTask, MemberNameToElement, OutputFile);
+			}
+			else if (OutputFile.HasExtension(".html"))
+			{
+				WriteDocumentationHTML(NameToTask, MemberNameToElement, OutputFile);
+			}
+			else
+			{
+				throw new BuildException("Unable to detect format from extension of output file ({0})", OutputFile);
+			}
+		}
+
+		/// <summary>
+		/// Writes documentation to a UDN file
+		/// </summary>
+		/// <param name="NameToTask">Map of name to script task</param>
+		/// <param name="MemberNameToElement">Map of field name to XML documenation element</param>
+		/// <param name="OutputFile">The output file to write to</param>
+		static void WriteDocumentationUDN(Dictionary<string, ScriptTask> NameToTask, Dictionary<string, XmlElement> MemberNameToElement, FileReference OutputFile)
+		{
+			using (StreamWriter Writer = new StreamWriter(OutputFile.FullName))
+			{
+				Writer.WriteLine("Availability: NoPublish");
+				Writer.WriteLine("Title: BuildGraph Predefined Tasks");
+				Writer.WriteLine("Crumbs: %ROOT%, Programming, Programming/Development, Programming/Development/BuildGraph, Programming/Development/BuildGraph/BuildGraphScriptTasks");
+				Writer.WriteLine("Description: This is a procedurally generated markdown page.");
+				Writer.WriteLine("version: {0}.{1}", ReadOnlyBuildVersion.Current.MajorVersion, ReadOnlyBuildVersion.Current.MinorVersion);
+				Writer.WriteLine("parent:Programming/Development/BuildGraph/BuildGraphScriptTasks");
+				Writer.WriteLine();
+				foreach (string TaskName in NameToTask.Keys.OrderBy(x => x))
+				{
+					// Get the task object
+					ScriptTask Task = NameToTask[TaskName];
+
+					// Get the documentation for this task
+					XmlElement TaskElement;
+					if (MemberNameToElement.TryGetValue("T:" + Task.TaskClass.FullName, out TaskElement))
+					{
+						// Write the task heading
+						Writer.WriteLine("### {0}", TaskName);
+						Writer.WriteLine();
+						Writer.WriteLine(ConvertToMarkdown(TaskElement.SelectSingleNode("summary")));
+						Writer.WriteLine();
+
+						// Document the parameters
+						List<string[]> Rows = new List<string[]>();
+						foreach (string ParameterName in Task.NameToParameter.Keys)
+						{
+							// Get the parameter data
+							ScriptTaskParameter Parameter = Task.NameToParameter[ParameterName];
+
+							// Get the documentation for this parameter
+							XmlElement ParameterElement;
+							if (MemberNameToElement.TryGetValue("F:" + Parameter.FieldInfo.DeclaringType.FullName + "." + Parameter.Name, out ParameterElement))
+							{
+								string TypeName = Parameter.FieldInfo.FieldType.Name;
+								if (Parameter.ValidationType != TaskParameterValidationType.Default)
+								{
+									StringBuilder NewTypeName = new StringBuilder(Parameter.ValidationType.ToString());
+									for (int Idx = 1; Idx < NewTypeName.Length; Idx++)
+									{
+										if (Char.IsLower(NewTypeName[Idx - 1]) && Char.IsUpper(NewTypeName[Idx]))
+										{
+											NewTypeName.Insert(Idx, ' ');
+										}
+									}
+									TypeName = NewTypeName.ToString();
+								}
+
+								string[] Columns = new string[4];
+								Columns[0] = ParameterName;
+								Columns[1] = TypeName;
+								Columns[2] = Parameter.bOptional ? "Optional" : "Required";
+								Columns[3] = ConvertToMarkdown(ParameterElement.SelectSingleNode("summary"));
+								Rows.Add(Columns);
+							}
+						}
+
+						// Always include the "If" attribute
+						string[] IfColumns = new string[4];
+						IfColumns[0] = "If";
+						IfColumns[1] = "Condition";
+						IfColumns[2] = "Optional";
+						IfColumns[3] = "Whether to execute this task. It is ignored if this condition evaluates to false.";
+						Rows.Add(IfColumns);
+
+						// Get the width of each column
+						int[] Widths = new int[4];
+						for (int Idx = 0; Idx < 4; Idx++)
+						{
+							Widths[Idx] = Rows.Max(x => x[Idx].Length);
+						}
+
+						// Format the markdown table
+						string Format = String.Format("| {{0,-{0}}} | {{1,-{1}}} | {{2,-{2}}} | {{3,-{3}}} |", Widths[0], Widths[1], Widths[2], Widths[3]);
+						Writer.WriteLine(Format, "", "", "", "");
+						Writer.WriteLine(Format, new string('-', Widths[0]), new string('-', Widths[1]), new string('-', Widths[2]), new string('-', Widths[3]));
+						for (int Idx = 0; Idx < Rows.Count; Idx++)
+						{
+							Writer.WriteLine(Format, Rows[Idx][0], Rows[Idx][1], Rows[Idx][2], Rows[Idx][3]);
+						}
+
+						// Blank line before next task
+						Writer.WriteLine();
+					}
+				}
+			}
+		}
+
+		/// <summary>
+		/// Writes documentation to an HTML file
+		/// </summary>
+		/// <param name="NameToTask">Map of name to script task</param>
+		/// <param name="MemberNameToElement">Map of field name to XML documenation element</param>
+		/// <param name="OutputFile">The output file to write to</param>
+		static void WriteDocumentationHTML(Dictionary<string, ScriptTask> NameToTask, Dictionary<string, XmlElement> MemberNameToElement, FileReference OutputFile)
+		{
 			LogInformation("Writing {0}...", OutputFile);
 			using (StreamWriter Writer = new StreamWriter(OutputFile.FullName))
 			{
diff --git a/Engine/Source/Programs/UnrealBuildTool/System/RulesDocumentation.cs b/Engine/Source/Programs/UnrealBuildTool/System/RulesDocumentation.cs
--- a/Engine/Source/Programs/UnrealBuildTool/System/RulesDocumentation.cs
+++ b/Engine/Source/Programs/UnrealBuildTool/System/RulesDocumentation.cs
@@ -56,39 +56,134 @@ namespace UnrealBuildTool
 			}
 
 			// Generate the documentation file
+			if (OutputFile.HasExtension(".udn"))
+			{
+				WriteDocumentationUDN(OutputFile, ReadOnlyFields, ReadWriteFields, InputDocumentation);
+			}
+			else if (OutputFile.HasExtension(".html"))
+			{
+				WriteDocumentationHTML(OutputFile, ReadOnlyFields, ReadWriteFields, InputDocumentation);
+			}
+			else
+			{
+				throw new BuildException("Unable to detect format from extension of output file ({0})", OutputFile);
+			}
+
+			// Success!
+			Log.TraceInformation("Written documentation to {0}.", OutputFile);
+		}
+
+		static void WriteDocumentationUDN(FileReference OutputFile, List<FieldInfo> ReadOnlyFields, List<FieldInfo> ReadWriteFields, XmlDocument InputDocumentation)
+		{
+			// Generate the UDN documentation file
+			using (StreamWriter Writer = new StreamWriter(OutputFile.FullName))
+			{
+				Writer.WriteLine("Availability: NoPublish");
+				Writer.WriteLine("Title: Build Configuration Properties Page");
+				Writer.WriteLine("Crumbs:");
+				Writer.WriteLine("Description: This is a procedurally generated markdown page.");
+				Writer.WriteLine("Version: {0}.{1}", ReadOnlyBuildVersion.Current.MajorVersion, ReadOnlyBuildVersion.Current.MinorVersion);
+				Writer.WriteLine("");
+				if (ReadOnlyFields.Count > 0)
+				{
+					Writer.WriteLine("### Read-Only Properties");
+					Writer.WriteLine();
+					foreach (FieldInfo Field in ReadOnlyFields)
+					{
+						WriteFieldUDN(InputDocumentation, Field, Writer);
+					}
+					Writer.WriteLine();
+				}
+				if (ReadWriteFields.Count > 0)
+				{
+					Writer.WriteLine("### Read/Write Properties");
+					foreach (FieldInfo Field in ReadWriteFields)
+					{
+						WriteFieldUDN(InputDocumentation, Field, Writer);
+					}
+					Writer.WriteLine("");
+				}
+			}
+		}
+
+		static void WriteFieldUDN(XmlDocument InputDocumentation, FieldInfo Field, TextWriter Writer)
+		{
+			XmlNode Node = InputDocumentation.SelectSingleNode(String.Format("//member[@name='F:{0}.{1}']/summary", Field.DeclaringType.FullName, Field.Name));
+			if (Node != null)
+			{
+				// Reflow the comments into paragraphs, assuming that each paragraph will be separated by a blank line
+				List<string> Lines = new List<string>(Node.InnerText.Trim().Split('\n').Select(x => x.Trim()));
+				for (int Idx = Lines.Count - 1; Idx > 0; Idx--)
+				{
+					if (Lines[Idx - 1].Length > 0 && !Lines[Idx].StartsWith("*") && !Lines[Idx].StartsWith("-"))
+					{
+						Lines[Idx - 1] += " " + Lines[Idx];
+						Lines.RemoveAt(Idx);
+					}
+				}
+
+				// Write the values of the enum
+				/*				if(Field.FieldType.IsEnum)
+								{
+									Lines.Add("Valid values are:");
+									foreach(string Value in Enum.GetNames(Field.FieldType))
+									{
+										Lines.Add(String.Format("* {0}.{1}", Field.FieldType.Name, Value));
+									}
+								}
+				*/
+				// Write the result to the .udn file
+				if (Lines.Count > 0)
+				{
+					Writer.WriteLine("$ {0} ({1}): {2}", Field.Name, GetPrettyTypeName(Field.FieldType), Lines[0]);
+					for (int Idx = 1; Idx < Lines.Count; Idx++)
+					{
+						if (Lines[Idx].StartsWith("*") || Lines[Idx].StartsWith("-"))
+						{
+							Writer.WriteLine("        * {0}", Lines[Idx].Substring(1).TrimStart());
+						}
+						else
+						{
+							Writer.WriteLine("    * {0}", Lines[Idx]);
+						}
+					}
+					Writer.WriteLine();
+				}
+			}
+		}
+
+		static void WriteDocumentationHTML(FileReference OutputFile, List<FieldInfo> ReadOnlyFields, List<FieldInfo> ReadWriteFields, XmlDocument InputDocumentation)
+		{
 			using (StreamWriter Writer = new StreamWriter(OutputFile.FullName))
 			{
 				Writer.WriteLine("<html>");
 				Writer.WriteLine("  <body>");
-				if(ReadOnlyFields.Count > 0)
+				if (ReadOnlyFields.Count > 0)
 				{
 					Writer.WriteLine("    <h2>Read-Only Properties</h2>");
 					Writer.WriteLine("    <dl>");
-					foreach(FieldInfo Field in ReadOnlyFields)
+					foreach (FieldInfo Field in ReadOnlyFields)
 					{
-						OutputField(InputDocumentation, Field, Writer);
+						WriteFieldHTML(InputDocumentation, Field, Writer);
 					}
 					Writer.WriteLine("    </dl>");
 				}
-				if(ReadWriteFields.Count > 0)
+				if (ReadWriteFields.Count > 0)
 				{
 					Writer.WriteLine("    <h2>Read/Write Properties</h2>");
 					Writer.WriteLine("    <dl>");
-					foreach(FieldInfo Field in ReadWriteFields)
+					foreach (FieldInfo Field in ReadWriteFields)
 					{
-						OutputField(InputDocumentation, Field, Writer);
+						WriteFieldHTML(InputDocumentation, Field, Writer);
 					}
 					Writer.WriteLine("    </dl>");
 				}
 				Writer.WriteLine("  </body>");
 				Writer.WriteLine("</html>");
 			}
-
-			// Success!
-			Log.TraceInformation("Written documentation to {0}.", OutputFile);
 		}
 
-		static void OutputField(XmlDocument InputDocumentation, FieldInfo Field, TextWriter Writer)
+		static void WriteFieldHTML(XmlDocument InputDocumentation, FieldInfo Field, TextWriter Writer)
 		{
 			XmlNode Node = InputDocumentation.SelectSingleNode(String.Format("//member[@name='F:{0}.{1}']/summary", Field.DeclaringType.FullName, Field.Name));
 			if(Node != null)
diff --git a/Engine/Source/Programs/UnrealBuildTool/System/XmlConfig.cs b/Engine/Source/Programs/UnrealBuildTool/System/XmlConfig.cs
--- a/Engine/Source/Programs/UnrealBuildTool/System/XmlConfig.cs
+++ b/Engine/Source/Programs/UnrealBuildTool/System/XmlConfig.cs
@@ -702,31 +702,122 @@ namespace UnrealBuildTool
 			}
 
 			// Generate the documentation file
+			if(OutputFile.HasExtension(".udn"))
+			{
+				WriteDocumentationUDN(OutputFile, InputDocumentation, CategoryToFields);
+			}
+			else if(OutputFile.HasExtension(".html"))
+			{
+				WriteDocumentationHTML(OutputFile, InputDocumentation, CategoryToFields);
+			}
+			else
+			{
+				throw new BuildException("Unable to detect format from extension of output file ({0})", OutputFile);
+			}
+
+			// Success!
+			Log.TraceInformation("Written documentation to {0}.", OutputFile);
+		}
+
+		/// <summary>
+		/// Writes out documentation in UDN format
+		/// </summary>
+		/// <param name="OutputFile">The output file</param>
+		/// <param name="InputDocumentation">The XML documentation for this assembly</param>
+		/// <param name="CategoryToFields">Map of string to types to fields</param>
+		private static void WriteDocumentationUDN(FileReference OutputFile, XmlDocument InputDocumentation, Dictionary<string, Dictionary<string, FieldInfo>> CategoryToFields)
+		{
+			using (StreamWriter Writer = new StreamWriter(OutputFile.FullName))
+			{
+				Writer.WriteLine("Availability: NoPublish");
+				Writer.WriteLine("Title: Build Configuration Properties Page");
+				Writer.WriteLine("Crumbs:");
+				Writer.WriteLine("Description: This is a procedurally generated markdown page.");
+				Writer.WriteLine("Version: {0}.{1}", ReadOnlyBuildVersion.Current.MajorVersion, ReadOnlyBuildVersion.Current.MinorVersion);
+				Writer.WriteLine("");
+
+				foreach (KeyValuePair<string, Dictionary<string, FieldInfo>> CategoryPair in CategoryToFields)
+				{
+					string CategoryName = CategoryPair.Key;
+					Writer.WriteLine("### {0}", CategoryName);
+					Writer.WriteLine();
+
+					Dictionary<string, FieldInfo> Fields = CategoryPair.Value;
+					foreach (KeyValuePair<string, FieldInfo> FieldPair in Fields)
+					{
+						string FieldName = FieldPair.Key;
+
+						FieldInfo Field = FieldPair.Value;
+						XmlNode Node = InputDocumentation.SelectSingleNode(String.Format("//member[@name='F:{0}.{1}']/summary", Field.DeclaringType.FullName, Field.Name));
+						if (Node != null)
+						{
+							// Reflow the comments into paragraphs, assuming that each paragraph will be separated by a blank line
+							List<string> Lines = new List<string>(Node.InnerText.Trim().Split('\n').Select(x => x.Trim()));
+							for (int Idx = Lines.Count - 1; Idx > 0; Idx--)
+							{
+								if (Lines[Idx - 1].Length > 0 && !Lines[Idx].StartsWith("*") && !Lines[Idx].StartsWith("-"))
+								{
+									Lines[Idx - 1] += " " + Lines[Idx];
+									Lines.RemoveAt(Idx);
+								}
+							}
+
+							// Write the result to the .udn file
+							if (Lines.Count > 0)
+							{
+								Writer.WriteLine("$ {0} : {1}", FieldName, Lines[0]);
+								for (int Idx = 1; Idx < Lines.Count; Idx++)
+								{
+									if (Lines[Idx].StartsWith("*") || Lines[Idx].StartsWith("-"))
+									{
+										Writer.WriteLine("        * {0}", Lines[Idx].Substring(1).TrimStart());
+									}
+									else
+									{
+										Writer.WriteLine("    * {0}", Lines[Idx]);
+									}
+								}
+								Writer.WriteLine();
+							}
+						}
+					}
+				}
+			}
+		}
+
+		/// <summary>
+		/// Writes out documentation in HTML format
+		/// </summary>
+		/// <param name="OutputFile">The output file</param>
+		/// <param name="InputDocumentation">The XML documentation for this assembly</param>
+		/// <param name="CategoryToFields">Map of string to types to fields</param>
+		private static void WriteDocumentationHTML(FileReference OutputFile, XmlDocument InputDocumentation, Dictionary<string, Dictionary<string, FieldInfo>> CategoryToFields)
+		{
 			using (StreamWriter Writer = new StreamWriter(OutputFile.FullName))
 			{
 				Writer.WriteLine("<html>");
 				Writer.WriteLine("  <body>");
 				Writer.WriteLine("  <h2>BuildConfiguration Properties</h2>");
-				foreach(KeyValuePair<string, Dictionary<string, FieldInfo>> CategoryPair in CategoryToFields)
+				foreach (KeyValuePair<string, Dictionary<string, FieldInfo>> CategoryPair in CategoryToFields)
 				{
 					string CategoryName = CategoryPair.Key;
 					Writer.WriteLine("    <h3>{0}</h3>", CategoryName);
 					Writer.WriteLine("    <dl>");
 
 					Dictionary<string, FieldInfo> Fields = CategoryPair.Value;
-					foreach(KeyValuePair<string, FieldInfo> FieldPair in Fields)
+					foreach (KeyValuePair<string, FieldInfo> FieldPair in Fields)
 					{
 						string FieldName = FieldPair.Key;
 
 						FieldInfo Field = FieldPair.Value;
 						XmlNode Node = InputDocumentation.SelectSingleNode(String.Format("//member[@name='F:{0}.{1}']/summary", Field.DeclaringType.FullName, Field.Name));
-						if(Node != null)
+						if (Node != null)
 						{
 							// Reflow the comments into paragraphs, assuming that each paragraph will be separated by a blank line
 							List<string> Lines = new List<string>(Node.InnerText.Trim().Split('\n').Select(x => x.Trim()));
-							for(int Idx = Lines.Count - 1; Idx > 0; Idx--)
+							for (int Idx = Lines.Count - 1; Idx > 0; Idx--)
 							{
-								if(Lines[Idx - 1].Length > 0 && !Lines[Idx].StartsWith("*") && !Lines[Idx].StartsWith("-"))
+								if (Lines[Idx - 1].Length > 0 && !Lines[Idx].StartsWith("*") && !Lines[Idx].StartsWith("-"))
 								{
 									Lines[Idx - 1] += " " + Lines[Idx];
 									Lines.RemoveAt(Idx);
@@ -734,23 +825,23 @@ namespace UnrealBuildTool
 							}
 
 							// Write the result to the .udn file
-							if(Lines.Count > 0)
+							if (Lines.Count > 0)
 							{
 								Writer.WriteLine("      <dt>{0}</dt>", FieldName);
 
-								if(Lines.Count == 1)
+								if (Lines.Count == 1)
 								{
 									Writer.WriteLine("      <dd>{0}</dd>", Lines[0]);
 								}
 								else
 								{
 									Writer.WriteLine("      <dd>");
-									for(int Idx = 0; Idx < Lines.Count; Idx++)
+									for (int Idx = 0; Idx < Lines.Count; Idx++)
 									{
-										if(Lines[Idx].StartsWith("*") || Lines[Idx].StartsWith("-"))
+										if (Lines[Idx].StartsWith("*") || Lines[Idx].StartsWith("-"))
 										{
 											Writer.WriteLine("        <ul>");
-											for(; Idx < Lines.Count && (Lines[Idx].StartsWith("*") || Lines[Idx].StartsWith("-")); Idx++)
+											for (; Idx < Lines.Count && (Lines[Idx].StartsWith("*") || Lines[Idx].StartsWith("-")); Idx++)
 											{
 												Writer.WriteLine("          <li>{0}</li>", Lines[Idx].Substring(1).TrimStart());
 											}
@@ -772,9 +863,7 @@ namespace UnrealBuildTool
 				Writer.WriteLine("  </body>");
 				Writer.WriteLine("</html>");
 			}
-
-			// Success!
-			Log.TraceInformation("Written documentation to {0}.", OutputFile);
 		}
 	}
 }
+
