UBT: Warn about undocumented fields in BuildConfiguration files.

#rnx

diff --git a/Engine/Source/Programs/UnrealBuildTool/System/XmlConfig.cs b/Engine/Source/Programs/UnrealBuildTool/System/XmlConfig.cs
--- a/Engine/Source/Programs/UnrealBuildTool/System/XmlConfig.cs
+++ b/Engine/Source/Programs/UnrealBuildTool/System/XmlConfig.cs
@@ -679,6 +679,7 @@ namespace UnrealBuildTool
 			// Find all the configurable fields from the given types
 			Dictionary<string, Dictionary<string, FieldInfo>> CategoryToFields = new Dictionary<string, Dictionary<string, FieldInfo>>();
 			FindConfigurableFields(ConfigTypes, CategoryToFields);
+			CategoryToFields = CategoryToFields.Where(x => x.Value.Count > 0).ToDictionary(x => x.Key, x => x.Value);
 
 			// Get the path to the XML documentation
 			FileReference InputDocumentationFile = new FileReference(Assembly.GetExecutingAssembly().Location).ChangeExtension(".xml");
@@ -719,6 +720,37 @@ namespace UnrealBuildTool
 			Log.TraceInformation("Written documentation to {0}.", OutputFile);
 		}
 
+		/// <summary>
+		/// Gets the XML comment for a particular field
+		/// </summary>
+		/// <param name="Documentation">The XML documentation</param>
+		/// <param name="Field">The field to search for</param>
+		/// <param name="Lines">Receives the description for the requested field</param>
+		/// <returns>True if a comment was found for the field</returns>
+		private static bool TryGetXmlComment(XmlDocument Documentation, FieldInfo Field, out List<string> Lines)
+		{
+			XmlNode Node = Documentation.SelectSingleNode(String.Format("//member[@name='F:{0}.{1}']/summary", Field.DeclaringType.FullName, Field.Name));
+			if (Node == null)
+			{
+				Lines = null;
+				return false;
+			}
+			else
+			{
+				// Reflow the comments into paragraphs, assuming that each paragraph will be separated by a blank line
+				Lines = new List<string>(Node.InnerText.Trim().Split('\n').Select(x => x.Trim()));
+				for (int Idx = Lines.Count - 1; Idx > 0; Idx--)
+				{
+					if (Lines[Idx - 1].Length > 0 && !Lines[Idx].StartsWith("*") && !Lines[Idx].StartsWith("-"))
+					{
+						Lines[Idx - 1] += " " + Lines[Idx];
+						Lines.RemoveAt(Idx);
+					}
+				}
+				return true;
+			}
+		}
+
 		/// <summary>
 		/// Writes out documentation in UDN format
 		/// </summary>
@@ -745,41 +777,28 @@ namespace UnrealBuildTool
 					Dictionary<string, FieldInfo> Fields = CategoryPair.Value;
 					foreach (KeyValuePair<string, FieldInfo> FieldPair in Fields)
 					{
-						string FieldName = FieldPair.Key;
+						// Get the XML comment for this field
+						List<string> Lines;
+						if(!TryGetXmlComment(InputDocumentation, FieldPair.Value, out Lines) || Lines.Count == 0)
+						{
+							Log.TraceWarning("Missing XML comment for {0}", FieldPair.Value.Name);
+							continue;
+						}
 
-						FieldInfo Field = FieldPair.Value;
-						XmlNode Node = InputDocumentation.SelectSingleNode(String.Format("//member[@name='F:{0}.{1}']/summary", Field.DeclaringType.FullName, Field.Name));
-						if (Node != null)
+						// Write the result to the .udn file
+						Writer.WriteLine("$ {0} : {1}", FieldPair.Key, Lines[0]);
+						for (int Idx = 1; Idx < Lines.Count; Idx++)
 						{
-							// Reflow the comments into paragraphs, assuming that each paragraph will be separated by a blank line
-							List<string> Lines = new List<string>(Node.InnerText.Trim().Split('\n').Select(x => x.Trim()));
-							for (int Idx = Lines.Count - 1; Idx > 0; Idx--)
+							if (Lines[Idx].StartsWith("*") || Lines[Idx].StartsWith("-"))
 							{
-								if (Lines[Idx - 1].Length > 0 && !Lines[Idx].StartsWith("*") && !Lines[Idx].StartsWith("-"))
-								{
-									Lines[Idx - 1] += " " + Lines[Idx];
-									Lines.RemoveAt(Idx);
-								}
+								Writer.WriteLine("        * {0}", Lines[Idx].Substring(1).TrimStart());
 							}
-
-							// Write the result to the .udn file
-							if (Lines.Count > 0)
+							else
 							{
-								Writer.WriteLine("$ {0} : {1}", FieldName, Lines[0]);
-								for (int Idx = 1; Idx < Lines.Count; Idx++)
-								{
-									if (Lines[Idx].StartsWith("*") || Lines[Idx].StartsWith("-"))
-									{
-										Writer.WriteLine("        * {0}", Lines[Idx].Substring(1).TrimStart());
-									}
-									else
-									{
-										Writer.WriteLine("    * {0}", Lines[Idx]);
-									}
-								}
-								Writer.WriteLine();
+								Writer.WriteLine("    * {0}", Lines[Idx]);
 							}
 						}
+						Writer.WriteLine();
 					}
 				}
 			}
@@ -807,54 +826,41 @@ namespace UnrealBuildTool
 					Dictionary<string, FieldInfo> Fields = CategoryPair.Value;
 					foreach (KeyValuePair<string, FieldInfo> FieldPair in Fields)
 					{
-						string FieldName = FieldPair.Key;
-
-						FieldInfo Field = FieldPair.Value;
-						XmlNode Node = InputDocumentation.SelectSingleNode(String.Format("//member[@name='F:{0}.{1}']/summary", Field.DeclaringType.FullName, Field.Name));
-						if (Node != null)
+						// Get the XML comment for this field
+						List<string> Lines;
+						if (!TryGetXmlComment(InputDocumentation, FieldPair.Value, out Lines) || Lines.Count == 0)
 						{
-							// Reflow the comments into paragraphs, assuming that each paragraph will be separated by a blank line
-							List<string> Lines = new List<string>(Node.InnerText.Trim().Split('\n').Select(x => x.Trim()));
-							for (int Idx = Lines.Count - 1; Idx > 0; Idx--)
-							{
-								if (Lines[Idx - 1].Length > 0 && !Lines[Idx].StartsWith("*") && !Lines[Idx].StartsWith("-"))
-								{
-									Lines[Idx - 1] += " " + Lines[Idx];
-									Lines.RemoveAt(Idx);
-								}
-							}
+							Log.TraceWarning("Missing XML comment for {0}", FieldPair.Value.Name);
+							continue;
+						}
 
-							// Write the result to the .udn file
-							if (Lines.Count > 0)
-							{
-								Writer.WriteLine("      <dt>{0}</dt>", FieldName);
+						// Write the result to the .udn file
+						Writer.WriteLine("      <dt>{0}</dt>", FieldPair.Key);
 
-								if (Lines.Count == 1)
+						if (Lines.Count == 1)
+						{
+							Writer.WriteLine("      <dd>{0}</dd>", Lines[0]);
+						}
+						else
+						{
+							Writer.WriteLine("      <dd>");
+							for (int Idx = 0; Idx < Lines.Count; Idx++)
+							{
+								if (Lines[Idx].StartsWith("*") || Lines[Idx].StartsWith("-"))
 								{
-									Writer.WriteLine("      <dd>{0}</dd>", Lines[0]);
+									Writer.WriteLine("        <ul>");
+									for (; Idx < Lines.Count && (Lines[Idx].StartsWith("*") || Lines[Idx].StartsWith("-")); Idx++)
+									{
+										Writer.WriteLine("          <li>{0}</li>", Lines[Idx].Substring(1).TrimStart());
+									}
+									Writer.WriteLine("        </ul>");
 								}
 								else
 								{
-									Writer.WriteLine("      <dd>");
-									for (int Idx = 0; Idx < Lines.Count; Idx++)
-									{
-										if (Lines[Idx].StartsWith("*") || Lines[Idx].StartsWith("-"))
-										{
-											Writer.WriteLine("        <ul>");
-											for (; Idx < Lines.Count && (Lines[Idx].StartsWith("*") || Lines[Idx].StartsWith("-")); Idx++)
-											{
-												Writer.WriteLine("          <li>{0}</li>", Lines[Idx].Substring(1).TrimStart());
-											}
-											Writer.WriteLine("        </ul>");
-										}
-										else
-										{
-											Writer.WriteLine("        {0}", Lines[Idx]);
-										}
-									}
-									Writer.WriteLine("      </dd>");
+									Writer.WriteLine("        {0}", Lines[Idx]);
 								}
 							}
+							Writer.WriteLine("      </dd>");
 						}
 					}
 
