Allow optional for alternative login id determination

#android

diff --git a/Engine/Build/Android/Java/src/com/epicgames/ue4/GameActivity.java.template b/Engine/Build/Android/Java/src/com/epicgames/ue4/GameActivity.java.template
--- a/Engine/Build/Android/Java/src/com/epicgames/ue4/GameActivity.java.template
+++ b/Engine/Build/Android/Java/src/com/epicgames/ue4/GameActivity.java.template
@@ -4057,6 +4057,8 @@ public class GameActivity extends $${gameActivitySuperClass}$$ implements Surfac
 	{
 		String loginId = null;
 
+	//$${gameActivityGetLoginIdAdditions}$$
+
 		// check for existing GUID fallback in internal
 		File file = new File(InternalFilesDir + "login-identifier.txt");
 		if (file.exists())
diff --git a/Engine/Source/Programs/UnrealBuildTool/Platform/Android/UEDeployAndroid.cs b/Engine/Source/Programs/UnrealBuildTool/Platform/Android/UEDeployAndroid.cs
--- a/Engine/Source/Programs/UnrealBuildTool/Platform/Android/UEDeployAndroid.cs
+++ b/Engine/Source/Programs/UnrealBuildTool/Platform/Android/UEDeployAndroid.cs
@@ -4479,6 +4479,7 @@ namespace UnrealBuildTool
 				{ "//$${gameActivityResizeKeyboardAdditions}$$", UPL.ProcessPluginNode(NDKArch, "gameActivityResizeKeyboardAdditions", "")},
 				{ "//$${gameActivityLoggerCallbackAdditions}$$", UPL.ProcessPluginNode(NDKArch, "gameActivityLoggerCallbackAdditions", "")},
 				{ "//$${gameActivityGetCommandLineAdditions}$$", UPL.ProcessPluginNode(NDKArch, "gameActivityGetCommandLineAdditions", "")},
+				{ "//$${gameActivityGetLoginIdAdditions}$$", UPL.ProcessPluginNode(NDKArch, "gameActivityGetLoginIdAdditions", "")},
 				{ "//$${soLoadLibrary}$$", UPL.ProcessPluginNode(NDKArch, "soLoadLibrary", LoadLibraryDefaults)},
 				{ "$${gameActivitySuperClass}$$", SuperClassDefault},
 			};
