Add UnrealInsights to the default makefile target on Linux.

diff --git a/Engine/Source/Programs/UnrealBuildTool/ProjectFiles/Make/MakefileGenerator.cs b/Engine/Source/Programs/UnrealBuildTool/ProjectFiles/Make/MakefileGenerator.cs
--- a/Engine/Source/Programs/UnrealBuildTool/ProjectFiles/Make/MakefileGenerator.cs
+++ b/Engine/Source/Programs/UnrealBuildTool/ProjectFiles/Make/MakefileGenerator.cs
@@ -127,7 +127,7 @@ namespace UnrealBuildTool
 			MakefileContent.Append("\n\n" + BuildCommand + ProjectBuildCommand + "\n" +
 				"all: StandardSet\n\n" +
 				"RequiredTools: CrashReportClient-Linux-Shipping ShaderCompileWorker UnrealLightmass\n\n" +
-				"StandardSet: RequiredTools UnrealFrontend UE4Editor\n\n" +
+				"StandardSet: RequiredTools UnrealFrontend UE4Editor UnrealInsights\n\n" +
 				"DebugSet: RequiredTools UnrealFrontend-Linux-Debug UE4Editor-Linux-Debug\n\n"
 			);
 
