Re-adding GLES_SILENCE_DEPRECATION and removing IOS GLES version check to solve binary build issues

diff --git a/Engine/Source/Programs/UnrealBuildTool/Platform/IOS/UEBuildIOS.cs b/Engine/Source/Programs/UnrealBuildTool/Platform/IOS/UEBuildIOS.cs
--- a/Engine/Source/Programs/UnrealBuildTool/Platform/IOS/UEBuildIOS.cs
+++ b/Engine/Source/Programs/UnrealBuildTool/Platform/IOS/UEBuildIOS.cs
@@ -1059,7 +1059,11 @@ namespace UnrealBuildTool
 				CompileEnvironment.Definitions.Add("ENABLE_ADVERTISING_IDENTIFIER=1");
 			}
 
-			CompileEnvironment.Definitions.Add("HAS_OPENGL_ES=" + ((Target.IOSPlatform.RuntimeVersion < 12.0) ? "1" : "0"));
+			// For 4.23 the version check fails for binary builds, re-enabling along with deprecation warning suppression for all builds UE-77520
+			// ES support will be fully removed in 4.24
+			CompileEnvironment.Definitions.Add("HAS_OPENGL_ES=1");
+			CompileEnvironment.Definitions.Add("GLES_SILENCE_DEPRECATION=1"); 
+			//CompileEnvironment.Definitions.Add("HAS_OPENGL_ES=" + ((Target.IOSPlatform.RuntimeVersion < 12.0) ? "1" : "0"));
 
 			// if the project has an Oodle compression Dll, enable the decompressor on IOS
 			if (Target.ProjectFile != null)
diff --git a/Engine/Source/Runtime/Launch/Launch.Build.cs b/Engine/Source/Runtime/Launch/Launch.Build.cs
--- a/Engine/Source/Runtime/Launch/Launch.Build.cs
+++ b/Engine/Source/Runtime/Launch/Launch.Build.cs
@@ -249,10 +249,13 @@ public class Launch : ModuleRules
 				"IOSRuntimeSettings",
 			});
 
-			// no longer build GL for apps requiring iOS 12 or later
-			if (Target.IOSPlatform.RuntimeVersion < 12.0)
-			{
-				PublicFrameworks.Add("OpenGLES");
+            // For 4.23 the below check fails for binary builds, re-enabling for all builds UE-77520
+            // ES support will be fully removed in 4.24
+
+            // no longer build GL for apps requiring iOS 12 or later
+            //if (Target.IOSPlatform.RuntimeVersion < 12.0)
+            {
+                PublicFrameworks.Add("OpenGLES");
 				PrivateDependencyModuleNames.Add("OpenGLDrv");
 			}
 			// needed for Metal layer
