UBT: Add a mechanism for versioning which build settings UBT should use by default.

Any build targets have a DefaultBuildSettings property. For engine targets, this defaults to BuildSettingsVersion.Latest. For project targets, this defaults to BuildSettingsVersion.Release_4_23. For new projects, this will default to the engine version they are created from.

If a target is not using the latest default build settings, they will receive a message describing the settings that have changed, like this:

    [Upgrade]
    [Upgrade] Using UE 4.23 compatible build settings. The latest version of UE4 sets the following values by default, which may require code changes:
    [Upgrade]     bLegacyPublicIncludePaths = false                 => Omits subfolders from public include paths to reduce compiler command line length.
    [Upgrade]     PCHUsage = PCHUsageMode.UseExplicitOrSharedPCHs   => Set in build.cs files to enables IWYU-style PCH model. See https://docs.unrealengine.com/en-US/Programming/BuildTools/UnrealBuildTool/IWYU/index.html.
    [Upgrade] Suppress this message by setting 'DefaultBuildSettings = BuildSettingsVersion.Release_4_24;' in UnrealPak.Target.cs, and explicitly overriding desired settings.
    [Upgrade]

Intent is to reduce friction for users initially upgrading to new engine versions, while notifying them of more optimal build settings being available, and letting them choose when (or if) to use them.

diff --git a/Engine/Source/Editor/GameProjectGeneration/Private/GameProjectUtils.cpp b/Engine/Source/Editor/GameProjectGeneration/Private/GameProjectUtils.cpp
--- a/Engine/Source/Editor/GameProjectGeneration/Private/GameProjectUtils.cpp
+++ b/Engine/Source/Editor/GameProjectGeneration/Private/GameProjectUtils.cpp
@@ -1381,6 +1381,13 @@ bool GameProjectUtils::CreateProjectFromTemplate(const FProjectInformation& InPr
 
 	SlowTask.EnterProgressFrame();
 
+	// Add a rule to replace the build settings version with the appropriate number
+	FTemplateReplacement& DefaultBuildSettingsRepl = TemplateDefs->ReplacementsInFiles.AddZeroed_GetRef();
+	DefaultBuildSettingsRepl.Extensions.Add("cs");
+	DefaultBuildSettingsRepl.From = TEXT("BuildSettingsVersion.Latest");
+	DefaultBuildSettingsRepl.To = GetDefaultBuildSettingsVersion();
+	DefaultBuildSettingsRepl.bCaseSensitive = true;
+
 	// Fix up the replacement strings using the specified project name
 	TemplateDefs->FixupStrings(TemplateName, ProjectName);
 
@@ -2809,6 +2816,11 @@ void GameProjectUtils::UpdateAdditionalPluginDirectory(const FString& InDir, con
 	}
 }
 
+const TCHAR* GameProjectUtils::GetDefaultBuildSettingsVersion()
+{
+	return TEXT("BuildSettingsVersion.V2");
+}
+
 bool GameProjectUtils::ReadTemplateFile(const FString& TemplateFileName, FString& OutFileContents, FText& OutFailReason)
 {
 	const FString FullFileName = FPaths::EngineContentDir() / TEXT("Editor") / TEXT("Templates") / TemplateFileName;
@@ -3263,6 +3275,7 @@ bool GameProjectUtils::GenerateGameModuleTargetFile(const FString& NewBuildFileN
 	FinalOutput = FinalOutput.Replace(TEXT("%EXTRA_MODULE_NAMES%"), *MakeCommaDelimitedList(ExtraModuleNames), ESearchCase::CaseSensitive);
 	FinalOutput = FinalOutput.Replace(TEXT("%MODULE_NAME%"), *ModuleName, ESearchCase::CaseSensitive);
 	FinalOutput = FinalOutput.Replace(TEXT("%TARGET_TYPE%"), TEXT("Game"), ESearchCase::CaseSensitive);
+	FinalOutput = FinalOutput.Replace(TEXT("%DEFAULT_BUILD_SETTINGS_VERSION%"), GetDefaultBuildSettingsVersion(), ESearchCase::CaseSensitive);
 
 	return WriteOutputFile(NewBuildFileName, FinalOutput, OutFailReason);
 }
@@ -3297,6 +3310,7 @@ bool GameProjectUtils::GenerateEditorModuleTargetFile(const FString& NewBuildFil
 	FinalOutput = FinalOutput.Replace(TEXT("%EXTRA_MODULE_NAMES%"), *MakeCommaDelimitedList(ExtraModuleNames), ESearchCase::CaseSensitive);
 	FinalOutput = FinalOutput.Replace(TEXT("%MODULE_NAME%"), *ModuleName, ESearchCase::CaseSensitive);
 	FinalOutput = FinalOutput.Replace(TEXT("%TARGET_TYPE%"), TEXT("Editor"), ESearchCase::CaseSensitive);
+	FinalOutput = FinalOutput.Replace(TEXT("%DEFAULT_BUILD_SETTINGS_VERSION%"), GetDefaultBuildSettingsVersion(), ESearchCase::CaseSensitive);
 
 	return WriteOutputFile(NewBuildFileName, FinalOutput, OutFailReason);
 }
diff --git a/Engine/Source/Editor/GameProjectGeneration/Public/GameProjectUtils.h b/Engine/Source/Editor/GameProjectGeneration/Public/GameProjectUtils.h
--- a/Engine/Source/Editor/GameProjectGeneration/Public/GameProjectUtils.h
+++ b/Engine/Source/Editor/GameProjectGeneration/Public/GameProjectUtils.h
@@ -277,6 +277,9 @@ public:
 	 */
 	static void UpdateAdditionalPluginDirectory(const FString& InDir, const bool bAddOrRemove);
 
+	/** Gets the default build settings version for UBT */
+	static const TCHAR* GetDefaultBuildSettingsVersion();
+
 private:
 
 	static FString GetHardwareConfigString(const FProjectInformation& InProjectInfo);
diff --git a/Engine/Source/Programs/UnrealBuildTool/Configuration/ModuleRules.cs b/Engine/Source/Programs/UnrealBuildTool/Configuration/ModuleRules.cs
--- a/Engine/Source/Programs/UnrealBuildTool/Configuration/ModuleRules.cs
+++ b/Engine/Source/Programs/UnrealBuildTool/Configuration/ModuleRules.cs
@@ -485,7 +485,34 @@ namespace UnrealBuildTool
 		/// <summary>
 		/// Precompiled header usage for this module
 		/// </summary>
-		public PCHUsageMode PCHUsage = PCHUsageMode.Default;
+		public PCHUsageMode PCHUsage
+		{
+			get
+			{
+				if (PCHUsagePrivate.HasValue)
+				{
+					// Use the override
+					return PCHUsagePrivate.Value;
+				}
+				else if(Target.bIWYU || DefaultBuildSettings >= BuildSettingsVersion.V2)
+				{
+					// Use shared or explicit PCHs, and enable IWYU
+					return PCHUsageMode.UseExplicitOrSharedPCHs;
+				}
+				else if(Plugin != null)
+				{
+					// Older plugins use shared PCHs by default, since they aren't typically large enough to warrant their own PCH.
+					return PCHUsageMode.UseSharedPCHs;
+				}
+				else
+				{
+					// Older game modules do not enable shared PCHs by default, because games usually have a large precompiled header of their own.
+					return PCHUsageMode.NoSharedPCHs;
+				}
+			}
+			set { PCHUsagePrivate = value; }
+		}
+		private PCHUsageMode? PCHUsagePrivate;
 
 		/// <summary>
 		/// Whether this module should be treated as an engine module (eg. using engine definitions, PCHs, compiled with optimizations enabled in DebugGame configurations, etc...).
@@ -493,11 +520,26 @@ namespace UnrealBuildTool
 		/// </summary>
 		public bool bTreatAsEngineModule;
 
+		/// <summary>
+		/// Which engine version's build settings to use by default. 
+		/// </summary>
+		public BuildSettingsVersion DefaultBuildSettings
+		{
+			get { return DefaultBuildSettingsPrivate ?? Target.DefaultBuildSettings; }
+			set { DefaultBuildSettingsPrivate = value; }
+		}
+		private BuildSettingsVersion? DefaultBuildSettingsPrivate;
+
 		/// <summary>
 		/// Whether to use backwards compatible defaults for this module. By default, engine modules always use the latest default settings, while project modules do not (to support
 		/// an easier migration path).
 		/// </summary>
-		public bool bUseBackwardsCompatibleDefaults;
+		[Obsolete("Set DefaultBuildSettings to the appropriate engine version instead")]
+		public bool bUseBackwardsCompatibleDefaults
+		{
+			get { return DefaultBuildSettings != BuildSettingsVersion.Latest; }
+			set { DefaultBuildSettings = bUseBackwardsCompatibleDefaults? BuildSettingsVersion.V1 : BuildSettingsVersion.Latest; }
+		}
 
 		/// <summary>
 		/// Use run time type information
@@ -836,7 +878,7 @@ namespace UnrealBuildTool
 		public bool bLegacyPublicIncludePaths
 		{
 			set { bLegacyPublicIncludePathsPrivate = value; }
-			get { return bLegacyPublicIncludePathsPrivate ?? (bUseBackwardsCompatibleDefaults? Target.bLegacyPublicIncludePaths : false); }
+			get { return bLegacyPublicIncludePathsPrivate ?? ((DefaultBuildSettings < BuildSettingsVersion.V2)? Target.bLegacyPublicIncludePaths : false); }
 		}
 		private bool? bLegacyPublicIncludePathsPrivate;
 
diff --git a/Engine/Source/Programs/UnrealBuildTool/Configuration/TargetRules.cs b/Engine/Source/Programs/UnrealBuildTool/Configuration/TargetRules.cs
--- a/Engine/Source/Programs/UnrealBuildTool/Configuration/TargetRules.cs
+++ b/Engine/Source/Programs/UnrealBuildTool/Configuration/TargetRules.cs
@@ -82,6 +82,29 @@ namespace UnrealBuildTool
 		Unique,
 	}
 
+	/// <summary>
+	/// Determines which version of the engine to take default build settings from. This allows for backwards compatibility as new options are enabled by default.
+	/// </summary>
+	public enum BuildSettingsVersion
+	{
+		/// <summary>
+		/// Legacy default build settings for 4.23 and earlier.
+		/// </summary>
+		V1,
+
+		/// <summary>
+		/// New defaults for 4.24: ModuleRules.PCHUsage = PCHUsageMode.UseExplicitOrSharedPCHs, ModuleRules.bLegacyPublicIncludePaths = false.
+		/// </summary>
+		V2,
+
+		// *** When adding new entries here, be sure to update GameProjectUtils::GetDefaultBuildSettingsVersion() to ensure that new projects are created correctly. ***
+
+		/// <summary>
+		/// Always use the defaults for the current engine version. Note that this may cause compatibility issues when upgrading.
+		/// </summary>
+		Latest
+	}
+
 	/// <summary>
 	/// Attribute used to mark fields which much match between targets in the shared build environment
 	/// </summary>
@@ -174,6 +197,12 @@ namespace UnrealBuildTool
 		/// </summary>
 		public global::UnrealBuildTool.TargetType Type = global::UnrealBuildTool.TargetType.Game;
 
+		/// <summary>
+		/// Specifies the engine version to maintain backwards-compatible default build settings with (eg. DefaultSettingsVersion.Release_4_23, DefaultSettingsVersion.Release_4_24). Specify DefaultSettingsVersion.Latest to always
+		/// use defaults for the current engine version, at the risk of introducing build errors while upgrading.
+		/// </summary>
+		public BuildSettingsVersion DefaultBuildSettings = BuildSettingsVersion.V1;
+
 		/// <summary>
 		/// Tracks a list of config values read while constructing this target
 		/// </summary>
@@ -635,10 +664,15 @@ namespace UnrealBuildTool
 		/// Whether to use backwards compatible defaults for this module. By default, engine modules always use the latest default settings, while project modules do not (to support
 		/// an easier migration path).
 		/// </summary>
-		public bool bUseBackwardsCompatibleDefaults = true;
+		[Obsolete("Set DefaultBuildSettings to the appropriate engine version (eg. BuildSettingsVersion.Release_4_23) or BuildSettingsVersion.Latest instead")]
+		public bool bUseBackwardsCompatibleDefaults
+		{
+			get { return DefaultBuildSettings != BuildSettingsVersion.Latest; }
+			set { DefaultBuildSettings = (value ? BuildSettingsVersion.V1 : BuildSettingsVersion.Latest); }
+		}
 
 		/// <summary>
-		/// Enables "include what you use" by default for modules in this target. Changes the default PCH mode for any module in this project to PCHUsageModule.UseExplicitOrSharedPCHs.
+		/// Enables "include what you use" by default for modules in this target. Changes the default PCH mode for any module in this project to PCHUsageMode.UseExplicitOrSharedPCHs.
 		/// </summary>
 		[CommandLine("-IWYU")]
 		public bool bIWYU = false;
@@ -1055,7 +1089,12 @@ namespace UnrealBuildTool
 		/// <summary>
 		/// Add all the public folders as include paths for the compile environment.
 		/// </summary>
-		public bool bLegacyPublicIncludePaths = true;
+		public bool bLegacyPublicIncludePaths
+		{
+			get { return bLegacyPublicIncludePathsPrivate ?? (DefaultBuildSettings < BuildSettingsVersion.V2); }
+			set { bLegacyPublicIncludePathsPrivate = value; }
+		}
+		private bool? bLegacyPublicIncludePathsPrivate;
 
 		/// <summary>
 		/// Which C++ stanard to use for compiling this target
@@ -1632,6 +1671,38 @@ namespace UnrealBuildTool
 		{
 			get { return UnrealBuildTool.IsEngineInstalled(); }
 		}
+
+		/// <summary>
+		/// Gets diagnostic messages about default settings which have changed in newer versions of the engine
+		/// </summary>
+		/// <param name="Diagnostics">List of diagnostic messages</param>
+		internal void GetBuildSettingsInfo(List<string> Diagnostics)
+		{
+			if(DefaultBuildSettings < BuildSettingsVersion.Latest)
+			{
+				Diagnostics.Add("[Upgrade]");
+				string Version = DefaultBuildSettings.ToString().Replace("Release_", "").Replace("_", ".");
+				Diagnostics.Add(String.Format("[Upgrade] Using UE {0} compatible build settings. The latest version of UE4 sets the following values by default, which may require code changes:", Version));
+
+				List<Tuple<string, string>> ModifiedSettings = new List<Tuple<string, string>>();
+				if(DefaultBuildSettings < BuildSettingsVersion.V2)
+				{
+					ModifiedSettings.Add(Tuple.Create(String.Format("{0} = false", nameof(bLegacyPublicIncludePaths)), "Omits subfolders from public include paths to reduce compiler command line length."));
+					ModifiedSettings.Add(Tuple.Create(String.Format("{0} = PCHUsageMode.UseExplicitOrSharedPCHs", nameof(ModuleRules.PCHUsage)), "Set in build.cs files to enables IWYU-style PCH model. See https://docs.unrealengine.com/en-US/Programming/BuildTools/UnrealBuildTool/IWYU/index.html."));
+				}
+
+				if(ModifiedSettings.Count > 0)
+				{
+					string FormatString = String.Format("[Upgrade]     {{0,-{0}}}   => {{1}}", ModifiedSettings.Max(x => x.Item1.Length));
+					foreach (Tuple<string, string> ModifiedSetting in ModifiedSettings)
+					{
+						Diagnostics.Add(String.Format(FormatString, ModifiedSetting.Item1, ModifiedSetting.Item2));
+					}
+				}
+				Diagnostics.Add(String.Format("[Upgrade] Suppress this message by setting 'DefaultBuildSettings = BuildSettingsVersion.{1};' in {2}, and explicitly overriding desired settings.", Version, (BuildSettingsVersion)(BuildSettingsVersion.Latest - 1), File.GetFileName()));
+				Diagnostics.Add("[Upgrade]");
+			}
+		}
 	}
 
 	/// <summary>
@@ -1712,6 +1783,11 @@ namespace UnrealBuildTool
 			get { return Inner.Type; }
 		}
 
+		public BuildSettingsVersion DefaultBuildSettings
+		{
+			get { return Inner.DefaultBuildSettings; }
+		}
+
 		internal ConfigValueTracker ConfigValueTracker
 		{
 			get { return Inner.ConfigValueTracker; }
@@ -2038,9 +2114,12 @@ namespace UnrealBuildTool
 			get { return Inner.bEventDrivenLoader; }
 		}
 
+		[Obsolete("Use DefaultBuildSettings to determine the default settings used for this target instead")]
 		public bool bUseBackwardsCompatibleDefaults
 		{
+			#pragma warning disable CS0618
 			get { return Inner.bUseBackwardsCompatibleDefaults; }
+			#pragma warning restore CS0618
 		}
 
 		public bool bIWYU
@@ -2569,5 +2648,14 @@ namespace UnrealBuildTool
 		{
 			return UEBuildPlatform.IsPlatformInGroup(Platform, Group);
 		}
+
+		/// <summary>
+		/// Gets diagnostic messages about default settings which have changed in newer versions of the engine
+		/// </summary>
+		/// <param name="Diagnostics">List of messages to be appended to</param>
+		internal void GetBuildSettingsInfo(List<string> Diagnostics)
+		{
+			Inner.GetBuildSettingsInfo(Diagnostics);
+		}
 	}
 }
diff --git a/Engine/Source/Programs/UnrealBuildTool/Configuration/UEBuildTarget.cs b/Engine/Source/Programs/UnrealBuildTool/Configuration/UEBuildTarget.cs
--- a/Engine/Source/Programs/UnrealBuildTool/Configuration/UEBuildTarget.cs
+++ b/Engine/Source/Programs/UnrealBuildTool/Configuration/UEBuildTarget.cs
@@ -1556,7 +1556,11 @@ namespace UnrealBuildTool
 
 			// Create the makefile
 			string ExternalMetadata = UEBuildPlatform.GetBuildPlatform(Platform).GetExternalBuildMetadata(ProjectFile);
-			TargetMakefile Makefile = new TargetMakefile(TargetToolChain.GetVersionInfo(), ExternalMetadata, Binaries[0].OutputFilePaths[0], ReceiptFileName, ProjectIntermediateDirectory, TargetType, Rules.ConfigValueTracker, bDeployAfterCompile, bHasProjectScriptPlugin);
+			TargetMakefile Makefile = new TargetMakefile(ExternalMetadata, Binaries[0].OutputFilePaths[0], ReceiptFileName, ProjectIntermediateDirectory, TargetType, Rules.ConfigValueTracker, bDeployAfterCompile, bHasProjectScriptPlugin);
+
+			// Get diagnostic info to be printed before each build
+			TargetToolChain.GetVersionInfo(Makefile.Diagnostics);
+			Rules.GetBuildSettingsInfo(Makefile.Diagnostics);
 
 			// Setup the hot reload module list
 			Makefile.HotReloadModuleNames = GetHotReloadModuleNames();
@@ -3653,29 +3657,6 @@ namespace UnrealBuildTool
 					}
 				}
 
-				// Disable shared PCHs for game modules by default (but not game plugins, since they won't depend on the game's PCH!)
-				if (RulesObject.PCHUsage == ModuleRules.PCHUsageMode.Default)
-				{
-					if(RulesObject.bUseBackwardsCompatibleDefaults && !Rules.bIWYU)
-					{
-						if(RulesObject.Plugin != null)
-						{
-							// Game plugin.  Enable shared PCHs by default, since they aren't typically large enough to warrant their own PCH.
-							RulesObject.PCHUsage = ModuleRules.PCHUsageMode.UseSharedPCHs;
-						}
-						else
-						{
-							// Game module.  Do not enable shared PCHs by default, because games usually have a large precompiled header of their own and compile times would suffer.
-							RulesObject.PCHUsage = ModuleRules.PCHUsageMode.NoSharedPCHs;
-						}
-					}
-					else
-					{
-						// Engine module or plugin module -- allow shared PCHs
-						RulesObject.PCHUsage = ModuleRules.PCHUsageMode.UseExplicitOrSharedPCHs;
-					}
-				}
-
 				// If we can't use a shared PCH, check there's a private PCH set
 				if(RulesObject.PCHUsage != ModuleRules.PCHUsageMode.NoPCHs && RulesObject.PCHUsage != ModuleRules.PCHUsageMode.UseExplicitOrSharedPCHs && RulesObject.PrivatePCHHeaderFile == null)
 				{
diff --git a/Engine/Source/Programs/UnrealBuildTool/Modes/BuildMode.cs b/Engine/Source/Programs/UnrealBuildTool/Modes/BuildMode.cs
--- a/Engine/Source/Programs/UnrealBuildTool/Modes/BuildMode.cs
+++ b/Engine/Source/Programs/UnrealBuildTool/Modes/BuildMode.cs
@@ -304,7 +304,7 @@ namespace UnrealBuildTool
 					// Execute the actions
 					if(MergedActionsToExecute.Count == 0)
 					{
-						if((Options & BuildOptions.Quiet) == 0)
+						if ((Options & BuildOptions.Quiet) == 0)
 						{
 							Log.TraceInformation((TargetDescriptors.Count == 1)? "Target is up to date" : "Targets are up to date");
 						}
@@ -346,18 +346,20 @@ namespace UnrealBuildTool
 		/// <param name="Makefiles">Matching array of makefiles for each target</param>
 		static void OutputToolchainInfo(List<TargetDescriptor> TargetDescriptors, TargetMakefile[] Makefiles)
 		{
-			List<string> UniqueStrings = new List<string>(Makefiles.Select(x => x.ToolchainInfo).Where(x => x != null).Distinct());
-			if(UniqueStrings.Count == 1)
+			if(Makefiles.Length == 1 || Makefiles.Skip(1).All(x => Enumerable.SequenceEqual(x.Diagnostics, Makefiles[0].Diagnostics)))
 			{
-				Log.TraceInformation("{0}", UniqueStrings[0]);
+				foreach(string Diagnostic in Makefiles[0].Diagnostics)
+				{
+					Log.TraceInformation("{0}", Diagnostic);
+				}
 			}
 			else
 			{
 				for(int Idx = 0; Idx < TargetDescriptors.Count; Idx++)
 				{
-					if(Makefiles[Idx].ToolchainInfo != null)
+					foreach(string Diagnostic in Makefiles[Idx].Diagnostics)
 					{
-						Log.TraceInformation("For {0}: {1}", TargetDescriptors[Idx], Makefiles[Idx].ToolchainInfo);
+						Log.TraceInformation("For {0}: {1}", TargetDescriptors[Idx], Diagnostic);
 					}
 				}
 			}
diff --git a/Engine/Source/Programs/UnrealBuildTool/Platform/Windows/VCToolChain.cs b/Engine/Source/Programs/UnrealBuildTool/Platform/Windows/VCToolChain.cs
--- a/Engine/Source/Programs/UnrealBuildTool/Platform/Windows/VCToolChain.cs
+++ b/Engine/Source/Programs/UnrealBuildTool/Platform/Windows/VCToolChain.cs
@@ -58,15 +58,15 @@ namespace UnrealBuildTool
 		/// Returns the version info for the toolchain. This will be output before building.
 		/// </summary>
 		/// <returns>String describing the current toolchain</returns>
-		public override string GetVersionInfo()
+		public override void GetVersionInfo(List<string> Lines)
 		{
 			if(EnvVars.Compiler == EnvVars.ToolChain)
 			{
-				return String.Format("Using {0} {1} toolchain ({2}) and Windows {3} SDK ({4}).", WindowsPlatform.GetCompilerName(EnvVars.Compiler), EnvVars.ToolChainVersion, EnvVars.ToolChainDir, EnvVars.WindowsSdkVersion, EnvVars.WindowsSdkDir);
+				Lines.Add(String.Format("Using {0} {1} toolchain ({2}) and Windows {3} SDK ({4}).", WindowsPlatform.GetCompilerName(EnvVars.Compiler), EnvVars.ToolChainVersion, EnvVars.ToolChainDir, EnvVars.WindowsSdkVersion, EnvVars.WindowsSdkDir));
 			}
 			else
 			{
-				return String.Format("Using {0} {1} compiler ({2}) with {3} {4} runtime ({5}) and Windows {6} SDK ({7}).", WindowsPlatform.GetCompilerName(EnvVars.Compiler), EnvVars.CompilerVersion, EnvVars.CompilerDir, WindowsPlatform.GetCompilerName(EnvVars.ToolChain), EnvVars.ToolChainVersion, EnvVars.ToolChainDir, EnvVars.WindowsSdkVersion, EnvVars.WindowsSdkDir);
+				Lines.Add(String.Format("Using {0} {1} compiler ({2}) with {3} {4} runtime ({5}) and Windows {6} SDK ({7}).", WindowsPlatform.GetCompilerName(EnvVars.Compiler), EnvVars.CompilerVersion, EnvVars.CompilerDir, WindowsPlatform.GetCompilerName(EnvVars.ToolChain), EnvVars.ToolChainVersion, EnvVars.ToolChainDir, EnvVars.WindowsSdkVersion, EnvVars.WindowsSdkDir));
 			}
 		}
 
diff --git a/Engine/Source/Programs/UnrealBuildTool/System/RulesAssembly.cs b/Engine/Source/Programs/UnrealBuildTool/System/RulesAssembly.cs
--- a/Engine/Source/Programs/UnrealBuildTool/System/RulesAssembly.cs
+++ b/Engine/Source/Programs/UnrealBuildTool/System/RulesAssembly.cs
@@ -66,7 +66,7 @@ namespace UnrealBuildTool
 		/// Whether to use backwards compatible default settings for module and target rules. This is enabled by default for game projects to support a simpler migration path, but
 		/// is disabled for engine modules.
 		/// </summary>
-		private bool bUseBackwardsCompatibleDefaults;
+		private BuildSettingsVersion? DefaultBuildSettings;
 
 		/// <summary>
 		/// Whether the modules and targets in this assembly are read-only
@@ -88,18 +88,18 @@ namespace UnrealBuildTool
 		/// <param name="TargetFiles">List of target files to compile</param>
 		/// <param name="AssemblyFileName">The output path for the compiled assembly</param>
 		/// <param name="bContainsEngineModules">Whether this assembly contains engine modules. Used to initialize the default value for ModuleRules.bTreatAsEngineModule.</param>
-		/// <param name="bUseBackwardsCompatibleDefaults">Whether modules in this assembly should use backwards-compatible defaults.</param>
+		/// <param name="DefaultBuildSettings">Optional override for the default build settings version for modules created from this assembly.</param>
 		/// <param name="bReadOnly">Whether the modules and targets in this assembly are installed, and should be created with the bUsePrecompiled flag set</param> 
 		/// <param name="bSkipCompile">Whether to skip compiling this assembly</param>
 		/// <param name="Parent">The parent rules assembly</param>
-		internal RulesAssembly(RulesScope Scope, DirectoryReference BaseDir, IReadOnlyList<PluginInfo> Plugins, Dictionary<FileReference, ModuleRulesContext> ModuleFileToContext, List<FileReference> TargetFiles, FileReference AssemblyFileName, bool bContainsEngineModules, bool bUseBackwardsCompatibleDefaults, bool bReadOnly, bool bSkipCompile, RulesAssembly Parent)
+		internal RulesAssembly(RulesScope Scope, DirectoryReference BaseDir, IReadOnlyList<PluginInfo> Plugins, Dictionary<FileReference, ModuleRulesContext> ModuleFileToContext, List<FileReference> TargetFiles, FileReference AssemblyFileName, bool bContainsEngineModules, BuildSettingsVersion? DefaultBuildSettings, bool bReadOnly, bool bSkipCompile, RulesAssembly Parent)
 		{
 			this.Scope = Scope;
 			this.BaseDir = BaseDir;
 			this.Plugins = Plugins;
 			this.ModuleFileToContext = ModuleFileToContext;
 			this.bContainsEngineModules = bContainsEngineModules;
-			this.bUseBackwardsCompatibleDefaults = bUseBackwardsCompatibleDefaults;
+			this.DefaultBuildSettings = DefaultBuildSettings;
 			this.bReadOnly = bReadOnly;
 			this.Parent = Parent;
 
@@ -421,7 +421,10 @@ namespace UnrealBuildTool
 				RulesObject.Context = ModuleFileToContext[RulesObject.File];
 				RulesObject.Plugin = RulesObject.Context.Plugin;
 				RulesObject.bTreatAsEngineModule = bContainsEngineModules;
-				RulesObject.bUseBackwardsCompatibleDefaults = bUseBackwardsCompatibleDefaults && Target.bUseBackwardsCompatibleDefaults;
+				if(DefaultBuildSettings.HasValue)
+				{
+					RulesObject.DefaultBuildSettings = DefaultBuildSettings.Value;
+				}
 				RulesObject.bPrecompile = (RulesObject.bTreatAsEngineModule || ModuleName.Equals("UE4Game", StringComparison.OrdinalIgnoreCase)) && Target.bPrecompile;
 				RulesObject.bUsePrecompiled = bReadOnly;
 
@@ -479,7 +482,10 @@ namespace UnrealBuildTool
 
 			// Create an instance of the module's rules object, and set some defaults before calling the constructor.
 			TargetRules Rules = (TargetRules)FormatterServices.GetUninitializedObject(RulesType);
-			Rules.bUseBackwardsCompatibleDefaults = bUseBackwardsCompatibleDefaults;
+			if (DefaultBuildSettings.HasValue)
+			{
+				Rules.DefaultBuildSettings = DefaultBuildSettings.Value;
+			}
 
 			// Find the constructor
 			ConstructorInfo Constructor = RulesType.GetConstructor(new Type[] { typeof(TargetInfo) });
diff --git a/Engine/Source/Programs/UnrealBuildTool/System/RulesCompiler.cs b/Engine/Source/Programs/UnrealBuildTool/System/RulesCompiler.cs
--- a/Engine/Source/Programs/UnrealBuildTool/System/RulesCompiler.cs
+++ b/Engine/Source/Programs/UnrealBuildTool/System/RulesCompiler.cs
@@ -471,7 +471,7 @@ namespace UnrealBuildTool
 
 			// Create the assembly
 			FileReference EngineAssemblyFileName = FileReference.Combine(RootDirectories[0], "Intermediate", "Build", "BuildRules", AssemblyPrefix + "Rules" + FrameworkAssemblyExtension);
-			RulesAssembly EngineAssembly = new RulesAssembly(Scope, RootDirectories[0], Plugins, ModuleFileToContext, new List<FileReference>(), EngineAssemblyFileName, bContainsEngineModules: true, bUseBackwardsCompatibleDefaults: false, bReadOnly: bReadOnly, bSkipCompile: bSkipCompile, Parent: Parent);
+			RulesAssembly EngineAssembly = new RulesAssembly(Scope, RootDirectories[0], Plugins, ModuleFileToContext, new List<FileReference>(), EngineAssemblyFileName, bContainsEngineModules: true, DefaultBuildSettings: BuildSettingsVersion.Latest, bReadOnly: bReadOnly, bSkipCompile: bSkipCompile, Parent: Parent);
 
 			List<FileReference> ProgramTargetFiles = new List<FileReference>();
 			Dictionary<FileReference, ModuleRulesContext> ProgramModuleFiles = new Dictionary<FileReference, ModuleRulesContext>();
@@ -498,7 +498,7 @@ namespace UnrealBuildTool
 
 			// Create a path to the assembly that we'll either load or compile
 			FileReference ProgramAssemblyFileName = FileReference.Combine(RootDirectories[0], "Intermediate", "Build", "BuildRules", AssemblyPrefix + "ProgramRules" + FrameworkAssemblyExtension);
-			RulesAssembly ProgramAssembly = new RulesAssembly(ProgramsScope, RootDirectories[0], new List<PluginInfo>().AsReadOnly(), ProgramModuleFiles, ProgramTargetFiles, ProgramAssemblyFileName, bContainsEngineModules: false, bUseBackwardsCompatibleDefaults: false, bReadOnly: bReadOnly, bSkipCompile: bSkipCompile, Parent: EngineAssembly);
+			RulesAssembly ProgramAssembly = new RulesAssembly(ProgramsScope, RootDirectories[0], new List<PluginInfo>().AsReadOnly(), ProgramModuleFiles, ProgramTargetFiles, ProgramAssemblyFileName, bContainsEngineModules: false, DefaultBuildSettings: BuildSettingsVersion.Latest, bReadOnly: bReadOnly, bSkipCompile: bSkipCompile, Parent: EngineAssembly);
 
 			// Return the combined assembly
 			return ProgramAssembly;
@@ -589,7 +589,7 @@ namespace UnrealBuildTool
 				}
 				else
 				{
-					ProjectRulesAssembly = new RulesAssembly(Scope, MainProjectDirectory, ProjectPlugins, ModuleFiles, TargetFiles, AssemblyFileName, bContainsEngineModules: false, bUseBackwardsCompatibleDefaults: true, bReadOnly: UnrealBuildTool.IsProjectInstalled(), bSkipCompile: bSkipCompile, Parent: Parent);
+					ProjectRulesAssembly = new RulesAssembly(Scope, MainProjectDirectory, ProjectPlugins, ModuleFiles, TargetFiles, AssemblyFileName, bContainsEngineModules: false, DefaultBuildSettings: null, bReadOnly: UnrealBuildTool.IsProjectInstalled(), bSkipCompile: bSkipCompile, Parent: Parent);
 				}
 				LoadedAssemblyMap.Add(ProjectFileName, ProjectRulesAssembly);
 			}
@@ -631,7 +631,7 @@ namespace UnrealBuildTool
 
 				// Compile the assembly
 				FileReference AssemblyFileName = FileReference.Combine(PluginFileName.Directory, "Intermediate", "Build", "BuildRules", Path.GetFileNameWithoutExtension(PluginFileName.FullName) + "ModuleRules" + FrameworkAssemblyExtension);
-				PluginRulesAssembly = new RulesAssembly(Scope, PluginFileName.Directory, ForeignPlugins, ModuleFiles, TargetFiles, AssemblyFileName, bContainsEngineModules, bUseBackwardsCompatibleDefaults: false, bReadOnly: false, bSkipCompile: bSkipCompile, Parent: Parent);
+				PluginRulesAssembly = new RulesAssembly(Scope, PluginFileName.Directory, ForeignPlugins, ModuleFiles, TargetFiles, AssemblyFileName, bContainsEngineModules, DefaultBuildSettings: null, bReadOnly: false, bSkipCompile: bSkipCompile, Parent: Parent);
 				LoadedAssemblyMap.Add(PluginFileName, PluginRulesAssembly);
 			}
 			return PluginRulesAssembly;
diff --git a/Engine/Source/Programs/UnrealBuildTool/System/TargetMakefile.cs b/Engine/Source/Programs/UnrealBuildTool/System/TargetMakefile.cs
--- a/Engine/Source/Programs/UnrealBuildTool/System/TargetMakefile.cs
+++ b/Engine/Source/Programs/UnrealBuildTool/System/TargetMakefile.cs
@@ -29,9 +29,9 @@ namespace UnrealBuildTool
 		public DateTime CreateTimeUtc;
 
 		/// <summary>
-		/// Information about the toolchain used to build. This string will be output before building.
+		/// Additional diagnostic output to print before building this target (toolchain version, etc...)
 		/// </summary>
-		public string ToolchainInfo;
+		public List<string> Diagnostics = new List<string>();
 
 		/// <summary>
 		/// Any additional information about the build environment which the platform can use to invalidate the makefile
@@ -157,7 +157,6 @@ namespace UnrealBuildTool
 		/// <summary>
 		/// Constructor
 		/// </summary>
-		/// <param name="ToolchainInfo">String describing the toolchain used to build. This will be output before executing actions.</param>
 		/// <param name="ExternalMetadata">External build metadata from the platform</param>
 		/// <param name="ExecutableFile">Path to the executable or primary output binary for this target</param>
 		/// <param name="ReceiptFile">Path to the receipt file</param>
@@ -166,10 +165,10 @@ namespace UnrealBuildTool
 		/// <param name="ConfigValueTracker">Set of dependencies on config files</param>
 		/// <param name="bDeployAfterCompile">Whether to deploy the target after compiling</param>
 		/// <param name="bHasProjectScriptPlugin">Whether the target has a project script plugin</param>
-		public TargetMakefile(string ToolchainInfo, string ExternalMetadata, FileReference ExecutableFile, FileReference ReceiptFile, DirectoryReference ProjectIntermediateDirectory, TargetType TargetType, ConfigValueTracker ConfigValueTracker, bool bDeployAfterCompile, bool bHasProjectScriptPlugin)
+		public TargetMakefile(string ExternalMetadata, FileReference ExecutableFile, FileReference ReceiptFile, DirectoryReference ProjectIntermediateDirectory, TargetType TargetType, ConfigValueTracker ConfigValueTracker, bool bDeployAfterCompile, bool bHasProjectScriptPlugin)
 		{
 			this.CreateTimeUtc = DateTime.UtcNow;
-			this.ToolchainInfo = ToolchainInfo;
+			this.Diagnostics = new List<string>();
 			this.ExternalMetadata = ExternalMetadata;
 			this.ExecutableFile = ExecutableFile;
 			this.ReceiptFile = ReceiptFile;
@@ -199,7 +198,7 @@ namespace UnrealBuildTool
 		public TargetMakefile(BinaryArchiveReader Reader)
 		{
 			CreateTimeUtc = new DateTime(Reader.ReadLong(), DateTimeKind.Utc);
-			ToolchainInfo = Reader.ReadString();
+			Diagnostics = Reader.ReadList(() => Reader.ReadString());
 			ExternalMetadata = Reader.ReadString();
 			ExecutableFile = Reader.ReadFileReference();
 			ReceiptFile = Reader.ReadFileReference();
@@ -232,7 +231,7 @@ namespace UnrealBuildTool
 		public void Write(BinaryArchiveWriter Writer)
 		{
 			Writer.WriteLong(CreateTimeUtc.Ticks);
-			Writer.WriteString(ToolchainInfo);
+			Writer.WriteList(Diagnostics, x => Writer.WriteString(x));
 			Writer.WriteString(ExternalMetadata);
 			Writer.WriteFileReference(ExecutableFile);
 			Writer.WriteFileReference(ReceiptFile);
diff --git a/Engine/Source/Programs/UnrealBuildTool/System/Unity.cs b/Engine/Source/Programs/UnrealBuildTool/System/Unity.cs
--- a/Engine/Source/Programs/UnrealBuildTool/System/Unity.cs
+++ b/Engine/Source/Programs/UnrealBuildTool/System/Unity.cs
@@ -12,15 +12,10 @@ namespace UnrealBuildTool
 {
 	class Unity
 	{
-		/// <summary>
-		/// Set of target names we've printed out adaptive non-unity messages for.
-		/// </summary>
-		public static HashSet<string> PrintedSettingsForTargets  = new HashSet<string>();
-
 		/// <summary>
 		/// Prefix used for all dynamically created Unity modules
 		/// </summary>
-		public static string ModulePrefix = "Module.";
+		public const string ModulePrefix = "Module.";
 
 		/// <summary>
 		/// A class which represents a list of files and the sum of their lengths.
@@ -281,27 +276,25 @@ namespace UnrealBuildTool
 
 				if (AdaptiveUnityBuildInfoString.Length > 0)
 				{
-					if (PrintedSettingsForTargets.Add(Target.Name))
+					if (Target.bAdaptiveUnityCreatesDedicatedPCH)
 					{
-						if (Target.bAdaptiveUnityCreatesDedicatedPCH)
-						{
-							Log.TraceInformation("[Adaptive unity build] Creating dedicated PCH for each excluded file. Set bAdaptiveUnityCreatesDedicatedPCH to false in BuildConfiguration.xml to change this behavior.");
-						}
-						else if (Target.bAdaptiveUnityDisablesPCH)
-						{
-							Log.TraceInformation("[Adaptive unity build] Disabling PCH for excluded files. Set bAdaptiveUnityDisablesPCH to false in BuildConfiguration.xml to change this behavior.");
-						}
+						AddUniqueDiagnostic(Makefile, "[Adaptive unity build] Creating dedicated PCH for each excluded file. Set bAdaptiveUnityCreatesDedicatedPCH to false in BuildConfiguration.xml to change this behavior.");
+					}
+					else if (Target.bAdaptiveUnityDisablesPCH)
+					{
+						AddUniqueDiagnostic(Makefile, "[Adaptive unity build] Disabling PCH for excluded files. Set bAdaptiveUnityDisablesPCH to false in BuildConfiguration.xml to change this behavior.");
+					}
 
-						if (Target.bAdaptiveUnityDisablesOptimizations)
-						{
-							Log.TraceInformation("[Adaptive unity build] Disabling optimizations for excluded files. Set bAdaptiveUnityDisablesOptimizations to false in BuildConfiguration.xml to change this behavior.");
-						}
-						if (Target.bAdaptiveUnityEnablesEditAndContinue)
-						{
-							Log.TraceInformation("[Adaptive unity build] Enabling Edit & Continue for excluded files. Set bAdaptiveUnityEnablesEditAndContinue to false in BuildConfiguration.xml to change this behavior.");
-						}
+					if (Target.bAdaptiveUnityDisablesOptimizations)
+					{
+						AddUniqueDiagnostic(Makefile, "[Adaptive unity build] Disabling optimizations for excluded files. Set bAdaptiveUnityDisablesOptimizations to false in BuildConfiguration.xml to change this behavior.");
 					}
-					Log.TraceInformation(AdaptiveUnityBuildInfoString.ToString());
+					if (Target.bAdaptiveUnityEnablesEditAndContinue)
+					{
+						AddUniqueDiagnostic(Makefile, "[Adaptive unity build] Enabling Edit & Continue for excluded files. Set bAdaptiveUnityEnablesEditAndContinue to false in BuildConfiguration.xml to change this behavior.");
+					}
+
+					Makefile.Diagnostics.Add(AdaptiveUnityBuildInfoString.ToString());
 				}
 
 				AllUnityFiles = CPPUnityFileBuilder.GetUnityFiles();
@@ -353,5 +346,13 @@ namespace UnrealBuildTool
 
 			return NewCPPFiles;
 		}
+
+		static void AddUniqueDiagnostic(TargetMakefile Makefile, string Message)
+		{
+			if(!Makefile.Diagnostics.Contains(Message, StringComparer.Ordinal))
+			{
+				Makefile.Diagnostics.Add(Message);
+			}
+		}
 	}
 }
diff --git a/Engine/Source/Programs/UnrealBuildTool/ToolChain/UEToolChain.cs b/Engine/Source/Programs/UnrealBuildTool/ToolChain/UEToolChain.cs
--- a/Engine/Source/Programs/UnrealBuildTool/ToolChain/UEToolChain.cs
+++ b/Engine/Source/Programs/UnrealBuildTool/ToolChain/UEToolChain.cs
@@ -20,9 +20,8 @@ namespace UnrealBuildTool
 		{
 		}
 
-		public virtual string GetVersionInfo()
+		public virtual void GetVersionInfo(List<string> Lines)
 		{
-			return null;
 		}
 
 		public abstract CPPOutput CompileCPPFiles(CppCompileEnvironment CompileEnvironment, List<FileItem> InputFiles, DirectoryReference OutputDir, string ModuleName, List<Action> Actions);
diff --git a/Templates/FP_FirstPerson/Source/FP_FirstPerson.Target.cs b/Templates/FP_FirstPerson/Source/FP_FirstPerson.Target.cs
--- a/Templates/FP_FirstPerson/Source/FP_FirstPerson.Target.cs
+++ b/Templates/FP_FirstPerson/Source/FP_FirstPerson.Target.cs
@@ -8,6 +8,7 @@ public class FP_FirstPersonTarget : TargetRules
 	public FP_FirstPersonTarget(TargetInfo Target) : base(Target)
 	{
 		Type = TargetType.Game;
+		DefaultBuildSettings = BuildSettingsVersion.Latest;
 		ExtraModuleNames.Add("FP_FirstPerson");
 	}
 }
diff --git a/Templates/FP_FirstPerson/Source/FP_FirstPersonEditor.Target.cs b/Templates/FP_FirstPerson/Source/FP_FirstPersonEditor.Target.cs
--- a/Templates/FP_FirstPerson/Source/FP_FirstPersonEditor.Target.cs
+++ b/Templates/FP_FirstPerson/Source/FP_FirstPersonEditor.Target.cs
@@ -8,6 +8,7 @@ public class FP_FirstPersonEditorTarget : TargetRules
 	public FP_FirstPersonEditorTarget(TargetInfo Target) : base(Target)
 	{
 		Type = TargetType.Editor;
+		DefaultBuildSettings = BuildSettingsVersion.Latest;
 		ExtraModuleNames.Add("FP_FirstPerson");
 	}
 }
diff --git a/Templates/TP_2DSideScroller/Source/TP_2DSideScroller.Target.cs b/Templates/TP_2DSideScroller/Source/TP_2DSideScroller.Target.cs
--- a/Templates/TP_2DSideScroller/Source/TP_2DSideScroller.Target.cs
+++ b/Templates/TP_2DSideScroller/Source/TP_2DSideScroller.Target.cs
@@ -8,6 +8,7 @@ public class TP_2DSideScrollerTarget : TargetRules
 	public TP_2DSideScrollerTarget(TargetInfo Target) : base(Target)
 	{
 		Type = TargetType.Game;
+		DefaultBuildSettings = BuildSettingsVersion.Latest;
 		ExtraModuleNames.Add("TP_2DSideScroller");
 	}
 }
diff --git a/Templates/TP_2DSideScroller/Source/TP_2DSideScrollerEditor.Target.cs b/Templates/TP_2DSideScroller/Source/TP_2DSideScrollerEditor.Target.cs
--- a/Templates/TP_2DSideScroller/Source/TP_2DSideScrollerEditor.Target.cs
+++ b/Templates/TP_2DSideScroller/Source/TP_2DSideScrollerEditor.Target.cs
@@ -8,6 +8,7 @@ public class TP_2DSideScrollerEditorTarget : TargetRules
 	public TP_2DSideScrollerEditorTarget(TargetInfo Target) : base(Target)
 	{
 		Type = TargetType.Editor;
+		DefaultBuildSettings = BuildSettingsVersion.Latest;
 		ExtraModuleNames.Add("TP_2DSideScroller");
 	}
 }
diff --git a/Templates/TP_FirstPerson/Source/TP_FirstPerson.Target.cs b/Templates/TP_FirstPerson/Source/TP_FirstPerson.Target.cs
--- a/Templates/TP_FirstPerson/Source/TP_FirstPerson.Target.cs
+++ b/Templates/TP_FirstPerson/Source/TP_FirstPerson.Target.cs
@@ -8,6 +8,7 @@ public class TP_FirstPersonTarget : TargetRules
 	public TP_FirstPersonTarget(TargetInfo Target) : base(Target)
 	{
 		Type = TargetType.Game;
+		DefaultBuildSettings = BuildSettingsVersion.Latest;
 		ExtraModuleNames.Add("TP_FirstPerson");
 	}
 }
diff --git a/Templates/TP_FirstPerson/Source/TP_FirstPersonEditor.Target.cs b/Templates/TP_FirstPerson/Source/TP_FirstPersonEditor.Target.cs
--- a/Templates/TP_FirstPerson/Source/TP_FirstPersonEditor.Target.cs
+++ b/Templates/TP_FirstPerson/Source/TP_FirstPersonEditor.Target.cs
@@ -8,6 +8,7 @@ public class TP_FirstPersonEditorTarget : TargetRules
 	public TP_FirstPersonEditorTarget(TargetInfo Target) : base(Target)
 	{
 		Type = TargetType.Editor;
+		DefaultBuildSettings = BuildSettingsVersion.Latest;
 		ExtraModuleNames.Add("TP_FirstPerson");
 	}
 }
diff --git a/Templates/TP_Flying/Source/TP_Flying.Target.cs b/Templates/TP_Flying/Source/TP_Flying.Target.cs
--- a/Templates/TP_Flying/Source/TP_Flying.Target.cs
+++ b/Templates/TP_Flying/Source/TP_Flying.Target.cs
@@ -8,6 +8,7 @@ public class TP_FlyingTarget : TargetRules
 	public TP_FlyingTarget(TargetInfo Target) : base(Target)
 	{
 		Type = TargetType.Game;
+		DefaultBuildSettings = BuildSettingsVersion.Latest;
 		ExtraModuleNames.Add("TP_Flying");
 	}
 }
diff --git a/Templates/TP_Flying/Source/TP_FlyingEditor.Target.cs b/Templates/TP_Flying/Source/TP_FlyingEditor.Target.cs
--- a/Templates/TP_Flying/Source/TP_FlyingEditor.Target.cs
+++ b/Templates/TP_Flying/Source/TP_FlyingEditor.Target.cs
@@ -8,6 +8,7 @@ public class TP_FlyingEditorTarget : TargetRules
 	public TP_FlyingEditorTarget(TargetInfo Target) : base(Target)
 	{
 		Type = TargetType.Editor;
+		DefaultBuildSettings = BuildSettingsVersion.Latest;
 		ExtraModuleNames.Add("TP_Flying");
 	}
 }
diff --git a/Templates/TP_Puzzle/Source/TP_Puzzle.Target.cs b/Templates/TP_Puzzle/Source/TP_Puzzle.Target.cs
--- a/Templates/TP_Puzzle/Source/TP_Puzzle.Target.cs
+++ b/Templates/TP_Puzzle/Source/TP_Puzzle.Target.cs
@@ -8,6 +8,7 @@ public class TP_PuzzleTarget : TargetRules
     public TP_PuzzleTarget(TargetInfo Target) : base(Target)
 	{
 		Type = TargetType.Game;
+		DefaultBuildSettings = BuildSettingsVersion.Latest;
 		ExtraModuleNames.Add("TP_Puzzle");
 	}
 }
diff --git a/Templates/TP_Puzzle/Source/TP_PuzzleEditor.Target.cs b/Templates/TP_Puzzle/Source/TP_PuzzleEditor.Target.cs
--- a/Templates/TP_Puzzle/Source/TP_PuzzleEditor.Target.cs
+++ b/Templates/TP_Puzzle/Source/TP_PuzzleEditor.Target.cs
@@ -8,6 +8,7 @@ public class TP_PuzzleEditorTarget : TargetRules
     public TP_PuzzleEditorTarget(TargetInfo Target) : base(Target)
 	{
 		Type = TargetType.Editor;
+		DefaultBuildSettings = BuildSettingsVersion.Latest;
 		ExtraModuleNames.Add("TP_Puzzle");
 	}
 }
diff --git a/Templates/TP_Rolling/Source/TP_Rolling.Target.cs b/Templates/TP_Rolling/Source/TP_Rolling.Target.cs
--- a/Templates/TP_Rolling/Source/TP_Rolling.Target.cs
+++ b/Templates/TP_Rolling/Source/TP_Rolling.Target.cs
@@ -8,6 +8,7 @@ public class TP_RollingTarget : TargetRules
     public TP_RollingTarget(TargetInfo Target) : base(Target)
 	{
 		Type = TargetType.Game;
+		DefaultBuildSettings = BuildSettingsVersion.Latest;
 		ExtraModuleNames.Add("TP_Rolling");
 	}
 }
diff --git a/Templates/TP_Rolling/Source/TP_RollingEditor.Target.cs b/Templates/TP_Rolling/Source/TP_RollingEditor.Target.cs
--- a/Templates/TP_Rolling/Source/TP_RollingEditor.Target.cs
+++ b/Templates/TP_Rolling/Source/TP_RollingEditor.Target.cs
@@ -8,6 +8,7 @@ public class TP_RollingEditorTarget : TargetRules
     public TP_RollingEditorTarget(TargetInfo Target) : base(Target)
 	{
 		Type = TargetType.Editor;
+		DefaultBuildSettings = BuildSettingsVersion.Latest;
 		ExtraModuleNames.Add("TP_Rolling");
 	}
 }
diff --git a/Templates/TP_SideScroller/Source/TP_SideScroller.Target.cs b/Templates/TP_SideScroller/Source/TP_SideScroller.Target.cs
--- a/Templates/TP_SideScroller/Source/TP_SideScroller.Target.cs
+++ b/Templates/TP_SideScroller/Source/TP_SideScroller.Target.cs
@@ -8,6 +8,7 @@ public class TP_SideScrollerTarget : TargetRules
 	public TP_SideScrollerTarget(TargetInfo Target) : base(Target)
 	{
 		Type = TargetType.Game;
+		DefaultBuildSettings = BuildSettingsVersion.Latest;
 		ExtraModuleNames.Add("TP_SideScroller");
 	}
 }
diff --git a/Templates/TP_SideScroller/Source/TP_SideScrollerEditor.Target.cs b/Templates/TP_SideScroller/Source/TP_SideScrollerEditor.Target.cs
--- a/Templates/TP_SideScroller/Source/TP_SideScrollerEditor.Target.cs
+++ b/Templates/TP_SideScroller/Source/TP_SideScrollerEditor.Target.cs
@@ -8,6 +8,7 @@ public class TP_SideScrollerEditorTarget : TargetRules
 	public TP_SideScrollerEditorTarget(TargetInfo Target) : base(Target)
 	{
 		Type = TargetType.Editor;
+		DefaultBuildSettings = BuildSettingsVersion.Latest;
 		ExtraModuleNames.Add("TP_SideScroller");
 	}
 }
diff --git a/Templates/TP_ThirdPerson/Source/TP_ThirdPerson.Target.cs b/Templates/TP_ThirdPerson/Source/TP_ThirdPerson.Target.cs
--- a/Templates/TP_ThirdPerson/Source/TP_ThirdPerson.Target.cs
+++ b/Templates/TP_ThirdPerson/Source/TP_ThirdPerson.Target.cs
@@ -8,6 +8,7 @@ public class TP_ThirdPersonTarget : TargetRules
 	public TP_ThirdPersonTarget(TargetInfo Target) : base(Target)
 	{
 		Type = TargetType.Game;
+		DefaultBuildSettings = BuildSettingsVersion.Latest;
 		ExtraModuleNames.Add("TP_ThirdPerson");
 	}
 }
diff --git a/Templates/TP_ThirdPerson/Source/TP_ThirdPersonEditor.Target.cs b/Templates/TP_ThirdPerson/Source/TP_ThirdPersonEditor.Target.cs
--- a/Templates/TP_ThirdPerson/Source/TP_ThirdPersonEditor.Target.cs
+++ b/Templates/TP_ThirdPerson/Source/TP_ThirdPersonEditor.Target.cs
@@ -8,6 +8,7 @@ public class TP_ThirdPersonEditorTarget : TargetRules
 	public TP_ThirdPersonEditorTarget(TargetInfo Target) : base(Target)
 	{
 		Type = TargetType.Editor;
+		DefaultBuildSettings = BuildSettingsVersion.Latest;
 		ExtraModuleNames.Add("TP_ThirdPerson");
 	}
 }
diff --git a/Templates/TP_TopDown/Source/TP_TopDown.Target.cs b/Templates/TP_TopDown/Source/TP_TopDown.Target.cs
--- a/Templates/TP_TopDown/Source/TP_TopDown.Target.cs
+++ b/Templates/TP_TopDown/Source/TP_TopDown.Target.cs
@@ -8,6 +8,7 @@ public class TP_TopDownTarget : TargetRules
 	public TP_TopDownTarget(TargetInfo Target) : base(Target)
 	{
 		Type = TargetType.Game;
+		DefaultBuildSettings = BuildSettingsVersion.Latest;
 		ExtraModuleNames.Add("TP_TopDown");
 	}
 }
diff --git a/Templates/TP_TopDown/Source/TP_TopDownEditor.Target.cs b/Templates/TP_TopDown/Source/TP_TopDownEditor.Target.cs
--- a/Templates/TP_TopDown/Source/TP_TopDownEditor.Target.cs
+++ b/Templates/TP_TopDown/Source/TP_TopDownEditor.Target.cs
@@ -8,6 +8,7 @@ public class TP_TopDownEditorTarget : TargetRules
 	public TP_TopDownEditorTarget(TargetInfo Target) : base(Target)
 	{
 		Type = TargetType.Editor;
+		DefaultBuildSettings = BuildSettingsVersion.Latest;
 		ExtraModuleNames.Add("TP_TopDown");
 	}
 }
diff --git a/Templates/TP_TwinStick/Source/TP_TwinStick.Target.cs b/Templates/TP_TwinStick/Source/TP_TwinStick.Target.cs
--- a/Templates/TP_TwinStick/Source/TP_TwinStick.Target.cs
+++ b/Templates/TP_TwinStick/Source/TP_TwinStick.Target.cs
@@ -8,6 +8,7 @@ public class TP_TwinStickTarget : TargetRules
 	public TP_TwinStickTarget(TargetInfo Target) : base(Target)
 	{
 		Type = TargetType.Game;
+		DefaultBuildSettings = BuildSettingsVersion.Latest;
 		ExtraModuleNames.Add("TP_TwinStick");
 	}
 }
diff --git a/Templates/TP_TwinStick/Source/TP_TwinStickEditor.Target.cs b/Templates/TP_TwinStick/Source/TP_TwinStickEditor.Target.cs
--- a/Templates/TP_TwinStick/Source/TP_TwinStickEditor.Target.cs
+++ b/Templates/TP_TwinStick/Source/TP_TwinStickEditor.Target.cs
@@ -8,6 +8,7 @@ public class TP_TwinStickEditorTarget : TargetRules
 	public TP_TwinStickEditorTarget(TargetInfo Target) : base(Target)
 	{
 		Type = TargetType.Editor;
+		DefaultBuildSettings = BuildSettingsVersion.Latest;
 		ExtraModuleNames.Add("TP_TwinStick");
 	}
 }
diff --git a/Templates/TP_Vehicle/Source/TP_Vehicle.Target.cs b/Templates/TP_Vehicle/Source/TP_Vehicle.Target.cs
--- a/Templates/TP_Vehicle/Source/TP_Vehicle.Target.cs
+++ b/Templates/TP_Vehicle/Source/TP_Vehicle.Target.cs
@@ -8,6 +8,7 @@ public class TP_VehicleTarget : TargetRules
 	public TP_VehicleTarget(TargetInfo Target) : base(Target)
 	{
 		Type = TargetType.Game;
+		DefaultBuildSettings = BuildSettingsVersion.Latest;
 		ExtraModuleNames.Add("TP_Vehicle");
 	}
 }
diff --git a/Templates/TP_Vehicle/Source/TP_VehicleEditor.Target.cs b/Templates/TP_Vehicle/Source/TP_VehicleEditor.Target.cs
--- a/Templates/TP_Vehicle/Source/TP_VehicleEditor.Target.cs
+++ b/Templates/TP_Vehicle/Source/TP_VehicleEditor.Target.cs
@@ -8,6 +8,7 @@ public class TP_VehicleEditorTarget : TargetRules
 	public TP_VehicleEditorTarget(TargetInfo Target) : base(Target)
 	{
 		Type = TargetType.Editor;
+		DefaultBuildSettings = BuildSettingsVersion.Latest;
 		ExtraModuleNames.Add("TP_Vehicle");
 	}
 }
diff --git a/Templates/TP_VehicleAdv/Source/TP_VehicleAdv.Target.cs b/Templates/TP_VehicleAdv/Source/TP_VehicleAdv.Target.cs
--- a/Templates/TP_VehicleAdv/Source/TP_VehicleAdv.Target.cs
+++ b/Templates/TP_VehicleAdv/Source/TP_VehicleAdv.Target.cs
@@ -8,6 +8,7 @@ public class TP_VehicleAdvTarget : TargetRules
 	public TP_VehicleAdvTarget(TargetInfo Target) : base(Target)
 	{
 		Type = TargetType.Game;
+		DefaultBuildSettings = BuildSettingsVersion.Latest;
 		ExtraModuleNames.Add("TP_VehicleAdv");
 	}
 }
diff --git a/Templates/TP_VehicleAdv/Source/TP_VehicleAdvEditor.Target.cs b/Templates/TP_VehicleAdv/Source/TP_VehicleAdvEditor.Target.cs
--- a/Templates/TP_VehicleAdv/Source/TP_VehicleAdvEditor.Target.cs
+++ b/Templates/TP_VehicleAdv/Source/TP_VehicleAdvEditor.Target.cs
@@ -8,6 +8,7 @@ public class TP_VehicleAdvEditorTarget : TargetRules
 	public TP_VehicleAdvEditorTarget(TargetInfo Target) : base(Target)
 	{
 		Type = TargetType.Editor;
+		DefaultBuildSettings = BuildSettingsVersion.Latest;
 		ExtraModuleNames.Add("TP_VehicleAdv");
 	}
 }
