UBT: Determine TargetType related settings dynamically using properties, so that individual targets can override them.

#rnx

diff --git a/Engine/Source/Programs/UnrealBuildTool/Configuration/TargetRules.cs b/Engine/Source/Programs/UnrealBuildTool/Configuration/TargetRules.cs
--- a/Engine/Source/Programs/UnrealBuildTool/Configuration/TargetRules.cs
+++ b/Engine/Source/Programs/UnrealBuildTool/Configuration/TargetRules.cs
@@ -414,14 +414,23 @@ namespace UnrealBuildTool
 		/// Whether to compile code related to building assets. Consoles generally cannot build assets. Desktop platforms generally can.
 		/// </summary>
 		[RequiresUniqueBuildEnvironment]
-		public bool bBuildRequiresCookedData = false;
+		public bool bBuildRequiresCookedData
+		{
+			get { return bBuildRequiresCookedDataOverride ?? (Type == TargetType.Game || Type == TargetType.Client || Type == TargetType.Server); }
+			set { bBuildRequiresCookedDataOverride = value; }
+		}
+		bool? bBuildRequiresCookedDataOverride;
 
 		/// <summary>
 		/// Whether to compile WITH_EDITORONLY_DATA disabled. Only Windows will use this, other platforms force this to false.
 		/// </summary>
 		[RequiresUniqueBuildEnvironment]
-		[CommandLine("-NoEditorOnlyData", Value = "false")]
-		public bool bBuildWithEditorOnlyData = true;
+		public bool bBuildWithEditorOnlyData
+		{
+			get { return bBuildWithEditorOnlyDataOverride ?? (Type == TargetType.Editor || Type == TargetType.Program); }
+			set { bBuildWithEditorOnlyDataOverride = value; }
+		}
+		private bool? bBuildWithEditorOnlyDataOverride;
 
 		/// <summary>
 		/// Manually specified value for bBuildDeveloperTools.
@@ -542,7 +551,12 @@ namespace UnrealBuildTool
 		/// Compile server-only code.
 		/// </summary>
 		[RequiresUniqueBuildEnvironment]
-		public bool bWithServerCode = true;
+		public bool bWithServerCode
+		{
+			get { return bWithServerCodeOverride ?? (Type != TargetType.Client); }
+			set { bWithServerCodeOverride = value; }
+		}
+		private bool? bWithServerCodeOverride;
 
 		/// <summary>
 		/// Whether to include stats support even without the engine.
@@ -561,7 +575,12 @@ namespace UnrealBuildTool
 		/// Whether to allow plugins which support all target platforms.
 		/// </summary>
 		[RequiresUniqueBuildEnvironment]
-		public bool bIncludePluginsForTargetPlatforms = false;
+		public bool bIncludePluginsForTargetPlatforms
+		{
+			get { return bIncludePluginsForTargetPlatformsOverride ?? (Type == TargetType.Editor); }
+			set { bIncludePluginsForTargetPlatformsOverride = value; }
+		}
+		private bool? bIncludePluginsForTargetPlatformsOverride;
 
 		/// <summary>
 		/// Whether to allow accessibility code in both Slate and the OS layer.
@@ -573,8 +592,14 @@ namespace UnrealBuildTool
 		/// Whether to include PerfCounters support.
 		/// </summary>
 		[RequiresUniqueBuildEnvironment]
+        public bool bWithPerfCounters
+		{
+			get { return bWithPerfCountersOverride ?? (Type == TargetType.Editor || Type == TargetType.Server); }
+			set { bWithPerfCountersOverride = value; }
+		}
+
 		[ConfigFile(ConfigHierarchyType.Engine, "/Script/BuildSettings.BuildSettings", "bWithPerfCounters")]
-        public bool bWithPerfCounters = false;
+		bool? bWithPerfCountersOverride;
 
 		/// <summary>
 		/// Whether to enable support for live coding
@@ -685,7 +710,12 @@ namespace UnrealBuildTool
 		/// <summary>
 		/// Whether the final executable should export symbols.
 		/// </summary>
-		public bool bHasExports = false;
+		public bool bHasExports
+		{
+			get { return bHasExportsOverride ?? (LinkType == TargetLinkType.Modular); }
+			set { bHasExportsOverride = value; }
+		}
+		private bool? bHasExportsOverride;
 
 		/// <summary>
 		/// Make static libraries for all engine modules as intermediates for this target.
@@ -707,7 +737,12 @@ namespace UnrealBuildTool
 		/// If true, creates an additional console application. Hack for Windows, where it's not possible to conditionally inherit a parent's console Window depending on how
 		/// the application is invoked; you have to link the same executable with a different subsystem setting.
 		/// </summary>
-		public bool bBuildAdditionalConsoleApp = false;
+		public bool bBuildAdditionalConsoleApp
+		{
+			get { return bBuildAdditionalConsoleAppOverride ?? (Type == TargetType.Editor); }
+			set { bBuildAdditionalConsoleAppOverride = value; }
+		}
+		private bool? bBuildAdditionalConsoleAppOverride;
 
 		/// <summary>
 		/// True if debug symbols that are cached for some platforms should not be created.
@@ -1436,85 +1471,18 @@ namespace UnrealBuildTool
 		{
 			if(Type == global::UnrealBuildTool.TargetType.Game)
 			{
-				// Do not include the editor
-				bBuildWithEditorOnlyData = false;
-
-				// Require cooked data
-				bBuildRequiresCookedData = true;
-
-				// Compile the engine
-				bCompileAgainstEngine = true;
-
-				// only have exports in modular builds
-				bHasExports = (LinkType == TargetLinkType.Modular);
-
-				// Tag it as a 'Game' build
 				GlobalDefinitions.Add("UE_GAME=1");
 			}
 			else if(Type == global::UnrealBuildTool.TargetType.Client)
 			{
-				// Do not include the editor
-				bBuildWithEditorOnlyData = false;
-
-				// Require cooked data
-				bBuildRequiresCookedData = true;
-
-				// Compile the engine
-				bCompileAgainstEngine = true;
-
-				// Disable server code
-				bWithServerCode = false;
-
-				// only have exports in modular builds
-				bHasExports = (LinkType == TargetLinkType.Modular);
-
-				// Tag it as a 'Game' build
 				GlobalDefinitions.Add("UE_GAME=1");
 			}
 			else if(Type == global::UnrealBuildTool.TargetType.Editor)
 			{
-				// Do not include the editor
-				bBuildWithEditorOnlyData = true;
-
-				// Require cooked data
-				bBuildRequiresCookedData = false;
-
-				// Compile the engine
-				bCompileAgainstEngine = true;
-
-				//enable PerfCounters
-				bWithPerfCounters = true;
-
-				// Include all plugins
-				bIncludePluginsForTargetPlatforms = true;
-
-				// Create an additional console app for the editor
-				bBuildAdditionalConsoleApp = true;
-
-				// only have exports in modular builds
-				bHasExports = (LinkType == TargetLinkType.Modular);
-
-				// Tag it as a 'Editor' build
 				GlobalDefinitions.Add("UE_EDITOR=1");
 			}
 			else if(Type == global::UnrealBuildTool.TargetType.Server)
 			{
-				// Do not include the editor
-				bBuildWithEditorOnlyData = false;
-
-				// Require cooked data
-				bBuildRequiresCookedData = true;
-
-				// Compile the engine
-				bCompileAgainstEngine = true;
-			
-				//enable PerfCounters
-				bWithPerfCounters = true;
-
-				// only have exports in modular builds
-				bHasExports = (LinkType == TargetLinkType.Modular);
-
-				// Tag it as a 'Server' build
 				GlobalDefinitions.Add("UE_SERVER=1");
 				GlobalDefinitions.Add("USE_NULL_RHI=1");
 			}
