Fixed issue were using -Platforms to generate solution with a subset of platforms would break the editor finding the solution, reported in UDN https://udn.unrealengine.com/questions/523843/platform-specific-solution-rename-breaks-visualstu.html

diff --git a/Engine/Source/Programs/UnrealBuildTool/ProjectFiles/ProjectFileGenerator.cs b/Engine/Source/Programs/UnrealBuildTool/ProjectFiles/ProjectFileGenerator.cs
--- a/Engine/Source/Programs/UnrealBuildTool/ProjectFiles/ProjectFileGenerator.cs
+++ b/Engine/Source/Programs/UnrealBuildTool/ProjectFiles/ProjectFileGenerator.cs
@@ -648,6 +648,11 @@ namespace UnrealBuildTool
 					MasterProjectName += SortedPlatform;
 					IntermediateProjectFilesPath = new DirectoryReference(IntermediateProjectFilesPath.FullName + SortedPlatform);
 				}
+
+				// the master project name is always read from our intermediate directory and not the overriden one for this set of platforms
+				FileReference MasterProjectNameLocation = FileReference.Combine(UnrealBuildTool.EngineDirectory, "Intermediate", "ProjectFiles", "MasterProjectName.txt");
+				DirectoryReference.CreateDirectory(MasterProjectNameLocation.Directory);
+				FileReference.WriteAllText(MasterProjectNameLocation, MasterProjectName);
 			}
 
 			bool bCleanProjectFiles = Arguments.Any(x => x.Equals("-CleanProjects", StringComparison.InvariantCultureIgnoreCase));
