Tweak upgrade messages for UBT settings version.

#rnx

diff --git a/Engine/Source/Programs/UnrealBuildTool/Configuration/TargetRules.cs b/Engine/Source/Programs/UnrealBuildTool/Configuration/TargetRules.cs
--- a/Engine/Source/Programs/UnrealBuildTool/Configuration/TargetRules.cs
+++ b/Engine/Source/Programs/UnrealBuildTool/Configuration/TargetRules.cs
@@ -1646,11 +1646,10 @@ namespace UnrealBuildTool
 		/// <param name="Diagnostics">List of diagnostic messages</param>
 		internal void GetBuildSettingsInfo(List<string> Diagnostics)
 		{
-			if(DefaultBuildSettings < BuildSettingsVersion.Latest)
+			if(DefaultBuildSettings < BuildSettingsVersion.V2)
 			{
 				Diagnostics.Add("[Upgrade]");
-				string Version = DefaultBuildSettings.ToString().Replace("Release_", "").Replace("_", ".");
-				Diagnostics.Add(String.Format("[Upgrade] Using UE {0} compatible build settings. The latest version of UE4 sets the following values by default, which may require code changes:", Version));
+				Diagnostics.Add("[Upgrade] Using backwards-compatible build settings. The latest version of UE4 sets the following values by default, which may require code changes:");
 
 				List<Tuple<string, string>> ModifiedSettings = new List<Tuple<string, string>>();
 				if(DefaultBuildSettings < BuildSettingsVersion.V2)
diff --git a/Engine/Source/Programs/UnrealBuildTool/Modes/BuildMode.cs b/Engine/Source/Programs/UnrealBuildTool/Modes/BuildMode.cs
--- a/Engine/Source/Programs/UnrealBuildTool/Modes/BuildMode.cs
+++ b/Engine/Source/Programs/UnrealBuildTool/Modes/BuildMode.cs
@@ -359,7 +359,7 @@ namespace UnrealBuildTool
 				{
 					foreach(string Diagnostic in Makefiles[Idx].Diagnostics)
 					{
-						Log.TraceInformation("For {0}: {1}", TargetDescriptors[Idx], Diagnostic);
+						Log.TraceInformation("{0}: {1}", TargetDescriptors[Idx].Name, Diagnostic);
 					}
 				}
 			}
