UBT: Fix single-file compile option modifying response files for files other than the one being built. Causes more files to be rebuilt in the next regular build than necessary.

diff --git a/Engine/Source/Programs/UnrealBuildTool/Configuration/UEBuildBinary.cs b/Engine/Source/Programs/UnrealBuildTool/Configuration/UEBuildBinary.cs
--- a/Engine/Source/Programs/UnrealBuildTool/Configuration/UEBuildBinary.cs
+++ b/Engine/Source/Programs/UnrealBuildTool/Configuration/UEBuildBinary.cs
@@ -169,11 +169,12 @@ namespace UnrealBuildTool
 		/// <param name="ToolChain">The toolchain which to use for building</param>
 		/// <param name="CompileEnvironment">The environment to compile the binary in</param>
 		/// <param name="LinkEnvironment">The environment to link the binary in</param>
+		/// <param name="SingleFileToCompile">If non-null, specifies a single cpp file to be compiled</param>
 		/// <param name="WorkingSet">The working set of source files</param>
 		/// <param name="ExeDir">Directory containing the output executable</param>
 		/// <param name="Makefile">The makefile to add actions to</param>
 		/// <returns>Set of built products</returns>
-		public List<FileItem> Build(ReadOnlyTargetRules Target, UEToolChain ToolChain, CppCompileEnvironment CompileEnvironment, LinkEnvironment LinkEnvironment, ISourceFileWorkingSet WorkingSet, DirectoryReference ExeDir, TargetMakefile Makefile)
+		public List<FileItem> Build(ReadOnlyTargetRules Target, UEToolChain ToolChain, CppCompileEnvironment CompileEnvironment, LinkEnvironment LinkEnvironment, FileReference SingleFileToCompile, ISourceFileWorkingSet WorkingSet, DirectoryReference ExeDir, TargetMakefile Makefile)
 		{
 			// Return nothing if we're using precompiled binaries. If we're not linking, we might want just one module to be compiled (eg. a foreign plugin), so allow any actions to run.
 			if (bUsePrecompiled && !Target.bDisableLinking)
@@ -182,7 +183,7 @@ namespace UnrealBuildTool
 			}
 
 			// Setup linking environment.
-			LinkEnvironment BinaryLinkEnvironment = SetupBinaryLinkEnvironment(Target, ToolChain, LinkEnvironment, CompileEnvironment, WorkingSet, ExeDir, Makefile);
+			LinkEnvironment BinaryLinkEnvironment = SetupBinaryLinkEnvironment(Target, ToolChain, LinkEnvironment, CompileEnvironment, SingleFileToCompile, WorkingSet, ExeDir, Makefile);
 
 			// If we're generating projects, we only need include paths and definitions, there is no need to run the linking logic.
 			if (ProjectFileGenerator.bGenerateProjectFiles)
@@ -676,7 +677,7 @@ namespace UnrealBuildTool
 			return BinaryCompileEnvironment;
 		}
 
-		private LinkEnvironment SetupBinaryLinkEnvironment(ReadOnlyTargetRules Target, UEToolChain ToolChain, LinkEnvironment LinkEnvironment, CppCompileEnvironment CompileEnvironment, ISourceFileWorkingSet WorkingSet, DirectoryReference ExeDir, TargetMakefile Makefile)
+		private LinkEnvironment SetupBinaryLinkEnvironment(ReadOnlyTargetRules Target, UEToolChain ToolChain, LinkEnvironment LinkEnvironment, CppCompileEnvironment CompileEnvironment, FileReference SingleFileToCompile, ISourceFileWorkingSet WorkingSet, DirectoryReference ExeDir, TargetMakefile Makefile)
 		{
 			LinkEnvironment BinaryLinkEnvironment = new LinkEnvironment(LinkEnvironment);
 			HashSet<UEBuildModule> LinkEnvironmentVisitedModules = new HashSet<UEBuildModule>();
@@ -696,7 +697,7 @@ namespace UnrealBuildTool
 				{
 					// Compile each module.
 					Log.TraceVerbose("Compile module: " + Module.Name);
-					LinkInputFiles = Module.Compile(Target, ToolChain, BinaryCompileEnvironment, WorkingSet, Makefile);
+					LinkInputFiles = Module.Compile(Target, ToolChain, BinaryCompileEnvironment, SingleFileToCompile, WorkingSet, Makefile);
 
 					// Save the module outputs. In monolithic builds, this is just the object files.
 					if (Target.LinkType == TargetLinkType.Monolithic)
diff --git a/Engine/Source/Programs/UnrealBuildTool/Configuration/UEBuildModule.cs b/Engine/Source/Programs/UnrealBuildTool/Configuration/UEBuildModule.cs
--- a/Engine/Source/Programs/UnrealBuildTool/Configuration/UEBuildModule.cs
+++ b/Engine/Source/Programs/UnrealBuildTool/Configuration/UEBuildModule.cs
@@ -776,7 +776,7 @@ namespace UnrealBuildTool
 		/// <summary>
 		/// Compiles the module, and returns a list of files output by the compiler.
 		/// </summary>
-		public abstract List<FileItem> Compile(ReadOnlyTargetRules Target, UEToolChain ToolChain, CppCompileEnvironment CompileEnvironment, ISourceFileWorkingSet WorkingSet, TargetMakefile Makefile);
+		public abstract List<FileItem> Compile(ReadOnlyTargetRules Target, UEToolChain ToolChain, CppCompileEnvironment CompileEnvironment, FileReference SingleFileToCompile, ISourceFileWorkingSet WorkingSet, TargetMakefile Makefile);
 
 		// Object interface.
 		public override string ToString()
diff --git a/Engine/Source/Programs/UnrealBuildTool/Configuration/UEBuildModuleCPP.cs b/Engine/Source/Programs/UnrealBuildTool/Configuration/UEBuildModuleCPP.cs
--- a/Engine/Source/Programs/UnrealBuildTool/Configuration/UEBuildModuleCPP.cs
+++ b/Engine/Source/Programs/UnrealBuildTool/Configuration/UEBuildModuleCPP.cs
@@ -252,7 +252,7 @@ namespace UnrealBuildTool
 		}
 
 		// UEBuildModule interface.
-		public override List<FileItem> Compile(ReadOnlyTargetRules Target, UEToolChain ToolChain, CppCompileEnvironment BinaryCompileEnvironment, ISourceFileWorkingSet WorkingSet, TargetMakefile Makefile)
+		public override List<FileItem> Compile(ReadOnlyTargetRules Target, UEToolChain ToolChain, CppCompileEnvironment BinaryCompileEnvironment, FileReference SingleFileToCompile, ISourceFileWorkingSet WorkingSet, TargetMakefile Makefile)
 		{
 			UEBuildPlatform BuildPlatform = UEBuildPlatform.GetBuildPlatform(BinaryCompileEnvironment.Platform);
 
@@ -291,6 +291,20 @@ namespace UnrealBuildTool
 			// Find all the input files
 			InputFileCollection InputFiles = FindInputFiles(Target.Platform, Makefile.DirectoryToSourceFiles);
 
+			// If we're compiling a single file, strip out anything else. This prevents us clobbering response files for anything we're 
+			// not going to build, triggering a larger build than necessary when we do a regular build again.
+			if(SingleFileToCompile != null)
+			{
+				InputFiles.CPPFiles.RemoveAll(x => x.Location != SingleFileToCompile);
+				InputFiles.CCFiles.RemoveAll(x => x.Location != SingleFileToCompile);
+				InputFiles.CFiles.RemoveAll(x => x.Location != SingleFileToCompile);
+
+				if(InputFiles.CPPFiles.Count == 0 && InputFiles.CCFiles.Count == 0 && InputFiles.CFiles.Count == 0)
+				{
+					return new List<FileItem>();
+				}
+			}
+
 			// Process all of the header file dependencies for this module
 			CheckFirstIncludeMatchesEachCppFile(Target, ModuleCompileEnvironment, InputFiles.HeaderFiles, InputFiles.CPPFiles);
 
@@ -449,7 +463,7 @@ namespace UnrealBuildTool
 			}
 
 			// Compile all the generated CPP files
-			if (GeneratedCodeWildcard != null && !CompileEnvironment.bHackHeaderGenerator)
+			if (GeneratedCodeWildcard != null && !CompileEnvironment.bHackHeaderGenerator && SingleFileToCompile == null)
 			{
 				string[] GeneratedFiles = Directory.GetFiles(Path.GetDirectoryName(GeneratedCodeWildcard), Path.GetFileName(GeneratedCodeWildcard));
 				if(GeneratedFiles.Length > 0)
diff --git a/Engine/Source/Programs/UnrealBuildTool/Configuration/UEBuildModuleExternal.cs b/Engine/Source/Programs/UnrealBuildTool/Configuration/UEBuildModuleExternal.cs
--- a/Engine/Source/Programs/UnrealBuildTool/Configuration/UEBuildModuleExternal.cs
+++ b/Engine/Source/Programs/UnrealBuildTool/Configuration/UEBuildModuleExternal.cs
@@ -20,7 +20,7 @@ namespace UnrealBuildTool
 		}
 
 		// UEBuildModule interface.
-		public override List<FileItem> Compile(ReadOnlyTargetRules Target, UEToolChain ToolChain, CppCompileEnvironment CompileEnvironment, ISourceFileWorkingSet WorkingSet, TargetMakefile Makefile)
+		public override List<FileItem> Compile(ReadOnlyTargetRules Target, UEToolChain ToolChain, CppCompileEnvironment CompileEnvironment, FileReference SingleFileToCompile, ISourceFileWorkingSet WorkingSet, TargetMakefile Makefile)
 		{
 			return new List<FileItem>();
 		}
diff --git a/Engine/Source/Programs/UnrealBuildTool/Configuration/UEBuildTarget.cs b/Engine/Source/Programs/UnrealBuildTool/Configuration/UEBuildTarget.cs
--- a/Engine/Source/Programs/UnrealBuildTool/Configuration/UEBuildTarget.cs
+++ b/Engine/Source/Programs/UnrealBuildTool/Configuration/UEBuildTarget.cs
@@ -1671,7 +1671,7 @@ namespace UnrealBuildTool
 			{
 				foreach (UEBuildBinary Binary in BuildBinaries)
 				{
-					List<FileItem> BinaryOutputItems = Binary.Build(Rules, TargetToolChain, GlobalCompileEnvironment, GlobalLinkEnvironment, WorkingSet, ExeDir, Makefile);
+					List<FileItem> BinaryOutputItems = Binary.Build(Rules, TargetToolChain, GlobalCompileEnvironment, GlobalLinkEnvironment, SingleFileToCompile, WorkingSet, ExeDir, Makefile);
 					Makefile.OutputItems.AddRange(BinaryOutputItems);
 				}
 			}
