LiveCoding: Fix determination of linker path now that link-filter is used to preprocess linker output.

#rnx

diff --git a/Engine/Source/Programs/UnrealBuildTool/System/HotReload.cs b/Engine/Source/Programs/UnrealBuildTool/System/HotReload.cs
--- a/Engine/Source/Programs/UnrealBuildTool/System/HotReload.cs
+++ b/Engine/Source/Programs/UnrealBuildTool/System/HotReload.cs
@@ -888,7 +888,20 @@ namespace UnrealBuildTool
 				Action LinkAction = Actions.FirstOrDefault(x => x.ActionType == ActionType.Link && x.ProducedItems.Any(y => y.HasExtension(".exe") || y.HasExtension(".dll")));
 				if(LinkAction != null)
 				{
-					Writer.WriteValue("LinkerPath", LinkAction.CommandPath.FullName);
+					FileReference LinkerPath = LinkAction.CommandPath;
+					if(String.Compare(LinkerPath.GetFileName(), "link-filter.exe", StringComparison.OrdinalIgnoreCase) == 0)
+					{
+						string[] Arguments = CommandLineArguments.Split(LinkAction.CommandArguments);
+						for(int Idx = 0; Idx + 1 < Arguments.Length; Idx++)
+						{
+							if(Arguments[Idx] == "--")
+							{
+								LinkerPath = new FileReference(Arguments[Idx + 1]);
+								break;
+							}
+						}
+					}
+					Writer.WriteValue("LinkerPath", LinkerPath.FullName);
 				}
 
 				Writer.WriteObjectStart("LinkerEnvironment");
