Fix error trying to delete file when hot reloading with the debugger attached.

diff --git a/Engine/Source/Programs/UnrealBuildTool/System/HotReload.cs b/Engine/Source/Programs/UnrealBuildTool/System/HotReload.cs
--- a/Engine/Source/Programs/UnrealBuildTool/System/HotReload.cs
+++ b/Engine/Source/Programs/UnrealBuildTool/System/HotReload.cs
@@ -683,6 +683,16 @@ namespace UnrealBuildTool
 					}
 				}
 
+				// Fix up the list of items to delete too
+				for(int Idx = 0; Idx < Action.DeleteItems.Count; Idx++)
+				{
+					FileItem NewItem;
+					if(AffectedOriginalFileItemAndNewFileItemMap.TryGetValue(Action.DeleteItems[Idx], out NewItem))
+					{
+						Action.DeleteItems[Idx] = NewItem;
+					}
+				}
+
 				// The status description of the item has the file name, so we'll update it too
 				Action.StatusDescription = ReplaceBaseFileName(Action.StatusDescription, OriginalFileNameWithoutExtension, NewFileNameWithoutExtension);
 
