Remove the bBuildEditor flag passed to ModuleDescriptor.IsCompiledInConfiguration(). This is always determined from the target type now.

#rnx

diff --git a/Engine/Source/Programs/AutomationTool/Scripts/BuildPluginCommand.Automation.cs b/Engine/Source/Programs/AutomationTool/Scripts/BuildPluginCommand.Automation.cs
--- a/Engine/Source/Programs/AutomationTool/Scripts/BuildPluginCommand.Automation.cs
+++ b/Engine/Source/Programs/AutomationTool/Scripts/BuildPluginCommand.Automation.cs
@@ -182,9 +182,8 @@ class BuildPlugin : BuildCommand
 			foreach (ModuleDescriptor Module in Plugin.Modules)
 			{
 				bool bBuildDeveloperTools = (TargetType == TargetType.Editor || TargetType == TargetType.Program);
-				bool bBuildEditor = (TargetType == TargetType.Editor);
 				bool bBuildRequiresCookedData = (TargetType != TargetType.Editor && TargetType != TargetType.Program);
-				if (Module.IsCompiledInConfiguration(Platform, Configuration, TargetName, TargetType, bBuildDeveloperTools, bBuildEditor, bBuildRequiresCookedData))
+				if (Module.IsCompiledInConfiguration(Platform, Configuration, TargetName, TargetType, bBuildDeveloperTools, bBuildRequiresCookedData))
 				{
 					bCompilePlatform = true;
 				}
diff --git a/Engine/Source/Programs/UnrealBuildTool/Configuration/UEBuildTarget.cs b/Engine/Source/Programs/UnrealBuildTool/Configuration/UEBuildTarget.cs
--- a/Engine/Source/Programs/UnrealBuildTool/Configuration/UEBuildTarget.cs
+++ b/Engine/Source/Programs/UnrealBuildTool/Configuration/UEBuildTarget.cs
@@ -2677,7 +2677,7 @@ namespace UnrealBuildTool
 				// Add all the modules
 				foreach (ModuleDescriptor ModuleDescriptor in Plugin.Descriptor.Modules)
 				{
-					if (ModuleDescriptor.IsCompiledInConfiguration(Platform, Configuration, TargetName, TargetType, Rules.bBuildDeveloperTools, Rules.bBuildEditor, Rules.bBuildRequiresCookedData))
+					if (ModuleDescriptor.IsCompiledInConfiguration(Platform, Configuration, TargetName, TargetType, Rules.bBuildDeveloperTools, Rules.bBuildRequiresCookedData))
 					{
 						FileReference ModuleFileName = RulesAssembly.GetModuleFileName(ModuleDescriptor.Name);
 						if(ModuleFileName == null)
@@ -3141,7 +3141,7 @@ namespace UnrealBuildTool
 				{
 					foreach (ModuleDescriptor ModuleInfo in Info.Descriptor.Modules)
 					{
-						if (ModuleInfo.IsCompiledInConfiguration(Platform, Configuration, TargetName, TargetType, Rules.bBuildDeveloperTools, Rules.bBuildEditor, Rules.bBuildRequiresCookedData))
+						if (ModuleInfo.IsCompiledInConfiguration(Platform, Configuration, TargetName, TargetType, Rules.bBuildDeveloperTools, Rules.bBuildRequiresCookedData))
 						{
 							UEBuildModuleCPP Module = FindOrCreateCppModuleByName(ModuleInfo.Name, PluginReferenceChain);
 							if(!Instance.Modules.Contains(Module))
diff --git a/Engine/Source/Programs/UnrealBuildTool/System/ModuleDescriptor.cs b/Engine/Source/Programs/UnrealBuildTool/System/ModuleDescriptor.cs
--- a/Engine/Source/Programs/UnrealBuildTool/System/ModuleDescriptor.cs
+++ b/Engine/Source/Programs/UnrealBuildTool/System/ModuleDescriptor.cs
@@ -418,13 +418,12 @@ namespace UnrealBuildTool
 		/// Determines whether the given plugin module is part of the current build.
 		/// </summary>
 		/// <param name="Platform">The platform being compiled for</param>
-		/// <param name="TargetConfiguration">The target configuration being compiled for</param>
+		/// <param name="Configuration">The target configuration being compiled for</param>
 		/// <param name="TargetName">Name of the target being built</param>
 		/// <param name="TargetType">The type of the target being compiled</param>
 		/// <param name="bBuildDeveloperTools">Whether the configuration includes developer tools (typically UEBuildConfiguration.bBuildDeveloperTools for UBT callers)</param>
-		/// <param name="bBuildEditor">Whether the configuration includes the editor (typically UEBuildConfiguration.bBuildEditor for UBT callers)</param>
 		/// <param name="bBuildRequiresCookedData">Whether the configuration requires cooked content (typically UEBuildConfiguration.bBuildRequiresCookedData for UBT callers)</param>
-		public bool IsCompiledInConfiguration(UnrealTargetPlatform Platform, UnrealTargetConfiguration TargetConfiguration, string TargetName, TargetType TargetType, bool bBuildDeveloperTools, bool bBuildEditor, bool bBuildRequiresCookedData)
+		public bool IsCompiledInConfiguration(UnrealTargetPlatform Platform, UnrealTargetConfiguration Configuration, string TargetName, TargetType TargetType, bool bBuildDeveloperTools, bool bBuildRequiresCookedData)
 		{
 			// Check the platform is whitelisted
 			if (WhitelistPlatforms != null && WhitelistPlatforms.Length > 0 && !WhitelistPlatforms.Contains(Platform))
@@ -451,13 +450,13 @@ namespace UnrealBuildTool
 			}
 
 			// Check the target configuration is whitelisted
-			if (WhitelistTargetConfigurations != null && WhitelistTargetConfigurations.Length > 0 && !WhitelistTargetConfigurations.Contains(TargetConfiguration))
+			if (WhitelistTargetConfigurations != null && WhitelistTargetConfigurations.Length > 0 && !WhitelistTargetConfigurations.Contains(Configuration))
 			{
 				return false;
 			}
 
 			// Check the target configuration is not blacklisted
-			if (BlacklistTargetConfigurations != null && BlacklistTargetConfigurations.Contains(TargetConfiguration))
+			if (BlacklistTargetConfigurations != null && BlacklistTargetConfigurations.Contains(Configuration))
 			{
 				return false;
 			}
@@ -491,14 +490,14 @@ namespace UnrealBuildTool
 				case ModuleHostType.UncookedOnly:
 					return !bBuildRequiresCookedData;
 				case ModuleHostType.Developer:
-					return TargetType == TargetType.Editor || bBuildEditor || TargetType == TargetType.Program;
+					return TargetType == TargetType.Editor || TargetType == TargetType.Program;
 				case ModuleHostType.DeveloperTool:
 					return bBuildDeveloperTools;
 				case ModuleHostType.Editor:
 				case ModuleHostType.EditorNoCommandlet:
-					return TargetType == TargetType.Editor || bBuildEditor;
+					return TargetType == TargetType.Editor;
 				case ModuleHostType.EditorAndProgram:
-					return TargetType == TargetType.Editor || bBuildEditor || TargetType == TargetType.Program;
+					return TargetType == TargetType.Editor || TargetType == TargetType.Program;
 				case ModuleHostType.Program:
 					return TargetType == TargetType.Program;
                 case ModuleHostType.ServerOnly:
