Detect console application targets when generating Xcode project and set proper application type for them to avoid code signing and to allow launching them directly from Xcode

diff --git a/Engine/Source/Programs/UnrealBuildTool/ProjectFiles/Xcode/XcodeProject.cs b/Engine/Source/Programs/UnrealBuildTool/ProjectFiles/Xcode/XcodeProject.cs
--- a/Engine/Source/Programs/UnrealBuildTool/ProjectFiles/Xcode/XcodeProject.cs
+++ b/Engine/Source/Programs/UnrealBuildTool/ProjectFiles/Xcode/XcodeProject.cs
@@ -307,7 +307,7 @@ namespace UnrealBuildTool
 		/// <summary>
 		/// Generates bodies of all sections that contain a list of source files plus a dictionary of project navigator groups.
 		/// </summary>
-		private void GenerateSectionsWithSourceFiles(StringBuilder PBXBuildFileSection, StringBuilder PBXFileReferenceSection, StringBuilder PBXSourcesBuildPhaseSection, string TargetAppGuid, string TargetName)
+		private void GenerateSectionsWithSourceFiles(StringBuilder PBXBuildFileSection, StringBuilder PBXFileReferenceSection, StringBuilder PBXSourcesBuildPhaseSection, string TargetAppGuid, string TargetName, bool bIsAppBundle)
 		{
 			SourceFiles.Sort((x, y) => { return x.Reference.FullName.CompareTo(y.Reference.FullName); });
 
@@ -346,7 +346,7 @@ namespace UnrealBuildTool
 				}
 			}
 
-			PBXFileReferenceSection.Append(string.Format("\t\t{0} /* {1} */ = {{isa = PBXFileReference; explicitFileType = wrapper.application; path = {1}; sourceTree = BUILT_PRODUCTS_DIR; }};" + ProjectFileGenerator.NewLine, TargetAppGuid, TargetName));
+			PBXFileReferenceSection.Append(string.Format("\t\t{0} /* {1} */ = {{isa = PBXFileReference; explicitFileType = {2}; path = {1}; sourceTree = BUILT_PRODUCTS_DIR; }};" + ProjectFileGenerator.NewLine, TargetAppGuid, TargetName, bIsAppBundle ? "wrapper.application" : "\"compiled.mach-o.executable\""));
 		}
 
 		private void GenerateSectionsWithExtensions(StringBuilder PBXBuildFileSection, StringBuilder PBXFileReferenceSection, StringBuilder PBXCopyFilesBuildPhaseSection, StringBuilder PBXResourcesBuildPhaseSection,
@@ -645,7 +645,7 @@ namespace UnrealBuildTool
 		}
 
 		private void AppendRunTargetSection(StringBuilder Content, string TargetName, string TargetGuid, string TargetBuildConfigGuid, string TargetDependencyGuid, 
-				string TargetAppGuid, string CopyExtensionsBuildPhaseGuid, string ShellScriptSectionGuid, List<XcodeExtensionInfo> AllExtensions)
+				string TargetAppGuid, string CopyExtensionsBuildPhaseGuid, string ShellScriptSectionGuid, List<XcodeExtensionInfo> AllExtensions, bool bIsAppBundle)
 		{
 			List<string> DependencyGuids = new List<string>();
 			// depends on the Run target if we want one
@@ -671,7 +671,7 @@ namespace UnrealBuildTool
 			}
 
 			// use generica target section function for an application type
-			AppendGenericTargetSection(Content, TargetName, TargetGuid, "com.apple.product-type.application", TargetBuildConfigGuid, TargetAppGuid, DependencyGuids, BuildPhases);
+			AppendGenericTargetSection(Content, TargetName, TargetGuid, bIsAppBundle ? "com.apple.product-type.application" : "com.apple.product-type.tool", TargetBuildConfigGuid, TargetAppGuid, DependencyGuids, BuildPhases);
 		}
 
 		private void AppendGenericTargetSection(StringBuilder Content, string TargetName, string TargetGuid, string TargetType, string TargetBuildConfigGuid, string TargetAppGuid, IEnumerable<string> TargetDependencyGuids, Dictionary<string, string> BuildPhases)
@@ -1773,6 +1773,7 @@ namespace UnrealBuildTool
 			}
 
 			bool bHasEditorConfiguration = false;
+			bool bIsAppBundle = true;
 
 			Dictionary<string, XcodeBuildConfig> ProjectBuildConfigs = new Dictionary<string, XcodeBuildConfig>();
 			Dictionary<string, XcodeBuildConfig> TargetBuildConfigs = new Dictionary<string, XcodeBuildConfig>();
@@ -1789,6 +1790,10 @@ namespace UnrealBuildTool
 				{
 					bHasEditorConfiguration = true;
 				}
+				else if (Config.ProjectTarget.TargetRules.bIsBuildingConsoleApplication)
+				{
+					bIsAppBundle = false;
+				}
 			}
 
 			StringBuilder PBXBuildFileSection = new StringBuilder();
@@ -1797,7 +1802,7 @@ namespace UnrealBuildTool
 			StringBuilder PBXCopyExtensionsBuildPhaseSection = new StringBuilder();
 			StringBuilder PBXResourcesBuildPhaseSection = new StringBuilder();
 			List<XcodeExtensionInfo> AllExtensions = new List<XcodeExtensionInfo>();
-			GenerateSectionsWithSourceFiles(PBXBuildFileSection, PBXFileReferenceSection, PBXSourcesBuildPhaseSection, TargetAppGuid, TargetName);
+			GenerateSectionsWithSourceFiles(PBXBuildFileSection, PBXFileReferenceSection, PBXSourcesBuildPhaseSection, TargetAppGuid, TargetName, bIsAppBundle);
 			GenerateSectionsWithExtensions(PBXBuildFileSection, PBXFileReferenceSection, PBXCopyExtensionsBuildPhaseSection, PBXResourcesBuildPhaseSection, AllExtensions, GameProjectPath, BuildConfigs);
 
 			StringBuilder ProjectFileContent = new StringBuilder();
@@ -1827,7 +1832,7 @@ namespace UnrealBuildTool
 			}
 			AppendGroupSection(ProjectFileContent, MainGroupGuid, ProductRefGroupGuid, TargetAppGuid, TargetName, AllExtensions);
 			AppendLegacyTargetSection(ProjectFileContent, BuildTargetName, BuildTargetGuid, BuildTargetConfigListGuid, GameProjectPath, bHasEditorConfiguration);
-			AppendRunTargetSection(ProjectFileContent, TargetName, TargetGuid, TargetConfigListGuid, TargetDependencyGuid, TargetAppGuid, CopyExtensionsBuildPhaseGuid, ShellScriptSectionGuid, AllExtensions);
+			AppendRunTargetSection(ProjectFileContent, TargetName, TargetGuid, TargetConfigListGuid, TargetDependencyGuid, TargetAppGuid, CopyExtensionsBuildPhaseGuid, ShellScriptSectionGuid, AllExtensions, bIsAppBundle);
 			AppendIndexTargetSection(ProjectFileContent, IndexTargetName, IndexTargetGuid, IndexTargetConfigListGuid, SourcesBuildPhaseGuid);
 			AppendExtensionTargetSections(ProjectFileContent, AllExtensions);
 			AppendProjectSection(ProjectFileContent, TargetName, TargetGuid, BuildTargetName, BuildTargetGuid, IndexTargetName, IndexTargetGuid, MainGroupGuid, ProductRefGroupGuid, ProjectGuid, ProjectConfigListGuid, GameProjectPath, AllExtensions);
