[NO NON-PII COMMIT MESSAGE]

diff --git a/Engine/Source/Programs/UnrealBuildTool/Platform/Lumin/UEBuildLumin.cs b/Engine/Source/Programs/UnrealBuildTool/Platform/Lumin/UEBuildLumin.cs
--- a/Engine/Source/Programs/UnrealBuildTool/Platform/Lumin/UEBuildLumin.cs
+++ b/Engine/Source/Programs/UnrealBuildTool/Platform/Lumin/UEBuildLumin.cs
@@ -139,6 +139,13 @@ namespace UnrealBuildTool
 			// 			PlatformExtraModules.Add("VulkanRHI");
 			PlatformExtraModules.Add("MagicLeapAudio");
 			PlatformExtraModules.Add("MagicLeapAudioCapture");
+
+			// Hack: GoogleOboe is an explicit dependency of AudioCaptureAndroid, which is built when you compile for Lumin with the -allmodules flag.
+			// Doing this suppresses a warning that GoogleOboe isn't supported for Lumin, even though it is.
+			if(Target.bBuildAllModules)
+			{
+				PlatformExtraModules.Add("GoogleOboe");
+			}
 		}
 
 		/// <summary>
