Linux: reduce monolithic binary sizes.

diff --git a/Engine/Source/Programs/UnrealBuildTool/Platform/Linux/LinuxToolChain.cs b/Engine/Source/Programs/UnrealBuildTool/Platform/Linux/LinuxToolChain.cs
--- a/Engine/Source/Programs/UnrealBuildTool/Platform/Linux/LinuxToolChain.cs
+++ b/Engine/Source/Programs/UnrealBuildTool/Platform/Linux/LinuxToolChain.cs
@@ -750,6 +750,11 @@ namespace UnrealBuildTool
 				// glibc/ld.so limit (DTV_SURPLUS) for number of dlopen()'ed DSOs with static TLS (see e.g. https://www.cygwin.com/ml/libc-help/2013-11/msg00033.html)
 				Result += " -ftls-model=local-dynamic";
 			}
+			else
+			{
+				Result += " -ffunction-sections";
+				Result += " -fdata-sections";
+			}
 
 			if (CompileEnvironment.bEnableExceptions)
 			{
@@ -949,9 +954,14 @@ namespace UnrealBuildTool
 
 			// This apparently can help LLDB speed up symbol lookups
 			Result += " -Wl,--build-id";
-			if (bSuppressPIE && !LinkEnvironment.bIsBuildingDLL)
+			if (!LinkEnvironment.bIsBuildingDLL)
 			{
-				Result += " -Wl,-nopie";
+				Result += " -Wl,--gc-sections";
+
+				if (bSuppressPIE)
+				{
+					Result += " -Wl,-nopie";
+				}
 			}
 
 			// Profile Guided Optimization (PGO) and Link Time Optimization (LTO)
