Make unresolved symbols a compile-time error on Linux.

- This enables reporting of unresolved symbols within a single module. Full reporting (i.e. symbols missing from both library currently being linked and all other
  libraries it is linking to) is not enabled, because it doesn't work with our circular (re)linking scheme. When full reporting is enabled, if module A links to module B
  that links to module A again, then A's symbols referenced in B are reported as unresolved.

- Even with this, some modules cannot be solved at this moment, so this change introduces a setting to opt-out.
   Initially there were 5 modules that needed this opt-out. I solved most of them (some by disabling, since they missed the necessary libs), and now there is only one left.
   Even the last one is solvable (UE-70768), so it is tempting to remove this option. However, I anticipate that there will be moment when a new module addition
   will be breaking CIS due to unresolved symbols, and there will be no immediate and easy way to disable that except for nixing the whole system. Also, licensees may have more.

diff --git a/Engine/Plugins/Experimental/PythonScriptPlugin/Source/PythonScriptPlugin/PythonScriptPlugin.Build.cs b/Engine/Plugins/Experimental/PythonScriptPlugin/Source/PythonScriptPlugin/PythonScriptPlugin.Build.cs
--- a/Engine/Plugins/Experimental/PythonScriptPlugin/Source/PythonScriptPlugin/PythonScriptPlugin.Build.cs
+++ b/Engine/Plugins/Experimental/PythonScriptPlugin/Source/PythonScriptPlugin/PythonScriptPlugin.Build.cs
@@ -38,6 +38,9 @@ namespace UnrealBuildTool.Rules
 					}
 				);
 			}
+
+			// some of the symbols libpython.a is referencing are unresolved (see UE-70768)
+			bIgnoreUnresolvedSymbols = true;
 		}
 	}
 }
diff --git a/Engine/Source/Programs/UnrealBuildTool/Configuration/ModuleRules.cs b/Engine/Source/Programs/UnrealBuildTool/Configuration/ModuleRules.cs
--- a/Engine/Source/Programs/UnrealBuildTool/Configuration/ModuleRules.cs
+++ b/Engine/Source/Programs/UnrealBuildTool/Configuration/ModuleRules.cs
@@ -529,6 +529,11 @@ namespace UnrealBuildTool
 		/// </summary>
 		public bool bAddDefaultIncludePaths = true;
 
+		/// <summary>
+		/// Whether to ignore dangling (i.e. unresolved external) symbols in modules
+		/// </summary>
+		public bool bIgnoreUnresolvedSymbols = false;
+
 		/// <summary>
 		/// Whether this module should be precompiled. Defaults to the bPrecompile flag from the target. Clear this flag to prevent a module being precompiled.
 		/// </summary>
diff --git a/Engine/Source/Programs/UnrealBuildTool/Configuration/UEBuildModule.cs b/Engine/Source/Programs/UnrealBuildTool/Configuration/UEBuildModule.cs
--- a/Engine/Source/Programs/UnrealBuildTool/Configuration/UEBuildModule.cs
+++ b/Engine/Source/Programs/UnrealBuildTool/Configuration/UEBuildModule.cs
@@ -731,6 +731,9 @@ namespace UnrealBuildTool
 
 			// Add all the additional properties
 			LinkEnvironment.AdditionalProperties.AddRange(Rules.AdditionalPropertiesForReceipt.Inner);
+
+			// this is a link-time property that needs to be accumulated (if any modules contributing to this module is ignoring, all are ignoring)
+			LinkEnvironment.bIgnoreUnresolvedSymbols |= Rules.bIgnoreUnresolvedSymbols;
 		}
 
 		/// <summary>
diff --git a/Engine/Source/Programs/UnrealBuildTool/Platform/Linux/LinuxToolChain.cs b/Engine/Source/Programs/UnrealBuildTool/Platform/Linux/LinuxToolChain.cs
--- a/Engine/Source/Programs/UnrealBuildTool/Platform/Linux/LinuxToolChain.cs
+++ b/Engine/Source/Programs/UnrealBuildTool/Platform/Linux/LinuxToolChain.cs
@@ -1719,6 +1719,14 @@ namespace UnrealBuildTool
 
 			LinkCommandString += " -Wl,--start-group";
 			LinkCommandString += ExternalLibraries;
+
+			// make unresolved symbols an error, unless a) building a cross-referenced DSO  b) we opted out
+			if ((!LinkEnvironment.bIsBuildingDLL || !LinkEnvironment.bIsCrossReferenced) && !LinkEnvironment.bIgnoreUnresolvedSymbols)
+			{
+				// This will make the linker report undefined symbols the current module, but ignore in the dependent DSOs.
+				// It is tempting, but may not be possible to change that report-all - due to circular dependencies between our libs.
+				LinkCommandString += " -Wl,--unresolved-symbols=ignore-in-shared-libs";
+			}
 			LinkCommandString += " -Wl,--end-group";
 
 			LinkCommandString += " -lrt"; // needed for clock_gettime()
@@ -1733,6 +1741,7 @@ namespace UnrealBuildTool
 				LinkCommandString += " " + "ThirdParty/Linux/LibCxx/lib/Linux/" + LinkEnvironment.Architecture + "/libc++abi.a";
 				LinkCommandString += " -lm";
 				LinkCommandString += " -lc";
+				LinkCommandString += " -lpthread"; // pthread_mutex_trylock is missing from libc stubs
 				LinkCommandString += " -lgcc_s";
 				LinkCommandString += " -lgcc";
 			}
diff --git a/Engine/Source/Programs/UnrealBuildTool/System/LinkEnvironment.cs b/Engine/Source/Programs/UnrealBuildTool/System/LinkEnvironment.cs
--- a/Engine/Source/Programs/UnrealBuildTool/System/LinkEnvironment.cs
+++ b/Engine/Source/Programs/UnrealBuildTool/System/LinkEnvironment.cs
@@ -244,6 +244,11 @@ namespace UnrealBuildTool
 		/// </summary>
 		public bool bUseFastPDBLinking;
 
+		/// <summary>
+		/// Whether to ignore dangling (i.e. unresolved external) symbols in modules
+		/// </summary>
+		public bool bIgnoreUnresolvedSymbols;
+
 		/// <summary>
 		/// Whether to log detailed timing information
 		/// </summary>
@@ -347,6 +352,7 @@ namespace UnrealBuildTool
             bAllowASLR = Other.bAllowASLR;
 			bUsePDBFiles = Other.bUsePDBFiles;
 			bUseFastPDBLinking = Other.bUseFastPDBLinking;
+			bIgnoreUnresolvedSymbols = Other.bIgnoreUnresolvedSymbols;
 			bPrintTimingInfo = Other.bPrintTimingInfo;
 			BundleVersion = Other.BundleVersion;
 			InputFiles.AddRange(Other.InputFiles);
