UBT: Make the -Quiet option per-target, and use it to suppress toolchain and adaptive non-unity info for secondary build targets.

#rnx

diff --git a/Engine/Source/Programs/UnrealBuildTool/Configuration/TargetDescriptor.cs b/Engine/Source/Programs/UnrealBuildTool/Configuration/TargetDescriptor.cs
--- a/Engine/Source/Programs/UnrealBuildTool/Configuration/TargetDescriptor.cs
+++ b/Engine/Source/Programs/UnrealBuildTool/Configuration/TargetDescriptor.cs
@@ -70,6 +70,12 @@ namespace UnrealBuildTool
 		[CommandLine("-LiveCodingManifest=")]
 		public FileReference LiveCodingManifest = null;
 
+		/// <summary>
+		/// Suppress messages about building this target
+		/// </summary>
+		[CommandLine("-Quiet")]
+		public bool bQuiet;
+
 		/// <summary>
 		/// Constructor
 		/// </summary>
diff --git a/Engine/Source/Programs/UnrealBuildTool/Modes/BuildMode.cs b/Engine/Source/Programs/UnrealBuildTool/Modes/BuildMode.cs
--- a/Engine/Source/Programs/UnrealBuildTool/Modes/BuildMode.cs
+++ b/Engine/Source/Programs/UnrealBuildTool/Modes/BuildMode.cs
@@ -23,20 +23,15 @@ namespace UnrealBuildTool
 		/// </summary>
 		None = 0,
 
-		/// <summary>
-		/// Don't output any messages unless we're going to build something
-		/// </summary>
-		Quiet = 1,
-
 		/// <summary>
 		/// Don't build anything, just do target setup and terminate
 		/// </summary>
-		SkipBuild = 2,
+		SkipBuild = 1,
 
 		/// <summary>
 		/// Just output a list of XGE actions; don't build anything
 		/// </summary>
-		XGEExport = 4,
+		XGEExport = 2,
 	}
 
 	/// <summary>
@@ -304,14 +299,14 @@ namespace UnrealBuildTool
 					// Execute the actions
 					if(MergedActionsToExecute.Count == 0)
 					{
-						if ((Options & BuildOptions.Quiet) == 0)
+						if (TargetDescriptors.Any(x => !x.bQuiet))
 						{
 							Log.TraceInformation((TargetDescriptors.Count == 1)? "Target is up to date" : "Targets are up to date");
 						}
 					}
 					else
 					{
-						if((Options & BuildOptions.Quiet) != 0)
+						if (TargetDescriptors.Any(x => !x.bQuiet))
 						{
 							Log.TraceInformation("Building {0}...", StringUtils.FormatList(TargetDescriptors.Select(x => x.Name).Distinct()));
 						}
@@ -346,20 +341,29 @@ namespace UnrealBuildTool
 		/// <param name="Makefiles">Matching array of makefiles for each target</param>
 		static void OutputToolchainInfo(List<TargetDescriptor> TargetDescriptors, TargetMakefile[] Makefiles)
 		{
-			if(Makefiles.Length == 1 || Makefiles.Skip(1).All(x => Enumerable.SequenceEqual(x.Diagnostics, Makefiles[0].Diagnostics)))
+			List<int> OutputIndices = new List<int>();
+			for (int Idx = 0; Idx < TargetDescriptors.Count; Idx++)
 			{
-				foreach(string Diagnostic in Makefiles[0].Diagnostics)
+				if (!TargetDescriptors[Idx].bQuiet)
+				{
+					OutputIndices.Add(Idx);
+				}
+			}
+
+			if(OutputIndices.Count == 1)
+			{
+				foreach(string Diagnostic in Makefiles[OutputIndices[0]].Diagnostics)
 				{
 					Log.TraceInformation("{0}", Diagnostic);
 				}
 			}
 			else
 			{
-				for(int Idx = 0; Idx < TargetDescriptors.Count; Idx++)
+				foreach(int OutputIndex in OutputIndices)
 				{
-					foreach(string Diagnostic in Makefiles[Idx].Diagnostics)
+					foreach(string Diagnostic in Makefiles[OutputIndex].Diagnostics)
 					{
-						Log.TraceInformation("{0}: {1}", TargetDescriptors[Idx].Name, Diagnostic);
+						Log.TraceInformation("{0}: {1}", TargetDescriptors[OutputIndex].Name, Diagnostic);
 					}
 				}
 			}
diff --git a/Engine/Source/Programs/UnrealBuildTool/ProjectFiles/VisualStudio/VCProject.cs b/Engine/Source/Programs/UnrealBuildTool/ProjectFiles/VisualStudio/VCProject.cs
--- a/Engine/Source/Programs/UnrealBuildTool/ProjectFiles/VisualStudio/VCProject.cs
+++ b/Engine/Source/Programs/UnrealBuildTool/ProjectFiles/VisualStudio/VCProject.cs
@@ -1397,7 +1397,7 @@ namespace UnrealBuildTool
 						BuildArguments.Append("\"");
 						foreach(string ExtraTarget in ExtraTargets)
 						{
-							BuildArguments.AppendFormat(" -Target=\"{0}\"", ExtraTarget);
+							BuildArguments.AppendFormat(" -Target=\"{0} -Quiet\"", ExtraTarget);
 						}
 					}
 
diff --git a/Engine/Source/Programs/UnrealBuildTool/System/ExternalExecution.cs b/Engine/Source/Programs/UnrealBuildTool/System/ExternalExecution.cs
--- a/Engine/Source/Programs/UnrealBuildTool/System/ExternalExecution.cs
+++ b/Engine/Source/Programs/UnrealBuildTool/System/ExternalExecution.cs
@@ -1117,12 +1117,12 @@ namespace UnrealBuildTool
 						}
 
 						// Create the target descriptor
-						List<TargetDescriptor> TargetDescriptors = new List<TargetDescriptor>();
-						TargetDescriptors.Add(new TargetDescriptor(ScriptProjectFile, "UnrealHeaderTool", Platform, Configuration, Architecture, null));
+						TargetDescriptor TargetDescriptor = new TargetDescriptor(ScriptProjectFile, "UnrealHeaderTool", Platform, Configuration, Architecture, null);
+						TargetDescriptor.bQuiet = true;
 
 						using(Timeline.ScopeEvent("Buildng UnrealHeaderTool"))
 						{
-							BuildMode.Build(TargetDescriptors, BuildConfiguration, WorkingSet, BuildOptions.Quiet, null);
+							BuildMode.Build(new List<TargetDescriptor>{ TargetDescriptor }, BuildConfiguration, WorkingSet, BuildOptions.None, null);
 						}
 					}
 
