Speculative fix for building UHT plugins from installed engine builds.

diff --git a/Engine/Source/Programs/UnrealBuildTool/Configuration/UEBuildTarget.cs b/Engine/Source/Programs/UnrealBuildTool/Configuration/UEBuildTarget.cs
--- a/Engine/Source/Programs/UnrealBuildTool/Configuration/UEBuildTarget.cs
+++ b/Engine/Source/Programs/UnrealBuildTool/Configuration/UEBuildTarget.cs
@@ -617,8 +617,8 @@ namespace UnrealBuildTool
 				throw new BuildException("{0} does not support the {1} configuration", Descriptor.Name, Descriptor.Configuration);
 			}
 
-			// Make sure this target type is supported
-			if (!InstalledPlatformInfo.IsValid(RulesObject.Type, Descriptor.Platform, Descriptor.Configuration, EProjectType.Code, InstalledPlatformState.Downloaded))
+			// Make sure this target type is supported. Allow UHT in installed builds as a special case for now.
+			if (!InstalledPlatformInfo.IsValid(RulesObject.Type, Descriptor.Platform, Descriptor.Configuration, EProjectType.Code, InstalledPlatformState.Downloaded) && Descriptor.Name != "UnrealHeaderTool")
 			{
 				if (InstalledPlatformInfo.IsValid(RulesObject.Type, Descriptor.Platform, Descriptor.Configuration, EProjectType.Code, InstalledPlatformState.Supported))
 				{
@@ -1182,12 +1182,12 @@ namespace UnrealBuildTool
 						{
 							ResolveLibraryName(LibraryPath, LibraryName, LibraryExtension, Files);
 						}
-					}
+							}
 
 					foreach (string LibraryName in Rules.PublicSystemLibraryPaths)
-					{
+							{
 						foreach (string LibraryPath in SystemLibraryPaths)
-						{
+								{
 							ResolveLibraryName(LibraryPath, LibraryName, LibraryExtension, Files);
 						}
 					}
@@ -2690,7 +2690,7 @@ namespace UnrealBuildTool
 					}
 				}
 			}
-			
+
 			// Create rules for each remaining module, and check that it's set to be compiled
 			foreach(string FilteredModuleName in FilteredModuleNames)
 			{
