Fixing timing data generation with cl.exe when using the VS2019 compiler in UBT.

#rnx

diff --git a/Engine/Source/Programs/UnrealBuildTool/Platform/Windows/VCToolChain.cs b/Engine/Source/Programs/UnrealBuildTool/Platform/Windows/VCToolChain.cs
--- a/Engine/Source/Programs/UnrealBuildTool/Platform/Windows/VCToolChain.cs
+++ b/Engine/Source/Programs/UnrealBuildTool/Platform/Windows/VCToolChain.cs
@@ -898,7 +898,9 @@ namespace UnrealBuildTool
 
 			if (CompileEnvironment.bPrintTimingInfo || Target.WindowsPlatform.bCompilerTrace)
 			{
-				if(Target.WindowsPlatform.Compiler == WindowsCompiler.VisualStudio2015_DEPRECATED || Target.WindowsPlatform.Compiler == WindowsCompiler.VisualStudio2017)
+				if (Target.WindowsPlatform.Compiler == WindowsCompiler.VisualStudio2015_DEPRECATED ||
+					Target.WindowsPlatform.Compiler == WindowsCompiler.VisualStudio2017 ||
+					Target.WindowsPlatform.Compiler == WindowsCompiler.VisualStudio2019)
 				{
 					if (CompileEnvironment.bPrintTimingInfo)
 					{
