Fixing bug in UBT where it would try to run the aggregation step when generating timing information and only compiling one file even though there is only file to aggegate.

#rnx

diff --git a/Engine/Source/Programs/UnrealBuildTool/Platform/Windows/VCToolChain.cs b/Engine/Source/Programs/UnrealBuildTool/Platform/Windows/VCToolChain.cs
--- a/Engine/Source/Programs/UnrealBuildTool/Platform/Windows/VCToolChain.cs
+++ b/Engine/Source/Programs/UnrealBuildTool/Platform/Windows/VCToolChain.cs
@@ -1233,7 +1233,6 @@ namespace UnrealBuildTool
 					AggregateTimingInfoAction.StatusDescription = $"Aggregating {TimingJsonFiles.Count} Timing File(s)";
 					AggregateTimingInfoAction.bCanExecuteRemotely = false;
 					AggregateTimingInfoAction.bCanExecuteRemotelyWithSNDBS = false;
-					AggregateTimingInfoAction.PrerequisiteItems.AddRange(ParseTimingActions.SelectMany(a => a.PrerequisiteItems));
 					AggregateTimingInfoAction.PrerequisiteItems.AddRange(TimingJsonFiles);
 
 					FileItem AggregateOutputFile = FileItem.GetItemByFileReference(FileReference.Combine(Makefile.ProjectIntermediateDirectory, $"{Target.Name}.timing.bin"));
