Fix exception when trying to package a project that uses content-only plugins.

#rnx

diff --git a/Engine/Source/Programs/UnrealBuildTool/System/Plugins.cs b/Engine/Source/Programs/UnrealBuildTool/System/Plugins.cs
--- a/Engine/Source/Programs/UnrealBuildTool/System/Plugins.cs
+++ b/Engine/Source/Programs/UnrealBuildTool/System/Plugins.cs
@@ -549,7 +549,7 @@ namespace UnrealBuildTool
 		public static bool IsPluginCompiledForTarget(PluginInfo Plugin, ProjectDescriptor Project, UnrealTargetPlatform Platform, UnrealTargetConfiguration Configuration, TargetType TargetType, bool bRequiresCookedData)
 		{
 			bool bCompiledForTarget = false;
-			if (IsPluginEnabledForTarget(Plugin, Project, Platform, Configuration, TargetType))
+			if (IsPluginEnabledForTarget(Plugin, Project, Platform, Configuration, TargetType) && Plugin.Descriptor.Modules != null)
 			{
 				bool bBuildDeveloperTools = (TargetType == TargetType.Editor || TargetType == TargetType.Program || (Configuration != UnrealTargetConfiguration.Test && Configuration != UnrealTargetConfiguration.Shipping));
 				foreach (ModuleDescriptor Module in Plugin.Descriptor.Modules)
