Merging LiveCoding fix from Dev-Build.

#rnx

diff --git a/Engine/Source/Programs/UnrealBuildTool/System/HotReload.cs b/Engine/Source/Programs/UnrealBuildTool/System/HotReload.cs
--- a/Engine/Source/Programs/UnrealBuildTool/System/HotReload.cs
+++ b/Engine/Source/Programs/UnrealBuildTool/System/HotReload.cs
@@ -683,6 +683,16 @@ namespace UnrealBuildTool
 					}
 				}
 
+				// Fix up the list of items to delete too
+				for(int Idx = 0; Idx < Action.DeleteItems.Count; Idx++)
+				{
+					FileItem NewItem;
+					if(AffectedOriginalFileItemAndNewFileItemMap.TryGetValue(Action.DeleteItems[Idx], out NewItem))
+					{
+						Action.DeleteItems[Idx] = NewItem;
+					}
+				}
+
 				// The status description of the item has the file name, so we'll update it too
 				Action.StatusDescription = ReplaceBaseFileName(Action.StatusDescription, OriginalFileNameWithoutExtension, NewFileNameWithoutExtension);
 
@@ -888,7 +898,20 @@ namespace UnrealBuildTool
 				Action LinkAction = Actions.FirstOrDefault(x => x.ActionType == ActionType.Link && x.ProducedItems.Any(y => y.HasExtension(".exe") || y.HasExtension(".dll")));
 				if(LinkAction != null)
 				{
-					Writer.WriteValue("LinkerPath", LinkAction.CommandPath.FullName);
+					FileReference LinkerPath = LinkAction.CommandPath;
+					if(String.Compare(LinkerPath.GetFileName(), "link-filter.exe", StringComparison.OrdinalIgnoreCase) == 0)
+					{
+						string[] Arguments = CommandLineArguments.Split(LinkAction.CommandArguments);
+						for(int Idx = 0; Idx + 1 < Arguments.Length; Idx++)
+						{
+							if(Arguments[Idx] == "--")
+							{
+								LinkerPath = new FileReference(Arguments[Idx + 1]);
+								break;
+							}
+						}
+					}
+					Writer.WriteValue("LinkerPath", LinkerPath.FullName);
 				}
 
 				Writer.WriteObjectStart("LinkerEnvironment");
