Tweak feedback message for using old build settings.

#rnx

diff --git a/Engine/Source/Programs/UnrealBuildTool/Configuration/TargetRules.cs b/Engine/Source/Programs/UnrealBuildTool/Configuration/TargetRules.cs
--- a/Engine/Source/Programs/UnrealBuildTool/Configuration/TargetRules.cs
+++ b/Engine/Source/Programs/UnrealBuildTool/Configuration/TargetRules.cs
@@ -1667,7 +1667,7 @@ namespace UnrealBuildTool
 						Diagnostics.Add(String.Format(FormatString, ModifiedSetting.Item1, ModifiedSetting.Item2));
 					}
 				}
-				Diagnostics.Add(String.Format("[Upgrade] Suppress this message by setting 'DefaultBuildSettings = BuildSettingsVersion.{1};' in {2}, and explicitly overriding desired settings.", Version, (BuildSettingsVersion)(BuildSettingsVersion.Latest - 1), File.GetFileName()));
+				Diagnostics.Add(String.Format("[Upgrade] Suppress this message by setting 'DefaultBuildSettings = BuildSettingsVersion.{1};' in {2}, and explicitly overriding settings that differ from the new defaults.", Version, (BuildSettingsVersion)(BuildSettingsVersion.Latest - 1), File.GetFileName()));
 				Diagnostics.Add("[Upgrade]");
 			}
 		}
