UBT: Add a -preprocess argument, which will just preprocess the source code rather than compile it. Can be used in conjunction with -SingleFile=... to preprocess a single source file.

diff --git a/Engine/Source/Programs/UnrealBuildTool/Configuration/TargetRules.cs b/Engine/Source/Programs/UnrealBuildTool/Configuration/TargetRules.cs
--- a/Engine/Source/Programs/UnrealBuildTool/Configuration/TargetRules.cs
+++ b/Engine/Source/Programs/UnrealBuildTool/Configuration/TargetRules.cs
@@ -829,6 +829,12 @@ namespace UnrealBuildTool
 		[XmlConfigFile(Category = "BuildConfiguration")]
 		public bool bUsePCHFiles = true;
 
+		/// <summary>
+		/// Whether to just preprocess source files for this target, and skip compilation
+		/// </summary>
+		[CommandLine("-Preprocess")]
+		public bool bPreprocessOnly = false;
+
 		/// <summary>
 		/// The minimum number of files that must use a pre-compiled header before it will be created and used.
 		/// </summary>
@@ -2178,6 +2184,11 @@ namespace UnrealBuildTool
 			get { return Inner.bUsePCHFiles; }
 		}
 
+		public bool bPreprocessOnly
+		{
+			get { return Inner.bPreprocessOnly; }
+		}
+
 		public int MinFilesUsingPrecompiledHeader
 		{
 			get { return Inner.MinFilesUsingPrecompiledHeader; }
diff --git a/Engine/Source/Programs/UnrealBuildTool/Configuration/UEBuildTarget.cs b/Engine/Source/Programs/UnrealBuildTool/Configuration/UEBuildTarget.cs
--- a/Engine/Source/Programs/UnrealBuildTool/Configuration/UEBuildTarget.cs
+++ b/Engine/Source/Programs/UnrealBuildTool/Configuration/UEBuildTarget.cs
@@ -704,6 +704,12 @@ namespace UnrealBuildTool
 				RulesObject.bUseSharedPCHs = false;
 			}
 
+			// Don't link if we're just preprocessing
+			if (RulesObject.bPreprocessOnly)
+			{
+				RulesObject.bDisableLinking = true;
+			}
+
 			// Include generated code plugin if not building an editor target and project is configured for nativization
 			FileReference NativizedPluginFile = RulesObject.GetNativizedPlugin();
 			if(NativizedPluginFile != null)
@@ -3299,6 +3305,7 @@ namespace UnrealBuildTool
 			GlobalCompileEnvironment.bOmitFramePointers = Rules.bOmitFramePointers;
 			GlobalCompileEnvironment.bUsePDBFiles = Rules.bUsePDBFiles;
 			GlobalCompileEnvironment.bSupportEditAndContinue = Rules.bSupportEditAndContinue;
+			GlobalCompileEnvironment.bPreprocessOnly = Rules.bPreprocessOnly;
 			GlobalCompileEnvironment.bUseIncrementalLinking = Rules.bUseIncrementalLinking;
 			GlobalCompileEnvironment.bAllowLTCG = Rules.bAllowLTCG;
 			GlobalCompileEnvironment.bPGOOptimize = Rules.bPGOOptimize;
diff --git a/Engine/Source/Programs/UnrealBuildTool/Platform/Windows/VCToolChain.cs b/Engine/Source/Programs/UnrealBuildTool/Platform/Windows/VCToolChain.cs
--- a/Engine/Source/Programs/UnrealBuildTool/Platform/Windows/VCToolChain.cs
+++ b/Engine/Source/Programs/UnrealBuildTool/Platform/Windows/VCToolChain.cs
@@ -1014,6 +1014,19 @@ namespace UnrealBuildTool
 					FileArguments.Add(String.Format("/FI\"{0}\"", ForceIncludeFile.Location));
 				}
 
+				if (CompileEnvironment.bPreprocessOnly)
+				{
+					FileItem PreprocessedFile = FileItem.GetItemByFileReference(FileReference.Combine(OutputDir, SourceFile.Location.GetFileName() + ".i"));
+
+					FileArguments.Add("/P"); // Preprocess
+					FileArguments.Add("/C"); // Preserve comments when preprocessing
+					FileArguments.Add(String.Format("/Fi\"{0}\"", PreprocessedFile)); // Preprocess to a file
+
+					CompileAction.ProducedItems.Add(PreprocessedFile);
+
+					bEmitsObjectFile = false;
+				}
+
 				if (bEmitsObjectFile)
 				{
 					// Add the object file to the produced item list.
diff --git a/Engine/Source/Programs/UnrealBuildTool/System/CppCompileEnvironment.cs b/Engine/Source/Programs/UnrealBuildTool/System/CppCompileEnvironment.cs
--- a/Engine/Source/Programs/UnrealBuildTool/System/CppCompileEnvironment.cs
+++ b/Engine/Source/Programs/UnrealBuildTool/System/CppCompileEnvironment.cs
@@ -245,6 +245,11 @@ namespace UnrealBuildTool
 		/// </summary>
 		public bool bUsePDBFiles = false;
 
+		/// <summary>
+		/// Whether to just preprocess source files
+		/// </summary>
+		public bool bPreprocessOnly = false;
+
 		/// <summary>
 		/// Whether to support edit and continue.  Only works on Microsoft compilers in 32-bit compiles.
 		/// </summary>
@@ -407,6 +412,7 @@ namespace UnrealBuildTool
 			bOmitFramePointers = Other.bOmitFramePointers;
 			bEnableOSX109Support = Other.bEnableOSX109Support;
 			bUsePDBFiles = Other.bUsePDBFiles;
+			bPreprocessOnly = Other.bPreprocessOnly;
 			bSupportEditAndContinue = Other.bSupportEditAndContinue;
 			bUseIncrementalLinking = Other.bUseIncrementalLinking;
 			bAllowLTCG = Other.bAllowLTCG;
