Change PVS Studio toolchain to use an instance of the VCToolChain to preprocess all the input files. Should fix issues with dependency checking not working correctly.

diff --git a/Engine/Source/Programs/UnrealBuildTool/Platform/Windows/PVSToolChain.cs b/Engine/Source/Programs/UnrealBuildTool/Platform/Windows/PVSToolChain.cs
--- a/Engine/Source/Programs/UnrealBuildTool/Platform/Windows/PVSToolChain.cs
+++ b/Engine/Source/Programs/UnrealBuildTool/Platform/Windows/PVSToolChain.cs
@@ -131,7 +131,8 @@ namespace UnrealBuildTool
 
 	class PVSToolChain : UEToolChain
 	{
-		VCEnvironment EnvVars;
+		ReadOnlyTargetRules Target;
+		VCToolChain InnerToolChain;
 		FileReference AnalyzerFile;
 		FileReference LicenseFile;
 		PVSApplicationSettings ApplicationSettings;
@@ -139,8 +140,9 @@ namespace UnrealBuildTool
 
 		public PVSToolChain(ReadOnlyTargetRules Target)
 		{
+			this.Target = Target;
 			Platform = Target.Platform;
-			EnvVars = Target.WindowsPlatform.Environment;
+			InnerToolChain = new VCToolChain(Target);
 
 			AnalyzerFile = FileReference.Combine(new DirectoryReference(Environment.GetFolderPath(Environment.SpecialFolder.ProgramFilesX86)), "PVS-Studio", "x64", "PVS-Studio.exe");
 			if(!FileReference.Exists(AnalyzerFile))
@@ -195,71 +197,45 @@ namespace UnrealBuildTool
 
 		public override CPPOutput CompileCPPFiles(CppCompileEnvironment CompileEnvironment, List<FileItem> InputFiles, DirectoryReference OutputDir, string ModuleName, List<Action> Actions)
 		{
-			// Get the MSVC arguments required to compile all files in this batch
-			List<string> SharedArguments = new List<string>();
-			SharedArguments.Add("/nologo");
-			SharedArguments.Add("/P"); // Preprocess
-			SharedArguments.Add("/C"); // Preserve comments when preprocessing
-			SharedArguments.Add("/D PVS_STUDIO");
-			SharedArguments.Add("/wd4005");
-			if (EnvVars.Compiler >= WindowsCompiler.VisualStudio2015_DEPRECATED)
+			// Preprocess the source files with the regular toolchain
+			CppCompileEnvironment PreprocessCompileEnvironment = new CppCompileEnvironment(CompileEnvironment);
+			PreprocessCompileEnvironment.bPreprocessOnly = true;
+			PreprocessCompileEnvironment.bEnableUndefinedIdentifierWarnings = false; // Not sure why THIRD_PARTY_INCLUDES_START doesn't pick this up; the _Pragma appears in the preprocessed output. Perhaps in preprocess-only mode the compiler doesn't respect these?
+			PreprocessCompileEnvironment.Definitions.Add("PVS_STUDIO");
+
+			List<Action> PreprocessActions = new List<Action>();
+			CPPOutput Result = InnerToolChain.CompileCPPFiles(PreprocessCompileEnvironment, InputFiles, OutputDir, ModuleName, PreprocessActions);
+			Actions.AddRange(PreprocessActions);
+
+			// Run the source files through PVS-Studio
+			foreach(Action PreprocessAction in PreprocessActions)
 			{
-				SharedArguments.Add("/D _SILENCE_STDEXT_HASH_DEPRECATION_WARNINGS=1");
-			}
-			foreach (DirectoryReference IncludePath in CompileEnvironment.UserIncludePaths)
-			{
-				SharedArguments.Add(String.Format("/I \"{0}\"", IncludePath));
-			}
-			foreach (DirectoryReference IncludePath in CompileEnvironment.SystemIncludePaths)
-			{
-				SharedArguments.Add(String.Format("/I \"{0}\"", IncludePath));
-			}
-			foreach (DirectoryReference IncludePath in EnvVars.IncludePaths)
-			{
-				SharedArguments.Add(String.Format("/I \"{0}\"", IncludePath));
-			}
-			foreach (string Definition in CompileEnvironment.Definitions)
-			{
-				SharedArguments.Add(String.Format("/D \"{0}\"", Definition));
-			}
-			foreach(FileItem ForceIncludeFile in CompileEnvironment.ForceIncludeFiles)
-			{
-				SharedArguments.Add(String.Format("/FI\"{0}\"", ForceIncludeFile.Location));
-			}
-
-			CPPOutput Result = new CPPOutput();
-			foreach (FileItem SourceFile in InputFiles)
-			{
-				// Get the file names for everything we need
-				string BaseFileName = SourceFile.Location.GetFileName();
-
-				// Write the response file
-				FileReference PreprocessedFileLocation = FileReference.Combine(OutputDir, BaseFileName + ".i");
-
-				List<string> Arguments = new List<string>(SharedArguments);
-				Arguments.Add(String.Format("/Fi\"{0}\"", PreprocessedFileLocation)); // Preprocess to a file
-				Arguments.Add(String.Format("\"{0}\"", SourceFile.AbsolutePath));
+				if (PreprocessAction.ActionType != ActionType.Compile)
+				{
+					continue;
+				}
 
-				FileReference ResponseFileLocation = FileReference.Combine(OutputDir, BaseFileName + ".i.response");
-				FileItem ResponseFileItem = FileItem.CreateIntermediateTextFile(ResponseFileLocation, String.Join("\n", Arguments));
+				FileItem SourceFileItem = PreprocessAction.PrerequisiteItems.FirstOrDefault(x => x.HasExtension(".c") || x.HasExtension(".cc") || x.HasExtension(".cpp"));
+				if (SourceFileItem == null)
+				{
+					Log.TraceWarning("Unable to find source file from command: {0} {1}", PreprocessAction.CommandArguments);
+					continue;
+				}
 
-				// Preprocess the source file
-				FileItem PreprocessedFileItem = FileItem.GetItemByFileReference(PreprocessedFileLocation);
+				FileItem PreprocessedFileItem = PreprocessAction.ProducedItems.FirstOrDefault(x => x.HasExtension(".i"));
+				if (PreprocessedFileItem == null)
+				{
+					Log.TraceWarning("Unable to find preprocessed output file from command: {0} {1}", PreprocessAction.CommandArguments);
+					continue;
+				}
 
-				Action PreprocessAction = new Action(ActionType.Compile);
-				PreprocessAction.CommandPath = EnvVars.CompilerPath;
-				PreprocessAction.WorkingDirectory = UnrealBuildTool.EngineSourceDirectory;
-				PreprocessAction.CommandArguments = " @\"" + ResponseFileItem.AbsolutePath + "\"";
-				PreprocessAction.PrerequisiteItems.AddRange(CompileEnvironment.ForceIncludeFiles);
-				PreprocessAction.PrerequisiteItems.Add(SourceFile);
-				PreprocessAction.PrerequisiteItems.Add(ResponseFileItem);
-				PreprocessAction.ProducedItems.Add(PreprocessedFileItem);
-				PreprocessAction.bShouldOutputStatusDescription = false;
-				Actions.Add(PreprocessAction);
+				// Disable a few warnings that seem to come from the preprocessor not respecting _Pragma
+				PreprocessAction.CommandArguments += " /wd4005"; // macro redefinition
+				PreprocessAction.CommandArguments += " /wd4828"; // file contains a character starting at offset xxxx that is illegal in the current source character set
 
 				// Write the PVS studio config file
 				StringBuilder ConfigFileContents = new StringBuilder();
-				foreach(DirectoryReference IncludePath in EnvVars.IncludePaths)
+				foreach(DirectoryReference IncludePath in Target.WindowsPlatform.Environment.IncludePaths)
 				{
 					ConfigFileContents.AppendFormat("exclude-path={0}\n", IncludePath.FullName);
 				}
@@ -297,6 +273,8 @@ namespace UnrealBuildTool
 				ConfigFileContents.Append("skip-cl-exe=yes\n");
 				ConfigFileContents.AppendFormat("i-file={0}\n", PreprocessedFileItem.Location.FullName);
 
+				string BaseFileName = PreprocessedFileItem.Location.GetFileNameWithoutExtension();
+
 				FileReference ConfigFileLocation = FileReference.Combine(OutputDir, BaseFileName + ".cfg");
 				FileItem ConfigFileItem = FileItem.CreateIntermediateTextFile(ConfigFileLocation, ConfigFileContents.ToString());
 
@@ -309,7 +287,7 @@ namespace UnrealBuildTool
 				AnalyzeAction.StatusDescription = BaseFileName;
 				AnalyzeAction.WorkingDirectory = UnrealBuildTool.EngineSourceDirectory;
 				AnalyzeAction.CommandPath = AnalyzerFile;
-				AnalyzeAction.CommandArguments = String.Format("--cl-params \"{0}\" --source-file \"{1}\" --output-file \"{2}\" --cfg \"{3}\" --analysis-mode 4", PreprocessAction.CommandArguments, SourceFile.AbsolutePath, OutputFileLocation, ConfigFileItem.AbsolutePath);
+				AnalyzeAction.CommandArguments = String.Format("--cl-params \"{0}\" --source-file \"{1}\" --output-file \"{2}\" --cfg \"{3}\" --analysis-mode 4", PreprocessAction.CommandArguments, SourceFileItem.AbsolutePath, OutputFileLocation, ConfigFileItem.AbsolutePath);
 				if (LicenseFile != null)
 				{
 					AnalyzeAction.CommandArguments += String.Format(" --lic-file \"{0}\"", LicenseFile);
