UBT: Allow filtering the platforms we generate project files for, from the BuildConfiguration.xml file. Reduces memory usage in the IDE by a bit.

diff --git a/Engine/Source/Programs/UnrealBuildTool/ProjectFiles/ProjectFileGenerator.cs b/Engine/Source/Programs/UnrealBuildTool/ProjectFiles/ProjectFileGenerator.cs
--- a/Engine/Source/Programs/UnrealBuildTool/ProjectFiles/ProjectFileGenerator.cs
+++ b/Engine/Source/Programs/UnrealBuildTool/ProjectFiles/ProjectFileGenerator.cs
@@ -196,6 +196,11 @@ namespace UnrealBuildTool
 		/// </summary>
 		readonly List<UnrealTargetPlatform> ProjectPlatforms = new List<UnrealTargetPlatform>();
 
+		/// <summary>
+		/// Whether to append the list of platform names after the solution
+		/// </summary>
+		public bool bAppendPlatformSuffix;
+
 		/// <summary>
 		/// When bGeneratingGameProjectFiles=true, this is the game name we're generating projects for
 		/// </summary>
@@ -298,6 +303,12 @@ namespace UnrealBuildTool
 		[XmlConfigFile]
 		bool bKeepSourceSubDirectories = true;
 
+		/// <summary>
+		/// Names of platforms to include in the generated project files
+		/// </summary>
+		[XmlConfigFile(Name = "Platforms")]
+		string[] PlatformNames = null;
+
 		/// <summary>
 		/// Relative path to the directory where the master project file will be saved to
 		/// </summary>
@@ -578,7 +589,7 @@ namespace UnrealBuildTool
 		public virtual bool GenerateProjectFiles( PlatformProjectGeneratorCollection PlatformProjectGenerators, String[] Arguments )
 		{
 			bool bSuccess = true;
-
+	
 			// Parse project generator options
 			bool IncludeAllPlatforms = true;
 			ConfigureProjectFileGeneration( Arguments, ref IncludeAllPlatforms);
@@ -632,7 +643,7 @@ namespace UnrealBuildTool
 			}
 
 			// Modify the name if specific platforms were given
-			if (ProjectPlatforms.Count > 0)
+			if (ProjectPlatforms.Count > 0 && bAppendPlatformSuffix)
 			{
 				// Sort the platforms names so we get consistent names
 				List<string> SortedPlatformNames = new List<string>();
@@ -1039,6 +1050,18 @@ namespace UnrealBuildTool
 		/// <param name="IncludeAllPlatforms">True if all platforms should be included</param>
 		protected virtual void ConfigureProjectFileGeneration( String[] Arguments, ref bool IncludeAllPlatforms )
 		{
+			if (PlatformNames != null)
+			{
+				foreach (string PlatformName in PlatformNames)
+				{
+					UnrealTargetPlatform Platform;
+					if (UnrealTargetPlatform.TryParse(PlatformName, out Platform) && !ProjectPlatforms.Contains(Platform))
+					{
+						ProjectPlatforms.Add(Platform);
+					}
+				}
+			}
+
 			bool bAlwaysIncludeEngineModules = false;
 			foreach( string CurArgument in Arguments )
 			{
@@ -1076,6 +1099,9 @@ namespace UnrealBuildTool
 							{
 								Log.TraceWarning("ProjectFiles invalid platform specified: {0}", PlatformString);
 							}
+
+							// Append the platform suffix to the solution name
+							bAppendPlatformSuffix = true;
 						}
 					}
 					else switch( CurArgument.ToUpperInvariant() )
@@ -1878,7 +1904,7 @@ namespace UnrealBuildTool
 			}
 			if (BuildHostPlatform.Current.Platform == UnrealTargetPlatform.Win64)
 			{
-				return InPlatform == UnrealTargetPlatform.Win32 || InPlatform == UnrealTargetPlatform.Win64;
+				return InPlatform == UnrealTargetPlatform.Win64;
 			}
 
 			throw new BuildException("Invalid RuntimePlatform:" + BuildHostPlatform.Current.Platform);
