Show additional logging about failures to perform an incremental link.

diff --git a/Engine/Source/Programs/UnrealBuildTool/Platform/Windows/VCToolChain.cs b/Engine/Source/Programs/UnrealBuildTool/Platform/Windows/VCToolChain.cs
--- a/Engine/Source/Programs/UnrealBuildTool/Platform/Windows/VCToolChain.cs
+++ b/Engine/Source/Programs/UnrealBuildTool/Platform/Windows/VCToolChain.cs
@@ -820,6 +820,7 @@ namespace UnrealBuildTool
 			if (LinkEnvironment.bUseIncrementalLinking)
 			{
 				Arguments.Add("/INCREMENTAL");
+				Arguments.Add("/verbose:incr");
 			}
 			else
 			{
