UBT: Add native support for parsing command line arguments inside Target.cs files, via the [CommandLine] attribute on fields.

diff --git a/Engine/Source/Programs/AutomationTool/AutomationUtils/ProjectUtils.cs b/Engine/Source/Programs/AutomationTool/AutomationUtils/ProjectUtils.cs
--- a/Engine/Source/Programs/AutomationTool/AutomationUtils/ProjectUtils.cs
+++ b/Engine/Source/Programs/AutomationTool/AutomationUtils/ProjectUtils.cs
@@ -561,7 +561,7 @@ namespace AutomationTool
 				{
 					string TargetName = GetTargetName(TargetType);
 
-					TargetInfo DummyTargetInfo = new TargetInfo(TargetName, BuildHostPlatform.Current.Platform, UnrealTargetConfiguration.Development, "", Properties.RawProjectPath);
+					TargetInfo DummyTargetInfo = new TargetInfo(TargetName, BuildHostPlatform.Current.Platform, UnrealTargetConfiguration.Development, "", Properties.RawProjectPath, null);
 
 					// Create an instance of this type
 					CommandUtils.LogVerbose("Creating target rules object: {0}", TargetType.Name);
diff --git a/Engine/Source/Programs/UnrealBuildTool/Configuration/TargetInfo.cs b/Engine/Source/Programs/UnrealBuildTool/Configuration/TargetInfo.cs
--- a/Engine/Source/Programs/UnrealBuildTool/Configuration/TargetInfo.cs
+++ b/Engine/Source/Programs/UnrealBuildTool/Configuration/TargetInfo.cs
@@ -48,6 +48,11 @@ namespace UnrealBuildTool
 			get { return ReadOnlyBuildVersion.Current; }
 		}
 
+		/// <summary>
+		/// Additional command line arguments for this target
+		/// </summary>
+		public CommandLineArguments Arguments;
+
 		/// <summary>
 		/// Constructs a TargetInfo for passing to the TargetRules constructor.
 		/// </summary>
@@ -56,13 +61,15 @@ namespace UnrealBuildTool
 		/// <param name="Configuration">The configuration being built</param>
 		/// <param name="Architecture">The architecture being built for</param>
 		/// <param name="ProjectFile">Path to the project file containing the target</param>
-		public TargetInfo(string Name, UnrealTargetPlatform Platform, UnrealTargetConfiguration Configuration, string Architecture, FileReference ProjectFile)
+		/// <param name="Arguments">Additional command line arguments for this target</param>
+		public TargetInfo(string Name, UnrealTargetPlatform Platform, UnrealTargetConfiguration Configuration, string Architecture, FileReference ProjectFile, CommandLineArguments Arguments)
 		{
 			this.Name = Name;
 			this.Platform = Platform;
 			this.Configuration = Configuration;
 			this.Architecture = Architecture;
 			this.ProjectFile = ProjectFile;
+			this.Arguments = Arguments;
 		}
 	}
 }
diff --git a/Engine/Source/Programs/UnrealBuildTool/Configuration/TargetRules.cs b/Engine/Source/Programs/UnrealBuildTool/Configuration/TargetRules.cs
--- a/Engine/Source/Programs/UnrealBuildTool/Configuration/TargetRules.cs
+++ b/Engine/Source/Programs/UnrealBuildTool/Configuration/TargetRules.cs
@@ -1418,6 +1418,10 @@ namespace UnrealBuildTool
 			{
 				ConfigCache.ReadSettings(DirectoryReference.FromFile(ProjectFile), Platform, ConfigurableObject, ConfigValueTracker);
 				XmlConfig.ApplyTo(ConfigurableObject);
+				if (Target.Arguments != null)
+				{
+					Target.Arguments.ApplyTo(ConfigurableObject);
+				}
 			}
 
 			// If we've got a changelist set, set that we're making a formal build
diff --git a/Engine/Source/Programs/UnrealBuildTool/System/RulesAssembly.cs b/Engine/Source/Programs/UnrealBuildTool/System/RulesAssembly.cs
--- a/Engine/Source/Programs/UnrealBuildTool/System/RulesAssembly.cs
+++ b/Engine/Source/Programs/UnrealBuildTool/System/RulesAssembly.cs
@@ -469,9 +469,8 @@ namespace UnrealBuildTool
 		/// </summary>
 		/// <param name="TypeName">Type name of the target rules</param>
 		/// <param name="TargetInfo">Target configuration information to pass to the constructor</param>
-		/// <param name="Arguments">Command line arguments for this target</param>
 		/// <returns>Instance of the corresponding TargetRules</returns>
-		protected TargetRules CreateTargetRulesInstance(string TypeName, TargetInfo TargetInfo, CommandLineArguments Arguments)
+		protected TargetRules CreateTargetRulesInstance(string TypeName, TargetInfo TargetInfo)
 		{
 			// The build module must define a type named '<TargetName>Target' that derives from our 'TargetRules' type.  
 			Type RulesType = CompiledAssembly.GetType(TypeName);
@@ -510,15 +509,6 @@ namespace UnrealBuildTool
 			// Set the default overriddes for the configured target type
 			Rules.SetOverridesForTargetType();
 
-			// Parse any additional command-line arguments. These override default settings specified in config files or the .target.cs files.
-			if(Arguments != null)
-			{
-				foreach(object ConfigurableObject in Rules.GetConfigurableObjects())
-				{
-					Arguments.ApplyTo(ConfigurableObject);
-				}
-			}
-
 			// Set the final value for the link type in the target rules
 			if(Rules.LinkType == TargetLinkType.Default)
 			{
@@ -630,7 +620,7 @@ namespace UnrealBuildTool
 			string TargetTypeName = TargetName + "Target";
 
 			// The build module must define a type named '<TargetName>Target' that derives from our 'TargetRules' type.  
-			return CreateTargetRulesInstance(TargetTypeName, new TargetInfo(TargetName, Platform, Configuration, Architecture, ProjectFile), Arguments);
+			return CreateTargetRulesInstance(TargetTypeName, new TargetInfo(TargetName, Platform, Configuration, Architecture, ProjectFile, Arguments));
 		}
 
 		/// <summary>
@@ -648,7 +638,7 @@ namespace UnrealBuildTool
 			List<string> Matches = new List<string>();
 			foreach(KeyValuePair<string, FileReference> TargetPair in TargetNameToTargetFile)
 			{
-				TargetRules Rules = CreateTargetRulesInstance(TargetPair.Key + "Target", new TargetInfo(TargetPair.Key, Platform, Configuration, Architecture, ProjectFile), null);
+				TargetRules Rules = CreateTargetRulesInstance(TargetPair.Key + "Target", new TargetInfo(TargetPair.Key, Platform, Configuration, Architecture, ProjectFile, null));
 				if(Rules.Type == Type)
 				{
 					Matches.Add(TargetPair.Key);
