Allow of a list of prefered windows toolchains, this allows us to reference both the latest toolchains as we version them in AutoSDK (were we version each inplace updated seperetly) while also still allowing for the version number which is used when installed locally.

diff --git a/Engine/Source/Programs/UnrealBuildTool/Platform/Windows/UEBuildWindows.cs b/Engine/Source/Programs/UnrealBuildTool/Platform/Windows/UEBuildWindows.cs
--- a/Engine/Source/Programs/UnrealBuildTool/Platform/Windows/UEBuildWindows.cs
+++ b/Engine/Source/Programs/UnrealBuildTool/Platform/Windows/UEBuildWindows.cs
@@ -558,9 +558,13 @@ namespace UnrealBuildTool
 		static readonly VersionNumber DefaultClangToolChainVersion = VersionNumber.Parse("8.0.0");
 
 		/// <summary>
-		/// The default compiler version to be used, if installed. 
+		/// The compiler toolchains to be used if installed, the first installed in the list will be used.
 		/// </summary>
-		static readonly VersionNumber DefaultVisualStudioToolChainVersion = VersionNumber.Parse("14.16.27023.2"); // VS2017 v15.9.15
+		static readonly VersionNumber[] PreferredVisualStudioToolChainVersion = new VersionNumber[]
+		{
+			VersionNumber.Parse("14.16.27023.2"), // VS2017 v15.9.15
+			VersionNumber.Parse("14.16.27023"), // fallback to VS2017 15.9 toolchain, microsoft updates these in places so for local installs only this version number is present
+		};
 
 		/// <summary>
 		/// The default Windows SDK version to be used, if installed.
@@ -1170,22 +1174,28 @@ namespace UnrealBuildTool
 			}
 			else
 			{
-				VersionNumber DefaultToolChainVersion;
+				VersionNumber[] PotentialToolchains;
 				if(Compiler == WindowsCompiler.Clang)
 				{
-					DefaultToolChainVersion = DefaultClangToolChainVersion;
+					PotentialToolchains = new VersionNumber[] { DefaultClangToolChainVersion};
 				}
 				else
 				{
-					DefaultToolChainVersion = DefaultVisualStudioToolChainVersion;
+					PotentialToolchains = PreferredVisualStudioToolChainVersion;
 				}
 
-				DirectoryReference DefaultToolChainDir;
-				if(ToolChainVersionToDir.TryGetValue(DefaultToolChainVersion, out DefaultToolChainDir) && (Compiler == WindowsCompiler.Clang || Has64BitToolChain(DefaultToolChainDir)))
+				foreach (VersionNumber PreferredToolchain in PotentialToolchains)
 				{
-					ToolChainVersion = DefaultToolChainVersion;
+					DirectoryReference DefaultToolChainDir;
+					if (ToolChainVersionToDir.TryGetValue(PreferredToolchain, out DefaultToolChainDir) && (Compiler == WindowsCompiler.Clang || Has64BitToolChain(DefaultToolChainDir)))
+					{
+						ToolChainVersion = PreferredToolchain;
+						break;
+					}
 				}
-				else if(ToolChainVersionToDir.Count > 0)
+
+				// if we failed to find any of our preferred toolchains we pick the newest (highest version number)
+				if (ToolChainVersion == null && ToolChainVersionToDir.Count > 0)
 				{
 					ToolChainVersion = ToolChainVersionToDir.OrderBy(x => Has64BitToolChain(x.Value)).ThenBy(x => x.Key).Last().Key;
 				}
