UBT: Prevent the rules assembly override for the default build settings version being overridden by calling the constructor.

#rnx

diff --git a/Engine/Source/Programs/UnrealBuildTool/Configuration/TargetRules.cs b/Engine/Source/Programs/UnrealBuildTool/Configuration/TargetRules.cs
--- a/Engine/Source/Programs/UnrealBuildTool/Configuration/TargetRules.cs
+++ b/Engine/Source/Programs/UnrealBuildTool/Configuration/TargetRules.cs
@@ -201,7 +201,12 @@ namespace UnrealBuildTool
 		/// Specifies the engine version to maintain backwards-compatible default build settings with (eg. DefaultSettingsVersion.Release_4_23, DefaultSettingsVersion.Release_4_24). Specify DefaultSettingsVersion.Latest to always
 		/// use defaults for the current engine version, at the risk of introducing build errors while upgrading.
 		/// </summary>
-		public BuildSettingsVersion DefaultBuildSettings = BuildSettingsVersion.V1;
+		public BuildSettingsVersion DefaultBuildSettings
+		{
+			get { return DefaultBuildSettingsPrivate ?? BuildSettingsVersion.V1; }
+			set { DefaultBuildSettingsPrivate = value; }
+		}
+		private BuildSettingsVersion? DefaultBuildSettingsPrivate; // Cannot be initialized inline; potentially overridden before the constructor is called.
 
 		/// <summary>
 		/// Tracks a list of config values read while constructing this target
@@ -1408,7 +1413,7 @@ namespace UnrealBuildTool
 			this.WindowsPlatform = new WindowsTargetRules(this);
 
 			// Read settings from config files
-			foreach(object ConfigurableObject in GetConfigurableObjects())
+			foreach (object ConfigurableObject in GetConfigurableObjects())
 			{
 				ConfigCache.ReadSettings(DirectoryReference.FromFile(ProjectFile), Platform, ConfigurableObject, ConfigValueTracker);
 				XmlConfig.ApplyTo(ConfigurableObject);
