update error message to reference LINUX_MULTIARCH_ROOT since LINUX_ROOT is deprecated

diff --git a/Engine/Source/Programs/UnrealBuildTool/Platform/Linux/LinuxToolChain.cs b/Engine/Source/Programs/UnrealBuildTool/Platform/Linux/LinuxToolChain.cs
--- a/Engine/Source/Programs/UnrealBuildTool/Platform/Linux/LinuxToolChain.cs
+++ b/Engine/Source/Programs/UnrealBuildTool/Platform/Linux/LinuxToolChain.cs
@@ -84,7 +84,7 @@ namespace UnrealBuildTool
 				// don't register if BaseLinuxPath is not specified and cannot use the system compiler
 				if (String.IsNullOrEmpty(BaseLinuxPath))
 				{
-					throw new BuildException("LINUX_ROOT environment variable is not set; cannot instantiate Linux toolchain");
+					throw new BuildException("LINUX_MULTIARCH_ROOT environment variable is not set; cannot instantiate Linux toolchain");
 				}
 			}
 
