Copying fix from Dev-VR

diff --git a/Engine/Source/Programs/UnrealBuildTool/Platform/HoloLens/HoloLensToolChain.cs b/Engine/Source/Programs/UnrealBuildTool/Platform/HoloLens/HoloLensToolChain.cs
--- a/Engine/Source/Programs/UnrealBuildTool/Platform/HoloLens/HoloLensToolChain.cs
+++ b/Engine/Source/Programs/UnrealBuildTool/Platform/HoloLens/HoloLensToolChain.cs
@@ -53,6 +53,10 @@ namespace UnrealBuildTool
 			// only necessary when AndroidToolChain is used during UAT
 			CommandLine.ParseArguments(Environment.GetCommandLineArgs(), this);
 		}
+		public override void SetEnvironmentVariables()
+		{
+			EnvVars.SetEnvironmentVariables();
+		}
 
 		void AppendCLArguments_Global(CppCompileEnvironment CompileEnvironment, VCEnvironment EnvVars, List<string> Arguments)
 		{
