UBT: Remove hard dependency on Windows SDK assembly for Hololens. Not all users have the required SDK.

diff --git a/Engine/Source/Programs/UnrealBuildTool/Platform/HoloLens/WinMDRegistrationInfo.cs b/Engine/Source/Programs/UnrealBuildTool/Platform/HoloLens/WinMDRegistrationInfo.cs
--- a/Engine/Source/Programs/UnrealBuildTool/Platform/HoloLens/WinMDRegistrationInfo.cs
+++ b/Engine/Source/Programs/UnrealBuildTool/Platform/HoloLens/WinMDRegistrationInfo.cs
@@ -23,23 +23,18 @@ namespace UnrealBuildTool
 			/// WinMD type info constructor
 			/// </summary>
 			/// <param name="InTypeName"></param>
-			/// <param name="InThreadingModel"></param>
-			public ActivatableType(string InTypeName, object InThreadingModel)
+			/// <param name="InThreadingModelName"></param>
+			public ActivatableType(string InTypeName, string InThreadingModelName)
 			{
 				TypeName = InTypeName;
-				if (BuildHostPlatform.Current.Platform == UnrealTargetPlatform.Win64)
-				{
-#if !__MonoCS__ && !MONO
-					ThreadingModelName = ((Windows.Foundation.Metadata.ThreadingModel)InThreadingModel).ToString().ToLowerInvariant();
-#endif
-				}
+				ThreadingModelName = InThreadingModelName;
 			}
-			
+
 			/// <summary>
 			/// Type Name
 			/// </summary>
 			public string TypeName { get; private set; }
-			
+
 			/// <summary>
 			/// Threading model
 			/// </summary>
@@ -59,30 +54,29 @@ namespace UnrealBuildTool
 			if (BuildHostPlatform.Current.Platform == UnrealTargetPlatform.Win64)
 			{
 				ActivatableTypesList = new List<ActivatableType>();
-#if !__MonoCS__ && !MONO
-				var DependsOn = Assembly.ReflectionOnlyLoadFrom(InWindMDSourcePath.FullName);
-				foreach (var WinMDType in DependsOn.GetExportedTypes())
+
+				Assembly DependsOn = Assembly.ReflectionOnlyLoadFrom(InWindMDSourcePath.FullName);
+				foreach (Type WinMDType in DependsOn.GetExportedTypes())
 				{
 					bool IsActivatable = false;
-					object ThreadingModel = Windows.Foundation.Metadata.ThreadingModel.Both;
-					foreach (var Attr in WinMDType.CustomAttributes)
+					string ThreadingModelName = "both";
+					foreach (CustomAttributeData Attr in WinMDType.CustomAttributes)
 					{
-						if (Attr.AttributeType.AssemblyQualifiedName == typeof(Windows.Foundation.Metadata.ActivatableAttribute).AssemblyQualifiedName ||
-							Attr.AttributeType.AssemblyQualifiedName == typeof(Windows.Foundation.Metadata.StaticAttribute).AssemblyQualifiedName)
+						if (Attr.AttributeType.FullName == "Windows.Foundation.Metadata.ActivatableAttribute" || Attr.AttributeType.AssemblyQualifiedName == "Windows.Foundation.Metadata.StaticAttribute")
 						{
 							IsActivatable = true;
 						}
-						else if (Attr.AttributeType.AssemblyQualifiedName == typeof(Windows.Foundation.Metadata.ThreadingAttribute).AssemblyQualifiedName)
+						else if (Attr.AttributeType.FullName == "Windows.Foundation.Metadata.ThreadingAttribute")
 						{
-							ThreadingModel = Attr.ConstructorArguments[0].Value;
+							CustomAttributeTypedArgument Argument = Attr.ConstructorArguments[0];
+							ThreadingModelName = Enum.GetName(Argument.ArgumentType, Argument.Value).ToLowerInvariant();
 						}
 					}
 					if (IsActivatable)
 					{
-						ActivatableTypesList.Add(new ActivatableType(WinMDType.FullName, ThreadingModel));
+						ActivatableTypesList.Add(new ActivatableType(WinMDType.FullName, ThreadingModelName));
 					}
 				}
-#endif
 			}
 		}
 
diff --git a/Engine/Source/Programs/UnrealBuildTool/UnrealBuildTool.csproj b/Engine/Source/Programs/UnrealBuildTool/UnrealBuildTool.csproj
--- a/Engine/Source/Programs/UnrealBuildTool/UnrealBuildTool.csproj
+++ b/Engine/Source/Programs/UnrealBuildTool/UnrealBuildTool.csproj
@@ -314,18 +314,6 @@
         <Windows10SdkPath>$(Registry:HKEY_LOCAL_MACHINE\SOFTWARE\Wow6432Node\Microsoft\Microsoft SDKs\Windows\v10.0@InstallationFolder)</Windows10SdkPath>
         <Windows10SdkLatestVersion>$(Registry:HKEY_LOCAL_MACHINE\SOFTWARE\Wow6432Node\Microsoft\Microsoft SDKs\Windows\v10.0@ProductVersion)</Windows10SdkLatestVersion>
       </PropertyGroup>
-      <ItemGroup>
-        <Reference Include="Windows">
-          <!-- This is a dummy entry to throw off the parser in DotNETUtilities's CsProjectInfo which doesn't understand MSBuild functions -->
-          <HintPath />
-          <HintPath>$(Windows10SdkPath)\UnionMetadata\$(Windows10SdkLatestVersion)\Windows.winmd</HintPath>
-          <!-- 16299 SDK (possibly others?) has ProductVersion == 10.0.16299 in the registry, but pathing needs the final .0 -->
-          <HintPath Condition="'$([System.IO.File]::Exists(%(HintPath)))'=='false'">$(Windows10SdkPath)\UnionMetadata\$(Windows10SdkLatestVersion).0\Windows.winmd</HintPath>
-          <!-- Fall back to old, unversioned location -->
-          <HintPath Condition="'$([System.IO.File]::Exists(%(HintPath)))'=='false'">$(Windows10SdkPath)\UnionMetadata\Windows.winmd</HintPath>
-          <Private>false</Private>
-        </Reference>
-      </ItemGroup>
     </When>
   </Choose>
   <Choose>
