Enable DX12 Perf markers for Test configuration

diff --git a/Engine/Source/Programs/UnrealBuildTool/Platform/Windows/UEBuildWindows.cs b/Engine/Source/Programs/UnrealBuildTool/Platform/Windows/UEBuildWindows.cs
--- a/Engine/Source/Programs/UnrealBuildTool/Platform/Windows/UEBuildWindows.cs
+++ b/Engine/Source/Programs/UnrealBuildTool/Platform/Windows/UEBuildWindows.cs
@@ -1805,7 +1805,7 @@ namespace UnrealBuildTool
 
 			if (ModuleName == "D3D12RHI")
 			{
-				if (Target.WindowsPlatform.bPixProfilingEnabled && Target.Platform == UnrealTargetPlatform.Win64 && Target.Configuration != UnrealTargetConfiguration.Shipping && Target.Configuration != UnrealTargetConfiguration.Test)
+				if (Target.WindowsPlatform.bPixProfilingEnabled && Target.Platform == UnrealTargetPlatform.Win64 && Target.Configuration != UnrealTargetConfiguration.Shipping)
 				{
 					// Define to indicate profiling enabled (64-bit only)
 					Rules.PublicDefinitions.Add("D3D12_PROFILING_ENABLED=1");
diff --git a/Engine/Source/Runtime/D3D12RHI/Private/D3D12RHIPrivate.h b/Engine/Source/Runtime/D3D12RHI/Private/D3D12RHIPrivate.h
--- a/Engine/Source/Runtime/D3D12RHI/Private/D3D12RHIPrivate.h
+++ b/Engine/Source/Runtime/D3D12RHI/Private/D3D12RHIPrivate.h
@@ -27,10 +27,6 @@
 
 #define D3D12_RHI_RAYTRACING (RHI_RAYTRACING)
 
-#ifndef WITH_DX_PERF
-#define WITH_DX_PERF 0
-#endif
-
 // Dependencies.
 #include "CoreMinimal.h"
 #include "RHI.h"
diff --git a/Engine/Source/Runtime/D3D12RHI/Private/D3D12Stats.cpp b/Engine/Source/Runtime/D3D12RHI/Private/D3D12Stats.cpp
--- a/Engine/Source/Runtime/D3D12RHI/Private/D3D12Stats.cpp
+++ b/Engine/Source/Runtime/D3D12RHI/Private/D3D12Stats.cpp
@@ -191,10 +191,6 @@ void D3D12RHI::FD3DGPUProfiler::EndFrame(FD3D12DynamicRHI* InRHI)
 
 void D3D12RHI::FD3DGPUProfiler::PushEvent(const TCHAR* Name, FColor Color)
 {
-#if WITH_DX_PERF
-	D3DPERF_BeginEvent(Color.DWColor(), Name);
-#endif
-
 	FGPUProfiler::PushEvent(Name, Color);
 }
 
@@ -238,10 +234,6 @@ void D3D12RHI::FD3DGPUProfiler::PushEvent(const TCHAR* Name, FColor Color, GFSDK
 
 void D3D12RHI::FD3DGPUProfiler::PopEvent()
 {
-#if WITH_DX_PERF
-	D3DPERF_EndEvent();
-#endif
-
 #if NV_AFTERMATH
 	if (GDX12NVAfterMathEnabled && bTrackingGPUCrashData)
 	{
