Remove -glldb from LinuxToolchain.cs breaks gdb 8 debugging info

diff --git a/Engine/Source/Programs/UnrealBuildTool/Platform/Linux/LinuxToolChain.cs b/Engine/Source/Programs/UnrealBuildTool/Platform/Linux/LinuxToolChain.cs
--- a/Engine/Source/Programs/UnrealBuildTool/Platform/Linux/LinuxToolChain.cs
+++ b/Engine/Source/Programs/UnrealBuildTool/Platform/Linux/LinuxToolChain.cs
@@ -727,9 +727,6 @@ namespace UnrealBuildTool
 			{
 				// libdwarf (from elftoolchain 0.6.1) doesn't support DWARF4. If we need to go back to depending on elftoolchain revert this back to dwarf-3
 				Result += " -gdwarf-4";
-
-				// Make debug info LLDB friendly
-				Result += " -glldb";
 			}
 
 			// optimization level
