UBT: Fix the bHasProjectScriptPlugin flag not being set correctly when building UHT for program targets.

#rnx

diff --git a/Engine/Source/Programs/UnrealBuildTool/Configuration/UEBuildTarget.cs b/Engine/Source/Programs/UnrealBuildTool/Configuration/UEBuildTarget.cs
--- a/Engine/Source/Programs/UnrealBuildTool/Configuration/UEBuildTarget.cs
+++ b/Engine/Source/Programs/UnrealBuildTool/Configuration/UEBuildTarget.cs
@@ -3047,9 +3047,6 @@ namespace UnrealBuildTool
 
 			// Set the list of plugins that should be built
 			BuildPlugins = new List<UEBuildPlugin>(NameToInstance.Values);
-
-			// Determine if the project has a script plugin. We will always build UHT if there is a script plugin in the game folder.
-			bHasProjectScriptPlugin = EnabledPlugins.Any(x => x.Descriptor.SupportedPrograms != null && x.Descriptor.SupportedPrograms.Contains("UnrealHeaderTool"));
 		}
 
 		/// <summary>
@@ -3130,6 +3127,12 @@ namespace UnrealBuildTool
 					return null;
 				}
 
+				// If this plugin supports UHT, flag that we need to pass the project file to UHT when generating headers
+				if (Info.Descriptor.SupportedPrograms != null && Info.Descriptor.SupportedPrograms.Contains("UnrealHeaderTool"))
+				{
+					bHasProjectScriptPlugin = true;
+				}
+
 				// Disable any plugins that aren't compatible with this program
 				if (Rules.Type == TargetType.Program && (Info.Descriptor.SupportedPrograms == null || !Info.Descriptor.SupportedPrograms.Contains(AppName)))
 				{
