Add a flag in TargetRules to optionally disable the generation of the runtime symbol files on some platforms.

diff --git a/Engine/Source/Programs/UnrealBuildTool/Configuration/TargetRules.cs b/Engine/Source/Programs/UnrealBuildTool/Configuration/TargetRules.cs
--- a/Engine/Source/Programs/UnrealBuildTool/Configuration/TargetRules.cs
+++ b/Engine/Source/Programs/UnrealBuildTool/Configuration/TargetRules.cs
@@ -1051,6 +1051,13 @@ namespace UnrealBuildTool
 		[XmlConfigFile(Category = "BuildConfiguration")]
 		public bool bCreateMapFile = false;
 
+		/// <summary>
+		/// True if runtime symbols files should be generated as a post build step for some platforms.
+		/// These files are used by the engine to resolve symbol names of callstack backtraces in logs.
+		/// </summary>
+		[XmlConfigFile(Category = "BuildConfiguration")]
+		public bool bAllowRuntimeSymbolFiles = true;
+
 		/// <summary>
 		/// Bundle version for Mac apps.
 		/// </summary>
@@ -2339,6 +2346,11 @@ namespace UnrealBuildTool
 			get { return Inner.bCreateMapFile; }
 		}
 
+		public bool bAllowRuntimeSymbolFiles
+		{
+			get { return Inner.bAllowRuntimeSymbolFiles; }
+		}
+
 		public string BundleVersion
 		{
 			get { return Inner.BundleVersion; }
diff --git a/Engine/Source/Programs/UnrealBuildTool/System/LinkEnvironment.cs b/Engine/Source/Programs/UnrealBuildTool/System/LinkEnvironment.cs
--- a/Engine/Source/Programs/UnrealBuildTool/System/LinkEnvironment.cs
+++ b/Engine/Source/Programs/UnrealBuildTool/System/LinkEnvironment.cs
@@ -120,6 +120,12 @@ namespace UnrealBuildTool
 		/// </summary>
 		public bool bCreateDebugInfo = true;
 
+		/// <summary>
+		/// True if runtime symbols files should be generated as a post build step for some platforms.
+		/// These files are used by the engine to resolve symbol names of callstack backtraces in logs.
+		/// </summary>
+		public bool bGenerateRuntimeSymbolFiles = true;
+
 		/// <summary>
 		/// True if debug symbols that are cached for some platforms should not be created.
 		/// </summary>
@@ -329,6 +335,7 @@ namespace UnrealBuildTool
 			DelayLoadDLLs.AddRange(Other.DelayLoadDLLs);
 			AdditionalArguments = Other.AdditionalArguments;
 			bCreateDebugInfo = Other.bCreateDebugInfo;
+			bGenerateRuntimeSymbolFiles = Other.bGenerateRuntimeSymbolFiles;
 			bIsBuildingLibrary = Other.bIsBuildingLibrary;
             bDisableSymbolCache = Other.bDisableSymbolCache;
 			bIsBuildingDLL = Other.bIsBuildingDLL;
