#DEV-VR - Removing warnings on LUMIN from trying to load missing modules

diff --git a/Engine/Config/Lumin/LuminEngine.ini b/Engine/Config/Lumin/LuminEngine.ini
--- a/Engine/Config/Lumin/LuminEngine.ini
+++ b/Engine/Config/Lumin/LuminEngine.ini
@@ -5,3 +5,11 @@ AudioDeviceModuleName=MagicLeapAudio
 AudioMixerModuleName=MagicLeapAudio
 ; Defines a platform-specific volume headroom (in dB) for audio to provide better platform consistency with respect to volume levels.
 PlatformHeadroomDB=-4
+
+;Override Android settings back to default
+[LocalNotification]
+DefaultPlatformService=
+
+[OnlineSubsystem]
+DefaultPlatformService=
+NativePlatformService=
diff --git a/Engine/Source/Programs/UnrealBuildTool/Platform/Lumin/UEBuildLumin.cs b/Engine/Source/Programs/UnrealBuildTool/Platform/Lumin/UEBuildLumin.cs
--- a/Engine/Source/Programs/UnrealBuildTool/Platform/Lumin/UEBuildLumin.cs
+++ b/Engine/Source/Programs/UnrealBuildTool/Platform/Lumin/UEBuildLumin.cs
@@ -157,6 +157,7 @@ namespace UnrealBuildTool
 		{
 			CompileEnvironment.Definitions.Add("PLATFORM_LUMIN=1");
 			CompileEnvironment.Definitions.Add("USE_ANDROID_JNI=0");
+			CompileEnvironment.Definitions.Add("USE_ANDROID_AUDIO=0");
 			CompileEnvironment.Definitions.Add("USE_ANDROID_FILE=0");
 			CompileEnvironment.Definitions.Add("USE_ANDROID_INPUT=0");
 			CompileEnvironment.Definitions.Add("USE_ANDROID_LAUNCH=0");
diff --git a/Engine/Source/Runtime/ApplicationCore/Private/Android/AndroidApplicationMisc.cpp b/Engine/Source/Runtime/ApplicationCore/Private/Android/AndroidApplicationMisc.cpp
--- a/Engine/Source/Runtime/ApplicationCore/Private/Android/AndroidApplicationMisc.cpp
+++ b/Engine/Source/Runtime/ApplicationCore/Private/Android/AndroidApplicationMisc.cpp
@@ -13,8 +13,10 @@
 void FAndroidApplicationMisc::LoadPreInitModules()
 {
 	FModuleManager::Get().LoadModule(TEXT("OpenGLDrv"));
+#if USE_ANDROID_AUDIO
 	FModuleManager::Get().LoadModule(TEXT("AndroidAudio"));
 	FModuleManager::Get().LoadModule(TEXT("AudioMixerAndroid"));
+#endif
 }
 
 class FOutputDeviceError* FAndroidApplicationMisc::GetErrorOutputDevice()
diff --git a/Engine/Source/Runtime/Core/Public/Android/AndroidPlatform.h b/Engine/Source/Runtime/Core/Public/Android/AndroidPlatform.h
--- a/Engine/Source/Runtime/Core/Public/Android/AndroidPlatform.h
+++ b/Engine/Source/Runtime/Core/Public/Android/AndroidPlatform.h
@@ -78,6 +78,9 @@ typedef FAndroidTypes FPlatformTypes;
 #ifndef USE_ANDROID_JNI
 	#define USE_ANDROID_JNI							1
 #endif
+#ifndef USE_ANDROID_AUDIO
+	#define USE_ANDROID_AUDIO						1
+#endif
 #ifndef USE_ANDROID_FILE
 	#define USE_ANDROID_FILE						1
 #endif
diff --git a/Engine/Source/Runtime/Launch/Launch.Build.cs b/Engine/Source/Runtime/Launch/Launch.Build.cs
--- a/Engine/Source/Runtime/Launch/Launch.Build.cs
+++ b/Engine/Source/Runtime/Launch/Launch.Build.cs
@@ -103,6 +103,7 @@ public class Launch : ModuleRules
 					"StreamingFile",
 					"CookedIterativeFile",
 					"AutomationWorker",
+					"FunctionalTesting"
 			});
 		}
 
