Merged in PR #4296 to fix CMake errors when using CMake on Windows and with the Steam SDK installed

diff --git a/Engine/Source/Programs/DotNETCommon/DotNETUtilities/FileSystemReference.cs b/Engine/Source/Programs/DotNETCommon/DotNETUtilities/FileSystemReference.cs
--- a/Engine/Source/Programs/DotNETCommon/DotNETUtilities/FileSystemReference.cs
+++ b/Engine/Source/Programs/DotNETCommon/DotNETUtilities/FileSystemReference.cs
@@ -314,6 +314,11 @@ namespace Tools.DotNETCommon
 			return Result.ToString();
 		}
 
+		public string ToNormalizedPath()
+		{
+			return FullName.Replace("\\", "/");
+		}
+
 		/// <summary>
 		/// Returns a string representation of this filesystem object
 		/// </summary>
diff --git a/Engine/Source/Programs/UnrealBuildTool/ProjectFiles/CMake/CMakefileGenerator.cs b/Engine/Source/Programs/UnrealBuildTool/ProjectFiles/CMake/CMakefileGenerator.cs
--- a/Engine/Source/Programs/UnrealBuildTool/ProjectFiles/CMake/CMakefileGenerator.cs
+++ b/Engine/Source/Programs/UnrealBuildTool/ProjectFiles/CMake/CMakefileGenerator.cs
@@ -326,19 +326,19 @@ namespace UnrealBuildTool
 			}
 
 			// Additional CMake file definitions
-			string EngineHeadersFilePath = FileReference.Combine(IntermediateProjectFilesPath, CMakeEngineHeadersFileName).ToString();
-			string ProjectHeadersFilePath = FileReference.Combine(IntermediateProjectFilesPath, CMakeProjectHeadersFileName).ToString();
-			string EngineSourcesFilePath = FileReference.Combine(IntermediateProjectFilesPath, CMakeEngineSourcesFileName).ToString();
-			string ProjectSourcesFilePath = FileReference.Combine(IntermediateProjectFilesPath, CMakeProjectSourcesFileName).ToString();
-			string ProjectFilePath = FileReference.Combine(IntermediateProjectFilesPath, CMakeProjectSourcesFileName).ToString();
-			string IncludeFilePath = FileReference.Combine(IntermediateProjectFilesPath, CMakeIncludesFileName).ToString();
-			string EngineConfigsFilePath = FileReference.Combine(IntermediateProjectFilesPath, CMakeEngineConfigsFileName).ToString();
-			string ProjectConfigsFilePath = FileReference.Combine(IntermediateProjectFilesPath, CMakeProjectConfigsFileName).ToString();
-			string EngineCSFilePath = FileReference.Combine(IntermediateProjectFilesPath, CMakeEngineCSFileName).ToString();
-			string ProjectCSFilePath = FileReference.Combine(IntermediateProjectFilesPath, CMakeProjectCSFileName).ToString();
-			string EngineShadersFilePath = FileReference.Combine(IntermediateProjectFilesPath, CMakeEngineShadersFileName).ToString();
-			string ProjectShadersFilePath = FileReference.Combine(IntermediateProjectFilesPath, CMakeProjectShadersFileName).ToString();
-			string DefinitionsFilePath = FileReference.Combine(IntermediateProjectFilesPath, CMakeDefinitionsFileName).ToString();
+			string EngineHeadersFilePath = FileReference.Combine(IntermediateProjectFilesPath, CMakeEngineHeadersFileName).ToNormalizedPath();
+			string ProjectHeadersFilePath = FileReference.Combine(IntermediateProjectFilesPath, CMakeProjectHeadersFileName).ToNormalizedPath();
+			string EngineSourcesFilePath = FileReference.Combine(IntermediateProjectFilesPath, CMakeEngineSourcesFileName).ToNormalizedPath();
+			string ProjectSourcesFilePath = FileReference.Combine(IntermediateProjectFilesPath, CMakeProjectSourcesFileName).ToNormalizedPath();
+			string ProjectFilePath = FileReference.Combine(IntermediateProjectFilesPath, CMakeProjectSourcesFileName).ToNormalizedPath();
+			string IncludeFilePath = FileReference.Combine(IntermediateProjectFilesPath, CMakeIncludesFileName).ToNormalizedPath();
+			string EngineConfigsFilePath = FileReference.Combine(IntermediateProjectFilesPath, CMakeEngineConfigsFileName).ToNormalizedPath();
+			string ProjectConfigsFilePath = FileReference.Combine(IntermediateProjectFilesPath, CMakeProjectConfigsFileName).ToNormalizedPath();
+			string EngineCSFilePath = FileReference.Combine(IntermediateProjectFilesPath, CMakeEngineCSFileName).ToNormalizedPath();
+			string ProjectCSFilePath = FileReference.Combine(IntermediateProjectFilesPath, CMakeProjectCSFileName).ToNormalizedPath();
+			string EngineShadersFilePath = FileReference.Combine(IntermediateProjectFilesPath, CMakeEngineShadersFileName).ToNormalizedPath();
+			string ProjectShadersFilePath = FileReference.Combine(IntermediateProjectFilesPath, CMakeProjectShadersFileName).ToNormalizedPath();
+			string DefinitionsFilePath = FileReference.Combine(IntermediateProjectFilesPath, CMakeDefinitionsFileName).ToNormalizedPath();
 
 			CMakefileContent.Append(
 				"# Makefile generated by CMakefileGenerator.cs (v1.2)\n" +
@@ -454,7 +454,17 @@ namespace UnrealBuildTool
 
 			foreach (string PreprocessorDefinition in PreprocessorDefinitions)
 			{
-				PreprocessorDefinitionsList.Append("\t-D" + PreprocessorDefinition + "\n");
+				int EqPos = PreprocessorDefinition.IndexOf("=");
+				if (EqPos >= 0)
+				{
+					string Key = PreprocessorDefinition.Substring(0, EqPos);
+					string Value = PreprocessorDefinition.Substring(EqPos).Replace("\"", "\\\"");
+					PreprocessorDefinitionsList.Append("\t\"-D" + Key + Value + "\"\n");
+				}
+				else
+				{
+					PreprocessorDefinitionsList.Append("\t\"-D" + PreprocessorDefinition + "\"\n");
+				}
 			}
 
 			// Add Engine/Shaders files (game are added via modules)
