Moved CLion specific functions from the general cmake project file to the CLion specific generator
Changed output of the ignore xml to be utf8

diff --git a/Engine/Source/Programs/UnrealBuildTool/ProjectFiles/CLion/CLionGenerator.cs b/Engine/Source/Programs/UnrealBuildTool/ProjectFiles/CLion/CLionGenerator.cs
--- a/Engine/Source/Programs/UnrealBuildTool/ProjectFiles/CLion/CLionGenerator.cs
+++ b/Engine/Source/Programs/UnrealBuildTool/ProjectFiles/CLion/CLionGenerator.cs
@@ -20,5 +20,51 @@ namespace UnrealBuildTool
 			: base(InOnlyGameProject)
 		{
 		}
+
+		/// <summary>
+		/// Writes a file that CLion uses to know what directories to exclude from indexing. This should speed up indexing
+		/// </summary>
+		protected void WriteCLionIgnoreDirs()
+		{
+			string CLionIngoreXml =
+				"<?xml version=\"1.0\" encoding=\"UTF-8\"?>" + Environment.NewLine +
+				"<project version=\"4\">" + Environment.NewLine +
+				  "\t<component name=\"CMakeWorkspace\" PROJECT_DIR=\"$PROJECT_DIR$\" />" + Environment.NewLine +
+				  "\t<component name=\"CidrRootsConfiguration\">" + Environment.NewLine +
+					"\t\t<excludeRoots>" + Environment.NewLine +
+					  "\t\t\t<file path=\"$PROJECT_DIR$/Binaries\" />" + Environment.NewLine +
+					  "\t\t\t<file path=\"$PROJECT_DIR$/Build\" />" + Environment.NewLine +
+					  "\t\t\t<file path=\"$PROJECT_DIR$/Content\" />" + Environment.NewLine +
+					  "\t\t\t<file path=\"$PROJECT_DIR$/DataTables\" />" + Environment.NewLine +
+					  "\t\t\t<file path=\"$PROJECT_DIR$/DerivedDataCache\" />" + Environment.NewLine +
+					  "\t\t\t<file path=\"$PROJECT_DIR$/FeaturePacks\" />" + Environment.NewLine +
+					  "\t\t\t<file path=\"$PROJECT_DIR$/Intermediate\" />" + Environment.NewLine +
+					  "\t\t\t<file path=\"$PROJECT_DIR$/LocalBuilds\" />" + Environment.NewLine +
+					  "\t\t\t<file path=\"$PROJECT_DIR$/Samples\" />" + Environment.NewLine +
+					  "\t\t\t<file path=\"$PROJECT_DIR$/Saved\" />" + Environment.NewLine +
+					  "\t\t\t<file path=\"$PROJECT_DIR$/Templates\" />" + Environment.NewLine +
+					  "" + Environment.NewLine +
+					  "\t\t\t<file path=\"$PROJECT_DIR$/Engine/Binaries\" />" + Environment.NewLine +
+					  "\t\t\t<file path=\"$PROJECT_DIR$/Engine/Build\" />" + Environment.NewLine +
+					  "\t\t\t<file path=\"$PROJECT_DIR$/Engine/Content\" />" + Environment.NewLine +
+					  "\t\t\t<file path=\"$PROJECT_DIR$/Engine/DerivedDataCache\" />" + Environment.NewLine +
+					  "\t\t\t<file path=\"$PROJECT_DIR$/Engine/Documentation\" />" + Environment.NewLine +
+					  "\t\t\t<file path=\"$PROJECT_DIR$/Engine/Extras\" />" + Environment.NewLine +
+					  "\t\t\t<file path=\"$PROJECT_DIR$/Engine/Intermediate\" />" + Environment.NewLine +
+					  "\t\t\t<file path=\"$PROJECT_DIR$/Engine/Programs\" />" + Environment.NewLine +
+					  "\t\t\t<file path=\"$PROJECT_DIR$/Engine/Saved\" />" + Environment.NewLine +
+					"\t\t</excludeRoots>" + Environment.NewLine +
+				  "\t</component>" + Environment.NewLine +
+				"</project>" + Environment.NewLine;
+
+			string FullFileName = Path.Combine(MasterProjectPath.FullName, ".idea/misc.xml");
+			WriteFileIfChanged(FullFileName, CLionIngoreXml, new UTF8Encoding());
+		}
+
+        protected override bool WriteProjectFiles()
+        {
+            WriteCLionIgnoreDirs();
+            return base.WriteProjectFiles();
+        }
 	}
 }
diff --git a/Engine/Source/Programs/UnrealBuildTool/ProjectFiles/CMake/CMakefileGenerator.cs b/Engine/Source/Programs/UnrealBuildTool/ProjectFiles/CMake/CMakefileGenerator.cs
--- a/Engine/Source/Programs/UnrealBuildTool/ProjectFiles/CMake/CMakefileGenerator.cs
+++ b/Engine/Source/Programs/UnrealBuildTool/ProjectFiles/CMake/CMakefileGenerator.cs
@@ -748,54 +748,10 @@ namespace UnrealBuildTool
 			return FullPath;
 		}
 
-		/// <summary>
-		/// Writes a file that CLion uses to know what directories to exclude from indexing. This should speed up indexing
-		/// </summary>
-		protected void WriteCLionIgnoreDirs()
-		{
-			string CLionIngoreXml =
-				"<?xml version=\"1.0\" encoding=\"UTF-8\"?>" + Environment.NewLine +
-				"<project version=\"4\">" + Environment.NewLine +
-				  "\t<component name=\"CMakeWorkspace\" PROJECT_DIR=\"$PROJECT_DIR$\" />" + Environment.NewLine +
-				  "\t<component name=\"CidrRootsConfiguration\">" + Environment.NewLine +
-					"\t\t<excludeRoots>" + Environment.NewLine +
-					  "\t\t\t<file path=\"$PROJECT_DIR$/Binaries\" />" + Environment.NewLine +
-					  "\t\t\t<file path=\"$PROJECT_DIR$/Build\" />" + Environment.NewLine +
-					  "\t\t\t<file path=\"$PROJECT_DIR$/Content\" />" + Environment.NewLine +
-					  "\t\t\t<file path=\"$PROJECT_DIR$/DataTables\" />" + Environment.NewLine +
-					  "\t\t\t<file path=\"$PROJECT_DIR$/DerivedDataCache\" />" + Environment.NewLine +
-					  "\t\t\t<file path=\"$PROJECT_DIR$/FeaturePacks\" />" + Environment.NewLine +
-					  "\t\t\t<file path=\"$PROJECT_DIR$/Intermediate\" />" + Environment.NewLine +
-					  "\t\t\t<file path=\"$PROJECT_DIR$/LocalBuilds\" />" + Environment.NewLine +
-					  "\t\t\t<file path=\"$PROJECT_DIR$/Samples\" />" + Environment.NewLine +
-					  "\t\t\t<file path=\"$PROJECT_DIR$/Saved\" />" + Environment.NewLine +
-					  "\t\t\t<file path=\"$PROJECT_DIR$/Templates\" />" + Environment.NewLine +
-					  "" + Environment.NewLine +
-					  "\t\t\t<file path=\"$PROJECT_DIR$/Engine/Binaries\" />" + Environment.NewLine +
-					  "\t\t\t<file path=\"$PROJECT_DIR$/Engine/Build\" />" + Environment.NewLine +
-					  "\t\t\t<file path=\"$PROJECT_DIR$/Engine/Content\" />" + Environment.NewLine +
-					  "\t\t\t<file path=\"$PROJECT_DIR$/Engine/DerivedDataCache\" />" + Environment.NewLine +
-					  "\t\t\t<file path=\"$PROJECT_DIR$/Engine/Documentation\" />" + Environment.NewLine +
-					  "\t\t\t<file path=\"$PROJECT_DIR$/Engine/Extras\" />" + Environment.NewLine +
-					  "\t\t\t<file path=\"$PROJECT_DIR$/Engine/Intermediate\" />" + Environment.NewLine +
-					  "\t\t\t<file path=\"$PROJECT_DIR$/Engine/Programs\" />" + Environment.NewLine +
-					  "\t\t\t<file path=\"$PROJECT_DIR$/Engine/Saved\" />" + Environment.NewLine +
-					"\t\t</excludeRoots>" + Environment.NewLine +
-				  "\t</component>" + Environment.NewLine +
-				"</project>" + Environment.NewLine;
-
-			string FullFileName = Path.Combine(MasterProjectPath.FullName, ".idea/misc.xml");
-			WriteFileIfChanged(FullFileName, CLionIngoreXml);
-		}
-
 		#region ProjectFileGenerator implementation
 
 		protected override bool WriteProjectFiles()
 		{
-			if (BuildHostPlatform.Current.Platform != UnrealTargetPlatform.Win64)
-			{
-				WriteCLionIgnoreDirs();
-			}
 			return WriteCMakeLists();
 		}
 
