Updated ML integration to MLSDK 0.19

diff --git a/Engine/Source/Programs/UnrealBuildTool/Platform/Lumin/UEDeployLumin.cs b/Engine/Source/Programs/UnrealBuildTool/Platform/Lumin/UEDeployLumin.cs
--- a/Engine/Source/Programs/UnrealBuildTool/Platform/Lumin/UEDeployLumin.cs
+++ b/Engine/Source/Programs/UnrealBuildTool/Platform/Lumin/UEDeployLumin.cs
@@ -77,14 +77,23 @@ namespace UnrealBuildTool
 			return ApplicationDisplayName;
 		}
 
-        private string GetMinimumOSVersionRequired()
-        {
-            string MinOSVersionRequired = GetRuntimeSetting("MinimumOSVersion");
-            return String.IsNullOrWhiteSpace(MinOSVersionRequired) ? "nova_1.0" : MinOSVersionRequired;
-        }
-
-        private void GetAppPrivileges(ConfigHierarchy EngineIni, StringBuilder Text)
-        {
+		private string GetMinimumAPILevelRequired()
+		{
+			const Int32 AbsoluteMinValue = 2;
+
+			Int32 Value = AbsoluteMinValue;
+			GetConfigCacheIni(ConfigHierarchyType.Engine).
+				GetInt32("/Script/LuminRuntimeSettings.LuminRuntimeSettings", "MinimumAPILevel", out Value);
+			if (Value < AbsoluteMinValue)
+			{
+				Log.TraceInformation("Config-specified MinimumAPILevel {0} is lower than engine's minimum {1}", Value, AbsoluteMinValue);
+				Value = AbsoluteMinValue;
+			}
+			return Value.ToString();
+		}
+
+		private void GetAppPrivileges(ConfigHierarchy EngineIni, StringBuilder Text)
+		{
 			List<string> AppPrivileges;
 			EngineIni.GetArray("/Script/LuminRuntimeSettings.LuminRuntimeSettings", "AppPrivileges", out AppPrivileges);
 			if (AppPrivileges != null)
@@ -102,7 +111,7 @@ namespace UnrealBuildTool
 					}
 				}
 			}
-        }
+		}
 
 		private string CleanFilePath(string FilePath)
 		{
@@ -124,44 +133,35 @@ namespace UnrealBuildTool
 			return "Icon/Portal";
 		}
 
-        public string GetMLSDKVersion(ConfigHierarchy EngineIni)
-        {
-            string MLSDKPath;
-            string Major = "0";
-            string Minor = "0";
-            if (EngineIni.GetString("/Script/LuminPlatformEditor.MagicLeapSDKSettings", "MLSDKPath", out MLSDKPath) && !string.IsNullOrWhiteSpace(MLSDKPath))
-            {
-                Dictionary<string, string> PathEntry;
-                ConfigHierarchy.TryParse(MLSDKPath, out PathEntry);
-                MLSDKPath = PathEntry.Values.First();
-            }
-            if (string.IsNullOrWhiteSpace(MLSDKPath))
-            {
-                MLSDKPath = Environment.GetEnvironmentVariable("MLSDK");
-            }
-            if (!string.IsNullOrEmpty(MLSDKPath))
-            {
-                if (Directory.Exists(MLSDKPath))
-                {
-                    String VersionFile = string.Format("{0}/include/ml_version.h", MLSDKPath).Replace('/', Path.DirectorySeparatorChar);
-                    if (File.Exists(VersionFile))
-                    {
-                        string FileText = File.ReadAllText(VersionFile);
-                        string Pattern = @"(MLSDK_VERSION_MAJOR) (?'MAJOR'\d+).*(MLSDK_VERSION_MINOR) (?'MINOR'\d+).*(MLSDK_VERSION_REVISION) (?'REV'\d+)";
-                        Regex VersionRegex = new Regex(Pattern, RegexOptions.Singleline);
-                        MatchCollection Matches = VersionRegex.Matches(FileText);
-                        if (Matches.Count > 0 &&
-                            !string.IsNullOrEmpty(Matches[0].Groups["MAJOR"].Value) &&
-                            !string.IsNullOrEmpty(Matches[0].Groups["MINOR"].Value))
-                        {
-                            Major = Matches[0].Groups["MAJOR"].Value;
-                            Minor = Matches[0].Groups["MINOR"].Value;
-                        }
-                    }
-                }
-            }
-            return string.Format("{0}.{1}", Major, Minor);
-        }
+		public string GetMLSDKVersion(ConfigHierarchy EngineIni)
+		{
+			string MLSDKPath;
+			string Major = "0";
+			string Minor = "0";
+			MLSDKPath = Environment.GetEnvironmentVariable("MLSDK");
+			if (!string.IsNullOrEmpty(MLSDKPath))
+			{
+				if (Directory.Exists(MLSDKPath))
+				{
+					String VersionFile = string.Format("{0}/include/ml_version.h", MLSDKPath).Replace('/', Path.DirectorySeparatorChar);
+					if (File.Exists(VersionFile))
+					{
+						string FileText = File.ReadAllText(VersionFile);
+						string Pattern = @"(MLSDK_VERSION_MAJOR) (?'MAJOR'\d+).*(MLSDK_VERSION_MINOR) (?'MINOR'\d+).*(MLSDK_VERSION_REVISION) (?'REV'\d+)";
+						Regex VersionRegex = new Regex(Pattern, RegexOptions.Singleline);
+						MatchCollection Matches = VersionRegex.Matches(FileText);
+						if (Matches.Count > 0 &&
+							!string.IsNullOrEmpty(Matches[0].Groups["MAJOR"].Value) &&
+							!string.IsNullOrEmpty(Matches[0].Groups["MINOR"].Value))
+						{
+							Major = Matches[0].Groups["MAJOR"].Value;
+							Minor = Matches[0].Groups["MINOR"].Value;
+						}
+					}
+				}
+			}
+			return string.Format("{0}.{1}", Major, Minor);
+		}
 
 		public string GenerateManifest(string ProjectName, bool bForDistribution, string Architecture)
 		{
@@ -183,15 +183,15 @@ namespace UnrealBuildTool
 
 			string PackageName = GetPackageName(ProjectName);
 			string ApplicationDisplayName = GetApplicationDisplayName(ProjectName);
-			string MinimumOSVersion = GetMinimumOSVersionRequired();
+			string MinimumAPILevel = GetMinimumAPILevelRequired();
 			string TargetExecutableName = "bin/" + ProjectName;
 
 			Text.AppendLine(string.Format("<manifest xmlns:ml=\"magicleap\" ml:package=\"{0}\" ml:version_name=\"{1}\" ml:version_code=\"{2}\">", PackageName, ProjectVersion, VersionCode));
-			Text.AppendLine(string.Format("\t<application ml:visible_name=\"{0}\" ml:sdk_version=\"{1}\" ml:minimum_os=\"{2}\">",
-                ApplicationDisplayName,
-                SDKVersion,
-                MinimumOSVersion));
-            GetAppPrivileges(EngineIni, Text);
+			Text.AppendLine(string.Format("\t<application ml:visible_name=\"{0}\" ml:sdk_version=\"{1}\" ml:min_api_level=\"{2}\">",
+				ApplicationDisplayName,
+				SDKVersion,
+				MinimumAPILevel));
+			GetAppPrivileges(EngineIni, Text);
 			Text.AppendLine(string.Format("\t\t<component ml:name=\".fullscreen\" ml:visible_name=\"{0}\" ml:binary_name=\"{1}\" ml:type=\"{2}\">", ApplicationDisplayName, TargetExecutableName, GetApplicationType()));
 
 			string IconTag = string.Format("<icon ml:model_folder=\"{0}\" ml:portal_folder=\"{1}\"/>", GetIconModelStagingPath(), GetIconPortalStagingPath());
@@ -296,6 +296,14 @@ namespace UnrealBuildTool
 			return UPL.ProcessPluginNode(Architecture, "stageFiles", "");
 		}
 
+		private bool GetRemoveDebugInfo()
+		{
+			ConfigHierarchy Ini = GetConfigCacheIni(ConfigHierarchyType.Engine);
+			bool Value = false;
+			Ini.GetBool("/Script/LuminRuntimeSettings.LuminRuntimeSettings", "bRemoveDebugInfo", out Value);
+			return Value;
+		}
+
 		private void MakeMabuPackage(string ProjectName, DirectoryReference ProjectDirectory, string ExePath, bool bForDistribution, string EngineDir, string MpkName)
 		{
 			string UE4BuildPath = Path.Combine(ProjectDirectory.FullName, "Intermediate/Lumin/Mabu");
@@ -306,19 +314,27 @@ namespace UnrealBuildTool
 
 
 			LuminToolChain ToolChain = new LuminToolChain(ProjectFile);
-
-			// If asked for, and if we are doing a distribution package, we strip debug symbols.
 			string ExecSrcFile = Path.Combine(UE4BuildPath, "Binaries", Path.GetFileName(ExePath));
-			if (bForDistribution)
+
+			if (bForDistribution || GetRemoveDebugInfo())
 			{
+				// If asked for, and if we are doing a distribution package, we strip debug symbols.
 				Directory.CreateDirectory(Path.GetDirectoryName(ExecSrcFile));
 				ToolChain.StripSymbols(new Tools.DotNETCommon.FileReference(ExePath), new Tools.DotNETCommon.FileReference(ExecSrcFile));
+				// We also create a SYM file to support debugging of stripped executables
+				string SymFile = Path.ChangeExtension(ExecSrcFile, "sym");
+				ToolChain.ExtractSymbols(new FileReference(ExePath), new FileReference(SymFile));
+				ToolChain.LinkSymbols(new FileReference(SymFile), new FileReference(ExecSrcFile));
 			}
 			else
 			{
 				// The generated mabu needs the src exe file. So we copy the original as-is so mabu can find it.
 				Directory.CreateDirectory(Path.GetDirectoryName(ExecSrcFile));
 				File.Copy(ExePath, ExecSrcFile, true);
+				// If the exe has debug info we need to remove a likely old sym file so that debugging doesn't use
+				// outdated information.
+				string SymFile = Path.ChangeExtension(ExecSrcFile, "sym");
+				File.Delete(SymFile);
 			}
 
 			// Generate manifest (after UPL is setup
@@ -330,7 +346,7 @@ namespace UnrealBuildTool
 			Certificate = CleanFilePath(Certificate);
 			if (!string.IsNullOrEmpty(Certificate))
 			{
-				Certificate = Path.GetFullPath(Path.Combine(EngineDir, "Binaries/Lumin", Certificate));
+				Certificate = Path.GetFullPath(Path.Combine(ProjectDirectory.FullName, Certificate));
 
 				if (File.Exists(Certificate))
 				{
diff --git a/Engine/Source/Programs/UnrealBuildTool/ProjectFiles/VisualStudio/VCProject.cs b/Engine/Source/Programs/UnrealBuildTool/ProjectFiles/VisualStudio/VCProject.cs
--- a/Engine/Source/Programs/UnrealBuildTool/ProjectFiles/VisualStudio/VCProject.cs
+++ b/Engine/Source/Programs/UnrealBuildTool/ProjectFiles/VisualStudio/VCProject.cs
@@ -1416,43 +1416,15 @@ namespace UnrealBuildTool
 						VCProjectFileContent.Append(PathStrings);
 					}
 
-					// This is the standard UE4 based project NMake build line:
-					//	..\..\Build\BatchFiles\Build.bat <TARGETNAME> <PLATFORM> <CONFIGURATION>
-					//	ie ..\..\Build\BatchFiles\Build.bat BlankProgram Win64 Debug
-
-					string BuildArguments = " " + TargetName + " " + UBTPlatformName + " " + UBTConfigurationName;
-					if (ProjectFileGenerator.bUsePrecompiled)
-					{
-						BuildArguments += " -UsePrecompiled";
-					}
-					if (IsForeignProject)
-					{
-						BuildArguments += " " + UProjectPath;
-					}
-
-					// Always wait for the mutex between UBT invocations, so that building the whole solution doesn't fail.
-					BuildArguments += " -WaitMutex";
-
-					// Always include a flag to format log messages for MSBuild
-					BuildArguments += " -FromMsBuild";
-
-					if (bUseFastPDB)
-					{
-						// Pass Fast PDB option to make use of Visual Studio's /DEBUG:FASTLINK option
-						BuildArguments += " -FastPDB";
-					}
+					UEPlatformProjectGenerator.ScriptArguments BatchScriptArguments = new UEPlatformProjectGenerator.ScriptArguments(ProjectFileGenerator.bUsePrecompiled, IsForeignProject, bUseFastPDB, TargetName, UBTPlatformName, UBTConfigurationName, BuildToolOverride);
 
 					DirectoryReference BatchFilesDirectory = DirectoryReference.Combine(UnrealBuildTool.EngineDirectory, "Build", "BatchFiles");
 
-					if(BuildToolOverride != null)
-					{
-						BuildArguments += " " + BuildToolOverride;
-					}
-
 					// NMake Build command line
-					VCProjectFileContent.AppendLine("    <NMakeBuildCommandLine>{0}{1}</NMakeBuildCommandLine>", EscapePath(NormalizeProjectPath(FileReference.Combine(BatchFilesDirectory, "Build.bat"))), BuildArguments.ToString());
-					VCProjectFileContent.AppendLine("    <NMakeReBuildCommandLine>{0}{1}</NMakeReBuildCommandLine>", EscapePath(NormalizeProjectPath(FileReference.Combine(BatchFilesDirectory, "Rebuild.bat"))), BuildArguments.ToString());
-					VCProjectFileContent.AppendLine("    <NMakeCleanCommandLine>{0}{1}</NMakeCleanCommandLine>", EscapePath(NormalizeProjectPath(FileReference.Combine(BatchFilesDirectory, "Clean.bat"))), BuildArguments.ToString());
+					VCProjectFileContent.AppendLine("    <NMakeBuildCommandLine>{0}</NMakeBuildCommandLine>", ProjGenerator.GetBuildCommandEntry(this, BatchFilesDirectory, BatchScriptArguments, OnlyGameProject));
+					// NMake Rebuild command line
+					VCProjectFileContent.AppendLine("    <NMakeReBuildCommandLine>{0}</NMakeReBuildCommandLine>", ProjGenerator.GetReBuildCommandEntry(this, BatchFilesDirectory, BatchScriptArguments, OnlyGameProject));
+					VCProjectFileContent.AppendLine("    <NMakeCleanCommandLine>{0}</NMakeCleanCommandLine>", ProjGenerator.GetCleanCommandEntry(this, BatchFilesDirectory, BatchScriptArguments, OnlyGameProject));
 					VCProjectFileContent.AppendLine("    <NMakeOutput>{0}</NMakeOutput>", NormalizeProjectPath(NMakePath.FullName));
 
 					if (TargetRulesObject.Type == TargetType.Game || TargetRulesObject.Type == TargetType.Client || TargetRulesObject.Type == TargetType.Server)
@@ -1669,4 +1641,4 @@ namespace UnrealBuildTool
 		protected readonly Dictionary<UnrealTargetConfiguration, CsProjectInfo> CachedProjectInfo = new Dictionary<UnrealTargetConfiguration, CsProjectInfo>();
 	}
 
-}
+}
\ No newline at end of file
diff --git a/Engine/Source/Programs/UnrealBuildTool/System/UEPlatformProjectGenerator.cs b/Engine/Source/Programs/UnrealBuildTool/System/UEPlatformProjectGenerator.cs
--- a/Engine/Source/Programs/UnrealBuildTool/System/UEPlatformProjectGenerator.cs
+++ b/Engine/Source/Programs/UnrealBuildTool/System/UEPlatformProjectGenerator.cs
@@ -14,6 +14,27 @@ namespace UnrealBuildTool
 	/// </summary>
 	abstract class UEPlatformProjectGenerator
 	{
+		public class ScriptArguments 
+		{
+			public ScriptArguments(bool usePrecompiled, bool isForeignProject, bool useFastPDB, string targetName, string platformName, string configurationName, string buildToolOverride)
+			{
+				UsePrecompiled = usePrecompiled;
+				IsForeignProject = isForeignProject;
+				UseFastPDB = useFastPDB;
+				TargetName = targetName;
+				PlatformName = platformName;
+				ConfigurationName = configurationName;
+				BuildToolOverride = buildToolOverride;
+			}
+			public bool UsePrecompiled { get; }
+			public bool IsForeignProject { get; }
+			public bool UseFastPDB { get; }
+			public string TargetName { get; }
+			public string PlatformName { get; }
+			public string ConfigurationName { get; }
+			public string BuildToolOverride { get; }
+		}
+
 		static Dictionary<UnrealTargetPlatform, UEPlatformProjectGenerator> ProjectGeneratorDictionary = new Dictionary<UnrealTargetPlatform, UEPlatformProjectGenerator>();
 
 		/// <summary>
@@ -167,6 +188,79 @@ namespace UnrealBuildTool
 			return InPlatform.ToString();
 		}
 
+		/// <summary>
+		///  Returns the default build batch script entry and its arguments
+		/// </summary>
+		/// <returns></returns>
+		public virtual string GetBuildCommandEntry(ProjectFile ProjectFile, DirectoryReference BatchFileDirectory, ScriptArguments FileArguments, FileReference OnlyGameProject)
+		{
+			return ProjectFile.EscapePath(ProjectFile.NormalizeProjectPath(FileReference.Combine(BatchFileDirectory, "Build.bat"))) + GetBatchFileArguments(FileArguments);
+		}
+
+		/// <summary>
+		///  Returns the default rebuild batch script entry and its arguments
+		/// </summary>
+		/// <returns></returns>
+		public virtual string GetReBuildCommandEntry(ProjectFile ProjectFile, DirectoryReference BatchFileDirectory, ScriptArguments FileArguments, FileReference OnlyGameProject)
+		{
+			return ProjectFile.EscapePath(ProjectFile.NormalizeProjectPath(FileReference.Combine(BatchFileDirectory, "Rebuild.bat"))) + GetBatchFileArguments(FileArguments);
+		}
+
+		/// <summary>
+		/// Returns the default clean batch script filename
+		/// </summary>
+		/// <returns></returns>
+		public virtual string GetCleanCommandEntry(ProjectFile ProjectFile, DirectoryReference BatchFileDirectory, ScriptArguments FileArguments, FileReference OnlyGameProject)
+		{
+			return ProjectFile.EscapePath(ProjectFile.NormalizeProjectPath(FileReference.Combine(BatchFileDirectory, "Clean.bat"))) + GetBatchFileArguments(FileArguments);
+		}
+
+		/// <summary>
+		/// Builds an argument string passed to build, rebuild and clean batch scripts
+		/// </summary>
+		/// <param name="FileArguments"></param>
+		/// <returns></returns>
+		public virtual string GetBatchFileArguments(ScriptArguments FileArguments)
+		{
+			// This is the standard UE4 based project NMake build line:
+			//	..\..\Build\BatchFiles\Build.bat <TARGETNAME> <PLATFORM> <CONFIGURATION>
+			//	ie ..\..\Build\BatchFiles\Build.bat BlankProgram Win64 Debu
+			string BuildArguments = " " + FileArguments.TargetName + " " + FileArguments.PlatformName + " " + FileArguments.ConfigurationName;
+			if (FileArguments.UsePrecompiled)
+			{
+				BuildArguments += " -UsePrecompiled";
+			}
+			if (FileArguments.IsForeignProject)
+			{
+				BuildArguments += " \"$(SolutionDir)$(ProjectName).uproject\"";
+			}
+
+			// Always wait for the mutex between UBT invocations, so that building the whole solution doesn't fail.
+			BuildArguments += " -WaitMutex";
+			// Always include a flag to format log messages for MSBuild
+			BuildArguments += " -FromMsBuild";
+			if (FileArguments.UseFastPDB)
+			{
+				BuildArguments += " -FastPDB";
+			}
+			if (!String.IsNullOrEmpty(FileArguments.BuildToolOverride))
+			{
+				// Pass Fast PDB option to make use of Visual Studio's /DEBUG:FASTLINK option
+				BuildArguments += " " + FileArguments.BuildToolOverride;
+			}
+			return BuildArguments;
+		}
+
+		/// <summary>
+		/// Returns the NMake output file for the particular generator.
+		/// </summary>
+		/// <param name="InNMakePath">  The current base nmake output file.</param>
+		/// <returns>FileReference The new NMake output file.</returns>
+		public virtual FileReference GetNMakeOutputPath(FileReference InNMakePath)
+		{
+			return InNMakePath;
+		}
+
 		/// <summary>
 		/// Return project configuration settings that must be included before the default props file
 		/// </summary>
diff --git a/Engine/Source/Runtime/ApplicationCore/Private/Lumin/LuminApplication.cpp b/Engine/Source/Runtime/ApplicationCore/Private/Lumin/LuminApplication.cpp
new file mode 100644
--- /dev/null
+++ b/Engine/Source/Runtime/ApplicationCore/Private/Lumin/LuminApplication.cpp
@@ -0,0 +1,159 @@
+// Copyright 1998-2018 Epic Games, Inc. All Rights Reserved.
+
+#include "Lumin/LuminApplication.h"
+#include "Misc/ScopeLock.h"
+#include "Misc/ConfigCacheIni.h"
+
+FLuminApplication* LuminApplication = nullptr;
+
+void FLuminApplication::DeferredKeyEvent::SendModified(
+	TSharedRef<FGenericApplicationMessageHandler> MessageHandler, uint32_t& InOutModifierMask) const
+{
+	switch (KeyEventType)
+	{
+	case Type::KeyDown:
+		InOutModifierMask = KeyData;
+		MessageHandler->OnKeyDown(KeyCode, 0, false);
+		break;
+	case Type::KeyUp:
+		InOutModifierMask = KeyData;
+		MessageHandler->OnKeyUp(KeyCode, 0, false);
+		break;
+	case Type::Char:
+		MessageHandler->OnKeyChar(KeyData, false);
+		break;
+	}
+}
+
+FLuminApplication::FLuminApplication()
+{
+	// Pull tracking preference from config file
+	MLInputConfiguration InputConfig = { };
+
+	InputConfig.dof[0] = InputConfig.dof[1] = MLInputControllerDof_6;
+
+	FString EnumVal;
+	GConfig->GetString(TEXT("/Script/LuminRuntimeSettings.LuminRuntimeSettings"),
+		TEXT("ControllerTrackingType"), EnumVal, GEngineIni);
+	if (EnumVal.Len() > 0)
+	{
+		// Can't get enums as strings at this level
+		if (0 == EnumVal.Compare("NotTracked", ESearchCase::IgnoreCase))
+		{
+			InputConfig.dof[0] = InputConfig.dof[1] = MLInputControllerDof_None;
+		}
+		else if (0 == EnumVal.Compare("InertialOnly", ESearchCase::IgnoreCase))
+		{
+			InputConfig.dof[0] = InputConfig.dof[1] = MLInputControllerDof_3;
+		}
+	}
+
+	// Create the input tracker
+	MLResult Result = MLInputCreate(&InputConfig, &InputTracker);
+	if (MLResult_Ok != Result)
+	{
+		UE_LOG(LogCore, Error, TEXT("FLuminApplication: Unable to initialize "
+			"input system: %s"), UTF8_TO_TCHAR(MLGetResultString(Result)));
+		InputTracker = ML_INVALID_HANDLE;
+	}
+	else
+	{
+		InitializeInputCallbacks();
+	}
+}
+
+FLuminApplication::~FLuminApplication()
+{
+	if (ML_INVALID_HANDLE != InputTracker)
+	{
+		auto Result = MLInputDestroy(InputTracker);
+		if (MLResult_Ok != Result)
+		{
+			UE_LOG(LogCore, Error, TEXT("FLuminApplication::~FLuminApplication: "
+				"Failure in MLInputDestroy: %s"), UTF8_TO_TCHAR(MLGetResultString(Result)));
+		}
+		InputTracker = ML_INVALID_HANDLE;
+	}
+	LuminApplication = nullptr;
+}
+
+FLuminApplication* FLuminApplication::CreateLuminApplication()
+{
+	LuminApplication = new FLuminApplication;
+	return LuminApplication;
+}
+
+void FLuminApplication::Tick(const float TimeDelta)
+{
+	// Roll through the pending events at our leisure
+	while (!DeferredKeyEvents.IsEmpty())
+	{
+		DeferredKeyEvent CurrKeyEvent;
+		DeferredKeyEvents.Dequeue(CurrKeyEvent);
+		CurrKeyEvent.SendModified(MessageHandler, ModifierMask);
+	}
+}
+
+void FLuminApplication::InitializeInputCallbacks()
+{
+	FMemory::Memset(&InputKeyboardCallbacks, 0, sizeof(InputKeyboardCallbacks));
+
+	InputKeyboardCallbacks.on_char = [](uint32_t char_utf32, void *data)
+	{
+		auto application = reinterpret_cast<FLuminApplication*>(data);
+		if (application)
+		{
+			application->AddDeferredKeyEvent(DeferredKeyEvent(char_utf32));
+		}
+	};
+
+	InputKeyboardCallbacks.on_key_down = [](MLKeyCode key_code, uint32 modifier_mask, void *data)
+	{
+		auto application = reinterpret_cast<FLuminApplication*>(data);
+		if (application)
+		{
+			application->AddDeferredKeyEvent(DeferredKeyEvent(
+				DeferredKeyEvent::Type::KeyDown, key_code, modifier_mask));
+		}
+	};
+
+	InputKeyboardCallbacks.on_key_up = [](MLKeyCode key_code, uint32 modifier_mask, void *data)
+	{
+		auto application = reinterpret_cast<FLuminApplication*>(data);
+		if (application)
+		{
+			application->AddDeferredKeyEvent(DeferredKeyEvent(
+				DeferredKeyEvent::Type::KeyUp, key_code, modifier_mask));
+		}
+	};
+
+	MLResult Result = MLInputSetKeyboardCallbacks(InputTracker, &InputKeyboardCallbacks, this);
+	if (MLResult_Ok != Result)
+	{
+		UE_LOG(LogCore, Error, TEXT("FLuminApplication::InitializeInputCallbacks: "
+			"unable to set keyboard callbacks: %s"), UTF8_TO_TCHAR(MLGetResultString(Result)));
+	}
+}
+
+void FLuminApplication::AddDeferredKeyEvent(const DeferredKeyEvent& InDeferredEvent)
+{
+	DeferredKeyEvents.Enqueue(InDeferredEvent);
+}
+
+FModifierKeysState FLuminApplication::GetModifierKeys() const
+{
+	// ML platform only exposes single Shift, Control, and Alt modifiers; we'll map them to the
+	// left because that is the most common location on single-modifier keyboards
+	return FModifierKeysState(
+		(ModifierMask & MLKEYMODIFIER_SHIFT) != 0,    // bInIsLeftShiftDown
+		false,                                        // bInIsRightShiftDown
+		(ModifierMask & MLKEYMODIFIER_CTRL) != 0,     // bInIsLeftControlDown
+		false,                                        // bInIsRightControlDown
+		(ModifierMask & MLKEYMODIFIER_ALT) != 0,      // bInIsLeftAltDown
+		false,                                        // bInIsRightAltDown
+		false,                                        // bInIsLeftCommandDown
+		false,                                        // bInIsRightCommandDown
+		(ModifierMask & MLKEYMODIFIER_CAPS_LOCK) != 0 // bInAreCapsLocked
+	);
+}
+
diff --git a/Engine/Source/Runtime/ApplicationCore/Private/Lumin/LuminPlatformApplicationMisc.cpp b/Engine/Source/Runtime/ApplicationCore/Private/Lumin/LuminPlatformApplicationMisc.cpp
--- a/Engine/Source/Runtime/ApplicationCore/Private/Lumin/LuminPlatformApplicationMisc.cpp
+++ b/Engine/Source/Runtime/ApplicationCore/Private/Lumin/LuminPlatformApplicationMisc.cpp
@@ -1,6 +1,7 @@
 // Copyright 1998-2018 Epic Games, Inc. All Rights Reserved.
 
 #include "Lumin/LuminPlatformApplicationMisc.h"
+#include "Lumin/LuminApplication.h"
 
 // #include "Android/AndroidApplication.h"
 // #include "Android/AndroidErrorOutputDevice.h"
@@ -10,3 +11,7 @@
 // #include "Internationalization/Regex.h"
 // #include "Modules/ModuleManager.h"
 
+GenericApplication* FLuminPlatformApplicationMisc::CreateApplication()
+{
+	return FLuminApplication::CreateLuminApplication();
+}
diff --git a/Engine/Source/Runtime/ApplicationCore/Public/Lumin/LuminApplication.h b/Engine/Source/Runtime/ApplicationCore/Public/Lumin/LuminApplication.h
new file mode 100644
--- /dev/null
+++ b/Engine/Source/Runtime/ApplicationCore/Public/Lumin/LuminApplication.h
@@ -0,0 +1,78 @@
+// Copyright 1998-2018 Epic Games, Inc. All Rights Reserved.
+
+#pragma once
+
+#include "CoreMinimal.h"
+#include "Android/AndroidApplication.h"
+#include <Containers/Queue.h>
+
+#include <ml_api.h>
+#include <ml_keycodes.h>
+#include <ml_input.h>
+
+/**
+* Lumin-specific application implementation.
+*/
+class APPLICATIONCORE_API FLuminApplication : public FAndroidApplication
+{
+public:
+	/** Key events come in on separate threads so we queue them up and process them on the main thread. */
+	class DeferredKeyEvent
+	{
+	public:
+		enum Type
+		{
+			KeyDown,
+			KeyUp,
+			Char,
+		};
+	protected:
+		Type      KeyEventType;
+		MLKeyCode KeyCode;
+		uint32_t  KeyData;
+	public:
+		inline DeferredKeyEvent()
+		{
+		}
+		inline DeferredKeyEvent(Type InKeyEventType, MLKeyCode InKeyCode, uint32 ModifierMask)
+			: KeyEventType(InKeyEventType)
+			, KeyCode(InKeyCode)
+			, KeyData(ModifierMask)
+		{
+		}
+		inline DeferredKeyEvent(uint32 CharUtf32)
+			: KeyEventType(Type::Char)
+			, KeyData(CharUtf32)
+		{
+		}
+		void SendModified(TSharedRef<FGenericApplicationMessageHandler> MessageHandler, uint32_t& InOutModifierMask) const;
+	};
+
+public:
+	static FLuminApplication* CreateLuminApplication();
+
+	FLuminApplication();
+	~FLuminApplication();
+
+	inline const MLHandle GetInputTracker() const
+	{
+		return InputTracker;
+	}
+
+	void AddDeferredKeyEvent(const DeferredKeyEvent& InDeferredEvent);
+
+	// GenericApplication
+	void Tick(const float TimeDelta) override;
+	FModifierKeysState GetModifierKeys() const override;
+
+protected:
+
+	void InitializeInputCallbacks();
+
+	MLHandle InputTracker = ML_INVALID_HANDLE;
+	MLInputKeyboardCallbacks InputKeyboardCallbacks;
+	TQueue<DeferredKeyEvent> DeferredKeyEvents;
+	uint32 ModifierMask = 0;
+};
+
+extern FLuminApplication* LuminApplication;
diff --git a/Engine/Source/Runtime/ApplicationCore/Public/Lumin/LuminPlatformApplicationMisc.h b/Engine/Source/Runtime/ApplicationCore/Public/Lumin/LuminPlatformApplicationMisc.h
--- a/Engine/Source/Runtime/ApplicationCore/Public/Lumin/LuminPlatformApplicationMisc.h
+++ b/Engine/Source/Runtime/ApplicationCore/Public/Lumin/LuminPlatformApplicationMisc.h
@@ -24,7 +24,7 @@ struct APPLICATIONCORE_API FLuminPlatformApplicationMisc : public FAndroidApplic
 		return false;
 	}
 
-
+	static class GenericApplication* CreateApplication();
 };
 
 typedef FLuminPlatformApplicationMisc FPlatformApplicationMisc;
diff --git a/Engine/Source/Runtime/AudioMixer/Private/AudioMixerSource.cpp b/Engine/Source/Runtime/AudioMixer/Private/AudioMixerSource.cpp
--- a/Engine/Source/Runtime/AudioMixer/Private/AudioMixerSource.cpp
+++ b/Engine/Source/Runtime/AudioMixer/Private/AudioMixerSource.cpp
@@ -364,14 +364,14 @@ namespace Audio
 			bIsDone = false;
 
 			// Not all wave data types have a non-zero duration
-			if (InWaveInstance->WaveData->Duration > 0)
+			if (!InWaveInstance->WaveData->IsLooping() && !InWaveInstance->WaveData->IsLooping())
 			{
 				if (!InWaveInstance->WaveData->bIsBus)
 				{
 					NumTotalFrames = InWaveInstance->WaveData->Duration * InWaveInstance->WaveData->GetSampleRateForCurrentPlatform();
 					check(NumTotalFrames > 0);
 				}
-				else if (!InWaveInstance->WaveData->IsLooping())
+				else
 				{
 					NumTotalFrames = InWaveInstance->WaveData->Duration * AudioDevice->GetSampleRate();
 					check(NumTotalFrames > 0);
diff --git a/Engine/Source/Runtime/Core/Private/HAL/ConsoleManager.cpp b/Engine/Source/Runtime/Core/Private/HAL/ConsoleManager.cpp
--- a/Engine/Source/Runtime/Core/Private/HAL/ConsoleManager.cpp
+++ b/Engine/Source/Runtime/Core/Private/HAL/ConsoleManager.cpp
@@ -1976,6 +1976,31 @@ static TAutoConsoleVariable<float> CVarAmbientOcclusionStaticFraction(TEXT("r.Am
 	TEXT(" 1: AO affects the stat lighting"),
 	ECVF_Default);
 
+static TAutoConsoleVariable<float> CVarDebugCanvasX(TEXT("r.LuminDebugCanvasX"),
+	11.0f,
+	TEXT("Allows the override of the lumin debug canvas 'x' coordinate."),
+	ECVF_Scalability | ECVF_Default);
+
+static TAutoConsoleVariable<float> CVarDebugCanvasY(TEXT("r.LuminDebugCanvasY"),
+	-18.0f,
+	TEXT("Allows the override of the lumin debug canvas 'y' coordinate."),
+	ECVF_Scalability | ECVF_Default);
+
+static TAutoConsoleVariable<float> CVarDebugCanvasZ(TEXT("r.LuminDebugCanvasZ"),
+	82.0f,
+	TEXT("Allows the override of the lumin debug canvas 'z' coordinate."),
+	ECVF_Scalability | ECVF_Default);
+
+static TAutoConsoleVariable<float> CVarDebugCanvasWidth(TEXT("r.LuminDebugCanvasWidth"),
+	75.0f,
+	TEXT("Allows the override of the lumin debug canvas width."),
+	ECVF_Scalability | ECVF_Default);
+
+static TAutoConsoleVariable<float> CVarDebugCanvasHeight(TEXT("r.LuminDebugCanvasHeight"),
+	40.0f,
+	TEXT("Allows the override of the lumin debug canvas height."),
+	ECVF_Scalability | ECVF_Default);
+
 static TAutoConsoleVariable<int32> CVarShadowQuality(
 	TEXT("r.ShadowQuality"),
 	5,
@@ -2545,3 +2570,11 @@ static TAutoConsoleVariable<int32> CVarSkyLightingQuality(
 	TEXT("  1: on\n"),
 	ECVF_Scalability | ECVF_RenderThreadSafe
 );
+
+static TAutoConsoleVariable<int32> CVarLuminUseThreadingForPerformance(
+	TEXT("core.Lumin.UseThreadingForPerformance"),
+	1,
+	TEXT("Determine if threading should be used for possible performance gains. Used for FApp::ShouldUseThreadingForPerformance(). (Lumin-only)\n")
+	TEXT("  0 = off\n")
+	TEXT("  1 = on [default]\n"),
+	ECVF_ReadOnly);
diff --git a/Engine/Source/Runtime/Core/Private/Lumin/LuminLifecycle.cpp b/Engine/Source/Runtime/Core/Private/Lumin/LuminLifecycle.cpp
--- a/Engine/Source/Runtime/Core/Private/Lumin/LuminLifecycle.cpp
+++ b/Engine/Source/Runtime/Core/Private/Lumin/LuminLifecycle.cpp
@@ -40,6 +40,13 @@ void FLuminLifecycle::Initialize()
 	}
 	else
 	{
+		// It is possible that FLuminLifecycle::Initialize() is called before LaunchLumin::InitCommandLine().
+		// So initialize command line here to take in args passed via mldb launch.
+		if (!FCommandLine::IsInitialized())
+		{
+			// initialize the command line to an empty string
+			FCommandLine::Set(TEXT(""));
+		}
 		OnNewInitArgs_Handler(nullptr);
 	}
 }
diff --git a/Engine/Source/Runtime/Core/Private/Lumin/LuminPlatformMisc.cpp b/Engine/Source/Runtime/Core/Private/Lumin/LuminPlatformMisc.cpp
--- a/Engine/Source/Runtime/Core/Private/Lumin/LuminPlatformMisc.cpp
+++ b/Engine/Source/Runtime/Core/Private/Lumin/LuminPlatformMisc.cpp
@@ -144,6 +144,15 @@ bool FLuminPlatformMisc::ShouldUseDesktopVulkan()
 	return ShouldUseVulkan() && !bUseMobileRendering;
 }
 
+const TCHAR* FLuminPlatformMisc::GetDefaultDeviceProfileName()
+{
+	if (ShouldUseDesktopVulkan() || ShouldUseDesktopOpenGL())
+	{
+		return TEXT("Lumin_Desktop");
+	}
+	return TEXT("Lumin");
+}
+
 bool FLuminPlatformMisc::ShouldUseDesktopOpenGL()
 {
 	bool bUseMobileRendering = false;
diff --git a/Engine/Source/Runtime/Core/Private/Misc/App.cpp b/Engine/Source/Runtime/Core/Private/Misc/App.cpp
--- a/Engine/Source/Runtime/Core/Private/Misc/App.cpp
+++ b/Engine/Source/Runtime/Core/Private/Misc/App.cpp
@@ -6,6 +6,7 @@
 #include "Misc/ConfigCacheIni.h"
 #include "Runtime/Launch/Resources/Version.h"
 #include "HAL/LowLevelMemTracker.h"
+#include "HAL/IConsoleManager.h"
 #include "BuildSettings.h"
 #include "UObject/DevObjectVersion.h"
 #include "Misc/EngineVersion.h"
diff --git a/Engine/Source/Runtime/Core/Public/Lumin/LuminPlatformMisc.h b/Engine/Source/Runtime/Core/Public/Lumin/LuminPlatformMisc.h
--- a/Engine/Source/Runtime/Core/Public/Lumin/LuminPlatformMisc.h
+++ b/Engine/Source/Runtime/Core/Public/Lumin/LuminPlatformMisc.h
@@ -70,6 +70,7 @@ struct CORE_API FLuminPlatformMisc : public FAndroidMisc
 	/** Android overrides. */
 	static bool ShouldUseVulkan();
 	static bool ShouldUseDesktopVulkan();
+	static const TCHAR* GetDefaultDeviceProfileName();
 
 	/** Lumin specific. */
 	static bool ShouldUseDesktopOpenGL();
diff --git a/Engine/Source/Runtime/Engine/Private/UnrealClient.cpp b/Engine/Source/Runtime/Engine/Private/UnrealClient.cpp
--- a/Engine/Source/Runtime/Engine/Private/UnrealClient.cpp
+++ b/Engine/Source/Runtime/Engine/Private/UnrealClient.cpp
@@ -223,7 +223,7 @@ bool FRenderTarget::ReadLinearColorPixels(TArray<FLinearColor> &OutImageData, FR
 		});
 	FlushRenderingCommands();
 
-	return true;
+	return OutImageData.Num() > 0;
 }
 
 /**
@@ -238,6 +238,7 @@ bool FRenderTarget::ReadLinearColorPixelsPtr(FLinearColor* OutImageBytes, FReadS
 	bool bResult = ReadLinearColorPixels(SurfaceData, InFlags, InRect);
 	if (bResult)
 	{
+		check(SurfaceData.Num() != 0);
 		FMemory::Memcpy(OutImageBytes, &SurfaceData[0], SurfaceData.Num() * sizeof(FLinearColor));
 	}
 
diff --git a/Engine/Source/Runtime/Engine/Public/SceneManagement.h b/Engine/Source/Runtime/Engine/Public/SceneManagement.h
--- a/Engine/Source/Runtime/Engine/Public/SceneManagement.h
+++ b/Engine/Source/Runtime/Engine/Public/SceneManagement.h
@@ -302,7 +302,7 @@ static const int32 MAX_NUM_LIGHTMAP_COEF = 2;
 
 /** Compile out low quality lightmaps to save memory */
 // @todo-mobile: Need to fix this!
-#define ALLOW_LQ_LIGHTMAPS (PLATFORM_DESKTOP || PLATFORM_IOS || PLATFORM_ANDROID || PLATFORM_HTML5 || PLATFORM_SWITCH || PLATFORM_LUMINGL4)
+#define ALLOW_LQ_LIGHTMAPS (PLATFORM_DESKTOP || PLATFORM_IOS || PLATFORM_ANDROID || PLATFORM_HTML5 || PLATFORM_SWITCH || PLATFORM_LUMIN)
 
 /** Compile out high quality lightmaps to save memory */
 #define ALLOW_HQ_LIGHTMAPS 1
diff --git a/Engine/Source/Runtime/Launch/Private/Lumin/LaunchLumin.cpp b/Engine/Source/Runtime/Launch/Private/Lumin/LaunchLumin.cpp
--- a/Engine/Source/Runtime/Launch/Private/Lumin/LaunchLumin.cpp
+++ b/Engine/Source/Runtime/Launch/Private/Lumin/LaunchLumin.cpp
@@ -26,8 +26,14 @@ static void InitCommandLine()
 {
 	static const uint32 CMD_LINE_MAX = 16384u;
 
-	// initialize the command line to an empty string
-	FCommandLine::Set(TEXT(""));
+	// It is possible that FLuminLifecycle::Initialize() is called before LaunchLumin::InitCommandLine().
+	// If so, FLuminLifecycle::Initialize() would have initialized the command line and filled it with args passed via mldb launch.
+	// So initialize command line here only if it hasnt been initialized already.
+	if (!FCommandLine::IsInitialized())
+	{
+		// initialize the command line to an empty string
+		FCommandLine::Set(TEXT(""));
+	}
 
 	// Adds command line args coming from lifecycle app init args.
 	FLuminPlatformMisc::InitLifecycle();	
diff --git a/Engine/Source/Runtime/Lumin/LuminRuntimeSettings/Classes/LuminRuntimeSettings.h b/Engine/Source/Runtime/Lumin/LuminRuntimeSettings/Classes/LuminRuntimeSettings.h
--- a/Engine/Source/Runtime/Lumin/LuminRuntimeSettings/Classes/LuminRuntimeSettings.h
+++ b/Engine/Source/Runtime/Lumin/LuminRuntimeSettings/Classes/LuminRuntimeSettings.h
@@ -122,9 +122,9 @@ public:
 	UPROPERTY(GlobalConfig, EditAnywhere, Category = "Advanced MPK Packaging", Meta = (DisplayName = "Version Code", ClampMin = 0))
 	int32 VersionCode;
 
-	/** Name of the minimum Lumin OS Version required */
-	UPROPERTY(GlobalConfig, EditAnywhere, Category = "Advanced MPK Packaging", Meta = (DisplayName = "Minimum OS Version Required"))
-	FString MinimumOSVersion;
+	/** Minimum API level required based on which APIs have been integrated into the base engine. Developers can set higher API levels if they are implementing newer APIs. */
+	UPROPERTY(GlobalConfig, EditAnywhere, Category = "Advanced MPK Packaging", Meta = (DisplayName = "Minimum API Level", ClampMin = 2))
+	int32 MinimumAPILevel;
 
 	/** Any privileges your app needs. */
 	UPROPERTY(GlobalConfig, EditAnywhere, Category = "Advanced MPK Packaging", Meta = (DisplayName = "App Privileges"))
@@ -149,4 +149,8 @@ public:
 	/** Which of the currently enabled occlusion plugins to use on Lumin. */
 	UPROPERTY(config, EditAnywhere, Category = "Audio")
 	FString OcclusionPlugin;
+
+	// Strip debug symbols from packaged builds even if they aren't shipping builds.
+	UPROPERTY(GlobalConfig, EditAnywhere, Category = AdvancedBuild, meta = (DisplayName = "Strip debug symbols from packaged builds even if they aren't shipping builds"))
+	bool bRemoveDebugInfo;
 };
diff --git a/Engine/Source/Runtime/Lumin/LuminRuntimeSettings/LuminRuntimeSettings.Build.cs b/Engine/Source/Runtime/Lumin/LuminRuntimeSettings/LuminRuntimeSettings.Build.cs
--- a/Engine/Source/Runtime/Lumin/LuminRuntimeSettings/LuminRuntimeSettings.Build.cs
+++ b/Engine/Source/Runtime/Lumin/LuminRuntimeSettings/LuminRuntimeSettings.Build.cs
@@ -1,6 +1,7 @@
 // Copyright 1998-2018 Epic Games, Inc. All Rights Reserved.
 
 using UnrealBuildTool;
+using System.IO;
 
 public class LuminRuntimeSettings : ModuleRules
 {
@@ -8,13 +9,13 @@ public class LuminRuntimeSettings : ModuleRules
 	{
 		BinariesSubFolder = "Lumin";
 
-		PrivateDependencyModuleNames.AddRange(
+        PrivateDependencyModuleNames.AddRange(
 			new string[]
 			{
 				"Core",
 				"CoreUObject",
 				"Engine",
-			}
+            }
 		);
 
         if (Target.Type == TargetRules.TargetType.Editor || Target.Type == TargetRules.TargetType.Program)
diff --git a/Engine/Source/Runtime/MRMesh/Public/MRMeshComponent.h b/Engine/Source/Runtime/MRMesh/Public/MRMeshComponent.h
--- a/Engine/Source/Runtime/MRMesh/Public/MRMeshComponent.h
+++ b/Engine/Source/Runtime/MRMesh/Public/MRMeshComponent.h
@@ -72,6 +72,8 @@ public:
 	UFUNCTION(BlueprintCallable, Category = "Mesh Reconstruction")
 	void Clear() override;
 
+	// UPrimitiveComponent.. public BP function needs to stay public to avoid nativization errors. (RR)
+	virtual void SetMaterial(int32 ElementIndex, class UMaterialInterface* InMaterial) override;
 protected:
 	virtual void OnActorEnableCollisionChanged() override;
 	virtual void UpdatePhysicsToRBChannels() override;
@@ -90,7 +92,6 @@ private:
 	virtual FPrimitiveSceneProxy* CreateSceneProxy() override;
 	virtual void GetUsedMaterials(TArray<UMaterialInterface*>& OutMaterials, bool bGetDebugMaterials = false) const override;
 	virtual FBoxSphereBounds CalcBounds(const FTransform& LocalToWorld) const override;
-	virtual void SetMaterial(int32 ElementIndex, class UMaterialInterface* InMaterial) override;
 	virtual bool DoCustomNavigableGeometryExport(FNavigableGeometryExport& GeomExport) const override;
 	virtual class UBodySetup* GetBodySetup();
 	//~ UPrimitiveComponent
diff --git a/Engine/Source/Runtime/Media/Public/IMediaPlayer.h b/Engine/Source/Runtime/Media/Public/IMediaPlayer.h
--- a/Engine/Source/Runtime/Media/Public/IMediaPlayer.h
+++ b/Engine/Source/Runtime/Media/Public/IMediaPlayer.h
@@ -231,7 +231,7 @@ public:
 	 *
 	 * @see TickInput, TickFetch
 	 */
-	virtual void TickAudio()
+	virtual void TickAudio(FTimespan LastAudioSampleTime)
 	{
 		// override in child class if needed
 	}
diff --git a/Engine/Source/Runtime/MediaUtils/Private/MediaPlayerFacade.cpp b/Engine/Source/Runtime/MediaUtils/Private/MediaPlayerFacade.cpp
--- a/Engine/Source/Runtime/MediaUtils/Private/MediaPlayerFacade.cpp
+++ b/Engine/Source/Runtime/MediaUtils/Private/MediaPlayerFacade.cpp
@@ -1063,7 +1063,7 @@ void FMediaPlayerFacade::TickTickable()
 		return;
 	}
 
-	Player->TickAudio();
+	Player->TickAudio(LastAudioRenderedSampleTime.Load());
 
 	// determine range of valid samples
 	TRange<FTimespan> AudioTimeRange;
@@ -1089,7 +1089,6 @@ void FMediaPlayerFacade::TickTickable()
 	ProcessMetadataSamples(Samples, MetadataTimeRange);	
 }
 
-
 /* FMediaPlayerFacade implementation
 *****************************************************************************/
 
diff --git a/Engine/Source/Runtime/Slate/Private/Framework/Application/SlateApplication.cpp b/Engine/Source/Runtime/Slate/Private/Framework/Application/SlateApplication.cpp
--- a/Engine/Source/Runtime/Slate/Private/Framework/Application/SlateApplication.cpp
+++ b/Engine/Source/Runtime/Slate/Private/Framework/Application/SlateApplication.cpp
@@ -885,7 +885,11 @@ void FSlateApplication::OverridePlatformApplication(TSharedPtr<class GenericAppl
 
 void FSlateApplication::Create()
 {
-	GSlateFastWidgetPath = GIsEditor ? 0 : 1;
+	// Force disable GSlateFastWidgetPath in editor
+	if (GIsEditor)
+	{
+		GSlateFastWidgetPath = 0;
+	}
 
 	Create(MakeShareable(FPlatformApplicationMisc::CreateApplication()));
 }
diff --git a/Engine/Source/Runtime/VulkanRHI/Private/Lumin/VulkanLuminPlatform.h b/Engine/Source/Runtime/VulkanRHI/Private/Lumin/VulkanLuminPlatform.h
--- a/Engine/Source/Runtime/VulkanRHI/Private/Lumin/VulkanLuminPlatform.h
+++ b/Engine/Source/Runtime/VulkanRHI/Private/Lumin/VulkanLuminPlatform.h
@@ -55,7 +55,7 @@ public:
 		GShaderPlatformForFeatureLevel[ERHIFeatureLevel::ES2] = SP_VULKAN_ES3_1_LUMIN;
 		GShaderPlatformForFeatureLevel[ERHIFeatureLevel::ES3_1] = SP_VULKAN_ES3_1_LUMIN;
 		GShaderPlatformForFeatureLevel[ERHIFeatureLevel::SM4] = SP_NumPlatforms;
-		GShaderPlatformForFeatureLevel[ERHIFeatureLevel::SM5] = PLATFORM_LUMINGL4 ? SP_VULKAN_SM5_LUMIN : SP_NumPlatforms;
+		GShaderPlatformForFeatureLevel[ERHIFeatureLevel::SM5] = SP_VULKAN_SM5_LUMIN;
 	}
 
 	static bool SupportsStandardSwapchain() { return false; }
diff --git a/Engine/Source/Runtime/VulkanRHI/Private/VulkanRenderTarget.cpp b/Engine/Source/Runtime/VulkanRHI/Private/VulkanRenderTarget.cpp
--- a/Engine/Source/Runtime/VulkanRHI/Private/VulkanRenderTarget.cpp
+++ b/Engine/Source/Runtime/VulkanRHI/Private/VulkanRenderTarget.cpp
@@ -943,6 +943,16 @@ void FVulkanDynamicRHI::RHIReadSurfaceData(FTextureRHIParamRef TextureRHI, FIntR
 	ImmediateContext.GetCommandBufferManager()->PrepareForNewActiveCommandBuffer();
 }
 
+void FVulkanDynamicRHI::RHIReadSurfaceData(FTextureRHIParamRef TextureRHI, FIntRect Rect, TArray<FLinearColor>& OutData, FReadSurfaceDataFlags InFlags)
+{
+	TArray<FColor> FromColorData;
+	RHIReadSurfaceData(TextureRHI, Rect, FromColorData, InFlags);
+	for (FColor& From : FromColorData)
+	{
+		OutData.Emplace(FLinearColor(From));
+	}
+}
+
 void FVulkanDynamicRHI::RHIMapStagingSurface(FTextureRHIParamRef TextureRHI,void*& OutData,int32& OutWidth,int32& OutHeight)
 {
 	FRHITexture2D* TextureRHI2D = TextureRHI->GetTexture2D();
