Magic Leap: Updated minimum SDK support for Magic Leap to 0.19

diff --git a/Engine/Source/Programs/UnrealBuildTool/Platform/Lumin/UEBuildLumin.cs b/Engine/Source/Programs/UnrealBuildTool/Platform/Lumin/UEBuildLumin.cs
--- a/Engine/Source/Programs/UnrealBuildTool/Platform/Lumin/UEBuildLumin.cs
+++ b/Engine/Source/Programs/UnrealBuildTool/Platform/Lumin/UEBuildLumin.cs
@@ -276,7 +276,7 @@ namespace UnrealBuildTool
 		/// This is the minimum SDK version we support.
 		/// </summary>
 		static uint MinimumSDKVersionMajor = 0;
-		static uint MinimumSDKVersionMinor = 16;
+		static uint MinimumSDKVersionMinor = 19;
 
 		public override string GetSDKTargetPlatformName()
 		{
