Revert UnrealBuildTool changes from CL-4636504 "Updated ML integration to MLSDK 0.19"

It appers to have made it so compiling in Visual Studio cooks and deploys to device. That is not consistent with any other platform.

diff --git a/Engine/Source/Programs/UnrealBuildTool/Platform/Lumin/LuminProjectGenerator.cs b/Engine/Source/Programs/UnrealBuildTool/Platform/Lumin/LuminProjectGenerator.cs
--- a/Engine/Source/Programs/UnrealBuildTool/Platform/Lumin/LuminProjectGenerator.cs
+++ b/Engine/Source/Programs/UnrealBuildTool/Platform/Lumin/LuminProjectGenerator.cs
@@ -82,7 +82,7 @@ namespace UnrealBuildTool
 
 			if (IsVSLuminSupportInstalled(InProjectFileFormat) && TargetType == TargetType.Game && InPlatform == UnrealTargetPlatform.Lumin)
 			{
-				string MLSDK = Utils.CleanDirectorySeparators(Environment.GetEnvironmentVariable("MLSDK"), '\\');
+				string MLSDK = Environment.GetEnvironmentVariable("MLSDK");
 
 				// TODO: Check if MPK name can be other than the project name.
 				string GameName = TargetRulesPath.GetFileNameWithoutExtension();
@@ -101,9 +101,6 @@ namespace UnrealBuildTool
                 ELFFile = Path.Combine(ELFFile, "..\\..\\Intermediate\\Lumin\\Mabu\\Packaged\\bin\\" + GameName);
 				string DebuggerFlavor = "MLDebugger";
 
-				string SymFile = Utils.MakePathRelativeTo(NMakeOutputPath.Directory.FullName, ProjectFilePath.Directory.FullName);
-				SymFile = Path.Combine(SymFile, "..\\..\\Intermediate\\Lumin\\Mabu\\Binaries", Path.ChangeExtension(NMakeOutputPath.GetFileName(), ".sym"));
-
 				// following are defaults for debugger options
 				string Attach = "false";
 				string EnableAutoStop = "true";
@@ -119,43 +116,9 @@ namespace UnrealBuildTool
 													"<EnableAutoStop>{5}</EnableAutoStop>" + ProjectFileGenerator.NewLine +
 													"<AutoStopAtFunction>{6}</AutoStopAtFunction>" + ProjectFileGenerator.NewLine +
 													"<EnablePrettyPrinting>{7}</EnablePrettyPrinting>" + ProjectFileGenerator.NewLine +
-													"<MLDownloadOnStart>{8}</MLDownloadOnStart>" + ProjectFileGenerator.NewLine +
-													"<SymFile>{9}</SymFile>" + ProjectFileGenerator.NewLine;
-				ProjectFileBuilder.Append(ProjectFileGenerator.NewLine + string.Format(CustomPathEntriesTemplate, MLSDK, PackageFile, ELFFile, DebuggerFlavor, Attach, EnableAutoStop, AutoStopAtFunction, EnablePrettyPrinting, MLDownloadOnStart, SymFile));
-			}
-		}
-
-		public override string GetBuildCommandEntry(ProjectFile ProjectFile, DirectoryReference BatchFileDirectory, ScriptArguments FileArguments, FileReference OnlyGameProject)
-		{
-			if (OnlyGameProject != null && File.Exists(OnlyGameProject.FullName))
-			{
-				return ProjectFile.EscapePath(ProjectFile.NormalizeProjectPath(FileReference.Combine(BatchFileDirectory, "RunUAT.bat"))) + GetBatchFileArguments(FileArguments);
-			}
-			return ProjectFile.EscapePath(ProjectFile.NormalizeProjectPath(FileReference.Combine(BatchFileDirectory, "Build.bat"))) + base.GetBatchFileArguments(FileArguments);
-		}
-
-		public override string GetReBuildCommandEntry(ProjectFile ProjectFile, DirectoryReference BatchFileDirectory, ScriptArguments FileArguments, FileReference OnlyGameProject)
-		{
-			if (OnlyGameProject != null && File.Exists(OnlyGameProject.FullName))
-			{
-				return GetBuildCommandEntry(ProjectFile, BatchFileDirectory, FileArguments, OnlyGameProject) + " -clean";
+													"<MLDownloadOnStart>{8}</MLDownloadOnStart>" + ProjectFileGenerator.NewLine;
+				ProjectFileBuilder.Append(ProjectFileGenerator.NewLine + string.Format(CustomPathEntriesTemplate, MLSDK, PackageFile, ELFFile, DebuggerFlavor, Attach, EnableAutoStop, AutoStopAtFunction, EnablePrettyPrinting, MLDownloadOnStart));
 			}
-			return ProjectFile.EscapePath(ProjectFile.NormalizeProjectPath(FileReference.Combine(BatchFileDirectory, "Rebuild.bat"))) + base.GetBatchFileArguments(FileArguments);
-		}
-
-		public override string GetCleanCommandEntry(ProjectFile ProjectFile, DirectoryReference BatchFileDirectory, ScriptArguments FileArguments, FileReference OnlyGameProject)
-		{
-			return ProjectFile.EscapePath(ProjectFile.NormalizeProjectPath(FileReference.Combine(BatchFileDirectory, "Clean.bat"))) + base.GetBatchFileArguments(FileArguments);
-		}
-
-		public override string GetBatchFileArguments(ScriptArguments FileArguments)
-		{
-			return String.Format(" BuildCookRun -nocompileeditor -nop4 -project=\"$(SolutionDir)$(ProjectName).uproject\" -cook -fastcook -stage -archive -archivedirectory=\"$(SolutionDir)..\\..\\Binaries\\Lumin\" -package -clientconfig={0} -ue4exe=UE4Editor-CMD.exe -SkipCookingEditorContent -pak -prereqs -targetplatform=Lumin -build -utf8output -compile", FileArguments.ConfigurationName);
-		}
-
-		public override FileReference GetNMakeOutputPath(FileReference InNMakePath)
-		{
-			return new FileReference(Path.Combine(InNMakePath.Directory.FullName, InNMakePath.GetFileNameWithoutExtension()+"-arm64-lumin"+InNMakePath.GetExtension()));
 		}
 	}
 }
diff --git a/Engine/Source/Programs/UnrealBuildTool/ProjectFiles/VisualStudio/VCProject.cs b/Engine/Source/Programs/UnrealBuildTool/ProjectFiles/VisualStudio/VCProject.cs
--- a/Engine/Source/Programs/UnrealBuildTool/ProjectFiles/VisualStudio/VCProject.cs
+++ b/Engine/Source/Programs/UnrealBuildTool/ProjectFiles/VisualStudio/VCProject.cs
@@ -1416,15 +1416,43 @@ namespace UnrealBuildTool
 						VCProjectFileContent.Append(PathStrings);
 					}
 
-					UEPlatformProjectGenerator.ScriptArguments BatchScriptArguments = new UEPlatformProjectGenerator.ScriptArguments(ProjectFileGenerator.bUsePrecompiled, IsForeignProject, bUseFastPDB, TargetName, UBTPlatformName, UBTConfigurationName, BuildToolOverride);
+					// This is the standard UE4 based project NMake build line:
+					//	..\..\Build\BatchFiles\Build.bat <TARGETNAME> <PLATFORM> <CONFIGURATION>
+					//	ie ..\..\Build\BatchFiles\Build.bat BlankProgram Win64 Debug
+
+					string BuildArguments = " " + TargetName + " " + UBTPlatformName + " " + UBTConfigurationName;
+					if (ProjectFileGenerator.bUsePrecompiled)
+					{
+						BuildArguments += " -UsePrecompiled";
+					}
+					if (IsForeignProject)
+					{
+						BuildArguments += " " + UProjectPath;
+					}
+
+					// Always wait for the mutex between UBT invocations, so that building the whole solution doesn't fail.
+					BuildArguments += " -WaitMutex";
+
+					// Always include a flag to format log messages for MSBuild
+					BuildArguments += " -FromMsBuild";
+
+					if (bUseFastPDB)
+					{
+						// Pass Fast PDB option to make use of Visual Studio's /DEBUG:FASTLINK option
+						BuildArguments += " -FastPDB";
+					}
 
 					DirectoryReference BatchFilesDirectory = DirectoryReference.Combine(UnrealBuildTool.EngineDirectory, "Build", "BatchFiles");
 
+					if(BuildToolOverride != null)
+					{
+						BuildArguments += " " + BuildToolOverride;
+					}
+
 					// NMake Build command line
-					VCProjectFileContent.AppendLine("    <NMakeBuildCommandLine>{0}</NMakeBuildCommandLine>", ProjGenerator.GetBuildCommandEntry(this, BatchFilesDirectory, BatchScriptArguments, OnlyGameProject));
-					// NMake Rebuild command line
-					VCProjectFileContent.AppendLine("    <NMakeReBuildCommandLine>{0}</NMakeReBuildCommandLine>", ProjGenerator.GetReBuildCommandEntry(this, BatchFilesDirectory, BatchScriptArguments, OnlyGameProject));
-					VCProjectFileContent.AppendLine("    <NMakeCleanCommandLine>{0}</NMakeCleanCommandLine>", ProjGenerator.GetCleanCommandEntry(this, BatchFilesDirectory, BatchScriptArguments, OnlyGameProject));
+					VCProjectFileContent.AppendLine("    <NMakeBuildCommandLine>{0}{1}</NMakeBuildCommandLine>", EscapePath(NormalizeProjectPath(FileReference.Combine(BatchFilesDirectory, "Build.bat"))), BuildArguments.ToString());
+					VCProjectFileContent.AppendLine("    <NMakeReBuildCommandLine>{0}{1}</NMakeReBuildCommandLine>", EscapePath(NormalizeProjectPath(FileReference.Combine(BatchFilesDirectory, "Rebuild.bat"))), BuildArguments.ToString());
+					VCProjectFileContent.AppendLine("    <NMakeCleanCommandLine>{0}{1}</NMakeCleanCommandLine>", EscapePath(NormalizeProjectPath(FileReference.Combine(BatchFilesDirectory, "Clean.bat"))), BuildArguments.ToString());
 					VCProjectFileContent.AppendLine("    <NMakeOutput>{0}</NMakeOutput>", NormalizeProjectPath(NMakePath.FullName));
 
 					if (TargetRulesObject.Type == TargetType.Game || TargetRulesObject.Type == TargetType.Client || TargetRulesObject.Type == TargetType.Server)
@@ -1641,4 +1669,4 @@ namespace UnrealBuildTool
 		protected readonly Dictionary<UnrealTargetConfiguration, CsProjectInfo> CachedProjectInfo = new Dictionary<UnrealTargetConfiguration, CsProjectInfo>();
 	}
 
-}
\ No newline at end of file
+}
diff --git a/Engine/Source/Programs/UnrealBuildTool/System/UEPlatformProjectGenerator.cs b/Engine/Source/Programs/UnrealBuildTool/System/UEPlatformProjectGenerator.cs
--- a/Engine/Source/Programs/UnrealBuildTool/System/UEPlatformProjectGenerator.cs
+++ b/Engine/Source/Programs/UnrealBuildTool/System/UEPlatformProjectGenerator.cs
@@ -14,27 +14,6 @@ namespace UnrealBuildTool
 	/// </summary>
 	abstract class UEPlatformProjectGenerator
 	{
-		public class ScriptArguments 
-		{
-			public ScriptArguments(bool usePrecompiled, bool isForeignProject, bool useFastPDB, string targetName, string platformName, string configurationName, string buildToolOverride)
-			{
-				UsePrecompiled = usePrecompiled;
-				IsForeignProject = isForeignProject;
-				UseFastPDB = useFastPDB;
-				TargetName = targetName;
-				PlatformName = platformName;
-				ConfigurationName = configurationName;
-				BuildToolOverride = buildToolOverride;
-			}
-			public bool UsePrecompiled { get; }
-			public bool IsForeignProject { get; }
-			public bool UseFastPDB { get; }
-			public string TargetName { get; }
-			public string PlatformName { get; }
-			public string ConfigurationName { get; }
-			public string BuildToolOverride { get; }
-		}
-
 		static Dictionary<UnrealTargetPlatform, UEPlatformProjectGenerator> ProjectGeneratorDictionary = new Dictionary<UnrealTargetPlatform, UEPlatformProjectGenerator>();
 
 		/// <summary>
@@ -188,79 +167,6 @@ namespace UnrealBuildTool
 			return InPlatform.ToString();
 		}
 
-		/// <summary>
-		///  Returns the default build batch script entry and its arguments
-		/// </summary>
-		/// <returns></returns>
-		public virtual string GetBuildCommandEntry(ProjectFile ProjectFile, DirectoryReference BatchFileDirectory, ScriptArguments FileArguments, FileReference OnlyGameProject)
-		{
-			return ProjectFile.EscapePath(ProjectFile.NormalizeProjectPath(FileReference.Combine(BatchFileDirectory, "Build.bat"))) + GetBatchFileArguments(FileArguments);
-		}
-
-		/// <summary>
-		///  Returns the default rebuild batch script entry and its arguments
-		/// </summary>
-		/// <returns></returns>
-		public virtual string GetReBuildCommandEntry(ProjectFile ProjectFile, DirectoryReference BatchFileDirectory, ScriptArguments FileArguments, FileReference OnlyGameProject)
-		{
-			return ProjectFile.EscapePath(ProjectFile.NormalizeProjectPath(FileReference.Combine(BatchFileDirectory, "Rebuild.bat"))) + GetBatchFileArguments(FileArguments);
-		}
-
-		/// <summary>
-		/// Returns the default clean batch script filename
-		/// </summary>
-		/// <returns></returns>
-		public virtual string GetCleanCommandEntry(ProjectFile ProjectFile, DirectoryReference BatchFileDirectory, ScriptArguments FileArguments, FileReference OnlyGameProject)
-		{
-			return ProjectFile.EscapePath(ProjectFile.NormalizeProjectPath(FileReference.Combine(BatchFileDirectory, "Clean.bat"))) + GetBatchFileArguments(FileArguments);
-		}
-
-		/// <summary>
-		/// Builds an argument string passed to build, rebuild and clean batch scripts
-		/// </summary>
-		/// <param name="FileArguments"></param>
-		/// <returns></returns>
-		public virtual string GetBatchFileArguments(ScriptArguments FileArguments)
-		{
-			// This is the standard UE4 based project NMake build line:
-			//	..\..\Build\BatchFiles\Build.bat <TARGETNAME> <PLATFORM> <CONFIGURATION>
-			//	ie ..\..\Build\BatchFiles\Build.bat BlankProgram Win64 Debu
-			string BuildArguments = " " + FileArguments.TargetName + " " + FileArguments.PlatformName + " " + FileArguments.ConfigurationName;
-			if (FileArguments.UsePrecompiled)
-			{
-				BuildArguments += " -UsePrecompiled";
-			}
-			if (FileArguments.IsForeignProject)
-			{
-				BuildArguments += " \"$(SolutionDir)$(ProjectName).uproject\"";
-			}
-
-			// Always wait for the mutex between UBT invocations, so that building the whole solution doesn't fail.
-			BuildArguments += " -WaitMutex";
-			// Always include a flag to format log messages for MSBuild
-			BuildArguments += " -FromMsBuild";
-			if (FileArguments.UseFastPDB)
-			{
-				BuildArguments += " -FastPDB";
-			}
-			if (!String.IsNullOrEmpty(FileArguments.BuildToolOverride))
-			{
-				// Pass Fast PDB option to make use of Visual Studio's /DEBUG:FASTLINK option
-				BuildArguments += " " + FileArguments.BuildToolOverride;
-			}
-			return BuildArguments;
-		}
-
-		/// <summary>
-		/// Returns the NMake output file for the particular generator.
-		/// </summary>
-		/// <param name="InNMakePath">  The current base nmake output file.</param>
-		/// <returns>FileReference The new NMake output file.</returns>
-		public virtual FileReference GetNMakeOutputPath(FileReference InNMakePath)
-		{
-			return InNMakePath;
-		}
-
 		/// <summary>
 		/// Return project configuration settings that must be included before the default props file
 		/// </summary>
