Bring back sym file linkage for debugging Lumin via VS. This is part of a change where we strip the debug symbols from the executable and put them in a separate .sym file. VS picks up the debug symbols directly from this file, helping reduce the size of the mpk significantly. This was reverted along with the other changes to LuminProjectGenerator in CL 4642351.

diff --git a/Engine/Source/Programs/UnrealBuildTool/Platform/Lumin/LuminProjectGenerator.cs b/Engine/Source/Programs/UnrealBuildTool/Platform/Lumin/LuminProjectGenerator.cs
--- a/Engine/Source/Programs/UnrealBuildTool/Platform/Lumin/LuminProjectGenerator.cs
+++ b/Engine/Source/Programs/UnrealBuildTool/Platform/Lumin/LuminProjectGenerator.cs
@@ -82,7 +82,7 @@ namespace UnrealBuildTool
 
 			if (IsVSLuminSupportInstalled(InProjectFileFormat) && TargetType == TargetType.Game && InPlatform == UnrealTargetPlatform.Lumin)
 			{
-				string MLSDK = Environment.GetEnvironmentVariable("MLSDK");
+				string MLSDK = Utils.CleanDirectorySeparators(Environment.GetEnvironmentVariable("MLSDK"), '\\');
 
 				// TODO: Check if MPK name can be other than the project name.
 				string GameName = TargetRulesPath.GetFileNameWithoutExtension();
@@ -101,6 +101,9 @@ namespace UnrealBuildTool
                 ELFFile = Path.Combine(ELFFile, "..\\..\\Intermediate\\Lumin\\Mabu\\Packaged\\bin\\" + GameName);
 				string DebuggerFlavor = "MLDebugger";
 
+				string SymFile = Utils.MakePathRelativeTo(NMakeOutputPath.Directory.FullName, ProjectFilePath.Directory.FullName);
+				SymFile = Path.Combine(SymFile, "..\\..\\Intermediate\\Lumin\\Mabu\\Binaries", Path.ChangeExtension(NMakeOutputPath.GetFileName(), ".sym"));
+
 				// following are defaults for debugger options
 				string Attach = "false";
 				string EnableAutoStop = "true";
@@ -116,8 +119,9 @@ namespace UnrealBuildTool
 													"<EnableAutoStop>{5}</EnableAutoStop>" + ProjectFileGenerator.NewLine +
 													"<AutoStopAtFunction>{6}</AutoStopAtFunction>" + ProjectFileGenerator.NewLine +
 													"<EnablePrettyPrinting>{7}</EnablePrettyPrinting>" + ProjectFileGenerator.NewLine +
-													"<MLDownloadOnStart>{8}</MLDownloadOnStart>" + ProjectFileGenerator.NewLine;
-				ProjectFileBuilder.Append(ProjectFileGenerator.NewLine + string.Format(CustomPathEntriesTemplate, MLSDK, PackageFile, ELFFile, DebuggerFlavor, Attach, EnableAutoStop, AutoStopAtFunction, EnablePrettyPrinting, MLDownloadOnStart));
+													"<MLDownloadOnStart>{8}</MLDownloadOnStart>" + ProjectFileGenerator.NewLine +
+													"<SymFile>{9}</SymFile>" + ProjectFileGenerator.NewLine;
+				ProjectFileBuilder.Append(ProjectFileGenerator.NewLine + string.Format(CustomPathEntriesTemplate, MLSDK, PackageFile, ELFFile, DebuggerFlavor, Attach, EnableAutoStop, AutoStopAtFunction, EnablePrettyPrinting, MLDownloadOnStart, SymFile));
 			}
 		}
 	}
