Clang path used to be in quotes, but some part of FileReference in AndroidToolChain was choking on the quotes.  Don't put your MLSDK in a directory with spaces for now, I guess.

diff --git a/Engine/Source/Programs/UnrealBuildTool/Platform/Lumin/LuminToolChain.cs b/Engine/Source/Programs/UnrealBuildTool/Platform/Lumin/LuminToolChain.cs
--- a/Engine/Source/Programs/UnrealBuildTool/Platform/Lumin/LuminToolChain.cs
+++ b/Engine/Source/Programs/UnrealBuildTool/Platform/Lumin/LuminToolChain.cs
@@ -114,7 +114,8 @@ namespace UnrealBuildTool
 			}
 
 			// set up the path to our toolchains
-			ClangPath = "\"" + ToolsDict["CXX"] + "\"";
+			// Clang path used to be in quotes, but some part of FileReference in AndroidToolChain was choking on the quotes.  Don't put your MLSDK in a directory with spaces for now, I guess.
+			ClangPath = ToolsDict["CXX"];
 			ArPathArm64 = "\"" + ToolsDict["AR"] + "\"";
 			// The strip command does not execute through the shell. Hence we don't quote it.
 			StripPath = ToolsDict["STRIP"];
