Merging //UE4/Dev-Main to Dev-VR (//UE4/Dev-VR)

diff --git a/Engine/Source/Developer/TargetPlatform/Private/TargetPlatformManagerModule.cpp b/Engine/Source/Developer/TargetPlatform/Private/TargetPlatformManagerModule.cpp
--- a/Engine/Source/Developer/TargetPlatform/Private/TargetPlatformManagerModule.cpp
+++ b/Engine/Source/Developer/TargetPlatform/Private/TargetPlatformManagerModule.cpp
@@ -92,7 +92,7 @@ public:
 			// amortize UBT cost by calling it once for all platforms, rather than once per platform.
 			if (FParse::Param(FCommandLine::Get(), TEXT("Multiprocess"))==false)
 			{
-				FString UBTParams(TEXT("-SetupPlatforms"));
+				FString UBTParams(TEXT("-Mode=SetupPlatforms"));
 				int32 UBTReturnCode = -1;
 				FString UBTOutput;
 				if (!FDesktopPlatformModule::Get()->InvokeUnrealBuildToolSync(UBTParams, *GLog, true, UBTReturnCode, UBTOutput))
diff --git a/Engine/Source/Programs/UnrealBuildTool/Configuration/UEBuildTarget.cs b/Engine/Source/Programs/UnrealBuildTool/Configuration/UEBuildTarget.cs
--- a/Engine/Source/Programs/UnrealBuildTool/Configuration/UEBuildTarget.cs
+++ b/Engine/Source/Programs/UnrealBuildTool/Configuration/UEBuildTarget.cs
@@ -1473,6 +1473,10 @@ namespace UnrealBuildTool
 					{
 						bResult &= Binary.CheckRestrictedFolders(DirectoryReference.FromFile(ProjectFile), ModuleRestrictedFolderCache);
 					}
+					foreach(KeyValuePair<FileReference, FileReference> Pair in RuntimeDependencyTargetFileToSourceFile)
+					{
+						bResult &= CheckRestrictedFolders(Pair.Key, Pair.Value);
+					}
 
 					if(!bResult)
 					{
@@ -1510,6 +1514,53 @@ namespace UnrealBuildTool
 			return Makefile;
 		}
 
+		/// <summary>
+		/// Check that copying a file from one location to another does not violate rules regarding restricted folders
+		/// </summary>
+		/// <param name="TargetFile">The destination location for the file</param>
+		/// <param name="SourceFile">The source location of the file</param>
+		/// <returns>True if the copy is permitted, false otherwise</returns>
+		bool CheckRestrictedFolders(FileReference TargetFile, FileReference SourceFile)
+		{
+			List<RestrictedFolder> TargetRestrictedFolders = GetRestrictedFolders(TargetFile);
+			List<RestrictedFolder> SourceRestrictedFolders = GetRestrictedFolders(SourceFile);
+			foreach(RestrictedFolder SourceRestrictedFolder in SourceRestrictedFolders)
+			{
+				if(!TargetRestrictedFolders.Contains(SourceRestrictedFolder))
+				{
+					Log.TraceError("Runtime dependency '{0}' is copied to '{1}', which does not contain a '{2}' folder.", SourceFile, TargetFile, SourceRestrictedFolder);
+					return false;
+				}
+			}
+			return true;
+		}
+
+		/// <summary>
+		/// Gets the restricted folders that the given file is in
+		/// </summary>
+		/// <param name="File">The file to test</param>
+		/// <returns>List of restricted folders for the file</returns>
+		List<RestrictedFolder> GetRestrictedFolders(FileReference File)
+		{
+			// Find the base directory for this binary
+			DirectoryReference BaseDir;
+			if(File.IsUnderDirectory(UnrealBuildTool.RootDirectory))
+			{
+				BaseDir = UnrealBuildTool.RootDirectory;
+			}
+			else if(ProjectDirectory != null && File.IsUnderDirectory(ProjectDirectory))
+			{
+				BaseDir = ProjectDirectory;
+			}
+			else
+			{
+				return new List<RestrictedFolder>();
+			}
+
+			// Find the restricted folders under the base directory
+			return RestrictedFolders.FindRestrictedFolders(BaseDir, File.Directory);
+		}
+
 		/// <summary>
 		/// Creates an action which copies a file from one location to another
 		/// </summary>
diff --git a/Engine/Source/Programs/UnrealBuildTool/Platform/Windows/UEBuildWindows.cs b/Engine/Source/Programs/UnrealBuildTool/Platform/Windows/UEBuildWindows.cs
--- a/Engine/Source/Programs/UnrealBuildTool/Platform/Windows/UEBuildWindows.cs
+++ b/Engine/Source/Programs/UnrealBuildTool/Platform/Windows/UEBuildWindows.cs
@@ -428,7 +428,7 @@ namespace UnrealBuildTool
 		/// <summary>
 		/// The default compiler version to be used, if installed. 
 		/// </summary>
-		static readonly VersionNumber DefaultVisualStudioToolChainVersion = VersionNumber.Parse("14.13.26128");
+		static readonly VersionNumber DefaultVisualStudioToolChainVersion = VersionNumber.Parse("14.16.27023");
 
 		/// <summary>
 		/// The default Windows SDK version to be used, if installed.
diff --git a/Engine/Source/Runtime/Engine/Private/Audio.cpp b/Engine/Source/Runtime/Engine/Private/Audio.cpp
--- a/Engine/Source/Runtime/Engine/Private/Audio.cpp
+++ b/Engine/Source/Runtime/Engine/Private/Audio.cpp
@@ -761,6 +761,7 @@ FWaveInstance::FWaveInstance( FActiveSound* InActiveSound )
 	: WaveData(nullptr)
 	, SoundClass(nullptr)
 	, SoundSubmix(nullptr)
+	, SourceEffectChain(nullptr)
 	, ActiveSound(InActiveSound)
 	, Volume(0.0f)
 	, DistanceAttenuation(1.0f)
@@ -790,11 +791,14 @@ FWaveInstance::FWaveInstance( FActiveSound* InActiveSound )
 	, bIsMusic(false)
 	, bReverb(true)
 	, bCenterChannelOnly(false)
+	, bIsPaused(false)
 	, bReportedSpatializationWarning(false)
 	, bIsAmbisonics(false)
 	, bIsStopping(false)
 	, SpatializationMethod(ESoundSpatializationAlgorithm::SPATIALIZATION_Default)
+	, SpatializationPluginSettings(nullptr)
 	, OcclusionPluginSettings(nullptr)
+	, ReverbPluginSettings(nullptr)
 	, OutputTarget(EAudioOutputTarget::Speaker)
 	, LowPassFilterFrequency(MAX_FILTER_FREQUENCY)
 	, OcclusionFilterFrequency(MAX_FILTER_FREQUENCY)
@@ -808,8 +812,13 @@ FWaveInstance::FWaveInstance( FActiveSound* InActiveSound )
 	, AttenuationDistance(0.0f)
 	, ListenerToSoundDistance(0.0f)
 	, AbsoluteAzimuth(0.0f)
+	, PlaybackTime(0.0f)
+	, ReverbSendMethod(EReverbSendMethod::Linear)
 	, ReverbSendLevelRange(0.0f, 0.0f)
 	, ReverbSendLevelDistanceRange(0.0f, 0.0f)
+	, ManualReverbSendLevel(0.0f)
+	, TypeHash(0)
+	, WaveInstanceHash(0)
 	, UserIndex(0)
 {
 	TypeHash = ++TypeHashCounter;
diff --git a/Engine/Source/Runtime/Launch/Private/Windows/LaunchWindows.cpp b/Engine/Source/Runtime/Launch/Private/Windows/LaunchWindows.cpp
--- a/Engine/Source/Runtime/Launch/Private/Windows/LaunchWindows.cpp
+++ b/Engine/Source/Runtime/Launch/Private/Windows/LaunchWindows.cpp
@@ -220,6 +220,12 @@ int32 WINAPI WinMain( _In_ HINSTANCE hInInstance, _In_opt_ HINSTANCE hPrevInstan
 		CmdLine = *GSavedCommandLine;
 	}
 
+	// If we're running in unattended mode, make sure we never display error dialogs if we crash.
+	if ( FParse::Param( CmdLine, TEXT("unattended") ) )
+	{
+		SetErrorMode(SEM_FAILCRITICALERRORS | SEM_NOGPFAULTERRORBOX | SEM_NOOPENFILEERRORBOX);
+	}
+
 #if !(UE_BUILD_SHIPPING && WITH_EDITOR)
 	// Named mutex we use to figure out whether we are the first instance of the game running. This is needed to e.g.
 	// make sure there is no contention when trying to save the shader cache.
diff --git a/Engine/Source/Runtime/MediaAssets/Private/Assets/MediaPlayer.cpp b/Engine/Source/Runtime/MediaAssets/Private/Assets/MediaPlayer.cpp
--- a/Engine/Source/Runtime/MediaAssets/Private/Assets/MediaPlayer.cpp
+++ b/Engine/Source/Runtime/MediaAssets/Private/Assets/MediaPlayer.cpp
@@ -225,6 +225,15 @@ FTimespan UMediaPlayer::GetTime() const
 	return PlayerFacade->GetTime();
 }
 
+FTimespan UMediaPlayer::GetLastAudioSampleProcessedTime() const
+{
+	return PlayerFacade->GetLastAudioSampleProcessedTime();
+}
+
+FTimespan UMediaPlayer::GetLastVideoSampleProcessedTime() const
+{
+	return PlayerFacade->GetLastVideoSampleProcessedTime();
+}
 
 FText UMediaPlayer::GetTrackDisplayName(EMediaPlayerTrack TrackType, int32 TrackIndex) const
 {
diff --git a/Engine/Source/Runtime/MediaAssets/Public/MediaPlayer.h b/Engine/Source/Runtime/MediaAssets/Public/MediaPlayer.h
--- a/Engine/Source/Runtime/MediaAssets/Public/MediaPlayer.h
+++ b/Engine/Source/Runtime/MediaAssets/Public/MediaPlayer.h
@@ -285,6 +285,22 @@ public:
 	UFUNCTION(BlueprintCallable, Category="Media|MediaPlayer")
 	FTimespan GetTime() const;
 
+	/**
+	 * Get time of last audio sample decoded
+	 *
+	 * @return Time of last audio sample decoded.
+	 */
+	UFUNCTION(BlueprintCallable, Category="Media|MediaPlayer")
+	FTimespan GetLastAudioSampleProcessedTime() const;
+
+	/**
+	 * Get time of last video sample decoded
+	 *
+	 * @return Time of last video sample decoded.
+	 */
+	UFUNCTION(BlueprintCallable, Category="Media|MediaPlayer")
+	FTimespan GetLastVideoSampleProcessedTime() const;
+
 	/**
 	 * Get the human readable name of the specified track.
 	 *
diff --git a/Engine/Source/Runtime/SlateCore/Private/Layout/Clipping.cpp b/Engine/Source/Runtime/SlateCore/Private/Layout/Clipping.cpp
--- a/Engine/Source/Runtime/SlateCore/Private/Layout/Clipping.cpp
+++ b/Engine/Source/Runtime/SlateCore/Private/Layout/Clipping.cpp
@@ -77,11 +77,19 @@ FSlateClippingZone::FSlateClippingZone(const FVector2D& InTopLeft, const FVector
 
 void FSlateClippingZone::InitializeFromArbitraryPoints(const FVector2D& InTopLeft, const FVector2D& InTopRight, const FVector2D& InBottomLeft, const FVector2D& InBottomRight)
 {
-	bIsAxisAligned =
-		FMath::RoundToInt(InTopLeft.X) == FMath::RoundToInt(InBottomLeft.X) &&
-		FMath::RoundToInt(InTopRight.X) == FMath::RoundToInt(InBottomRight.X) &&
-		FMath::RoundToInt(InTopLeft.Y) == FMath::RoundToInt(InTopRight.Y) &&
-		FMath::RoundToInt(InBottomLeft.Y) == FMath::RoundToInt(InBottomRight.Y);
+	FVector2D Points[4] = { InTopLeft, InTopRight, InBottomLeft, InBottomRight };
+
+	FBox2D AABB(Points, 4);
+
+	bIsAxisAligned = true;
+	for (int32 i = 0; i < 4; ++i)
+	{
+		if (!AABB.GetClosestPointTo(Points[i]).Equals(Points[i], 0.1f))
+		{
+			bIsAxisAligned = false;
+			break;
+		}
+	}
 
 	if ( bIsAxisAligned )
 	{
