Fixed TryGetWindowsSdkDir just checking for null instead of IsNullOrEmpty

diff --git a/Engine/Source/Programs/UnrealBuildTool/Platform/Windows/UEBuildWindows.cs b/Engine/Source/Programs/UnrealBuildTool/Platform/Windows/UEBuildWindows.cs
--- a/Engine/Source/Programs/UnrealBuildTool/Platform/Windows/UEBuildWindows.cs
+++ b/Engine/Source/Programs/UnrealBuildTool/Platform/Windows/UEBuildWindows.cs
@@ -1406,7 +1406,7 @@ namespace UnrealBuildTool
 
 			// Figure out which version number to look for
 			VersionNumber WindowsSdkVersion = null;
-			if(DesiredVersion != null)
+			if(!string.IsNullOrEmpty(DesiredVersion))
 			{
 				if(String.Compare(DesiredVersion, "Latest", StringComparison.InvariantCultureIgnoreCase) == 0 && CachedWindowsSdkDirs.Count > 0)
 				{
