Made UEBuildHoloLens skip checking the registry for the proper SDK when compiling with Mono

diff --git a/Engine/Source/Programs/UnrealBuildTool/Platform/HoloLens/UEBuildHoloLens.cs b/Engine/Source/Programs/UnrealBuildTool/Platform/HoloLens/UEBuildHoloLens.cs
--- a/Engine/Source/Programs/UnrealBuildTool/Platform/HoloLens/UEBuildHoloLens.cs
+++ b/Engine/Source/Programs/UnrealBuildTool/Platform/HoloLens/UEBuildHoloLens.cs
@@ -769,6 +769,7 @@ namespace UnrealBuildTool
 				@"HKEY_LOCAL_MACHINE\SOFTWARE\Wow6432Node\Microsoft\Microsoft SDKs\Windows\",
 				@"HKEY_CURRENT_USER\SOFTWARE\Wow6432Node\Microsoft\Microsoft SDKs\Windows\"
 			};
+#if !__MonoCS__
 			foreach (string regLocation in possibleRegLocations)
 			{
 				object Result = Microsoft.Win32.Registry.GetValue(@"HKEY_LOCAL_MACHINE\SOFTWARE\Wow6432Node\Microsoft\Microsoft SDKs\Windows\" + Version, "InstallationFolder", null);
@@ -781,7 +782,7 @@ namespace UnrealBuildTool
 					break;
 				}
 			}
-
+#endif
 		}
 
 		protected override SDKStatus HasRequiredManualSDKInternal()
