Added a define for __MonoCS__ on Mac since mono on Mac doesn't define it as expected
Wrapped some Windows runtime code in Mono checks

diff --git a/Engine/Build/BatchFiles/Mac/RunXBuild.sh b/Engine/Build/BatchFiles/Mac/RunXBuild.sh
--- a/Engine/Build/BatchFiles/Mac/RunXBuild.sh
+++ b/Engine/Build/BatchFiles/Mac/RunXBuild.sh
@@ -20,7 +20,7 @@ if [ ! -f Build/BatchFiles/Mac/RunXBuild.sh ]; then
 	exit 1
 fi
 
-xbuild /verbosity:quiet /nologo "$@" |grep -i error
+xbuild /verbosity:quiet /nologo /p:DefineConstants=__MonoCS__ "$@" |grep -i error
 if [ $? -ne 1 ]; then
 	exit 1
 else
diff --git a/Engine/Source/Programs/UnrealBuildTool/Platform/HoloLens/WinMDRegistrationInfo.cs b/Engine/Source/Programs/UnrealBuildTool/Platform/HoloLens/WinMDRegistrationInfo.cs
--- a/Engine/Source/Programs/UnrealBuildTool/Platform/HoloLens/WinMDRegistrationInfo.cs
+++ b/Engine/Source/Programs/UnrealBuildTool/Platform/HoloLens/WinMDRegistrationInfo.cs
@@ -3,7 +3,9 @@
 using System;
 using System.Collections.Generic;
 using System.Linq;
+#if !__MonoCS__
 using Windows.Foundation.Metadata;
+#endif
 using System.Runtime.InteropServices.WindowsRuntime;
 using System.Reflection;
 using Tools.DotNETCommon;
@@ -28,7 +30,9 @@ namespace UnrealBuildTool
 			public ActivatableType(string InTypeName, object InThreadingModel)
 			{
 				TypeName = InTypeName;
+#if !__MonoCS__
 				ThreadingModelName = ((ThreadingModel)InThreadingModel).ToString().ToLowerInvariant();
+#endif
 			}
 			
 			/// <summary>
@@ -52,6 +56,7 @@ namespace UnrealBuildTool
 			PackageRelativeDllPath = InPackageRelativeDllPath;
 			ResolveSearchPaths.Add(InWindMDSourcePath.Directory.FullName);
 
+#if !__MonoCS__
 			ActivatableTypesList = new List<ActivatableType>();
 			var DependsOn = Assembly.ReflectionOnlyLoadFrom(InWindMDSourcePath.FullName);
 			foreach (var WinMDType in DependsOn.GetExportedTypes())
@@ -75,6 +80,7 @@ namespace UnrealBuildTool
 					ActivatableTypesList.Add(new ActivatableType(WinMDType.FullName, ThreadingModel));
 				}
 			}
+#endif
 		}
 
 		/// <summary>
@@ -95,6 +101,7 @@ namespace UnrealBuildTool
 
 		static WinMDRegistrationInfo()
 		{
+#if !__MonoCS__
 			AppDomain.CurrentDomain.ReflectionOnlyAssemblyResolve += (Sender, EventArgs) => Assembly.ReflectionOnlyLoad(EventArgs.Name);
 			WindowsRuntimeMetadata.ReflectionOnlyNamespaceResolve += (Sender, EventArgs) =>
 			{
@@ -105,6 +112,7 @@ namespace UnrealBuildTool
 				}
 				EventArgs.ResolvedAssemblies.Add(Assembly.ReflectionOnlyLoadFrom(Path));
 			};
+#endif
 		}
 
 		private static List<string> ResolveSearchPaths = new List<string>();
