UnrealBuildTool project changes:
- Moved the HoloLens files to be conditional like other platforms
- Wrapped the registry SDK path check in a !__MonoCS__ choose block so it's only done on Windows platforms

diff --git a/Engine/Source/Programs/UnrealBuildTool/UnrealBuildTool.csproj b/Engine/Source/Programs/UnrealBuildTool/UnrealBuildTool.csproj
--- a/Engine/Source/Programs/UnrealBuildTool/UnrealBuildTool.csproj
+++ b/Engine/Source/Programs/UnrealBuildTool/UnrealBuildTool.csproj
@@ -271,31 +271,41 @@
   <ItemGroup>
     <Compile Include="ToolChain\RemoteMac.cs" />
     <Compile Include="UnrealBuildTool.cs" />
-    <Compile Include="Platform\HoloLens\UEBuildHoloLens.cs" />
-    <Compile Include="Platform\HoloLens\HoloLensDeploy.cs" />
-    <Compile Include="Platform\HoloLens\HoloLensExports.cs" />
-    <Compile Include="Platform\HoloLens\HoloLensProjectGenerator.cs" />
-    <Compile Include="Platform\HoloLens\HoloLensToolChain.cs" />
-    <Compile Include="Platform\HoloLens\ManifestGenerator.cs" />
-    <Compile Include="Platform\HoloLens\WinMDRegistrationInfo.cs" />
     <Compile Include="Platform\Windows\PVSToolChain.cs" />
   </ItemGroup>
-  <PropertyGroup>
-    <Windows10SdkPath>$(Registry:HKEY_LOCAL_MACHINE\SOFTWARE\Wow6432Node\Microsoft\Microsoft SDKs\Windows\v10.0@InstallationFolder)</Windows10SdkPath>
-    <Windows10SdkLatestVersion>$(Registry:HKEY_LOCAL_MACHINE\SOFTWARE\Wow6432Node\Microsoft\Microsoft SDKs\Windows\v10.0@ProductVersion)</Windows10SdkLatestVersion>
-  </PropertyGroup>
-  <ItemGroup>
-    <Reference Include="Windows">
-      <!-- This is a dummy entry to throw off the parser in DotNETUtilities's CsProjectInfo which doesn't understand MSBuild functions -->
-      <HintPath/>      
-      <HintPath>$(Windows10SdkPath)\UnionMetadata\$(Windows10SdkLatestVersion)\Windows.winmd</HintPath>
-      <!-- 16299 SDK (possibly others?) has ProductVersion == 10.0.16299 in the registry, but pathing needs the final .0 -->
-      <HintPath Condition="'$([System.IO.File]::Exists(%(HintPath)))'=='false'">$(Windows10SdkPath)\UnionMetadata\$(Windows10SdkLatestVersion).0\Windows.winmd</HintPath>
-      <!-- Fall back to old, unversioned location -->
-      <HintPath Condition="'$([System.IO.File]::Exists(%(HintPath)))'=='false'">$(Windows10SdkPath)\UnionMetadata\Windows.winmd</HintPath>
-      <Private>false</Private>
-    </Reference>
-  </ItemGroup>
+  <Choose>
+    <When Condition="Exists('Platform\HoloLens\UEBuildHoloLens.cs')">
+      <ItemGroup>
+		<Compile Include="Platform\HoloLens\UEBuildHoloLens.cs" />
+		<Compile Include="Platform\HoloLens\HoloLensDeploy.cs" />
+		<Compile Include="Platform\HoloLens\HoloLensExports.cs" />
+		<Compile Include="Platform\HoloLens\HoloLensProjectGenerator.cs" />
+		<Compile Include="Platform\HoloLens\HoloLensToolChain.cs" />
+		<Compile Include="Platform\HoloLens\ManifestGenerator.cs" />
+		<Compile Include="Platform\HoloLens\WinMDRegistrationInfo.cs" />
+      </ItemGroup>
+    </When>
+  </Choose>
+  <Choose>
+    <When Condition="!$(DefineConstants.Contains(__MonoCS__))">
+	  <PropertyGroup>
+		<Windows10SdkPath>$(Registry:HKEY_LOCAL_MACHINE\SOFTWARE\Wow6432Node\Microsoft\Microsoft SDKs\Windows\v10.0@InstallationFolder)</Windows10SdkPath>
+		<Windows10SdkLatestVersion>$(Registry:HKEY_LOCAL_MACHINE\SOFTWARE\Wow6432Node\Microsoft\Microsoft SDKs\Windows\v10.0@ProductVersion)</Windows10SdkLatestVersion>
+	  </PropertyGroup>
+	  <ItemGroup>
+		<Reference Include="Windows">
+		  <!-- This is a dummy entry to throw off the parser in DotNETUtilities's CsProjectInfo which doesn't understand MSBuild functions -->
+		  <HintPath/>      
+		  <HintPath>$(Windows10SdkPath)\UnionMetadata\$(Windows10SdkLatestVersion)\Windows.winmd</HintPath>
+		  <!-- 16299 SDK (possibly others?) has ProductVersion == 10.0.16299 in the registry, but pathing needs the final .0 -->
+		  <HintPath Condition="'$([System.IO.File]::Exists(%(HintPath)))'=='false'">$(Windows10SdkPath)\UnionMetadata\$(Windows10SdkLatestVersion).0\Windows.winmd</HintPath>
+		  <!-- Fall back to old, unversioned location -->
+		  <HintPath Condition="'$([System.IO.File]::Exists(%(HintPath)))'=='false'">$(Windows10SdkPath)\UnionMetadata\Windows.winmd</HintPath>
+		  <Private>false</Private>
+		</Reference>
+	  </ItemGroup>
+	</When>
+  </Choose>
   <Choose>
     <When Condition="Exists('Platform\IOS\UEBuildIOS.cs')">
       <ItemGroup>
