Changed the conditional inclusion of the HoloLens registry check to be by OS type

diff --git a/Engine/Source/Programs/UnrealBuildTool/UnrealBuildTool.csproj b/Engine/Source/Programs/UnrealBuildTool/UnrealBuildTool.csproj
--- a/Engine/Source/Programs/UnrealBuildTool/UnrealBuildTool.csproj
+++ b/Engine/Source/Programs/UnrealBuildTool/UnrealBuildTool.csproj
@@ -287,7 +287,7 @@
     </When>
   </Choose>
   <Choose>
-    <When Condition="!$(DefineConstants.Contains(__MonoCS__))">
+    <When Condition="'$(OS)' == 'Windows_NT'">
 	  <PropertyGroup>
 		<Windows10SdkPath>$(Registry:HKEY_LOCAL_MACHINE\SOFTWARE\Wow6432Node\Microsoft\Microsoft SDKs\Windows\v10.0@InstallationFolder)</Windows10SdkPath>
 		<Windows10SdkLatestVersion>$(Registry:HKEY_LOCAL_MACHINE\SOFTWARE\Wow6432Node\Microsoft\Microsoft SDKs\Windows\v10.0@ProductVersion)</Windows10SdkLatestVersion>
