Don't do the windows stuff on linux

diff --git a/Engine/Source/Programs/UnrealBuildTool/Platform/HoloLens/WinMDRegistrationInfo.cs b/Engine/Source/Programs/UnrealBuildTool/Platform/HoloLens/WinMDRegistrationInfo.cs
--- a/Engine/Source/Programs/UnrealBuildTool/Platform/HoloLens/WinMDRegistrationInfo.cs
+++ b/Engine/Source/Programs/UnrealBuildTool/Platform/HoloLens/WinMDRegistrationInfo.cs
@@ -3,7 +3,7 @@
 using System;
 using System.Collections.Generic;
 using System.Linq;
-#if !__MonoCS__
+#if !__MonoCS__ && !PLATFORM_LINUX
 using Windows.Foundation.Metadata;
 #endif
 using System.Runtime.InteropServices.WindowsRuntime;
@@ -30,7 +30,7 @@ namespace UnrealBuildTool
 			public ActivatableType(string InTypeName, object InThreadingModel)
 			{
 				TypeName = InTypeName;
-#if !__MonoCS__
+#if !__MonoCS__ && !PLATFORM_LINUX
 				ThreadingModelName = ((ThreadingModel)InThreadingModel).ToString().ToLowerInvariant();
 #endif
 			}
@@ -56,7 +56,7 @@ namespace UnrealBuildTool
 			PackageRelativeDllPath = InPackageRelativeDllPath;
 			ResolveSearchPaths.Add(InWindMDSourcePath.Directory.FullName);
 
-#if !__MonoCS__
+#if !__MonoCS__ && !PLATFORM_LINUX
 			ActivatableTypesList = new List<ActivatableType>();
 			var DependsOn = Assembly.ReflectionOnlyLoadFrom(InWindMDSourcePath.FullName);
 			foreach (var WinMDType in DependsOn.GetExportedTypes())
@@ -101,7 +101,7 @@ namespace UnrealBuildTool
 
 		static WinMDRegistrationInfo()
 		{
-#if !__MonoCS__
+#if !__MonoCS__ && !PLATFORM_LINUX
 			AppDomain.CurrentDomain.ReflectionOnlyAssemblyResolve += (Sender, EventArgs) => Assembly.ReflectionOnlyLoad(EventArgs.Name);
 			WindowsRuntimeMetadata.ReflectionOnlyNamespaceResolve += (Sender, EventArgs) =>
 			{
