Fix for UEBuildTarget.cs issue with merge from main

diff --git a/Engine/Source/Programs/UnrealBuildTool/Configuration/UEBuildTarget.cs b/Engine/Source/Programs/UnrealBuildTool/Configuration/UEBuildTarget.cs
--- a/Engine/Source/Programs/UnrealBuildTool/Configuration/UEBuildTarget.cs
+++ b/Engine/Source/Programs/UnrealBuildTool/Configuration/UEBuildTarget.cs
@@ -200,12 +200,10 @@ namespace UnrealBuildTool
 		/// </summary>
 		public static UnrealTargetPlatform Win64 = FindOrAddByName("Win64");
 
-		// @ATG_CHANGE : BEGIN HoloLens support
 		/// <summary>
 		/// HoloLens
 		/// </summary>
-		HoloLens,
-		// @ATG_CHANGE : END
+		public static UnrealTargetPlatform HoloLens = FindOrAddByName("HoloLens");
 
 		/// <summary>
 		/// Mac
@@ -404,12 +402,10 @@ namespace UnrealBuildTool
 		/// </summary>
 		public static UnrealPlatformGroup Windows = FindOrAddByName("Windows");
 
-		// @ATG_CHANGE : BEGIN HoloLens support
 		/// <summary>
 		/// this group is just to lump HoloLens32 and HoloLens64 into HoloLens directories
 		/// </summary>
-		HoloLens,
-		// @ATG_CHANGE : END
+		public static UnrealPlatformGroup HoloLens = FindOrAddByName("HoloLens");
 
 		/// <summary>
 		/// Microsoft platforms
