Fix merge error dealing with the platform names in Ringtail -> Dev-VR merge

diff --git a/Engine/Source/Programs/UnrealBuildTool/Configuration/UEBuildTarget.cs b/Engine/Source/Programs/UnrealBuildTool/Configuration/UEBuildTarget.cs
--- a/Engine/Source/Programs/UnrealBuildTool/Configuration/UEBuildTarget.cs
+++ b/Engine/Source/Programs/UnrealBuildTool/Configuration/UEBuildTarget.cs
@@ -3484,8 +3484,8 @@ namespace UnrealBuildTool
 			GlobalCompileEnvironment.Definitions.Add(String.Format("UBT_MODULE_MANIFEST_DEBUGGAME=\"{0}\"", ModuleManifest.GetStandardFileName(AppName, Platform, UnrealTargetConfiguration.DebugGame, Architecture, true)));
 
 			// tell the compiled code the name of the UBT platform (this affects folder on disk, etc that the game may need to know)
-			GlobalCompileEnvironment.Definitions.Add(String.Format("UBT_COMPILED_PLATFORM=\"{0}\"", Platform.ToString()));
-			GlobalCompileEnvironment.Definitions.Add(String.Format("UBT_COMPILED_TARGET=\"{0}\"", TargetType.ToString()));
+			GlobalCompileEnvironment.Definitions.Add(String.Format("UBT_COMPILED_PLATFORM=" + Platform.ToString()));
+			GlobalCompileEnvironment.Definitions.Add(String.Format("UBT_COMPILED_TARGET=" + TargetType.ToString()));
 
 			// Set the global app name
 			GlobalCompileEnvironment.Definitions.Add(String.Format("UE_APP_NAME=\"{0}\"", AppName));
diff --git a/Engine/Source/Runtime/Core/Private/GenericPlatform/GenericPlatformMisc.cpp b/Engine/Source/Runtime/Core/Private/GenericPlatform/GenericPlatformMisc.cpp
--- a/Engine/Source/Runtime/Core/Private/GenericPlatform/GenericPlatformMisc.cpp
+++ b/Engine/Source/Runtime/Core/Private/GenericPlatform/GenericPlatformMisc.cpp
@@ -997,12 +997,12 @@ const TCHAR* FGenericPlatformMisc::GamePersistentDownloadDir()
 
 const TCHAR* FGenericPlatformMisc::GetUBTPlatform()
 {
-	return TEXT(UBT_COMPILED_PLATFORM);
+	return TEXT(PREPROCESSOR_TO_STRING(UBT_COMPILED_PLATFORM));
 }
 
 const TCHAR* FGenericPlatformMisc::GetUBTTarget()
 {
-    return TEXT(UBT_COMPILED_TARGET);
+	return TEXT(PREPROCESSOR_TO_STRING(UBT_COMPILED_TARGET));
 }
 
 const TCHAR* FGenericPlatformMisc::GetDefaultDeviceProfileName()
