Making checks for Mono runtime checks instead of #defines

diff --git a/Engine/Source/Programs/UnrealBuildTool/Platform/HoloLens/WinMDRegistrationInfo.cs b/Engine/Source/Programs/UnrealBuildTool/Platform/HoloLens/WinMDRegistrationInfo.cs
--- a/Engine/Source/Programs/UnrealBuildTool/Platform/HoloLens/WinMDRegistrationInfo.cs
+++ b/Engine/Source/Programs/UnrealBuildTool/Platform/HoloLens/WinMDRegistrationInfo.cs
@@ -3,9 +3,6 @@
 using System;
 using System.Collections.Generic;
 using System.Linq;
-#if !__MonoCS__
-using Windows.Foundation.Metadata;
-#endif
 using System.Runtime.InteropServices.WindowsRuntime;
 using System.Reflection;
 using Tools.DotNETCommon;
@@ -30,9 +27,10 @@ namespace UnrealBuildTool
 			public ActivatableType(string InTypeName, object InThreadingModel)
 			{
 				TypeName = InTypeName;
-#if !__MonoCS__
-				ThreadingModelName = ((ThreadingModel)InThreadingModel).ToString().ToLowerInvariant();
-#endif
+				if (BuildHostPlatform.Current.Platform == UnrealTargetPlatform.Win64)
+				{
+					ThreadingModelName = ((Windows.Foundation.Metadata.ThreadingModel)InThreadingModel).ToString().ToLowerInvariant();
+				}
 			}
 			
 			/// <summary>
@@ -56,31 +54,33 @@ namespace UnrealBuildTool
 			PackageRelativeDllPath = InPackageRelativeDllPath;
 			ResolveSearchPaths.Add(InWindMDSourcePath.Directory.FullName);
 
-#if !__MonoCS__
-			ActivatableTypesList = new List<ActivatableType>();
-			var DependsOn = Assembly.ReflectionOnlyLoadFrom(InWindMDSourcePath.FullName);
-			foreach (var WinMDType in DependsOn.GetExportedTypes())
+			if (BuildHostPlatform.Current.Platform == UnrealTargetPlatform.Win64)
 			{
-				bool IsActivatable = false;
-				object ThreadingModel = Windows.Foundation.Metadata.ThreadingModel.Both;
-				foreach (var Attr in WinMDType.CustomAttributes)
+				ActivatableTypesList = new List<ActivatableType>();
+				var DependsOn = Assembly.ReflectionOnlyLoadFrom(InWindMDSourcePath.FullName);
+				foreach (var WinMDType in DependsOn.GetExportedTypes())
 				{
-					if (Attr.AttributeType.AssemblyQualifiedName == typeof(ActivatableAttribute).AssemblyQualifiedName ||
-						Attr.AttributeType.AssemblyQualifiedName == typeof(StaticAttribute).AssemblyQualifiedName)
+					bool IsActivatable = false;
+					object ThreadingModel = Windows.Foundation.Metadata.ThreadingModel.Both;
+					foreach (var Attr in WinMDType.CustomAttributes)
 					{
-						IsActivatable = true;
+						if (Attr.AttributeType.AssemblyQualifiedName == typeof(Windows.Foundation.Metadata.ActivatableAttribute).AssemblyQualifiedName ||
+							Attr.AttributeType.AssemblyQualifiedName == typeof(Windows.Foundation.Metadata.StaticAttribute).AssemblyQualifiedName)
+						{
+							IsActivatable = true;
+						}
+						else if (Attr.AttributeType.AssemblyQualifiedName == typeof(Windows.Foundation.Metadata.ThreadingAttribute).AssemblyQualifiedName)
+						{
+							ThreadingModel = Attr.ConstructorArguments[0].Value;
+						}
 					}
-					else if (Attr.AttributeType.AssemblyQualifiedName == typeof(ThreadingAttribute).AssemblyQualifiedName)
+					if (IsActivatable)
 					{
-						ThreadingModel = Attr.ConstructorArguments[0].Value;
+						ActivatableTypesList.Add(new ActivatableType(WinMDType.FullName, ThreadingModel));
 					}
 				}
-				if (IsActivatable)
-				{
-					ActivatableTypesList.Add(new ActivatableType(WinMDType.FullName, ThreadingModel));
-				}
+
 			}
-#endif
 		}
 
 		/// <summary>
@@ -101,18 +101,20 @@ namespace UnrealBuildTool
 
 		static WinMDRegistrationInfo()
 		{
-#if !__MonoCS__
-			AppDomain.CurrentDomain.ReflectionOnlyAssemblyResolve += (Sender, EventArgs) => Assembly.ReflectionOnlyLoad(EventArgs.Name);
-			WindowsRuntimeMetadata.ReflectionOnlyNamespaceResolve += (Sender, EventArgs) =>
+			if (BuildHostPlatform.Current.Platform == UnrealTargetPlatform.Win64)
 			{
-				string Path = WindowsRuntimeMetadata.ResolveNamespace(EventArgs.NamespaceName, ResolveSearchPaths).FirstOrDefault();
-				if (Path == null)
+				AppDomain.CurrentDomain.ReflectionOnlyAssemblyResolve += (Sender, EventArgs) => Assembly.ReflectionOnlyLoad(EventArgs.Name);
+				WindowsRuntimeMetadata.ReflectionOnlyNamespaceResolve += (Sender, EventArgs) =>
 				{
-					return;
-				}
-				EventArgs.ResolvedAssemblies.Add(Assembly.ReflectionOnlyLoadFrom(Path));
-			};
-#endif
+					string Path = WindowsRuntimeMetadata.ResolveNamespace(EventArgs.NamespaceName, ResolveSearchPaths).FirstOrDefault();
+					if (Path == null)
+					{
+						return;
+					}
+					EventArgs.ResolvedAssemblies.Add(Assembly.ReflectionOnlyLoadFrom(Path));
+				};
+
+			}
 		}
 
 		private static List<string> ResolveSearchPaths = new List<string>();
