Rolled back UBT project change to nest choose/when for HoloLens bits because Mac/Linux didn't like it

diff --git a/Engine/Source/Programs/UnrealBuildTool/Platform/HoloLens/WinMDRegistrationInfo.cs b/Engine/Source/Programs/UnrealBuildTool/Platform/HoloLens/WinMDRegistrationInfo.cs
--- a/Engine/Source/Programs/UnrealBuildTool/Platform/HoloLens/WinMDRegistrationInfo.cs
+++ b/Engine/Source/Programs/UnrealBuildTool/Platform/HoloLens/WinMDRegistrationInfo.cs
@@ -29,7 +29,9 @@ namespace UnrealBuildTool
 				TypeName = InTypeName;
 				if (BuildHostPlatform.Current.Platform == UnrealTargetPlatform.Win64)
 				{
+#if !__MonoCS__ && !MONO
 					ThreadingModelName = ((Windows.Foundation.Metadata.ThreadingModel)InThreadingModel).ToString().ToLowerInvariant();
+#endif
 				}
 			}
 			
@@ -56,6 +58,7 @@ namespace UnrealBuildTool
 
 			if (BuildHostPlatform.Current.Platform == UnrealTargetPlatform.Win64)
 			{
+#if !__MonoCS__ && !MONO
 				ActivatableTypesList = new List<ActivatableType>();
 				var DependsOn = Assembly.ReflectionOnlyLoadFrom(InWindMDSourcePath.FullName);
 				foreach (var WinMDType in DependsOn.GetExportedTypes())
@@ -79,7 +82,7 @@ namespace UnrealBuildTool
 						ActivatableTypesList.Add(new ActivatableType(WinMDType.FullName, ThreadingModel));
 					}
 				}
-
+#endif
 			}
 		}
 
