Fix for 2019 project generation

diff --git a/Engine/Source/Programs/UnrealBuildTool/Platform/HoloLens/HoloLensProjectGenerator.cs b/Engine/Source/Programs/UnrealBuildTool/Platform/HoloLens/HoloLensProjectGenerator.cs
--- a/Engine/Source/Programs/UnrealBuildTool/Platform/HoloLens/HoloLensProjectGenerator.cs
+++ b/Engine/Source/Programs/UnrealBuildTool/Platform/HoloLens/HoloLensProjectGenerator.cs
@@ -77,7 +77,7 @@ namespace UnrealBuildTool
 				VCProjectFileContent.Append("		<WindowsTargetPlatformVersion>" + MaxTestedVersion + "</WindowsTargetPlatformVersion>" + ProjectFileGenerator.NewLine);
 			}
 
-			WindowsCompiler Compiler = WindowsCompiler.VisualStudio2017;
+			WindowsCompiler Compiler = WindowsCompiler.VisualStudio2019;
 			DirectoryReference PlatformWinMDLocation = HoloLens.GetCppCXMetadataLocation(Compiler, "Latest");
 			if (PlatformWinMDLocation == null || !FileReference.Exists(FileReference.Combine(PlatformWinMDLocation, "platform.winmd")))
 			{
diff --git a/Engine/Source/Programs/UnrealBuildTool/Platform/HoloLens/UEBuildHoloLens.cs b/Engine/Source/Programs/UnrealBuildTool/Platform/HoloLens/UEBuildHoloLens.cs
--- a/Engine/Source/Programs/UnrealBuildTool/Platform/HoloLens/UEBuildHoloLens.cs
+++ b/Engine/Source/Programs/UnrealBuildTool/Platform/HoloLens/UEBuildHoloLens.cs
@@ -326,6 +326,10 @@ namespace UnrealBuildTool
 			{
 				return DirectoryReference.Combine(SelectedToolChainDir, "lib", "x86", "Store", "references");
 			}
+			else if (Compiler >= WindowsCompiler.VisualStudio2019)
+			{
+				return DirectoryReference.Combine(SelectedToolChainDir, "lib", "x86", "Store", "references");
+			}
 			else
 			{
 				return null;
