Using UE_NAME_PROJECT_AFTER_FOLDER leads to "go to source" not working from the editor. Now handle SLN naming in one place thus what we named the SLN is written out consistently.

diff --git a/Engine/Source/Programs/UnrealBuildTool/ProjectFiles/ProjectFileGenerator.cs b/Engine/Source/Programs/UnrealBuildTool/ProjectFiles/ProjectFileGenerator.cs
--- a/Engine/Source/Programs/UnrealBuildTool/ProjectFiles/ProjectFileGenerator.cs
+++ b/Engine/Source/Programs/UnrealBuildTool/ProjectFiles/ProjectFileGenerator.cs
@@ -600,7 +600,11 @@ namespace UnrealBuildTool
 			else
 			{
 				// Set the master project name from the folder name
-				if(bMasterProjectNameFromFolder)
+				if (Environment.GetEnvironmentVariable("UE_NAME_PROJECT_AFTER_FOLDER") == "1")
+				{
+					MasterProjectName += "_" + Path.GetFileName(MasterProjectPath.ToString());
+				}
+				else if(bMasterProjectNameFromFolder)
 				{
 					string NewMasterProjectName = MasterProjectPath.GetDirectoryName();
 					if(!String.IsNullOrEmpty(NewMasterProjectName))
@@ -634,12 +638,6 @@ namespace UnrealBuildTool
 				}
 			}
 
-            // Optionally include the folder name in the project
-            if (Environment.GetEnvironmentVariable("UE_NAME_PROJECT_AFTER_FOLDER") == "1")
-            {
-                MasterProjectName += "_" + Path.GetFileName(MasterProjectPath.ToString());
-            }
-
 			bool bCleanProjectFiles = Arguments.Any(x => x.Equals("-CleanProjects", StringComparison.InvariantCultureIgnoreCase));
 			if (bCleanProjectFiles)
 			{
