Fix additional libraries so .so files do not become absolute on Android and correct the library filtering by architecture now lib isn't assumed

#ue4
#android

diff --git a/Engine/Source/Programs/UnrealBuildTool/Platform/Android/AndroidToolChain.cs b/Engine/Source/Programs/UnrealBuildTool/Platform/Android/AndroidToolChain.cs
--- a/Engine/Source/Programs/UnrealBuildTool/Platform/Android/AndroidToolChain.cs
+++ b/Engine/Source/Programs/UnrealBuildTool/Platform/Android/AndroidToolChain.cs
@@ -987,8 +987,15 @@ namespace UnrealBuildTool
 			return false;
 		}
 
-		bool ShouldSkipLib(string Lib, string Arch, string GPUArchitecture)
+		bool ShouldSkipLib(string FullLib, string Arch, string GPUArchitecture)
 		{
+			// strip any absolute path
+			string Lib = Path.GetFileNameWithoutExtension(FullLib);
+			if (Lib.StartsWith("lib"))
+			{
+				Lib = Lib.Substring(3);
+			}
+
 			// reject any libs we outright don't want to link with
 			foreach (string LibName in LibrariesToSkip[Arch])
 			{
@@ -999,9 +1006,9 @@ namespace UnrealBuildTool
 			}
 
 			// deal with .so files with wrong architecture
-			if (Path.GetExtension(Lib) == ".so")
+			if (Path.GetExtension(FullLib) == ".so")
 			{
-				string ParentDirectory = Path.GetDirectoryName(Lib);
+				string ParentDirectory = Path.GetDirectoryName(FullLib);
 				if (!IsDirectoryForArch(ParentDirectory, Arch))
 				{
 					return true;
@@ -1009,7 +1016,7 @@ namespace UnrealBuildTool
 			}
 
 			// apply the same directory filtering to libraries as we do to additional library paths
-			if (!IsDirectoryForArch(Path.GetDirectoryName(Lib), Arch))
+			if (!IsDirectoryForArch(Path.GetDirectoryName(FullLib), Arch))
 			{
 				return true;
 			}
@@ -1019,7 +1026,7 @@ namespace UnrealBuildTool
 			{
 				if (ComboName != Arch + GPUArchitecture)
 				{
-					if (Path.GetFileNameWithoutExtension(Lib).EndsWith(ComboName))
+					if (Lib.EndsWith(ComboName))
 					{
 						return true;
 					}
@@ -1654,7 +1661,11 @@ namespace UnrealBuildTool
 					// libs don't link in other libs
 					if (!LinkEnvironment.bIsBuildingLibrary)
 					{
-						// Add the library paths to the argument list.
+						// Make a list of library paths to search
+						List<string> AdditionalLibraryPaths = new List<string>();
+						List<string> AdditionalLibraries = new List<string>();
+
+						// Add the library paths to the additional path list
 						foreach (DirectoryReference LibraryPath in LinkEnvironment.LibraryPaths)
 						{
 							// LinkerPaths could be relative or absolute
@@ -1666,28 +1677,71 @@ namespace UnrealBuildTool
 								{
 									AbsoluteLibraryPath = Path.Combine(LinkerPath.FullName, AbsoluteLibraryPath);
 								}
-								LinkResponseArguments += string.Format(" -L\"{0}\"", Utils.CollapseRelativeDirectories(AbsoluteLibraryPath));
+								AbsoluteLibraryPath = Utils.CollapseRelativeDirectories(AbsoluteLibraryPath);
+								if (!AdditionalLibraryPaths.Contains(AbsoluteLibraryPath))
+								{
+									AdditionalLibraryPaths.Add(AbsoluteLibraryPath);
+								}
 							}
 						}
 
-						// add libraries in a library group
-						LinkResponseArguments += string.Format(" -Wl,--start-group");
+						// discover additional libraries and their paths
 						foreach (string AdditionalLibrary in LinkEnvironment.AdditionalLibraries)
 						{
 							if (!ShouldSkipLib(AdditionalLibrary, Arch, GPUArchitecture))
 							{
 								if (String.IsNullOrEmpty(Path.GetDirectoryName(AdditionalLibrary)))
 								{
-									LinkResponseArguments += string.Format(" \"-l{0}\"", AdditionalLibrary);
+									if (AdditionalLibrary.StartsWith("lib"))
+									{
+										AdditionalLibraries.Add(AdditionalLibrary);
+									}
+									else
+									{
+										AdditionalLibraries.Add("lib" + AdditionalLibrary);
+									}
 								}
 								else
 								{
-									// full pathed libs are compiled by us, so we depend on linking them
-									LinkResponseArguments += string.Format(" \"{0}\"", Path.GetFullPath(AdditionalLibrary));
+									string AbsoluteLibraryPath = Path.GetDirectoryName(Path.GetFullPath(AdditionalLibrary));
 									LinkAction.PrerequisiteItems.Add(FileItem.GetItemByPath(AdditionalLibrary));
+
+									string Lib = Path.GetFileNameWithoutExtension(AdditionalLibrary);
+									if (Lib.StartsWith("lib"))
+									{
+										AdditionalLibraries.Add(Lib);
+										if (!AdditionalLibraryPaths.Contains(AbsoluteLibraryPath))
+										{
+											AdditionalLibraryPaths.Add(AbsoluteLibraryPath);
+										}
+									}
+									else
+									{
+										AdditionalLibraries.Add(AbsoluteLibraryPath);
+									}
 								}
 							}
 						}
+
+						// add the library paths to response
+						foreach (string LibaryPath in AdditionalLibraryPaths)
+						{
+							LinkResponseArguments += string.Format(" -L\"{0}\"", LibaryPath);
+						}
+
+						// add libraries in a library group
+						LinkResponseArguments += string.Format(" -Wl,--start-group");
+						foreach (string AdditionalLibrary in AdditionalLibraries)
+						{
+							if (AdditionalLibrary.StartsWith("lib"))
+							{
+								LinkResponseArguments += string.Format(" \"-l{0}\"", AdditionalLibrary.Substring(3));
+							}
+							else
+							{
+								LinkResponseArguments += string.Format(" \"{0}\"", AdditionalLibrary);
+							}
+						}
 						LinkResponseArguments += string.Format(" -Wl,--end-group");
 
 						// Write the MAP file to the output directory.
