Copying //UE4/Partner-MagicLeap-Staging to Dev-VR (//UE4/Dev-VR)

diff --git a/Engine/Source/Programs/AutomationTool/Lumin/LuminPlatform.Automation.cs b/Engine/Source/Programs/AutomationTool/Lumin/LuminPlatform.Automation.cs
--- a/Engine/Source/Programs/AutomationTool/Lumin/LuminPlatform.Automation.cs
+++ b/Engine/Source/Programs/AutomationTool/Lumin/LuminPlatform.Automation.cs
@@ -19,7 +19,8 @@ public class LuminPlatform : Platform
 
 	// TODO Try to get these at runtime from the Lumin lifecycle service instead of hardcoding them here.
 	//private static string PackageInstallPath = "/package";
-	private static string PackageWritePath = "/documents/c2";
+	// Don't lowercase the sandbox directories.
+	private static string PackageWritePath = "/documents/C2";
 
 	private List<FileReference> RuntimeDependenciesForMabu;
 
@@ -65,20 +66,12 @@ public class LuminPlatform : Platform
 	{
 		// Read in any extra assets required to correctly install the application.
 		ConfigHierarchy Ini = ConfigCache.ReadHierarchy(ConfigHierarchyType.Engine, SC.RawProjectPath.Directory, UnrealTargetPlatform.Lumin);
-		// ConfigHierarchy Ini = ConfigCache.ReadHierarchy(ConfigHierarchyType.Engine, SC.RawProjectPath.Directory, SC.StageTargetPlatform.PlatformType);
 		string Value;
 		Ini.GetString("/Script/LuminRuntimeSettings.LuminRuntimeSettings", ConfigPropertyName, out Value);
 		Value = CleanFilePath(Value);
 		// We can have two different kinds of paths in the ini for icons: engine exec relative or project root relative.
 		DirectoryReference IconDir = GetFullPathFromRelativePath(Value, SC);
-		List<FileReference> IconFiles = SC.FindFilesToStage(IconDir, "*", StageFilesSearch.AllDirectories);
-		StringBuilder Builder = new StringBuilder();
-		foreach(FileReference IconFile in IconFiles)
-		{
-			string StagePath = IconFile.FullName.Replace(IconDir.FullName, IconStagePath);
-			Builder.AppendLine(String.Format("\"{0}\" : \"{1}\"\\", IconFile.FullName.Replace('\\', '/'), StagePath.Replace('\\', '/')));
-		}
-		return Builder.ToString();
+		return GenerateMabuPackageLine(EnsureTrailingSlashForDirectoryPath(IconDir.FullName), EnsureTrailingSlashForDirectoryPath(IconStagePath));
 	}
 
 	public override void GetFilesToDeployOrStage(ProjectParams Params, DeploymentContext SC)
@@ -506,8 +499,24 @@ public class LuminPlatform : Platform
 				// if we don't need to deploy any NonUFS files, and the exe is up to date, then there's no need to waste time packaging!
 				if (Lines.Length > 0 && !string.IsNullOrWhiteSpace(Lines[0]))
 				{
-					LogInformation("Need minimal package because {0} NonUFS files needed to be staged", Lines.Length);
-					return false;
+					bool NeedDeployNonUFSFiles = false;
+					foreach (string Line in Lines)
+					{
+						// We should not count ue4commandline.txt as a NonUFS file since it is deployed regardless (see Deploy() function).
+						// and counting it here as NonUFS causes an unnecesarry minimal package.
+						if (!(string.IsNullOrWhiteSpace(Line) || Line.Contains("ue4commandline.txt")))
+						{
+							NeedDeployNonUFSFiles = true;
+							break;
+						}
+					}
+					if (NeedDeployNonUFSFiles)
+					{
+						// We used to print the number of NonUFS files that need to be staged, but Lines.Length counts the new line at EOF and thus is off by 1.
+						// Print the actual files instead. Helps with debugging iterative deployment too.
+						LogInformation("Need minimal package because NonUFS files needed to be staged - {0}", NonUFSFiles.Replace("\n", "; "));
+						return false;
+					}
 				}
 			}
 			else
@@ -576,54 +585,142 @@ public class LuminPlatform : Platform
 			Builder.AppendLine("OPTIONS=\\");
 			Builder.AppendLine("stl/libgnustl\\");
 			Builder.AppendLine(string.Format("package/debuggable/{0}", Params.Distribution ? "off" : "on"));
-			Builder.AppendLine("DATAS=\\");
+
+			// Use .package file for mpk signing so that the user can override that if needed with the MLCERT env var.
+			string Certificate = Deploy.GetProjectRelativeCertificatePath();
+			if (!string.IsNullOrEmpty(Certificate))
+			{
+				Certificate = Path.GetFullPath(Path.Combine(Params.RawProjectPath.Directory.FullName, Certificate));
+				if (File.Exists(Certificate))
+				{
+					// Quote the path to handle spaces.
+					Builder.AppendLine(string.Format("MLCERT=\"{0}\"", Certificate));
+				}
+				else
+				{
+					throw new BuildException(string.Format("Certificate file does not exist at path {0}. Please enter a valid certificate file path in Project Settings > Magic Leap or clear the field if you do not intend to sign the package.", Certificate));
+				}
+			}
+
+			Builder.AppendLine("DATAS= \\");
 			if (additionalFiles.Length > 0)
 			{
 				//intentionally no newline because it should already be in the additionalFiles
 				Builder.Append(additionalFiles);
 			}
 
-			string StagedEngineDir = Path.Combine(SC.StageDirectory.FullName, "engine").ToLowerInvariant();
-			string StagedProjectDir = Path.Combine(SC.StageDirectory.FullName, Params.ShortProjectName).ToLowerInvariant();
-			string StagedMoviesDir = Path.Combine(StagedProjectDir, "content", "movies");
-
-			foreach (var FilePath in Directory.EnumerateFiles(SC.StageDirectory.FullName, "*", SearchOption.AllDirectories))
+			List<string> MinimalPackageFiles = new List<string>();
+			if (Params.IterativeDeploy)
 			{
-				if (Path.GetFileName(MabuFile) == Path.GetFileName(FilePath))
+				string NonUFSManifestFileName = CombinePaths(SC.StageDirectory.FullName, SC.GetNonUFSDeployedManifestFileName(null));
+				if (File.Exists(NonUFSManifestFileName))
 				{
-					continue;
+					string[] AllLines = File.ReadAllLines(NonUFSManifestFileName);
+					foreach (string NonUFSFileInfo in AllLines)
+					{
+						MinimalPackageFiles.Add(Path.Combine(SC.StageDirectory.FullName, NonUFSFileInfo.Split(null)[0]).ToLowerInvariant());
+					}
+				}
+				else
+				{
+					CommandUtils.LogError("NonUFS file manifest {0} does not exist at package time", NonUFSManifestFileName);
 				}
 
-				string LowercaseFilePath = FilePath.ToLowerInvariant();
-
-				// when doing iterative deploying, we only add the root files and Icon files
-				// @todo Lumin: Can we change iterative deploying to be at staging time, so that we have the list of NonUFS files to package
-				// @todo Lumin: Why aren't the manifest files around in the Staged dir? This would allow us to know what is Staged
-				if (Params.IterativeDeploy)
+				// Add files in the bin folder explicitely since they are not counted as NonUFS files.
+				string StagedBinDir = Path.Combine(SC.StageDirectory.FullName, "bin").ToLowerInvariant();
+				if (Directory.Exists(StagedBinDir))
 				{
-					// skip project files (except movies)
-					if (LowercaseFilePath.StartsWith(StagedEngineDir) || (LowercaseFilePath.StartsWith(StagedProjectDir) && !LowercaseFilePath.StartsWith(StagedMoviesDir)))
+					foreach (var BinFilePath in Directory.EnumerateFiles(StagedBinDir, "*", SearchOption.AllDirectories))
 					{
-						continue;
+						MinimalPackageFiles.Add(BinFilePath.ToLowerInvariant());
 					}
 				}
 
+				foreach (var FilePath in Directory.EnumerateFiles(SC.StageDirectory.FullName, "*", SearchOption.AllDirectories))
 				{
-					Builder.AppendLine(String.Format("\"{0}\" : \"{1}\"\\", FilePath, Utils.MakePathRelativeTo(FilePath, SC.StageDirectory.FullName, true).Replace('\\', '/')));
+					if (Path.GetFileName(MabuFile) == Path.GetFileName(FilePath))
+					{
+						continue;
+					}
+
+					string LowercaseFilePath = FilePath.ToLowerInvariant();
+
+					// when doing iterative deploying, we only add the NonUFS files and bin folder
+					if (!MinimalPackageFiles.Contains(LowercaseFilePath))
+					{
+						continue;
+					}
+
+					{
+						Builder.AppendLine(GenerateMabuPackageLine(FilePath, Utils.MakePathRelativeTo(FilePath, SC.StageDirectory.FullName, true).Replace('\\', '/')));
+					}
 				}
 			}
+			else
+			{
+				// For regular deployment, provide the entire StageDirectory for packaging.
+				Builder.AppendLine(GenerateMabuPackageLine(EnsureTrailingSlashForDirectoryPath(SC.StageDirectory.FullName), "./"));
+			}
 
 			// third party lbs should be packaged in the bin folder, without changing their filename casing.
 			foreach (FileReference LibFile in RuntimeDependenciesForMabu)
 			{
-				Builder.AppendLine(String.Format("\"{0}\" : \"bin/{1}\"\\", LibFile.FullName, LibFile.GetFileName()));
+				Builder.AppendLine(GenerateMabuPackageLine(LibFile.FullName, string.Format("bin/{0}", LibFile.GetFileName())));
+			}
+
+			// Stage vulkan validation libs
+			bool bStageVulkanValidationLayers = false;
+			if (Deploy.UseVulkan())
+			{
+				switch (SC.StageTargets[0].Receipt.Configuration)
+				{
+					case UnrealTargetConfiguration.Debug:
+					case UnrealTargetConfiguration.DebugGame:
+					case UnrealTargetConfiguration.Development:
+						bStageVulkanValidationLayers = true;
+						break;
+				}
+			}
+
+			if (bStageVulkanValidationLayers)
+			{
+				List<string> VulkanLayerLibsLookupPaths = new List<string>();
+
+				string VulkanValdationLayerLibsDir = Deploy.GetVulkanValdationLayerLibsDir();
+				if (!string.IsNullOrEmpty(VulkanValdationLayerLibsDir))
+				{
+					VulkanLayerLibsLookupPaths.Add(Path.GetFullPath(VulkanValdationLayerLibsDir));
+				}
+				// TODO: add path dir in MLSDK which contains these layer libs. Until that is available, use NDKROOT path as a contingency.
+				//  When we do get these libs in the MLSDK, remove the folder selection option from LuminRuntimeSettings.
+				VulkanLayerLibsLookupPaths.Add(Environment.ExpandEnvironmentVariables("%NDKROOT%/sources/third_party/vulkan/src/build-android/jniLibs/arm64-v8a"));
+
+				VulkanValdationLayerLibsDir = string.Empty;
+				foreach (string LookupPath in VulkanLayerLibsLookupPaths)
+				{
+					if (Directory.Exists(LookupPath))
+					{
+						VulkanValdationLayerLibsDir = LookupPath;
+						break;
+					}
+				}
+
+				if (!string.IsNullOrEmpty(VulkanValdationLayerLibsDir))
+				{
+					Builder.AppendLine(GenerateMabuPackageLine(EnsureTrailingSlashForDirectoryPath(VulkanValdationLayerLibsDir), "bin/"));
+					CommandUtils.LogInformation("Packaging vulkan validation layers from {0} into the mpk.", VulkanValdationLayerLibsDir);
+				}
+				else
+				{
+					CommandUtils.LogInformation("No vulkan validation layer libraries found while building with bUseVulkan enabled. Validation layers will not be enabled in this build.");
+				}
 			}
 
 			// Stage icon files directly to mabu instead of via Unreal to avoid the file casing issues.
 			// Icon files must be staged as is, without changing the case as the fbx, obj etc could have 
 			// embedded references to texture files.
-			Builder.Append(StageIconFileToMabu("IconModelPath", Deploy.GetIconModelStagingPath(), SC));
-			Builder.Append(StageIconFileToMabu("IconPortalPath", Deploy.GetIconPortalStagingPath(), SC));
+			Builder.AppendLine(StageIconFileToMabu("IconModelPath", Deploy.GetIconModelStagingPath(), SC));
+			Builder.AppendLine(StageIconFileToMabu("IconPortalPath", Deploy.GetIconPortalStagingPath(), SC));
 
 			// find executable
 			if (GetExecutableNames(SC).Count > 1)
@@ -633,7 +730,7 @@ public class LuminPlatform : Platform
 
 			// now put the exe into the data list (coming from original location)
 			string ExePath = GetExePath(Params, SC);
-			Builder.AppendLine(String.Format("\"{0}/Binaries/{1}\" : \"bin/{2}\"", Path.GetDirectoryName(MabuFile), Path.GetFileName(ExePath), Params.ShortProjectName));
+			Builder.AppendLine(GenerateMabuPackageLine(string.Format("{0}/Binaries/{1}", Path.GetDirectoryName(MabuFile), Path.GetFileName(ExePath)), string.Format("bin/{0}", Params.ShortProjectName)));
 
 			Directory.CreateDirectory(Path.GetDirectoryName(MabuFile));
 			WriteAllText(MabuFile, Builder.ToString());
@@ -787,7 +884,8 @@ public class LuminPlatform : Platform
 		// update the ue4commandline.txt
 		// update and deploy ue4commandline.txt
 		// always delete the existing commandline text file, so it doesn't reuse an old one
-		string IntermediateCmdLineFile = CombinePaths(SC.StageDirectory.FullName, "UE4CommandLine.txt");
+		// only the file name should be lowercased, not the entire path
+		string IntermediateCmdLineFile = CombinePaths(SC.StageDirectory.FullName, "ue4commandline.txt");
 		Project.WriteStageCommandline(new FileReference(IntermediateCmdLineFile), Params, SC);
 
 		// Where we put the files on device.
@@ -802,7 +900,7 @@ public class LuminPlatform : Platform
 			if (Params.IterativeDeploy)
 			{
 				LogInformation("ITERATIVE DEPLOY..");
-				// always send UE4CommandLine.txt (it was written above after delta checks applied)
+				// always send ue4commandline.txt (it was written above after delta checks applied)
 				EntriesToDeploy.Add(IntermediateCmdLineFile);
 
 				// Add non UFS files if any to deploy
@@ -1237,6 +1335,27 @@ public class LuminPlatform : Platform
 		return Configurations;
 	}
 
+	/// <summary>
+	/// Generates a formatted line for the mabu .package file to map the SrcFile to the Dst path in the mpk
+	/// </summary>
+	/// <param name="SrcFile">Full path to src file. Should not be quoted.</param>
+	/// <param name="DstFile">Path relative to mpk root where the SrcFile should be packaged. Should not be quoted.</param>
+	private string GenerateMabuPackageLine(string SrcFile, string DstFile)
+	{
+		// Ensure that the paths are quoted and that there is a space before the trailing slash
+		return string.Format("\"{0}\" : \"{1}\" \\", SrcFile.Replace('\\', '/'), DstFile.Replace('\\', '/'));
+	}
+
+	private string EnsureTrailingSlashForDirectoryPath(string DirPath)
+	{
+		if (!(DirPath.EndsWith(Path.DirectorySeparatorChar.ToString()) || DirPath.EndsWith(Path.AltDirectorySeparatorChar.ToString()) ))
+		{
+			// Ensure trailing slash in folder path for recursive folder packaging via mabu.
+			return DirPath + Path.DirectorySeparatorChar;
+		}
+
+		return DirPath;
+	}
 
 	#endregion
 }
\ No newline at end of file
diff --git a/Engine/Source/Programs/UnrealBuildTool/Platform/Android/AndroidToolChain.cs b/Engine/Source/Programs/UnrealBuildTool/Platform/Android/AndroidToolChain.cs
--- a/Engine/Source/Programs/UnrealBuildTool/Platform/Android/AndroidToolChain.cs
+++ b/Engine/Source/Programs/UnrealBuildTool/Platform/Android/AndroidToolChain.cs
@@ -1490,22 +1490,7 @@ namespace UnrealBuildTool
 						CompileAction.WorkingDirectory = UnrealBuildTool.EngineSourceDirectory;
 						if(bExecuteCompilerThroughShell)
 						{
-							string FixedClangPath = ClangPath;
-							if (FixedClangPath.Contains(' '))
-							{
-								FixedClangPath = "'" + FixedClangPath + "'";
-							}
-					
-							CompileAction.CommandPath = BuildHostPlatform.Current.Shell;
-							if (BuildHostPlatform.Current.ShellType == ShellType.Cmd)
-							{
-								CompileAction.CommandArguments = String.Format("/c \"{0} {1}\"", FixedClangPath, ResponseArgument);
-							}
-							else
-							{
-								CompileAction.CommandArguments = String.Format("-c \'{0} {1}\'", FixedClangPath, ResponseArgument);
-							}
-							CompileAction.CommandDescription = "Compile";
+							SetupActionToExecuteCompilerThroughShell(ref CompileAction, ClangPath, ResponseArgument, "Compile");
 						}
 						else
 						{
@@ -1777,22 +1762,7 @@ namespace UnrealBuildTool
 
 					if(bExecuteCompilerThroughShell)
 					{
-						string LinkCommandPath = LinkAction.CommandPath.FullName;
-						if (LinkCommandPath.Contains(' '))
-						{
-							LinkCommandPath = "'" + LinkCommandPath + "'";
-						}
-						
-						if (BuildHostPlatform.Current.ShellType == ShellType.Cmd)
-						{
-							LinkAction.CommandArguments = String.Format("/c \"{0} {1}\"", LinkCommandPath, LinkAction.CommandArguments);
-						}
-						else
-						{
-							LinkAction.CommandArguments = String.Format("-c \'{0} {1}\'", LinkCommandPath, LinkAction.CommandArguments);
-						}
-						LinkAction.CommandPath = BuildHostPlatform.Current.Shell;
-						LinkAction.CommandDescription = "Link";
+						SetupActionToExecuteCompilerThroughShell(ref LinkAction, LinkAction.CommandPath.FullName, LinkAction.CommandArguments, "Link");
 					}
 					Actions.Add(LinkAction);
 
@@ -1957,5 +1927,26 @@ namespace UnrealBuildTool
 			StartInfo.CreateNoWindow = true;
 			Utils.RunLocalProcessAndLogOutput(StartInfo);
 		}
+
+		protected virtual void SetupActionToExecuteCompilerThroughShell(ref Action CompileOrLinkAction, string CommandPath, string CommandArguments, string CommandDescription)
+		{
+			string QuotedCommandPath = CommandPath;
+			if (CommandPath.Contains(' '))
+			{
+				QuotedCommandPath = "'" + CommandPath + "'";
+			}
+	
+			if (BuildHostPlatform.Current.ShellType == ShellType.Cmd)
+			{
+				CompileOrLinkAction.CommandArguments = String.Format("/c \"{0} {1}\"", QuotedCommandPath, CommandArguments);
+			}
+			else
+			{
+				CompileOrLinkAction.CommandArguments = String.Format("-c \'{0} {1}\'", QuotedCommandPath, CommandArguments);
+			}
+
+			CompileOrLinkAction.CommandPath = BuildHostPlatform.Current.Shell;
+			CompileOrLinkAction.CommandDescription = CommandDescription;
+		}
 	};
 }
diff --git a/Engine/Source/Programs/UnrealBuildTool/Platform/Lumin/LuminExports.cs b/Engine/Source/Programs/UnrealBuildTool/Platform/Lumin/LuminExports.cs
--- a/Engine/Source/Programs/UnrealBuildTool/Platform/Lumin/LuminExports.cs
+++ b/Engine/Source/Programs/UnrealBuildTool/Platform/Lumin/LuminExports.cs
@@ -67,6 +67,24 @@ namespace UnrealBuildTool
 		/// </summary>
 		/// <returns>Directory path relative to root of MPK</returns>
 		string GetIconPortalStagingPath();
+
+		/// <summary>
+		/// Path of the certificate to sign the MPK with.
+		/// </summary>
+		/// <returns>Certificate file path relative to root of the project</returns>
+		string GetProjectRelativeCertificatePath();
+
+		/// <summary>
+		/// Whether the app should use vulkan
+		/// </summary>
+		/// <returns>true if app should use vulkan, false otherwise</returns>
+		bool UseVulkan();
+
+		/// <summary>
+		/// Path of the folder containing the vulkan validation layers.
+		/// </summary>
+		/// <returns>Path of the folder containing the vulkan validation layers.</returns>
+		string GetVulkanValdationLayerLibsDir();
 	}
 
 	/// <summary>
diff --git a/Engine/Source/Programs/UnrealBuildTool/Platform/Lumin/LuminProjectGenerator.cs b/Engine/Source/Programs/UnrealBuildTool/Platform/Lumin/LuminProjectGenerator.cs
--- a/Engine/Source/Programs/UnrealBuildTool/Platform/Lumin/LuminProjectGenerator.cs
+++ b/Engine/Source/Programs/UnrealBuildTool/Platform/Lumin/LuminProjectGenerator.cs
@@ -104,9 +104,10 @@ namespace UnrealBuildTool
 				}
 				PackageFile = Path.Combine(PackageFile, PackageName + ".mpk");
 
-                // TODO: Fix NMakeOutput to have the the correct architecture name and then set ELFFile to $(NMakeOutput)
+                // Can't use $(NMakeOutput) directly since that is defined after <ELFFile> tag and thus ends up being translated as an empty string.
                 string ELFFile = Utils.MakePathRelativeTo(NMakeOutputPath.Directory.FullName, ProjectFilePath.Directory.FullName);
-                ELFFile = Path.Combine(ELFFile, "..\\..\\Intermediate\\Lumin\\Mabu\\Packaged\\bin\\" + GameName);
+				// Provide path to stripped executable so all symbols are resolved from the external sym file instead.
+                ELFFile = Path.Combine(ELFFile, "..\\..\\Intermediate\\Lumin\\Mabu\\Binaries", GetElfName(NMakeOutputPath));
 				string DebuggerFlavor = "MLDebugger";
 
 				string SymFile = Utils.MakePathRelativeTo(NMakeOutputPath.Directory.FullName, ProjectFilePath.Directory.FullName);
@@ -127,10 +128,24 @@ namespace UnrealBuildTool
 													"<EnableAutoStop>{5}</EnableAutoStop>" + ProjectFileGenerator.NewLine +
 													"<AutoStopAtFunction>{6}</AutoStopAtFunction>" + ProjectFileGenerator.NewLine +
 													"<EnablePrettyPrinting>{7}</EnablePrettyPrinting>" + ProjectFileGenerator.NewLine +
-													"<MLDownloadOnStart>{8}</MLDownloadOnStart>" + ProjectFileGenerator.NewLine +
-													"<SymFile>{9}</SymFile>" + ProjectFileGenerator.NewLine;
-				ProjectFileBuilder.Append(ProjectFileGenerator.NewLine + string.Format(CustomPathEntriesTemplate, MLSDK, PackageFile, ELFFile, DebuggerFlavor, Attach, EnableAutoStop, AutoStopAtFunction, EnablePrettyPrinting, MLDownloadOnStart, SymFile));
+													"<MLDownloadOnStart>{8}</MLDownloadOnStart>" + ProjectFileGenerator.NewLine;
+													// No need to provide SymFile path. The file is automatically picked up when it resides in the same folder and has the same name as the ElfFile
+													// "<SymFile>{9}</SymFile>" + ProjectFileGenerator.NewLine;
+				ProjectFileBuilder.Append(ProjectFileGenerator.NewLine + string.Format(CustomPathEntriesTemplate, MLSDK, PackageFile, ELFFile, DebuggerFlavor, Attach, EnableAutoStop, AutoStopAtFunction, EnablePrettyPrinting, MLDownloadOnStart));
 			}
 		}
+
+		private string GetElfName(FileReference InNMakeOutputPath)
+		{
+			ConfigHierarchy Ini = ConfigCache.ReadHierarchy(ConfigHierarchyType.Engine, InNMakeOutputPath.Directory.ParentDirectory.ParentDirectory, UnrealTargetPlatform.Lumin);
+			bool bUseMobileRendering = false;
+			Ini.GetBool("/Script/LuminRuntimeSettings.LuminRuntimeSettings", "bUseMobileRendering", out bUseMobileRendering);
+			bool bUseVulkan = false;
+			Ini.GetBool("/Script/LuminRuntimeSettings.LuminRuntimeSettings", "bUseVulkan", out bUseVulkan);
+
+			string OutputFileName = string.Format("{0}-arm64-{1}{2}", InNMakeOutputPath.GetFileNameWithoutExtension(), !(bUseMobileRendering || bUseVulkan) ? "lumingl4" : "lumin", InNMakeOutputPath.GetExtension());
+
+			return OutputFileName;
+		}
 	}
 }
diff --git a/Engine/Source/Programs/UnrealBuildTool/Platform/Lumin/LuminToolChain.cs b/Engine/Source/Programs/UnrealBuildTool/Platform/Lumin/LuminToolChain.cs
--- a/Engine/Source/Programs/UnrealBuildTool/Platform/Lumin/LuminToolChain.cs
+++ b/Engine/Source/Programs/UnrealBuildTool/Platform/Lumin/LuminToolChain.cs
@@ -79,7 +79,7 @@ namespace UnrealBuildTool
 
 			MLSDKPath = MLSDKPath.Replace("\"", "");
 
-			string MabuSpec = RunMabuAndReadOutput("-t device --print-spec");
+			string MabuSpec = RunMabuAndReadOutput("-t device --print-spec -q");
 
 			// parse clange version
 			Regex SpecRegex = new Regex("\\s*(?:[a-z]+_lumin_clang-)(\\d)[.](\\d)\\s*");
@@ -95,7 +95,7 @@ namespace UnrealBuildTool
 
 			SetClangVersion(int.Parse(SpecMatch.Groups[1].Value), int.Parse(SpecMatch.Groups[2].Value), 0);
 
-			string MabuTools = RunMabuAndReadOutput("-t device --print-tools");
+			string MabuTools = RunMabuAndReadOutput("-t device --print-tools -q");
 
 			Dictionary<string, string> ToolsDict = new Dictionary<string, string>();
 			using (StringReader Reader = new StringReader(MabuTools))
@@ -114,9 +114,9 @@ namespace UnrealBuildTool
 			}
 
 			// set up the path to our toolchains
-			// Clang path used to be in quotes, but some part of FileReference in AndroidToolChain was choking on the quotes.  Don't put your MLSDK in a directory with spaces for now, I guess.
+			// ClangPath and ArPathArm64 are quoted in SetupActionToExecuteCompilerThroughShell() just before use if path has spaces in it.
 			ClangPath = ToolsDict["CXX"];
-			ArPathArm64 = "\"" + ToolsDict["AR"] + "\"";
+			ArPathArm64 = ToolsDict["AR"];
 			// The strip command does not execute through the shell. Hence we don't quote it.
 			StripPath = ToolsDict["STRIP"];
 			// The objcopy command does not execute through the shell. Hence we don't quote it.
@@ -351,5 +351,23 @@ namespace UnrealBuildTool
 			StartInfo.CreateNoWindow = true;
 			Utils.RunLocalProcessAndLogOutput(StartInfo);
 		}
+
+		protected override void SetupActionToExecuteCompilerThroughShell(ref Action CompileOrLinkAction, string CommandPath, string CommandArguments, string CommandDescription)
+		{
+			base.SetupActionToExecuteCompilerThroughShell(ref CompileOrLinkAction, CommandPath, CommandArguments, CommandDescription);
+
+			string QuotedCommandPath = CommandPath;
+			if (CommandPath.Contains(' '))
+			{
+				QuotedCommandPath = "'" + CommandPath + "'";
+			}
+			if (BuildHostPlatform.Current.ShellType != ShellType.Cmd)
+			{
+				// When quoting the command for the shell, we also need to escape any quotes. Otherwise
+				// they inadvertently unquote the interior arguments. For example it would unquote the
+				// command itself, if it was quoted to allow for spaces in the path, causing errors.
+				CompileOrLinkAction.CommandArguments = String.Format("-c \'{0} {1}\'", QuotedCommandPath.Replace("\'", "\\\'"), CommandArguments.Replace("\'", "\\\'"));
+			}
+		}
 	};
 }
diff --git a/Engine/Source/Programs/UnrealBuildTool/Platform/Lumin/UEBuildLumin.cs b/Engine/Source/Programs/UnrealBuildTool/Platform/Lumin/UEBuildLumin.cs
--- a/Engine/Source/Programs/UnrealBuildTool/Platform/Lumin/UEBuildLumin.cs
+++ b/Engine/Source/Programs/UnrealBuildTool/Platform/Lumin/UEBuildLumin.cs
@@ -21,6 +21,12 @@ namespace UnrealBuildTool
 		/// </summary>
 		[CommandLine("-GPUArchitectures=", ListSeparator = '+')]
 		public List<string> GPUArchitectures = new List<string>();
+
+		/// <summary>
+		/// If -distribution was passed on the commandline, this build is for distribution.
+		/// </summary>
+		[CommandLine("-distribution")]
+		public bool bForDistribution = false;
 	}
 
 	/// <summary>
@@ -54,12 +60,98 @@ namespace UnrealBuildTool
 		{
 			get { return Inner.GPUArchitectures.AsReadOnly(); }
 		}
+		public bool bForDistribution
+		{
+			get { return Inner.bForDistribution; }
+		}
 
 #if !__MonoCS__
 #pragma warning restore CS1591
 #endif
 		#endregion
 	}
+
+	/// <summary>
+	/// Helper class for shared lumin platform functions and members
+	/// </summary>
+	public class LuminSDKVersionHelper
+	{
+		/// <summary>
+		/// This is the minimum SDK version we support.
+		/// </summary>
+		static uint MinimumSDKVersionMajor = 0;
+		static uint MinimumSDKVersionMinor = 22;
+
+		/// <summary>
+		/// Gets a string defining the minimum required sdk version.
+		/// </summary>
+		/// <returns></returns>
+		public string GetRequiredSDKString()
+		{
+			return string.Format("{0}.{1}", MinimumSDKVersionMajor, MinimumSDKVersionMinor);
+		}
+
+		/// <summary>
+		/// checks if the sdk is installed or has been synced, sets environment variable
+		/// </summary>
+		/// <returns></returns>
+		public bool HasAnySDK()
+		{
+			string EnvVarKey = "MLSDK";
+
+			string MLSDKPath = Environment.GetEnvironmentVariable(EnvVarKey);
+			if (String.IsNullOrEmpty(MLSDKPath))
+			{
+				Log.TraceLog("*** Unable to determine MLSDK location ***");
+				return false;
+			}
+
+			// detected SDK version is < minimum major/minor
+			uint DetectedMajorVersion;
+			uint DetectedMinorVersion;
+			String VersionFile = string.Format("{0}/include/ml_version.h", MLSDKPath).Replace('/', Path.DirectorySeparatorChar);
+			if (File.Exists(VersionFile))
+			{
+				string[] VersionText = File.ReadAllLines(VersionFile);
+
+				String MajorVersion = FindVersionNumber("MLSDK_VERSION_MAJOR", VersionText);
+				String MinorVersion = FindVersionNumber("MLSDK_VERSION_MINOR", VersionText);
+				DetectedMajorVersion = Convert.ToUInt32(MajorVersion);
+				DetectedMinorVersion = Convert.ToUInt32(MinorVersion);
+			}
+			else
+			{
+				Log.TraceLog("*** Unable to locate MLSDK version file ml_version.h ***");
+				return false;
+			}
+
+			if (DetectedMajorVersion < MinimumSDKVersionMajor || (DetectedMajorVersion == MinimumSDKVersionMajor && DetectedMinorVersion < MinimumSDKVersionMinor))
+			{
+				Log.TraceLog("*** Found installed MLSDK version {0}.{1} at '{2}' but require at least {3}.{4} ***",
+					DetectedMajorVersion, DetectedMinorVersion, MLSDKPath, MinimumSDKVersionMajor, MinimumSDKVersionMinor);
+				return false;
+			}
+
+			Log.TraceLog("*** Found installed MLSDK version {0}.{1} at '{2}' ***", DetectedMajorVersion, DetectedMinorVersion, MLSDKPath);
+			return true;
+		}
+
+		private string FindVersionNumber(string StringToFind, string[] AllLines)
+		{
+			string FoundVersion = "Unknown";
+			foreach (string CurrentLine in AllLines)
+			{
+				int Index = CurrentLine.IndexOf(StringToFind);
+				if (Index != -1)
+				{
+					FoundVersion = CurrentLine.Substring(Index + StringToFind.Length);
+					break;
+				}
+			}
+			return FoundVersion.Trim();
+		}
+	}
+
 	class LuminPlatform : AndroidPlatform
 	{
 		public LuminPlatform(AndroidPlatformSDK InSDK)
@@ -79,7 +171,6 @@ namespace UnrealBuildTool
 			Ini.GetBool("/Script/LuminRuntimeSettings.LuminRuntimeSettings", "bUseVulkan", out bUseVulkan);
 
 			List<string> ConfigBoolKeys = new List<string>();
-			ConfigBoolKeys.Add("bBuildWithNvTegraGfxDebugger");
 			if (!bUseVulkan)
 			{
 				ConfigBoolKeys.Add("bUseMobileRendering");
@@ -180,73 +271,23 @@ namespace UnrealBuildTool
 			LinkEnvironment.LibraryPaths.Add(DirectoryReference.Combine(MLSDKDir, "lib/lumin"));
 			LinkEnvironment.LibraryPaths.Add(DirectoryReference.Combine(MLSDKDir, "lumin/stl/gnu-libstdc++/lib"));
 
-			if (!UseTegraGraphicsDebugger(Target.ProjectFile) || !UseTegraDebuggerStub(Target.ProjectFile))
-			{
-				LinkEnvironment.AdditionalLibraries.Add("GLESv2");
-				LinkEnvironment.AdditionalLibraries.Add("EGL");
-			}
-			if (UseTegraDebuggerStub(Target.ProjectFile))
-			{
-				DirectoryReference TegraDebuggerDirectoryReference = new DirectoryReference(TegraDebuggerDir);
-
-				LinkEnvironment.LibraryPaths.Add(DirectoryReference.Combine(TegraDebuggerDirectoryReference, "/target/android-kk-egl-t124-a32"));
-				LinkEnvironment.LibraryPaths.Add(DirectoryReference.Combine(TegraDebuggerDirectoryReference, "/target/android-L-egl-t132-a64"));
-				LinkEnvironment.AdditionalLibraries.Add("Nvidia_gfx_debugger_stub");
-			}
+			LinkEnvironment.AdditionalLibraries.Add("GLESv2");
+			LinkEnvironment.AdditionalLibraries.Add("EGL");
 
 			LinkEnvironment.AdditionalLibraries.Add("ml_lifecycle");
 			LinkEnvironment.AdditionalLibraries.Add("ml_ext_logging");
 			LinkEnvironment.AdditionalLibraries.Add("ml_dispatch");
 		}
 
-
-        public static bool? HaveTegraGraphicsDebugger = null;
-        public static string TegraDebuggerDir = null;
-        public static int[] TegraDebuggerVersion = null;
-
-        public static bool UseTegraGraphicsDebugger(FileReference ProjectFile)
-        {
-            if (!HaveTegraGraphicsDebugger.HasValue)
-            {
-                string ProgramsDir = Environment.GetEnvironmentVariable("ProgramFiles(x86)");
-                string NVDir = ProgramsDir + "/NVIDIA Corporation";
-                try
-                {
-                    string[] TegraDebuggerDirs = Directory.GetDirectories(NVDir, "Tegra Graphics Debugger *");
-                    if (TegraDebuggerDirs.Length > 0)
-                    {
-                        TegraDebuggerDir = TegraDebuggerDirs[0].Replace('\\', '/');
-                        HaveTegraGraphicsDebugger = true;
-                        string[] V = TegraDebuggerDir.Split(' ').Last().Split('.');
-                        TegraDebuggerVersion = new int[2];
-                        TegraDebuggerVersion[0] = Int32.Parse(V[0]);
-                        TegraDebuggerVersion[1] = Int32.Parse(V[1]);
-                    }
-                    else
-                    {
-                        HaveTegraGraphicsDebugger = false;
-                    }
-                }
-                catch (System.IO.IOException)
-                {
-                    HaveTegraGraphicsDebugger = false;
-                }
-            }
-            bool bBuild = false;
-			// TODO: do we need this?
-            ConfigHierarchy Ini = ConfigCache.ReadHierarchy(ConfigHierarchyType.Engine, DirectoryReference.FromFile(ProjectFile), UnrealTargetPlatform.Lumin);
-            Ini.GetBool(
-                "/Script/LuminRuntimeSettings.LuminRuntimeSettings",
-                "bBuildWithNvTegraGfxDebugger",
-                out bBuild);
-            return HaveTegraGraphicsDebugger.Value && bBuild;
-        }
-
-
-		public static bool UseTegraDebuggerStub(FileReference ProjectFile)
+		public override bool ShouldCreateDebugInfo(ReadOnlyTargetRules Target)
 		{
-			return UseTegraGraphicsDebugger(ProjectFile) &&
-				TegraDebuggerVersion[0] <= 2 && TegraDebuggerVersion[1] <= 1;
+			// Don't create debug info for shipping distribution builds
+			if (Target.Configuration == UnrealTargetConfiguration.Shipping && Target.LuminPlatform.bForDistribution)
+			{
+				return false;
+			}
+
+			return base.ShouldCreateDebugInfo(Target);
 		}
 
 		public override UEToolChain CreateToolChain(ReadOnlyTargetRules Target)
@@ -282,16 +323,16 @@ namespace UnrealBuildTool
 		/// <summary>
 		/// This is the minimum SDK version we support.
 		/// </summary>
-		static uint MinimumSDKVersionMajor = 0;
-		static uint MinimumSDKVersionMinor = 19;
+		LuminSDKVersionHelper SDKVersionHelper = new LuminSDKVersionHelper();
 
 		public override string GetSDKTargetPlatformName()
 		{
 			return "Lumin";
 		}
+
 		protected override string GetRequiredSDKString()
 		{
-			return string.Format("{0}.{1}", MinimumSDKVersionMajor, MinimumSDKVersionMinor);
+			return SDKVersionHelper.GetRequiredSDKString();
 		}
 
 		protected override String GetRequiredScriptVersionString()
@@ -299,64 +340,13 @@ namespace UnrealBuildTool
 			return "Lumin_15";
 		}
 
-		private string FindVersionNumber(string StringToFind, string[] AllLines)
-		{
-			string FoundVersion = "Unknown";
-			foreach (string CurrentLine in AllLines)
-			{
-				int Index = CurrentLine.IndexOf(StringToFind);
-				if (Index != -1)
-				{
-					FoundVersion = CurrentLine.Substring(Index + StringToFind.Length);
-					break;
-				}
-			}
-			return FoundVersion.Trim();
-		}
-
 		/// <summary>
 		/// checks if the sdk is installed or has been synced, sets environment variable
 		/// </summary>
 		/// <returns></returns>
 		protected override bool HasAnySDK()
 		{
-			string EnvVarKey = "MLSDK";
-
-			string MLSDKPath = Environment.GetEnvironmentVariable(EnvVarKey);
-			if (String.IsNullOrEmpty(MLSDKPath))
-			{
-				Log.TraceLog("*** Unable to determine MLSDK location ***");
-				return false;
-			}
-
-			// detected SDK version is < minimum major/minor
-			uint DetectedMajorVersion;
-			uint DetectedMinorVersion;
-			String VersionFile = string.Format("{0}/include/ml_version.h", MLSDKPath).Replace('/', Path.DirectorySeparatorChar);
-			if (File.Exists(VersionFile))
-			{
-				string[] VersionText = File.ReadAllLines(VersionFile);
-
-				String MajorVersion = FindVersionNumber("MLSDK_VERSION_MAJOR", VersionText);
-				String MinorVersion = FindVersionNumber("MLSDK_VERSION_MINOR", VersionText);
-				DetectedMajorVersion = Convert.ToUInt32(MajorVersion);
-				DetectedMinorVersion = Convert.ToUInt32(MinorVersion);
-			}
-			else
-			{
-				Log.TraceLog("*** Unable to locate MLSDK version file ml_version.h ***");
-				return false;
-			}
-
-			if (DetectedMajorVersion < MinimumSDKVersionMajor || (DetectedMajorVersion == MinimumSDKVersionMajor && DetectedMinorVersion < MinimumSDKVersionMinor))
-			{
-				Log.TraceLog("*** Found installed MLSDK version {0}.{1} at '{2}' but require at least {3}.{4} ***",
-					DetectedMajorVersion, DetectedMinorVersion, MLSDKPath, MinimumSDKVersionMajor, MinimumSDKVersionMinor);
-				return false;
-			}
-
-			Log.TraceLog("*** Found installed MLSDK version {0}.{1} at '{2}' ***", DetectedMajorVersion, DetectedMinorVersion, MLSDKPath);
-			return true;
+			return SDKVersionHelper.HasAnySDK();
 		}
 
 		protected override SDKStatus HasRequiredManualSDKInternal()
@@ -367,7 +357,7 @@ namespace UnrealBuildTool
 				return SDKStatus.Invalid;
 			}
 
-			if (HasAnySDK())
+			if (SDKVersionHelper.HasAnySDK())
 			{
 				return SDKStatus.Valid;
 			}
diff --git a/Engine/Source/Programs/UnrealBuildTool/Platform/Lumin/UEDeployLumin.cs b/Engine/Source/Programs/UnrealBuildTool/Platform/Lumin/UEDeployLumin.cs
--- a/Engine/Source/Programs/UnrealBuildTool/Platform/Lumin/UEDeployLumin.cs
+++ b/Engine/Source/Programs/UnrealBuildTool/Platform/Lumin/UEDeployLumin.cs
@@ -9,6 +9,7 @@ using System.Text.RegularExpressions;
 using Tools.DotNETCommon;
 using System.Linq;
 using System.Xml;
+using System.Xml.Serialization;
 using System.Xml.Linq;
 
 namespace UnrealBuildTool
@@ -17,6 +18,8 @@ namespace UnrealBuildTool
 	{
 		private FileReference ProjectFile;
 
+		public manifest PackageManifest = new manifest();
+
 		protected UnrealPluginLanguage UPL;
 
 		public void SetLuminPluginData(List<string> Architectures, List<string> inPluginExtraData)
@@ -53,16 +56,16 @@ namespace UnrealBuildTool
 			return PackageName.ToLower();
 		}
 
-		private string GetApplicationType()
+		private manifestApplicationComponentType GetApplicationType()
 		{
 			ConfigHierarchy Ini = GetConfigCacheIni(ConfigHierarchyType.Engine);
 			bool Value = false;
 			Ini.GetBool("/Script/LuminRuntimeSettings.LuminRuntimeSettings", "bIsScreensApp", out Value);
 			if (Value)
 			{
-				return "ScreensImmersive";
+				return manifestApplicationComponentType.ScreensImmersive;
 			}
-			return "Fullscreen";
+			return manifestApplicationComponentType.Fullscreen;
 		}
 
 		private string GetApplicationDisplayName(string ProjectName)
@@ -90,27 +93,6 @@ namespace UnrealBuildTool
 			return Value.ToString();
 		}
 
-		private void GetAppPrivileges(ConfigHierarchy EngineIni, StringBuilder Text)
-		{
-			List<string> AppPrivileges;
-			EngineIni.GetArray("/Script/LuminRuntimeSettings.LuminRuntimeSettings", "AppPrivileges", out AppPrivileges);
-			if (AppPrivileges != null)
-			{
-				foreach (string Privilege in AppPrivileges)
-				{
-					string TrimmedPrivilege = Privilege.Trim(' ');
-					if (TrimmedPrivilege != "")
-					{
-						string PrivilegeString = string.Format("\t\t\t<uses-privilege ml:name=\"{0}\"/>", TrimmedPrivilege);
-						if (!Text.ToString().Contains(PrivilegeString))
-						{
-							Text.AppendLine(PrivilegeString);
-						}
-					}
-				}
-			}
-		}
-
 		private string CleanFilePath(string FilePath)
 		{
 			// Removes the extra characters from a FFilePath parameter.
@@ -131,6 +113,24 @@ namespace UnrealBuildTool
 			return "Icon/Portal";
 		}
 
+		public string GetProjectRelativeCertificatePath()
+		{
+			return CleanFilePath(GetRuntimeSetting("Certificate"));
+		}
+
+		public bool UseVulkan()
+		{
+			ConfigHierarchy Ini = GetConfigCacheIni(ConfigHierarchyType.Engine);
+			bool Value = false;
+			Ini.GetBool("/Script/LuminRuntimeSettings.LuminRuntimeSettings", "bUseVulkan", out Value);
+			return Value;
+		}
+
+		public string GetVulkanValdationLayerLibsDir()
+		{
+			return CleanFilePath(GetRuntimeSetting("VulkanValidationLayerLibs"));
+		}
+
 		public string GetMLSDKVersion(ConfigHierarchy EngineIni)
 		{
 			string MLSDKPath;
@@ -161,6 +161,112 @@ namespace UnrealBuildTool
 			return string.Format("{0}.{1}", Major, Minor);
 		}
 
+		private object GetComponentSubElement(string ElementType, string ElementValue)
+		{
+			switch (ElementType)
+			{
+				case "FileExtension":
+					return new manifestApplicationComponentFileextension
+					{
+						name = ElementValue,
+					};
+				case "MimeType":
+					return new manifestApplicationComponentMimetype
+					{
+						name = ElementValue,
+					};
+				case "MusicAttribute":
+					return new manifestApplicationComponentMusicattribute
+					{
+						name = ElementValue,
+					};
+				case "Mode":
+					return new manifestApplicationComponentMode
+					{
+						shareable = ElementValue,
+					};
+				case "Schema":
+					return new manifestApplicationComponentSchema
+					{
+						name = ElementValue,
+					};
+				default:
+					Log.TraceInformation("Tried to use an unsupported component sub-element type: {0}", ElementType);
+					return null;
+			}
+		}
+
+		private object GetComponentElement(Dictionary<string, string> ComponentElement)
+		{
+			manifestApplicationComponent OutComponent = new manifestApplicationComponent
+			{
+				name = ComponentElement["Name"],
+				visible_name = ComponentElement["VisibleName"],
+			};
+
+			// App developer has the responsibility to package the executable in the bin folder,
+			// perhaps by using UPL. We simply generate the manifest correctly.
+			string BinaryName = ComponentElement["ExecutableName"];
+			if (BinaryName.IndexOf("bin/") != 0)
+			{
+				// Prepend bin folder string to executable name if it not there already.
+				BinaryName = string.Format("bin/{0}", BinaryName);
+			}
+			OutComponent.binary_name = BinaryName;
+
+			switch (ComponentElement["ComponentType"])
+			{
+				case "Universe":
+					OutComponent.type = manifestApplicationComponentType.Universe;
+					break;
+				case "Fullscreen":
+					OutComponent.type = manifestApplicationComponentType.Fullscreen;
+					break;
+				case "SearchProvider":
+					OutComponent.type = manifestApplicationComponentType.SearchProvider;
+					break;
+				case "MusicService":
+					OutComponent.type = manifestApplicationComponentType.MusicService;
+					break;
+				case "Screens":
+					OutComponent.type = manifestApplicationComponentType.Screens;
+					break;
+				case "ScreensImmersive":
+					OutComponent.type = manifestApplicationComponentType.ScreensImmersive;
+					break;
+				case "Console":
+				default:
+					OutComponent.type = manifestApplicationComponentType.Console;
+					break;
+				case "SystemUI":
+					OutComponent.type = manifestApplicationComponentType.SystemUI;
+					break;
+			}
+
+			if (ComponentElement.ContainsKey("ExtraComponentSubElements"))
+			{
+				// Unfortunately there are no config object array parsing functions in UBT
+				string SubElementsString = ComponentElement["ExtraComponentSubElements"];
+				string ConfObjArrayPattern = "\\([a-zA-Z0-9]+=[a-zA-Z0-9]+,[a-zA-Z0-9]+=\"?[a-zA-Z0-9]+\"?\\)";
+				Regex ConfigObjArrayRegex = new Regex(ConfObjArrayPattern);
+				MatchCollection ConfigObjMatches = ConfigObjArrayRegex.Matches(SubElementsString);
+				if (ConfigObjMatches.Count != 0)
+				{
+					OutComponent.Items = new object[ConfigObjMatches.Count];
+					for (int Index = 0; Index < ConfigObjMatches.Count; ++Index)
+					{
+						Match Match = ConfigObjMatches[Index];
+						Dictionary<string, string> SubElement;
+						if (ConfigHierarchy.TryParse(Match.Value, out SubElement))
+						{
+							OutComponent.Items[Index] = GetComponentSubElement(SubElement["ElementType"], SubElement["Value"]);
+						}
+					}
+				}
+			}
+			return OutComponent;
+		}
+
 		public string GenerateManifest(string ProjectName, bool bForDistribution, string Architecture)
 		{
 			ConfigHierarchy GameIni = GetConfigCacheIni(ConfigHierarchyType.Game);
@@ -172,63 +278,115 @@ namespace UnrealBuildTool
 			}
 
 			ConfigHierarchy EngineIni = GetConfigCacheIni(ConfigHierarchyType.Engine);
-			int VersionCode;
+			Int32 VersionCode;
 			EngineIni.GetInt32("/Script/LuminRuntimeSettings.LuminRuntimeSettings", "VersionCode", out VersionCode);
 
 			string SDKVersion = GetMLSDKVersion(EngineIni);
-
-			StringBuilder Text = new StringBuilder();
-
 			string PackageName = GetPackageName(ProjectName);
 			string ApplicationDisplayName = GetApplicationDisplayName(ProjectName);
 			string MinimumAPILevel = GetMinimumAPILevelRequired();
 			string TargetExecutableName = "bin/" + ProjectName;
 
-			Text.AppendLine(string.Format("<manifest xmlns:ml=\"magicleap\" ml:package=\"{0}\" ml:version_name=\"{1}\" ml:version_code=\"{2}\">", PackageName, ProjectVersion, VersionCode));
-			Text.AppendLine(string.Format("\t<application ml:visible_name=\"{0}\" ml:sdk_version=\"{1}\" ml:min_api_level=\"{2}\">",
-				ApplicationDisplayName,
-				SDKVersion,
-				MinimumAPILevel));
-			GetAppPrivileges(EngineIni, Text);
-			Text.AppendLine(string.Format("\t\t<component ml:name=\".fullscreen\" ml:visible_name=\"{0}\" ml:binary_name=\"{1}\" ml:type=\"{2}\">", ApplicationDisplayName, TargetExecutableName, GetApplicationType()));
+			PackageManifest.version_name = ProjectVersion;
+			PackageManifest.package = PackageName;
+			PackageManifest.version_code = Convert.ToUInt64(VersionCode);
+
+			PackageManifest.application = new manifestApplication
+			{
+				sdk_version = SDKVersion,
+				min_api_level = MinimumAPILevel,
+				visible_name = ApplicationDisplayName
+			};
 
-			string IconTag = string.Format("<icon ml:model_folder=\"{0}\" ml:portal_folder=\"{1}\"/>", GetIconModelStagingPath(), GetIconPortalStagingPath());
-			Text.AppendLine(string.Format("\t\t\t{0}", IconTag));
+			List<string> AppPrivileges;
+			EngineIni.GetArray("/Script/LuminRuntimeSettings.LuminRuntimeSettings", "AppPrivileges", out AppPrivileges);
 
-			List<string> ExtraComponentNodes;
-			EngineIni.GetArray("/Script/LuminRuntimeSettings.LuminRuntimeSettings", "ExtraComponentNodes", out ExtraComponentNodes);
-			if (ExtraComponentNodes != null)
+			List<string> ExtraComponentElements;
+			EngineIni.GetArray("/Script/LuminRuntimeSettings.LuminRuntimeSettings", "ExtraComponentElements", out ExtraComponentElements);
+
+			// We always add an additional item as that will be our 'root' <component>
+			int Size = (ExtraComponentElements == null ? AppPrivileges.Count() : AppPrivileges.Count() + ExtraComponentElements.Count()) + 1;
+			// Index used for sibling elements (app privileges, root component and any extra components)
+			int CurrentIndex = 0;
+			PackageManifest.application.Items = new object[Size];
+			// Remove all invalid strings from the list of strings
+			AppPrivileges.RemoveAll(item => item == "Invalid");
+			// Privileges get added first
+			for (int Index = 0; Index < AppPrivileges.Count(); ++Index)
 			{
-				foreach (string ComponentNode in ExtraComponentNodes)
+				string TrimmedPrivilege = AppPrivileges[Index].Trim(' ');
+				if (TrimmedPrivilege != "")
 				{
-					Text.AppendLine(string.Format("\t\t\t{0}", ComponentNode));
+					PackageManifest.application.Items[CurrentIndex] = new manifestApplicationUsesprivilege
+					{
+						name = TrimmedPrivilege,
+					};
+					CurrentIndex++;
 				}
 			}
 
-			Text.AppendLine("\t\t</component>");
+			// Then our root component, this is important as `mldb launch` will use the first component in the manifest
+			PackageManifest.application.Items[CurrentIndex] = new manifestApplicationComponent();
+			manifestApplicationComponent RootComponent = (manifestApplicationComponent)PackageManifest.application.Items[CurrentIndex];
+			RootComponent.name = ".fullscreen";
+			RootComponent.visible_name = ApplicationDisplayName;
+			RootComponent.binary_name = TargetExecutableName;
+			RootComponent.type = GetApplicationType();
+
+			// Sub-elements under root <component>
+			List<string> ExtraComponentSubElements;
+			EngineIni.GetArray("/Script/LuminRuntimeSettings.LuminRuntimeSettings", "ExtraComponentSubElements", out ExtraComponentSubElements);
+			RootComponent.Items = (ExtraComponentSubElements == null ? new object[1] : new object[ExtraComponentSubElements.Count() + 1]);
+
+			// Root component icon
+			RootComponent.Items[0] = new manifestApplicationComponentIcon();
+			((manifestApplicationComponentIcon)RootComponent.Items[0]).model_folder = GetIconModelStagingPath();
+			((manifestApplicationComponentIcon)RootComponent.Items[0]).portal_folder = GetIconPortalStagingPath();
+
+			if (ExtraComponentSubElements != null)
+			{
+				for (int Index = 0; Index < ExtraComponentSubElements.Count(); ++Index)
+				{
+					Dictionary<string, string> NodeContent;
+					if (ConfigHierarchy.TryParse(ExtraComponentSubElements[Index], out NodeContent))
+					{
+						RootComponent.Items[Index + 1] = GetComponentSubElement(NodeContent["ElementType"], NodeContent["Value"]);
+					}
+				}
+			}
 
-			List<string> ExtraApplicationNodes;
-			EngineIni.GetArray("/Script/LuminRuntimeSettings.LuminRuntimeSettings", "ExtraApplicationNodes", out ExtraApplicationNodes);
-			if (ExtraApplicationNodes != null)
+			// Finally, add additional components
+			CurrentIndex++;
+			if (ExtraComponentElements != null)
 			{
-				foreach (string ApplicationNode in ExtraApplicationNodes)
+				for (int Index = 0; Index < ExtraComponentElements.Count(); ++Index)
 				{
-					Text.AppendLine(string.Format("\t\t{0}", ApplicationNode));
+					Dictionary<string, string> ComponentElement;
+					if (ConfigHierarchy.TryParse(ExtraComponentElements[Index], out ComponentElement))
+					{
+						PackageManifest.application.Items[CurrentIndex] = GetComponentElement(ComponentElement);
+						CurrentIndex++;
+					}
 				}
 			}
 
-			Text.AppendLine("\t</application>");
-			Text.AppendLine("</manifest>");
+			// Wrap up serialization
+			XmlSerializer PackageManifestSerializer = new XmlSerializer(PackageManifest.GetType());
+			XmlSerializerNamespaces MLNamespace = new XmlSerializerNamespaces();
+			MLNamespace.Add("ml", "magicleap");
+			StringWriter Writer = new StringWriter();
+
+			PackageManifestSerializer.Serialize(Writer, PackageManifest, MLNamespace);
 
 			// allow plugins to modify final manifest HERE
 			XDocument XDoc;
 			try
 			{
-				XDoc = XDocument.Parse(Text.ToString());
+				XDoc = XDocument.Parse(Writer.ToString());
 			}
 			catch (Exception e)
 			{
-				throw new BuildException("LuminManifest.xml is invalid {0}\n{1}", e, Text.ToString());
+				throw new BuildException("LuminManifest.xml is invalid {0}\n{1}", e, Writer.ToString());
 			}
 
 			UPL.ProcessPluginNode(Architecture, "luminManifestUpdates", "", ref XDoc);
@@ -315,56 +473,42 @@ namespace UnrealBuildTool
 				// If asked for, and if we are doing a distribution package, we strip debug symbols.
 				Directory.CreateDirectory(Path.GetDirectoryName(ExecSrcFile));
 				ToolChain.StripSymbols(new Tools.DotNETCommon.FileReference(ExePath), new Tools.DotNETCommon.FileReference(ExecSrcFile));
-				// We also create a SYM file to support debugging of stripped executables
-				string SymFile = Path.ChangeExtension(ExecSrcFile, "sym");
-				ToolChain.ExtractSymbols(new FileReference(ExePath), new FileReference(SymFile));
-				ToolChain.LinkSymbols(new FileReference(SymFile), new FileReference(ExecSrcFile));
 			}
 			else
 			{
 				// The generated mabu needs the src exe file. So we copy the original as-is so mabu can find it.
 				Directory.CreateDirectory(Path.GetDirectoryName(ExecSrcFile));
 				File.Copy(ExePath, ExecSrcFile, true);
-				// If the exe has debug info we need to remove a likely old sym file so that debugging doesn't use
-				// outdated information.
-				string SymFile = Path.ChangeExtension(ExecSrcFile, "sym");
-				File.Delete(SymFile);
 			}
 
+			// We also create a SYM file to support debugging
+			string SymFile = Path.ChangeExtension(ExecSrcFile, "sym");
+			ToolChain.ExtractSymbols(new FileReference(ExePath), new FileReference(SymFile));
+			ToolChain.LinkSymbols(new FileReference(SymFile), new FileReference(ExecSrcFile));
+			
 			// Generate manifest (after UPL is setup
-			string Architecture = "arm64-v8a";
+			const string Architecture = "arm64-v8a";
 			var Manifest = GenerateManifest(ProjectName, bForDistribution, Architecture);
 			File.WriteAllText(ManifestFile, Manifest);
 
+			string MabuPackagingMessage = "Building mabu package....";
 			string Certificate = GetRuntimeSetting("Certificate");
 			Certificate = CleanFilePath(Certificate);
 			if (!string.IsNullOrEmpty(Certificate))
 			{
-				Certificate = Path.GetFullPath(Path.Combine(ProjectDirectory.FullName, Certificate));
-
-				if (File.Exists(Certificate))
-				{
-					ToolChain.RunMabuWithException(Path.GetDirectoryName(MabuFile), String.Format("-t device -s \"{0}\" -o \"{1}\" \"{2}\"", Certificate, MabuOutputPath, Path.GetFileName(MabuFile)), "Building signed mabu package....");
-				}
-				else
-				{
-					throw new BuildException(string.Format("Certificate file does not exist at path {0}. Please enter a valid certificate file path in Project Settings > Magic Leap or clear the field if you do not intend to sign the package.", Certificate));
-				}
+				// For legacy sakes. We used to print this message when signing via the mbu commnd line so should continue to do so.
+				// However, now this would only indicate if we are signing via the .package file and does not take into consideration
+				// the MLCERT env var.
+				MabuPackagingMessage = "Building signed mabu package....";
 			}
-			else
+			else if (bForDistribution)
 			{
-				if (bForDistribution)
-				{
-					// Certificate required for a distribution package.
-					throw new BuildException("Packaging for distribution, however no certificate file has been chosen. Please enter a certificate file in Project Settings > Magic Leap.");
-				}
-				// If a certificate file is not present but the package is not for distribution, package without signing.
-				else
-				{
-					ToolChain.RunMabuWithException(Path.GetDirectoryName(MabuFile), String.Format("-t device --allow-unsigned -o \"{0}\" \"{1}\"", MabuOutputPath, Path.GetFileName(MabuFile)), "Building mabu package....");
-				}
+				// The user could be signing via the MLCERT env var so instead of throwing an exception, we simply log a warning.
+				Log.TraceWarning("Packaging for distribution, however no certificate file has been chosen. Are you using the MLCERT environment variable instead?");
 			}
 
+			ToolChain.RunMabuWithException(Path.GetDirectoryName(MabuFile), String.Format("-t device --allow-unsigned -o \"{0}\" \"{1}\"", MabuOutputPath, Path.GetFileName(MabuFile)), MabuPackagingMessage);
+
 			// copy the .mpk into binaries
 			// @todo Lumin: Move this logic into a function in this class, and have AndroidAutomation call into it in GetFinalMpkName
 			// @todo Lumin: Handle the whole Prebuilt thing, it may need to go somwehere else, or maybe this isn't even called?
diff --git a/Engine/Source/Programs/UnrealBuildTool/UnrealBuildTool.csproj b/Engine/Source/Programs/UnrealBuildTool/UnrealBuildTool.csproj
--- a/Engine/Source/Programs/UnrealBuildTool/UnrealBuildTool.csproj
+++ b/Engine/Source/Programs/UnrealBuildTool/UnrealBuildTool.csproj
@@ -168,6 +168,7 @@
     <Compile Include="Modes\WriteMetadataMode.cs" />
     <Compile Include="Platform\IOS\IOSPostBuildSyncMode.cs" />
     <Compile Include="Platform\IOS\MobileProvisionContents.cs" />
+    <Compile Include="Platform\Lumin\LuminPackageManifest.cs" />
     <Compile Include="Platform\Lumin\UEDeployLumin.cs" />
     <Compile Include="ProjectFiles\CLion\CLionGenerator.cs" />
     <Compile Include="ProjectFiles\CMake\CMakefileGenerator.cs" />
diff --git a/Engine/Source/Runtime/Apple/MetalRHI/Public/MetalRHI.h b/Engine/Source/Runtime/Apple/MetalRHI/Public/MetalRHI.h
--- a/Engine/Source/Runtime/Apple/MetalRHI/Public/MetalRHI.h
+++ b/Engine/Source/Runtime/Apple/MetalRHI/Public/MetalRHI.h
@@ -136,7 +136,7 @@ public:
 	virtual void RHIVirtualTextureSetFirstMipInMemory(FRHITexture2D* Texture, uint32 FirstMip) final override;
 	virtual void RHIVirtualTextureSetFirstMipVisible(FRHITexture2D* Texture, uint32 FirstMip) final override;
 	virtual void RHIExecuteCommandList(FRHICommandList* CmdList) final override;
-	virtual void* RHIGetNativeDevice() final override;
+	virtual void* METALRHI_API RHIGetNativeDevice() final override;
 	virtual class IRHICommandContext* RHIGetDefaultContext() final override;
 	virtual IRHIComputeContext* RHIGetDefaultAsyncComputeContext() final override;
 	virtual class IRHICommandContextContainer* RHIGetCommandContextContainer(int32 Index, int32 Num) final override;
diff --git a/Engine/Source/Runtime/Apple/MetalRHI/Public/MetalResources.h b/Engine/Source/Runtime/Apple/MetalRHI/Public/MetalResources.h
--- a/Engine/Source/Runtime/Apple/MetalRHI/Public/MetalResources.h
+++ b/Engine/Source/Runtime/Apple/MetalRHI/Public/MetalResources.h
@@ -454,7 +454,7 @@ public:
 };
 
 /** Texture/RT wrapper. */
-class FMetalSurface
+class METALRHI_API FMetalSurface
 {
 public:
 
diff --git a/Engine/Source/Runtime/ApplicationCore/Private/Lumin/LuminApplication.cpp b/Engine/Source/Runtime/ApplicationCore/Private/Lumin/LuminApplication.cpp
--- a/Engine/Source/Runtime/ApplicationCore/Private/Lumin/LuminApplication.cpp
+++ b/Engine/Source/Runtime/ApplicationCore/Private/Lumin/LuminApplication.cpp
@@ -13,11 +13,13 @@ void FLuminApplication::DeferredKeyEvent::SendModified(
 	{
 	case Type::KeyDown:
 		InOutModifierMask = KeyData;
-		MessageHandler->OnKeyDown(KeyCode, 0, false);
+		// Engine requires a CharacterCode != 0, so KeyCode is passed in twice
+		MessageHandler->OnKeyDown(KeyCode, KeyCode, false);
 		break;
 	case Type::KeyUp:
 		InOutModifierMask = KeyData;
-		MessageHandler->OnKeyUp(KeyCode, 0, false);
+		// Engine requires a CharacterCode != 0, so KeyCode is passed in twice
+		MessageHandler->OnKeyUp(KeyCode, KeyCode, false);
 		break;
 	case Type::Char:
 		MessageHandler->OnKeyChar(KeyData, false);
@@ -122,6 +124,11 @@ void FLuminApplication::InitializeInputCallbacks()
 		auto application = reinterpret_cast<FLuminApplication*>(data);
 		if (application)
 		{
+			// BACKSPACE TEMP FIX UNTIL PLATFORM'S FIX
+			if (key_code == MLKeyCode::MLKEYCODE_DEL)
+			{
+				application->AddDeferredKeyEvent(DeferredKeyEvent('\b'));
+			}
 			application->AddDeferredKeyEvent(DeferredKeyEvent(
 				DeferredKeyEvent::Type::KeyUp, key_code, modifier_mask));
 		}
diff --git a/Engine/Source/Runtime/ApplicationCore/Public/Lumin/LuminApplication.h b/Engine/Source/Runtime/ApplicationCore/Public/Lumin/LuminApplication.h
--- a/Engine/Source/Runtime/ApplicationCore/Public/Lumin/LuminApplication.h
+++ b/Engine/Source/Runtime/ApplicationCore/Public/Lumin/LuminApplication.h
@@ -6,9 +6,8 @@
 #include "Android/AndroidApplication.h"
 #include <Containers/Queue.h>
 
-#include <ml_api.h>
-#include <ml_keycodes.h>
-#include <ml_input.h>
+#include "Lumin/CAPIShims/LuminAPIInput.h"
+#include "Lumin/CAPIShims/LuminAPIKeyCodes.h"
 
 /**
 * Lumin-specific application implementation.
diff --git a/Engine/Source/Runtime/Core/Private/Lumin/LuminLifecycle.cpp b/Engine/Source/Runtime/Core/Private/Lumin/LuminLifecycle.cpp
--- a/Engine/Source/Runtime/Core/Private/Lumin/LuminLifecycle.cpp
+++ b/Engine/Source/Runtime/Core/Private/Lumin/LuminLifecycle.cpp
@@ -5,14 +5,19 @@
 #include "Misc/CallbackDevice.h"
 #include "Async/TaskGraphInterfaces.h"
 #include "Misc/CommandLine.h"
+#include "GenericPlatform/GenericPlatformFile.h"
+#include "Lumin/LuminPlatformFile.h"
+#include "Lumin/CAPIShims/LuminAPIFileInfo.h"
 
 DEFINE_LOG_CATEGORY(LogLifecycle);
 
 bool FLuminLifecycle::bIsEngineLoopInitComplete = false;
 bool FLuminLifecycle::bIsAppPaused = false;
 MLResult FLuminLifecycle::LifecycleState = MLResult_UnspecifiedFailure;
-MLLifecycleCallbacks FLuminLifecycle::LifecycleCallbacks = { nullptr, nullptr, nullptr, nullptr, nullptr };
-TArray<FString> FLuminLifecycle::PendingStartupArgs;
+MLLifecycleInitArgList* FLuminLifecycle::InitArgList = nullptr;
+MLLifecycleCallbacksEx FLuminLifecycle::LifecycleCallbacks;
+TArray<FString> FLuminLifecycle::InitStringArgs;
+TArray<FLuminFileInfo> FLuminLifecycle::InitFileArgs;
 
 void FLuminLifecycle::Initialize()
 {
@@ -21,13 +26,17 @@ void FLuminLifecycle::Initialize()
 		return;
 	}
 
+	MLLifecycleCallbacksExInit(&LifecycleCallbacks);
 	LifecycleCallbacks.on_stop = Stop_Handler;
 	LifecycleCallbacks.on_pause = Pause_Handler;
 	LifecycleCallbacks.on_resume = Resume_Handler;
 	LifecycleCallbacks.on_unload_resources = UnloadResources_Handler;
 	LifecycleCallbacks.on_new_initarg = OnNewInitArgs_Handler;
+	LifecycleCallbacks.on_device_active = OnDeviceActive_Handler;
+	LifecycleCallbacks.on_device_reality = OnDeviceReality_Handler;
+	LifecycleCallbacks.on_device_standby = OnDeviceStandby_Handler;
 
-	LifecycleState = MLLifecycleInit(&LifecycleCallbacks, nullptr);
+	LifecycleState = MLLifecycleInitEx(&LifecycleCallbacks, nullptr);
 
 	FCoreDelegates::OnFEngineLoopInitComplete.AddStatic(FLuminLifecycle::OnFEngineLoopInitComplete_Handler);
 
@@ -69,6 +78,12 @@ void FLuminLifecycle::Stop_Handler(void* ApplicationContext)
 		FTaskGraphInterface::Get().WaitUntilTaskCompletes(WillTerminateTask);
 	}
 
+	if (InitArgList != nullptr)
+	{
+		MLResult Result = MLLifecycleFreeInitArgList(&InitArgList);
+		UE_CLOG(MLResult_Ok != Result, LogLifecycle, Error, TEXT("Error %s freeing init args list."), UTF8_TO_TCHAR(MLGetResultString(Result)));
+	}
+
 	FPlatformMisc::RequestExit(false);
 }
 
@@ -120,6 +135,11 @@ void FLuminLifecycle::Resume_Handler(void* ApplicationContext)
 		}
 
 		bIsAppPaused = false;
+
+		// if an app is resumed from a paused state, init_args callback is received before resume.
+		// in that case, we simply cache the args and fire them after the engine actually resumes.
+		// Otherwise, certain events like changing the map while the app is paused might cause a crash, or a deadlock.
+		FirePendingInitArgs();
 	}
 }
 
@@ -139,7 +159,10 @@ void FLuminLifecycle::UnloadResources_Handler(void* ApplicationContext)
 
 void FLuminLifecycle::OnNewInitArgs_Handler(void* ApplicationContext)
 {
-	MLLifecycleInitArgList* InitArgList = nullptr;
+	IPlatformFile& PlatformFile = IPlatformFile::GetPlatformPhysical();
+	FLuminPlatformFile* LuminPlatformFile = static_cast<FLuminPlatformFile*>(&PlatformFile);
+
+	InitArgList = nullptr;
 	MLResult Result = MLLifecycleGetInitArgList(&InitArgList);
 	if (Result == MLResult_Ok && InitArgList != nullptr)
 	{
@@ -147,7 +170,6 @@ void FLuminLifecycle::OnNewInitArgs_Handler(void* ApplicationContext)
 		Result = MLLifecycleGetInitArgListLength(InitArgList, &InitArgCount);
 		if (Result == MLResult_Ok && InitArgCount > 0)
 		{
-			TArray<FString> InitArgsArray;
 			for (int64 i = 0; i < InitArgCount; ++i)
 			{
 				const MLLifecycleInitArg* InitArg = nullptr;
@@ -163,39 +185,121 @@ void FLuminLifecycle::OnNewInitArgs_Handler(void* ApplicationContext)
 						ArgStr.Append(UTF8_TO_TCHAR(Arg));
 						ArgStr.TrimEndInline();
 						FCommandLine::Append(*ArgStr);
+	
+						InitStringArgs.Add(UTF8_TO_TCHAR(Arg));
+					}
+
+					int64_t FileInfoListLength = 0;
+					Result = MLLifecycleGetFileInfoListLength(InitArg, &FileInfoListLength);
+					FPlatformMisc::LowLevelOutputDebugStringf(TEXT("FileInfoListLength = %lld"), FileInfoListLength);
+					for (int64_t j = 0; j < FileInfoListLength; ++j)
+					{
+						const MLFileInfo* FileInfo = nullptr;
+						Result = MLLifecycleGetFileInfoByIndex(InitArg, j, &FileInfo);
+						if (Result == MLResult_Ok && FileInfo != nullptr)
+						{
+							FLuminFileInfo LuminFile;
+							const char* Filename = nullptr;
+							Result = MLFileInfoGetFileName(FileInfo, &Filename);
+							if (Result == MLResult_Ok && Filename != nullptr)
+							{
+								LuminFile.FileName = UTF8_TO_TCHAR(Filename);
+							}
+
+							const char* MimeType = nullptr;
+							Result = MLFileInfoGetMimeType(FileInfo, &MimeType);
+							if (Result == MLResult_Ok && MimeType != nullptr)
+							{
+								LuminFile.MimeType = UTF8_TO_TCHAR(MimeType);
+							}
+
+							LuminFile.FileHandle = LuminPlatformFile->GetFileHandleForMLFileInfo(FileInfo);
+							if (LuminFile.FileHandle != nullptr)
+							{
+								InitFileArgs.Add(LuminFile);
+							}
+						}
 					}
-					InitArgsArray.Add(UTF8_TO_TCHAR(Arg));
 				}
 			}
 
-			if (FTaskGraphInterface::IsRunning() && bIsEngineLoopInitComplete)
-			{
-				FGraphEventRef StartupArgumentsTask = FFunctionGraphTask::CreateAndDispatchWhenReady([&]()
-				{
-					FCoreDelegates::ApplicationReceivedStartupArgumentsDelegate.Broadcast(InitArgsArray);
-				}, TStatId(), nullptr, ENamedThreads::GameThread);
-				FTaskGraphInterface::Get().WaitUntilTaskCompletes(StartupArgumentsTask);
-			}
-			else
+			if (bIsEngineLoopInitComplete && !bIsAppPaused)
 			{
-				PendingStartupArgs = InitArgsArray;
+				FirePendingInitArgs();
 			}
 		}
 	}
 }
 
+void FLuminLifecycle::OnDeviceActive_Handler(void* ApplicationContext)
+{
+	UE_LOG(LogLifecycle, Log, TEXT("FLuminLifecycle : The device is active again."));
+	if (FTaskGraphInterface::IsRunning())
+	{
+		FGraphEventRef WillTerminateTask = FFunctionGraphTask::CreateAndDispatchWhenReady([&]()
+		{
+#if PLATFORM_LUMIN
+			FLuminDelegates::DeviceHasReactivatedDelegate.Broadcast();
+#endif // PLATFORM_LUMIN
+		}, TStatId(), nullptr, ENamedThreads::GameThread);
+		FTaskGraphInterface::Get().WaitUntilTaskCompletes(WillTerminateTask);
+	}
+}
+
+void FLuminLifecycle::OnDeviceReality_Handler(void* ApplicationContext)
+{
+	UE_LOG(LogLifecycle, Log, TEXT("FLuminLifecycle : The device's reality button has been pressed."));
+	if (FTaskGraphInterface::IsRunning())
+	{
+		FGraphEventRef WillTerminateTask = FFunctionGraphTask::CreateAndDispatchWhenReady([&]()
+		{
+#if PLATFORM_LUMIN
+			FLuminDelegates::DeviceWillEnterRealityModeDelegate.Broadcast();
+#endif // PLATFORM_LUMIN
+		}, TStatId(), nullptr, ENamedThreads::GameThread);
+		FTaskGraphInterface::Get().WaitUntilTaskCompletes(WillTerminateTask);
+	}
+}
+
+void FLuminLifecycle::OnDeviceStandby_Handler(void* ApplicationContext)
+{
+	UE_LOG(LogLifecycle, Log, TEXT("FLuminLifecycle : The device is going into standby."));
+	if (FTaskGraphInterface::IsRunning())
+	{
+		FGraphEventRef WillTerminateTask = FFunctionGraphTask::CreateAndDispatchWhenReady([&]()
+		{
+#if PLATFORM_LUMIN
+			FLuminDelegates::DeviceWillGoInStandbyDelegate.Broadcast();
+#endif // PLATFORM_LUMIN
+		}, TStatId(), nullptr, ENamedThreads::GameThread);
+		FTaskGraphInterface::Get().WaitUntilTaskCompletes(WillTerminateTask);
+	}
+}
+
 void FLuminLifecycle::OnFEngineLoopInitComplete_Handler()
 {
 	bIsEngineLoopInitComplete = true;
+	FirePendingInitArgs();
+}
 
-	if (PendingStartupArgs.Num() > 0 && FTaskGraphInterface::IsRunning())
+void FLuminLifecycle::FirePendingInitArgs()
+{
+	if ((InitStringArgs.Num() > 0 || InitFileArgs.Num() > 0) && FTaskGraphInterface::IsRunning())
 	{
 		FPlatformMisc::LowLevelOutputDebugString(TEXT("FLuminLifecycle :: Firing startup args..."));
 		FGraphEventRef StartupArgumentsTask = FFunctionGraphTask::CreateAndDispatchWhenReady([&]()
 		{
-			FCoreDelegates::ApplicationReceivedStartupArgumentsDelegate.Broadcast(PendingStartupArgs);
+			FCoreDelegates::ApplicationReceivedStartupArgumentsDelegate.Broadcast(InitStringArgs);
+		}, TStatId(), nullptr, ENamedThreads::GameThread);
+		FTaskGraphInterface::Get().WaitUntilTaskCompletes(StartupArgumentsTask);
+
+		StartupArgumentsTask = FFunctionGraphTask::CreateAndDispatchWhenReady([&]()
+		{
+			FLuminDelegates::LuminAppReceivedStartupArgumentsDelegate.Broadcast(InitStringArgs, InitFileArgs);
 		}, TStatId(), nullptr, ENamedThreads::GameThread);
 		FTaskGraphInterface::Get().WaitUntilTaskCompletes(StartupArgumentsTask);
-		PendingStartupArgs.Empty();
+
+		InitStringArgs.Empty();
+		InitFileArgs.Empty();
 	}
 }
diff --git a/Engine/Source/Runtime/Core/Private/Lumin/LuminLifecycle.h b/Engine/Source/Runtime/Core/Private/Lumin/LuminLifecycle.h
--- a/Engine/Source/Runtime/Core/Private/Lumin/LuminLifecycle.h
+++ b/Engine/Source/Runtime/Core/Private/Lumin/LuminLifecycle.h
@@ -1,8 +1,9 @@
 // Copyright 1998-2019 Epic Games, Inc. All Rights Reserved.
 
 #pragma once
-#include <ml_lifecycle.h>
+#include "Lumin/CAPIShims/LuminAPILifecycle.h"
 #include "Logging/LogMacros.h"
+#include "Lumin/LuminPlatformDelegates.h"
 
 class CORE_API FLuminLifecycle
 {
@@ -18,14 +19,22 @@ private:
 	// mldb launch -i "-arg1=value1 -arg2=value2" <package_name>
 	static void OnNewInitArgs_Handler(void* ApplicationContext);
 
+	static void OnDeviceActive_Handler(void* ApplicationContext);
+	static void OnDeviceReality_Handler(void* ApplicationContext);
+	static void OnDeviceStandby_Handler(void* ApplicationContext);
+
 	static void OnFEngineLoopInitComplete_Handler();
 
 private:
+	static void FirePendingInitArgs();
+
 	static bool bIsEngineLoopInitComplete;
 	static bool bIsAppPaused;
 	static MLResult LifecycleState;
-	static MLLifecycleCallbacks LifecycleCallbacks;
-	static TArray<FString> PendingStartupArgs;
+	static MLLifecycleCallbacksEx LifecycleCallbacks;
+	static MLLifecycleInitArgList* InitArgList;
+	static TArray<FString> InitStringArgs;
+	static TArray<FLuminFileInfo> InitFileArgs;
 };
 
 DECLARE_LOG_CATEGORY_EXTERN(LogLifecycle, Log, All);
diff --git a/Engine/Source/Runtime/Core/Private/Lumin/LuminPlatformDelegates.cpp b/Engine/Source/Runtime/Core/Private/Lumin/LuminPlatformDelegates.cpp
new file mode 100644
--- /dev/null
+++ b/Engine/Source/Runtime/Core/Private/Lumin/LuminPlatformDelegates.cpp
@@ -0,0 +1,7 @@
+#include "Lumin/LuminPlatformDelegates.h"
+
+FLuminDelegates::FLuminAppStartupArgumentsDelegate FLuminDelegates::LuminAppReceivedStartupArgumentsDelegate;
+
+FLuminDelegates::FLuminApplicationLifetimeDelegate FLuminDelegates::DeviceHasReactivatedDelegate;
+FLuminDelegates::FLuminApplicationLifetimeDelegate FLuminDelegates::DeviceWillEnterRealityModeDelegate;
+FLuminDelegates::FLuminApplicationLifetimeDelegate FLuminDelegates::DeviceWillGoInStandbyDelegate;
diff --git a/Engine/Source/Runtime/Core/Private/Lumin/LuminPlatformFile.cpp b/Engine/Source/Runtime/Core/Private/Lumin/LuminPlatformFile.cpp
--- a/Engine/Source/Runtime/Core/Private/Lumin/LuminPlatformFile.cpp
+++ b/Engine/Source/Runtime/Core/Private/Lumin/LuminPlatformFile.cpp
@@ -10,9 +10,14 @@
 #include "Misc/Paths.h"
 #include <sys/stat.h>   // mkdirp()
 #include <dirent.h>
+#include "Lumin/CAPIShims/LuminAPISharedFile.h"
 
 DEFINE_LOG_CATEGORY_STATIC(LogLuminPlatformFile, Log, All);
 
+FLuminFileInfo::FLuminFileInfo()
+: FileHandle(nullptr)
+{}
+
 // make an FTimeSpan object that represents the "epoch" for time_t (from a stat struct)
 const FDateTime UnixEpoch(1970, 1, 1);
 
@@ -46,6 +51,7 @@ namespace
 * files which are opened READ_ONLY.
 */
 // @todo lumin - Consider if we need managed handles or not.
+// TODO - handle shared files properly when using managed file handles.
 #define MANAGE_FILE_HANDLES 0
 
 /**
@@ -69,6 +75,9 @@ public:
 		, FileOffset(0)
 		, FileSize(0)
 #endif // MANAGE_FILE_HANDLES
+		, SharedFileList(nullptr)
+		, Fileinfo(nullptr)
+		, bReleaseFileInfo(false)
 	{
 		check(FileHandle > -1);
 #if MANAGE_FILE_HANDLES
@@ -88,6 +97,26 @@ public:
 #endif // MANAGE_FILE_HANDLES
 	}
 
+	FFileHandleLumin(int32 InFileHandle, MLSharedFileList* InSharedFileList)
+		: FileHandle(InFileHandle)
+		, SharedFileList(InSharedFileList)
+		, Fileinfo(nullptr)
+		, bReleaseFileInfo(false)
+	{
+		check(FileHandle > -1);
+		check(SharedFileList != nullptr)
+	}
+
+	FFileHandleLumin(int32 InFileHandle, const MLFileInfo* InFileInfo)
+		: FileHandle(InFileHandle)
+		, SharedFileList(nullptr)
+		, Fileinfo(InFileInfo)
+		, bReleaseFileInfo(false)
+	{
+		check(FileHandle > -1);
+		check(Fileinfo != nullptr)
+	}
+
 	virtual ~FFileHandleLumin()
 	{
 #if MANAGE_FILE_HANDLES
@@ -102,9 +131,28 @@ public:
 		else
 #endif // MANAGE_FILE_HANDLES
 		{
-			close(FileHandle);
+			if (SharedFileList != nullptr)
+			{
+				MLResult Result = MLSharedFileListRelease(&SharedFileList);
+				UE_CLOG(MLResult_Ok != Result, LogLuminPlatformFile, Error, TEXT("Error %s releasing shared file list for fd %d"), UTF8_TO_TCHAR(MLSharedFileGetResultString(Result)), FileHandle);
+			}
+			// Close if it's a normal unreal file. MLFileInfo fds are closed via MLLifecycleFreeInitArgList()
+			else if (Fileinfo == nullptr)
+			{
+				close(FileHandle);
+			}
 		}
+
 		FileHandle = -1;
+		SharedFileList = nullptr;
+		Fileinfo = nullptr;
+	}
+
+	bool SetMLFileInfoFD(MLFileInfo* InFileInfo) const
+	{
+		MLResult Result = MLFileInfoSetFD(InFileInfo, FileHandle);
+		UE_CLOG(MLResult_Ok != Result, LogLuminPlatformFile, Error, TEXT("Error setting MLFileInfo FD : %s"), UTF8_TO_TCHAR(MLGetResultString(Result)));
+		return MLResult_Ok == Result;
 	}
 
 	virtual int64 Tell() override
@@ -348,6 +396,10 @@ private:
 	static __thread FFileHandleLumin* ActiveHandles[ACTIVE_HANDLE_COUNT];
 	static __thread double AccessTimes[ACTIVE_HANDLE_COUNT];
 #endif // MANAGE_FILE_HANDLES
+
+	MLSharedFileList* SharedFileList;
+	const MLFileInfo* Fileinfo;
+	bool bReleaseFileInfo;
 };
 
 #if MANAGE_FILE_HANDLES
@@ -1097,14 +1149,12 @@ FString FLuminPlatformFile::ConvertToLuminPath(const FString& Filename, bool bFo
 	// Remove the writable path if present, we will prepend it the correct base path as needed.
 	Result.ReplaceInline(*FLuminPlatformMisc::GetApplicationWritableDirectoryPath(), TEXT(""));
 
+	FString lhs;
+
 	// If the file is for writing, then add it to the app writable directory path.
 	if (bForWrite)
 	{
-		FString lhs = FLuminPlatformMisc::GetApplicationWritableDirectoryPath();
-		FString rhs = Result;
-		lhs.RemoveFromEnd(TEXT("/"));
-		rhs.RemoveFromStart(TEXT("/"));
-		Result = lhs / rhs;
+		lhs = FLuminPlatformMisc::GetApplicationWritableDirectoryPath();
 	}
 	else
 	{
@@ -1112,35 +1162,19 @@ FString FLuminPlatformFile::ConvertToLuminPath(const FString& Filename, bool bFo
 		FString Value;
 		// Cache this value as the command line doesn't change.
 		static bool bHasHostIP = FParse::Value(FCommandLine::Get(), TEXT("filehostip"), Value) || FParse::Value(FCommandLine::Get(), TEXT("streaminghostip"), Value);
-		// @todo Lumin support iterative deployment!! See LuminPlatform.Automation.cs
 		static bool bIsIterative = FParse::Value(FCommandLine::Get(), TEXT("iterative"), Value);
 
 		if (bHasHostIP)
 		{
-			FString lhs = FLuminPlatformMisc::GetApplicationWritableDirectoryPath();
-			FString rhs = Result;
-			lhs.RemoveFromEnd(TEXT("/"));
-			rhs.RemoveFromStart(TEXT("/"));
-			Result = lhs / rhs;
+			lhs = FLuminPlatformMisc::GetApplicationWritableDirectoryPath();
 		}
 		else if (bIsIterative)
 		{
-			FString lhs = FLuminPlatformMisc::GetApplicationWritableDirectoryPath();
-			FString rhs = Result;
-			lhs.RemoveFromEnd(TEXT("/"));
-			rhs.RemoveFromStart(TEXT("/"));
-			Result = lhs / rhs;
+			lhs = FLuminPlatformMisc::GetApplicationWritableDirectoryPath();
 		}
 		else
 		{
-//			Result = Filename;
-			// @todo Lumin: Don't know why the Lumin code didn't have this
-			FString lhs = FLuminPlatformMisc::GetApplicationPackageDirectoryPath();
-			FString rhs = Result;
-			lhs.RemoveFromEnd(TEXT("/"));
-			rhs.RemoveFromStart(TEXT("/"));
-			Result = lhs / rhs;
-
+			lhs = FLuminPlatformMisc::GetApplicationPackageDirectoryPath();
 		}
 	}
 
@@ -1148,8 +1182,84 @@ FString FLuminPlatformFile::ConvertToLuminPath(const FString& Filename, bool bFo
 	FPlatformMisc::LowLevelOutputDebugStringf(TEXT("LOG_LUMIN_PATH Write = %d Input = %s Output = %s"), bForWrite, *Filename, *Result.ToLower());
 #endif
 
-	// always use lower case ... always
-	return Result.ToLower();
+	// Lower only unreal path.
+	FString rhs = Result.ToLower();
+	rhs.RemoveFromStart(TEXT("/"));
+	lhs.RemoveFromEnd(TEXT("/"));
+
+	Result = lhs / rhs;
+
+	return Result;
+}
+
+
+IFileHandle* GetHandleForSharedFile(const TCHAR* Filename, bool bForWrite)
+{
+	const char* Filename_UTF8 = TCHAR_TO_UTF8(Filename);
+	MLSharedFileList* SharedFileList = nullptr;
+	MLResult Result;
+	if (bForWrite)
+	{
+		Result = MLSharedFileWrite(&Filename_UTF8, 1, &SharedFileList);
+	}
+	else
+	{
+		Result = MLSharedFileRead(&Filename_UTF8, 1, &SharedFileList);
+	}
+
+	if (Result == MLResult_Ok && SharedFileList != nullptr)
+	{
+		MLHandle ListLength = 0;
+		Result = MLSharedFileGetListLength(SharedFileList, &ListLength);
+		if (Result == MLResult_Ok && ListLength > 0)
+		{
+			MLFileInfo* FileInfo = nullptr;
+			Result = MLSharedFileGetMLFileInfoByIndex(SharedFileList, 0, &FileInfo);
+			if (Result == MLResult_Ok && FileInfo != nullptr)
+			{
+				MLFileDescriptor FileHandle = -1;
+				Result = MLFileInfoGetFD(FileInfo, &FileHandle);
+				if (Result == MLResult_Ok && FileHandle > -1)
+				{
+					return new FFileHandleLumin(FileHandle, SharedFileList);
+				}
+			}
+		}
+	}
+
+	UE_LOG(LogLuminPlatformFile, Error, TEXT("Error %s opening shared file %s for read."), UTF8_TO_TCHAR(MLSharedFileGetResultString(Result)), Filename);
+	return nullptr;
+}
+
+IFileHandle* FLuminPlatformFile::SharedFileOpenRead(const TCHAR* Filename)
+{
+	return GetHandleForSharedFile(Filename, false);
+}
+
+IFileHandle* FLuminPlatformFile::SharedFileOpenWrite(const TCHAR* Filename)
+{
+	return GetHandleForSharedFile(Filename, true);
+}
+
+IFileHandle* FLuminPlatformFile::GetFileHandleForMLFileInfo(const void* InFileInfo)
+{
+	const MLFileInfo* Fileinfo = static_cast<const MLFileInfo*>(InFileInfo);
+	MLFileDescriptor FileHandle = -1;
+	MLResult Result = MLFileInfoGetFD(Fileinfo, &FileHandle);
+	if (Result == MLResult_Ok && FileHandle > -1)
+	{
+		return new FFileHandleLumin(FileHandle, Fileinfo);
+	}
+
+	return nullptr;
+}
+
+bool FLuminPlatformFile::SetMLFileInfoFD(const IFileHandle* FileHandle, void* InFileInfo)
+{
+	check(FileHandle);
+	check(InFileInfo);
+	const FFileHandleLumin* LuminFileHandle = static_cast<const FFileHandleLumin*>(FileHandle);
+	return LuminFileHandle->SetMLFileInfoFD(static_cast<MLFileInfo*>(InFileInfo));
 }
 
 bool FLuminPlatformFile::FileExistsInternal(const FString& NormalizedFilename) const
diff --git a/Engine/Source/Runtime/Core/Private/Lumin/LuminPlatformMisc.cpp b/Engine/Source/Runtime/Core/Private/Lumin/LuminPlatformMisc.cpp
--- a/Engine/Source/Runtime/Core/Private/Lumin/LuminPlatformMisc.cpp
+++ b/Engine/Source/Runtime/Core/Private/Lumin/LuminPlatformMisc.cpp
@@ -17,8 +17,9 @@
 #include "Misc/ConfigCacheIni.h"
 #include "Misc/CommandLine.h"
 #include "Misc/Parse.h"
-#include <ml_lifecycle.h>
-#include <ml_logging.h>
+#include "Lumin/CAPIShims/LuminAPILifecycle.h"
+#include "Lumin/CAPIShims/LuminAPILogging.h"
+#include "Lumin/CAPIShims/LuminAPILocale.h"
 
 // @todo Lumin: this is super gross - but HMD plugin had it hardcoded as well
 static int32 GBuiltinDisplayWidth = 2560;
@@ -278,6 +279,39 @@ const FString& FLuminPlatformMisc::GetApplicationComponentName()
 	return ComponentName;
 }
 
+FString FLuminPlatformMisc::GetDefaultLocale()
+{
+	const char* CountryCode = nullptr;
+	MLResult Result = MLLocaleGetSystemCountry(&CountryCode);
+	if (Result != MLResult_Ok)
+	{
+		UE_LOG(LogLumin, Error, TEXT("MLLocaleGetSystemCountry failed with error '%s'"), UTF8_TO_TCHAR(MLGetResultString(Result)));
+		return FString();
+	}
+	
+	return UTF8_TO_TCHAR(CountryCode);
+}
+
+FString FLuminPlatformMisc::GetDefaultLanguage()
+{
+	const char* LanguageCode = nullptr;
+	MLResult Result = MLLocaleGetSystemLanguage(&LanguageCode);
+	if (Result != MLResult_Ok)
+	{
+		UE_LOG(LogLumin, Error, TEXT("MLLocaleGetSystemLanguage failed with error '%s'"), UTF8_TO_TCHAR(MLGetResultString(Result)));
+		return FString();
+	}
+
+	FString DefaultLanguage = UTF8_TO_TCHAR(LanguageCode);
+	FString DefaultCountryCode = GetDefaultLocale();
+	if (DefaultLanguage != DefaultCountryCode)
+	{
+		DefaultCountryCode.ToUpperInline();
+		DefaultLanguage = FString::Printf(TEXT("%s-%s"), *DefaultLanguage, *DefaultCountryCode);
+	}
+	
+	return DefaultLanguage;
+}
 
 void FLuminPlatformMisc::InitApplicationPaths()
 {
diff --git a/Engine/Source/Runtime/Core/Private/Lumin/LuminPlatformProcess.cpp b/Engine/Source/Runtime/Core/Private/Lumin/LuminPlatformProcess.cpp
--- a/Engine/Source/Runtime/Core/Private/Lumin/LuminPlatformProcess.cpp
+++ b/Engine/Source/Runtime/Core/Private/Lumin/LuminPlatformProcess.cpp
@@ -10,7 +10,7 @@
 #include <pthread.h>
 #include <libgen.h>
 #include <dlfcn.h>
-#include <ml_dispatch.h>
+#include "Lumin/CAPIShims/LuminAPIDispatch.h"
 
 const TCHAR* FLuminPlatformProcess::ComputerName()
 {
diff --git a/Engine/Source/Runtime/Core/Private/Misc/App.cpp b/Engine/Source/Runtime/Core/Private/Misc/App.cpp
--- a/Engine/Source/Runtime/Core/Private/Misc/App.cpp
+++ b/Engine/Source/Runtime/Core/Private/Misc/App.cpp
@@ -282,7 +282,10 @@ bool FApp::ShouldUseThreadingForPerformance()
 		IsRunningDedicatedServer() ||
 		!FPlatformProcess::SupportsMultithreading() ||
 		FPlatformMisc::NumberOfCoresIncludingHyperthreads() < MIN_CORE_COUNT;
-	return !OnlyOneThread;
+
+	// Enable at runtime for experimentation by passing "useperfthreads" as a command line arg.
+	static bool bForceUsePerfThreads = FParse::Param(FCommandLine::Get(), TEXT("useperfthreads"));
+	return !OnlyOneThread || bForceUsePerfThreads;
 }
 #undef MIN_CORE_COUNT
 
diff --git a/Engine/Source/Runtime/Core/Public/Lumin/CAPIShims/LuminAPI.h b/Engine/Source/Runtime/Core/Public/Lumin/CAPIShims/LuminAPI.h
new file mode 100644
--- /dev/null
+++ b/Engine/Source/Runtime/Core/Public/Lumin/CAPIShims/LuminAPI.h
@@ -0,0 +1,135 @@
+// Copyright 1998-2019 Epic Games, Inc. All Rights Reserved.
+
+#pragma once
+
+#include "CoreMinimal.h"
+#include "HAL/Platform.h"
+#include "Stats/Stats.h"
+
+#if !defined(WITH_MLSDK) || WITH_MLSDK
+
+#if defined(ML_NO_DEPRECATION_WARNING)
+#define MLSDK_API_NO_DEPRECATION_WARNING
+#endif
+#define ML_NO_DEPRECATION_WARNING
+#define ML_NO_DEPRECATION_DISABLED_MSG
+
+#if PLATFORM_WINDOWS
+
+	#ifndef LUMIN_THIRD_PARTY_INCLUDES_START
+		#define LUMIN_THIRD_PARTY_INCLUDES_START \
+			__pragma(warning(push)) \
+			__pragma(warning(disable: 4201)) \
+			THIRD_PARTY_INCLUDES_START
+	#endif
+
+	#ifndef LUMIN_THIRD_PARTY_INCLUDES_END
+		#define LUMIN_THIRD_PARTY_INCLUDES_END \
+			__pragma(warning(pop)) \
+			THIRD_PARTY_INCLUDES_END
+	#endif
+
+#else
+
+	#ifndef LUMIN_THIRD_PARTY_INCLUDES_START
+		#define LUMIN_THIRD_PARTY_INCLUDES_START \
+			THIRD_PARTY_INCLUDES_START
+	#endif
+
+	#ifndef LUMIN_THIRD_PARTY_INCLUDES_END
+		#define LUMIN_THIRD_PARTY_INCLUDES_END \
+			THIRD_PARTY_INCLUDES_END
+	#endif
+
+#endif
+
+LUMIN_THIRD_PARTY_INCLUDES_START
+#include <ml_api.h>
+LUMIN_THIRD_PARTY_INCLUDES_END
+
+// We default to delay loaded calling as we have multiple
+// sets of libraries to load on desktop platforms.
+// Enabled for lumin as well to load only necesarry libraries.
+#ifndef MLSDK_API_USE_STUBS
+#define MLSDK_API_USE_STUBS 1
+#endif
+
+#include "LuminAPIImpl.h"
+
+DECLARE_STATS_GROUP(TEXT("MLAPI"), STATGROUP_MLAPI, STATCAT_Advanced);
+
+namespace MLSDK_API
+{
+#if MLSDK_API_USE_STUBS
+	#define CREATE_FUNCTION_SHIM(library_name, return_type, function_name) \
+		DECLARE_CYCLE_STAT(TEXT(#function_name), STAT_ ## function_name, STATGROUP_MLAPI); \
+		struct library_name; \
+		struct function_name ## DelayCall; \
+		template<typename... Args> return_type function_name ## Shim(Args... args) \
+		{ \
+			SCOPE_CYCLE_COUNTER(STAT_ ## function_name); \
+			static DelayCall<library_name, function_name ## DelayCall, return_type, Args...> Call{ #library_name, #function_name }; \
+			return Call(args...); \
+		}
+
+	#define CREATE_GLOBAL_SHIM(library_name, return_type, variable_name) \
+		struct library_name; \
+		struct variable_name ## DelayCall; \
+		inline return_type variable_name ## Shim() \
+		{ \
+			static DelayValue<library_name, variable_name ## DelayCall, return_type> DelayedValue{ #library_name, #variable_name }; \
+			return DelayedValue.Get(); \
+		}
+	#define CREATE_DEPRECATED_SHIM(library_name, return_type, function_name) \
+		DECLARE_CYCLE_STAT(TEXT(#function_name), STAT_ ## function_name, STATGROUP_MLAPI); \
+		struct library_name; \
+		struct function_name ## DelayCall; \
+		template<typename... Args> MLSDK_API_DEPRECATED MLResult function_name ## Shim(Args... args) \
+		{ \
+			SCOPE_CYCLE_COUNTER(STAT_ ## function_name); \
+			static DelayCall<library_name, function_name ## DelayCall, return_type, Args...> Call{ #library_name, #function_name }; \
+			return Call(args...); \
+		}
+
+	#define CREATE_DEPRECATED_MSG_SHIM(library_name, return_type, function_name, deprecation_msg) \
+		DECLARE_CYCLE_STAT(TEXT(#function_name), STAT_ ## function_name, STATGROUP_MLAPI); \
+		struct library_name; \
+		struct function_name ## DelayCall; \
+		template<typename... Args> MLSDK_API_DEPRECATED_MSG(deprecation_msg) MLResult function_name ## Shim(Args... args) \
+		{ \
+			SCOPE_CYCLE_COUNTER(STAT_ ## function_name); \
+			static DelayCall<library_name, function_name ## DelayCall, return_type, Args...> Call{ #library_name, #function_name }; \
+			return Call(args...); \
+		}
+#else
+	#define CREATE_FUNCTION_SHIM(library_name, return_type, function_name) \
+		DECLARE_CYCLE_STAT(TEXT(#function_name), STAT_ ## function_name, STATGROUP_MLAPI); \
+		template<typename... Args>  return_type function_name##Shim(Args... args) \
+		{ \
+			SCOPE_CYCLE_COUNTER(STAT_ ## function_name); \
+			return ::function_name(args...); \
+		}
+	#define CREATE_GLOBAL_SHIM(library_name, return_type, variable_name) \
+		inline return_type variable_name ## Shim() \
+		{ \
+			return ::variable_name; \
+		}
+	#define CREATE_DEPRECATED_SHIM(library_name, return_type, function_name) \
+		DECLARE_CYCLE_STAT(TEXT(#function_name), STAT_ ## function_name, STATGROUP_MLAPI); \
+		template<typename... Args>  MLSDK_API_DEPRECATED return_type function_name##Shim(Args... args) \
+		{ \
+			SCOPE_CYCLE_COUNTER(STAT_ ## function_name); \
+			return ::function_name(args...); \
+		}
+	#define CREATE_DEPRECATED_MSG_SHIM(library_name, return_type, function_name, deprecation_msg) \
+		DECLARE_CYCLE_STAT(TEXT(#function_name), STAT_ ## function_name, STATGROUP_MLAPI); \
+		template<typename... Args>  MLSDK_API_DEPRECATED_MSG(deprecation_msg) return_type function_name##Shim(Args... args) \
+		{ \
+			SCOPE_CYCLE_COUNTER(STAT_ ## function_name); \
+			return ::function_name(args...); \
+		}
+#endif
+
+}
+
+#endif // !defined(WITH_MLSDK) || WITH_MLSDK
diff --git a/Engine/Source/Runtime/Core/Public/Lumin/CAPIShims/LuminAPIDispatch.h b/Engine/Source/Runtime/Core/Public/Lumin/CAPIShims/LuminAPIDispatch.h
new file mode 100644
--- /dev/null
+++ b/Engine/Source/Runtime/Core/Public/Lumin/CAPIShims/LuminAPIDispatch.h
@@ -0,0 +1,41 @@
+// Copyright 1998-2019 Epic Games, Inc. All Rights Reserved.
+
+#pragma once
+
+#if !defined(WITH_MLSDK) || WITH_MLSDK
+
+#include "LuminAPI.h"
+
+LUMIN_THIRD_PARTY_INCLUDES_START
+#include <ml_dispatch.h>
+LUMIN_THIRD_PARTY_INCLUDES_END
+
+namespace MLSDK_API
+{
+
+CREATE_FUNCTION_SHIM(ml_dispatch, MLResult, MLDispatchAllocateEmptyPacket)
+#define MLDispatchAllocateEmptyPacket ::MLSDK_API::MLDispatchAllocateEmptyPacketShim
+CREATE_FUNCTION_SHIM(ml_dispatch, MLResult, MLDispatchReleasePacket)
+#define MLDispatchReleasePacket ::MLSDK_API::MLDispatchReleasePacketShim
+CREATE_FUNCTION_SHIM(ml_dispatch, MLResult, MLDispatchAllocateFileInfoList)
+#define MLDispatchAllocateFileInfoList ::MLSDK_API::MLDispatchAllocateFileInfoListShim
+CREATE_FUNCTION_SHIM(ml_dispatch, MLResult, MLDispatchGetFileInfoListLength)
+#define MLDispatchGetFileInfoListLength ::MLSDK_API::MLDispatchGetFileInfoListLengthShim
+CREATE_FUNCTION_SHIM(ml_dispatch, MLResult, MLDispatchGetFileInfoByIndex)
+#define MLDispatchGetFileInfoByIndex ::MLSDK_API::MLDispatchGetFileInfoByIndexShim
+CREATE_FUNCTION_SHIM(ml_dispatch, MLResult, MLDispatchAddFileInfo)
+#define MLDispatchAddFileInfo ::MLSDK_API::MLDispatchAddFileInfoShim
+CREATE_FUNCTION_SHIM(ml_dispatch, MLResult, MLDispatchReleaseFileInfoList)
+#define MLDispatchReleaseFileInfoList ::MLSDK_API::MLDispatchReleaseFileInfoListShim
+CREATE_FUNCTION_SHIM(ml_dispatch, MLResult, MLDispatchSetUri)
+#define MLDispatchSetUri ::MLSDK_API::MLDispatchSetUriShim
+CREATE_FUNCTION_SHIM(ml_dispatch, MLResult, MLDispatchReleaseUri)
+#define MLDispatchReleaseUri ::MLSDK_API::MLDispatchReleaseUriShim
+CREATE_FUNCTION_SHIM(ml_dispatch, MLResult, MLDispatchTryOpenApplication)
+#define MLDispatchTryOpenApplication ::MLSDK_API::MLDispatchTryOpenApplicationShim
+CREATE_FUNCTION_SHIM(ml_dispatch, const char*, MLDispatchGetResultString)
+#define MLDispatchGetResultString ::MLSDK_API::MLDispatchGetResultStringShim
+
+}
+
+#endif // !defined(WITH_MLSDK) || WITH_MLSDK
diff --git a/Engine/Source/Runtime/Core/Public/Lumin/CAPIShims/LuminAPIFileInfo.h b/Engine/Source/Runtime/Core/Public/Lumin/CAPIShims/LuminAPIFileInfo.h
new file mode 100644
--- /dev/null
+++ b/Engine/Source/Runtime/Core/Public/Lumin/CAPIShims/LuminAPIFileInfo.h
@@ -0,0 +1,35 @@
+// Copyright 1998-2019 Epic Games, Inc. All Rights Reserved.
+
+#pragma once
+
+#if !defined(WITH_MLSDK) || WITH_MLSDK
+
+#include "LuminAPI.h"
+
+LUMIN_THIRD_PARTY_INCLUDES_START
+#include <ml_fileinfo.h>
+LUMIN_THIRD_PARTY_INCLUDES_END
+
+namespace MLSDK_API
+{
+
+CREATE_FUNCTION_SHIM(ml_lifecycle, MLResult, MLFileInfoAllocateEmpty)
+#define MLFileInfoAllocateEmpty ::MLSDK_API::MLFileInfoAllocateEmptyShim
+CREATE_FUNCTION_SHIM(ml_lifecycle, MLResult, MLFileInfoGetMimeType)
+#define MLFileInfoGetMimeType ::MLSDK_API::MLFileInfoGetMimeTypeShim
+CREATE_FUNCTION_SHIM(ml_lifecycle, MLResult, MLFileInfoGetFileName)
+#define MLFileInfoGetFileName ::MLSDK_API::MLFileInfoGetFileNameShim
+CREATE_FUNCTION_SHIM(ml_lifecycle, MLResult, MLFileInfoGetFD)
+#define MLFileInfoGetFD ::MLSDK_API::MLFileInfoGetFDShim
+CREATE_FUNCTION_SHIM(ml_lifecycle, MLResult, MLFileInfoSetFD)
+#define MLFileInfoSetFD ::MLSDK_API::MLFileInfoSetFDShim
+CREATE_FUNCTION_SHIM(ml_lifecycle, MLResult, MLFileInfoSetFileName)
+#define MLFileInfoSetFileName ::MLSDK_API::MLFileInfoSetFileNameShim
+CREATE_FUNCTION_SHIM(ml_lifecycle, MLResult, MLFileInfoSetMimeType)
+#define MLFileInfoSetMimeType ::MLSDK_API::MLFileInfoSetMimeTypeShim
+CREATE_FUNCTION_SHIM(ml_lifecycle, MLResult, MLFileInfoRelease)
+#define MLFileInfoRelease ::MLSDK_API::MLFileInfoReleaseShim
+
+}
+
+#endif // !defined(WITH_MLSDK) || WITH_MLSDK
