Add missing XML documentation to suppress warnings.

#rnx

diff --git a/Engine/Source/Programs/UnrealBuildTool/ProjectFiles/VisualStudioCode/VSCodeProjectFileGenerator.cs b/Engine/Source/Programs/UnrealBuildTool/ProjectFiles/VisualStudioCode/VSCodeProjectFileGenerator.cs
--- a/Engine/Source/Programs/UnrealBuildTool/ProjectFiles/VisualStudioCode/VSCodeProjectFileGenerator.cs
+++ b/Engine/Source/Programs/UnrealBuildTool/ProjectFiles/VisualStudioCode/VSCodeProjectFileGenerator.cs
@@ -40,6 +40,9 @@ namespace UnrealBuildTool
 		private string FrameworkExecutableExtension;
 		private string FrameworkLibraryExtension = ".dll";
 
+		/// <summary>
+		/// Includes all files in the generated workspace.
+		/// </summary>
 		[XmlConfigFile(Name = "IncludeAllFiles")]
 		private bool IncludeAllFiles = false;
 
