Add LinuxAArch64Server and LinuxAArch64Client build targets

Cherry pick 9427927

diff --git a/Engine/Config/BaseDeviceProfiles.ini b/Engine/Config/BaseDeviceProfiles.ini
--- a/Engine/Config/BaseDeviceProfiles.ini
+++ b/Engine/Config/BaseDeviceProfiles.ini
@@ -61,7 +61,10 @@
 +DeviceProfileNameAndTypes=Linux,Linux
 +DeviceProfileNameAndTypes=LinuxNoEditor,Linux
 +DeviceProfileNameAndTypes=LinuxAArch64NoEditor,Linux
++DeviceProfileNameAndTypes=LinuxClient,Linux
++DeviceProfileNameAndTypes=LinuxAArch64Client,Linux
 +DeviceProfileNameAndTypes=LinuxServer,Linux
++DeviceProfileNameAndTypes=LinuxAArch64Server,Linux
 +DeviceProfileNameAndTypes=Android,Android
 +DeviceProfileNameAndTypes=Android_Low,Android
 +DeviceProfileNameAndTypes=Android_Mid,Android
@@ -744,12 +747,23 @@ DeviceType=Linux
 BaseProfileName=Linux
 +CVars=r.RHICmdBypass=0
 
+[LinuxAArch64Client DeviceProfile]
+DeviceType=Linux
+BaseProfileName=Linux
++CVars=r.RHICmdBypass=0
+
 [LinuxServer DeviceProfile]
 DeviceType=Linux
 BaseProfileName=Linux
 MeshLODSettings=
 TextureLODSettings=
 
+[LinuxAArch64Server DeviceProfile]
+DeviceType=Linux
+BaseProfileName=Linux
+MeshLODSettings=
+TextureLODSettings=
+
 [HoloLens DeviceProfile]
 DeviceType=HoloLens
 BaseProfileName=WindowsNoEditor
diff --git a/Engine/Source/Developer/LauncherServices/Private/Launcher/LauncherWorker.cpp b/Engine/Source/Developer/LauncherServices/Private/Launcher/LauncherWorker.cpp
--- a/Engine/Source/Developer/LauncherServices/Private/Launcher/LauncherWorker.cpp
+++ b/Engine/Source/Developer/LauncherServices/Private/Launcher/LauncherWorker.cpp
@@ -316,6 +316,10 @@ FString FLauncherWorker::CreateUATCommand( const ILauncherProfileRef& InProfile,
 			{
 				ServerPlatforms += TEXT("+Linux");
 			}
+			else if (PlatformInfo->TargetPlatformName == FName("LinuxAArch64Server"))
+			{
+				ServerPlatforms += TEXT("+LinuxAArch64");
+			}
 			else if (PlatformInfo->TargetPlatformName == FName("WindowsServer"))
 			{
 				ServerPlatforms += TEXT("+Win64");
@@ -328,7 +332,7 @@ FString FLauncherWorker::CreateUATCommand( const ILauncherProfileRef& InProfile,
 			{
 				Platforms += TEXT("+Linux");
 			}
-			else if (PlatformInfo->TargetPlatformName == FName("LinuxAArch64NoEditor"))
+			else if (PlatformInfo->TargetPlatformName == FName("LinuxAArch64NoEditor") || PlatformInfo->TargetPlatformName == FName("LinuxAArch64Client"))
 			{
 				Platforms += TEXT("+LinuxAArch64");
 			}
diff --git a/Engine/Source/Developer/LauncherServices/Private/Profiles/LauncherProfile.h b/Engine/Source/Developer/LauncherServices/Private/Profiles/LauncherProfile.h
--- a/Engine/Source/Developer/LauncherServices/Private/Profiles/LauncherProfile.h
+++ b/Engine/Source/Developer/LauncherServices/Private/Profiles/LauncherProfile.h
@@ -1573,14 +1573,18 @@ public:
 			{
 				ServerPlatforms.Add(TEXT("Mac"));
 			}
-			else if (PlatformInfo->TargetPlatformName == FName("LinuxNoEditor"))
+			else if (PlatformInfo->TargetPlatformName == FName("LinuxNoEditor") || PlatformInfo->TargetPlatformName == FName("LinuxClient"))
 			{
 				ClientPlatforms.Add(TEXT("Linux"));
 			}
-			else if (PlatformInfo->TargetPlatformName == FName("LinuxAArch64NoEditor"))
+			else if (PlatformInfo->TargetPlatformName == FName("LinuxAArch64NoEditor") || PlatformInfo->TargetPlatformName == FName("LinuxAArch64Client"))
 			{
 				ClientPlatforms.Add(TEXT("LinuxAArch64"));
 			}
+			else if (PlatformInfo->TargetPlatformName == FName("LinuxAArch64Server"))
+			{
+				ServerPlatforms.Add(TEXT("LinuxAArch64"));
+			}
 			else if (PlatformInfo->TargetPlatformName == FName("WindowsNoEditor") || PlatformInfo->TargetPlatformName == FName("Windows"))
 			{
 				ClientPlatforms.Add(TEXT("Win64"));
diff --git a/Engine/Source/Developer/Linux/LinuxAArch64ClientTargetPlatform/LinuxAArch64ClientTargetPlatform.Build.cs b/Engine/Source/Developer/Linux/LinuxAArch64ClientTargetPlatform/LinuxAArch64ClientTargetPlatform.Build.cs
new file mode 100644
--- /dev/null
+++ b/Engine/Source/Developer/Linux/LinuxAArch64ClientTargetPlatform/LinuxAArch64ClientTargetPlatform.Build.cs
@@ -0,0 +1,40 @@
+// Copyright 1998-2019 Epic Games, Inc. All Rights Reserved.
+
+using UnrealBuildTool;
+
+public class LinuxAArch64ClientTargetPlatform : ModuleRules
+{
+    public LinuxAArch64ClientTargetPlatform(ReadOnlyTargetRules Target) : base(Target)
+    {
+        BinariesSubFolder = "Linux";
+		if (Target.Platform == UnrealTargetPlatform.LinuxAArch64)
+		{
+			BinariesSubFolder += "AArch64";
+		}
+
+        PrivateDependencyModuleNames.AddRange(
+            new string[] {
+				"Core",
+				"TargetPlatform",
+				"DesktopPlatform",
+				"Projects"
+			}
+        );
+
+        if (Target.bCompileAgainstEngine)
+        {
+            PrivateDependencyModuleNames.AddRange(new string[] {
+				"Engine"
+				}
+            );
+
+            PrivateIncludePathModuleNames.Add("TextureCompressor");
+        }
+
+        PrivateIncludePaths.AddRange(
+            new string[] {
+				"Developer/Linux/LinuxTargetPlatform/Private"
+			}
+        );
+    }
+}
diff --git a/Engine/Source/Developer/Linux/LinuxAArch64ClientTargetPlatform/Private/LinuxAArch64ClientTargetPlatformModule.cpp b/Engine/Source/Developer/Linux/LinuxAArch64ClientTargetPlatform/Private/LinuxAArch64ClientTargetPlatformModule.cpp
new file mode 100644
--- /dev/null
+++ b/Engine/Source/Developer/Linux/LinuxAArch64ClientTargetPlatform/Private/LinuxAArch64ClientTargetPlatformModule.cpp
@@ -0,0 +1,47 @@
+// Copyright 1998-2019 Epic Games, Inc. All Rights Reserved.
+
+/*=============================================================================
+	LinuxAArch64ClientTargetPlatformModule.cpp: Implements the FLinuxAArch64ClientTargetPlatformModule class.
+=============================================================================*/
+
+#include "CoreMinimal.h"
+#include "Modules/ModuleManager.h"
+
+#include "Interfaces/ITargetPlatformModule.h"
+
+#include "LinuxTargetDevice.h"
+#include "LinuxTargetPlatform.h"
+
+
+/**
+ * Holds the target platform singleton.
+ */
+static ITargetPlatform* Singleton = NULL;
+
+
+/**
+ * Module for the Linux target platform (without editor).
+ */
+class FLinuxAArch64ClientTargetPlatformModule
+	: public ITargetPlatformModule
+{
+public:
+
+	virtual ~FLinuxAArch64ClientTargetPlatformModule( )
+	{
+		Singleton = NULL;
+	}
+
+	virtual ITargetPlatform* GetTargetPlatform( )
+	{
+		if (Singleton == NULL && TLinuxTargetPlatform<FLinuxPlatformProperties<false, false, true, true> >::IsUsable())
+		{
+ 			Singleton = new TLinuxTargetPlatform<FLinuxPlatformProperties<false, false, true, true> >();
+		}
+
+		return Singleton;
+	}
+};
+
+// 
+IMPLEMENT_MODULE(FLinuxAArch64ClientTargetPlatformModule, LinuxAArch64ClientTargetPlatform);
diff --git a/Engine/Source/Developer/Linux/LinuxAArch64ServerTargetPlatform/LinuxAArch64ServerTargetPlatform.Build.cs b/Engine/Source/Developer/Linux/LinuxAArch64ServerTargetPlatform/LinuxAArch64ServerTargetPlatform.Build.cs
new file mode 100644
--- /dev/null
+++ b/Engine/Source/Developer/Linux/LinuxAArch64ServerTargetPlatform/LinuxAArch64ServerTargetPlatform.Build.cs
@@ -0,0 +1,40 @@
+// Copyright 1998-2019 Epic Games, Inc. All Rights Reserved.
+
+using UnrealBuildTool;
+
+public class LinuxAArch64ServerTargetPlatform : ModuleRules
+{
+    public LinuxAArch64ServerTargetPlatform(ReadOnlyTargetRules Target) : base(Target)
+    {
+        BinariesSubFolder = "Linux";
+		if (Target.Platform == UnrealTargetPlatform.LinuxAArch64)
+		{
+			BinariesSubFolder += "AArch64";
+		}
+
+        PrivateDependencyModuleNames.AddRange(
+            new string[] {
+				"Core",
+				"TargetPlatform",
+				"DesktopPlatform",
+				"Projects"
+			}
+        );
+
+        if (Target.bCompileAgainstEngine)
+        {
+            PrivateDependencyModuleNames.AddRange(new string[] {
+				"Engine"
+				}
+            );
+
+            PrivateIncludePathModuleNames.Add("TextureCompressor");
+       }
+
+        PrivateIncludePaths.AddRange(
+            new string[] {
+				"Developer/Linux/LinuxTargetPlatform/Private"
+		    }
+        );
+    }
+}
diff --git a/Engine/Source/Developer/Linux/LinuxAArch64ServerTargetPlatform/Private/LinuxAArch64ServerTargetPlatformModule.cpp b/Engine/Source/Developer/Linux/LinuxAArch64ServerTargetPlatform/Private/LinuxAArch64ServerTargetPlatformModule.cpp
new file mode 100644
--- /dev/null
+++ b/Engine/Source/Developer/Linux/LinuxAArch64ServerTargetPlatform/Private/LinuxAArch64ServerTargetPlatformModule.cpp
@@ -0,0 +1,46 @@
+// Copyright 1998-2019 Epic Games, Inc. All Rights Reserved.
+
+/*=============================================================================
+	LinuxAArch64TargetPlatformModule.cpp: Implements the FLinuxAArch64ServerTargetPlatformModule class.
+=============================================================================*/
+
+#include "CoreMinimal.h"
+#include "Modules/ModuleManager.h"
+
+#include "Interfaces/ITargetPlatformModule.h"
+
+#include "LinuxTargetDevice.h"
+#include "LinuxTargetPlatform.h"
+
+/**
+ * Holds the target platform singleton.
+ */
+static ITargetPlatform* Singleton = NULL;
+
+
+/**
+ * Module for the Linux AArch64 Server target platform.
+ */
+class FLinuxAArch64ServerTargetPlatformModule
+	: public ITargetPlatformModule
+{
+public:
+
+	virtual ~FLinuxAArch64ServerTargetPlatformModule( )
+	{
+		Singleton = NULL;
+	}
+
+	virtual ITargetPlatform* GetTargetPlatform( )
+	{
+		if (Singleton == NULL && TLinuxTargetPlatform<FLinuxPlatformProperties<false, false, false, true> >::IsUsable())
+		{
+			Singleton = new TLinuxTargetPlatform<FLinuxPlatformProperties<false, true, false, true> >();
+		}
+		
+		return Singleton;
+	}
+};
+
+
+IMPLEMENT_MODULE( FLinuxAArch64ServerTargetPlatformModule, LinuxAArch64ServerTargetPlatform);
diff --git a/Engine/Source/Developer/Linux/LinuxServerTargetPlatform/Private/LinuxServerTargetPlatformModule.cpp b/Engine/Source/Developer/Linux/LinuxServerTargetPlatform/Private/LinuxServerTargetPlatformModule.cpp
--- a/Engine/Source/Developer/Linux/LinuxServerTargetPlatform/Private/LinuxServerTargetPlatformModule.cpp
+++ b/Engine/Source/Developer/Linux/LinuxServerTargetPlatform/Private/LinuxServerTargetPlatformModule.cpp
@@ -1,7 +1,7 @@
 // Copyright 1998-2019 Epic Games, Inc. All Rights Reserved.
 
 /*=============================================================================
-	LinuxTargetPlatformModule.cpp: Implements the FAndroidTargetPlatformModule class.
+	LinuxTargetPlatformModule.cpp: Implements the FLinuxServerTargetPlatformModule class.
 =============================================================================*/
 
 #include "CoreMinimal.h"
diff --git a/Engine/Source/Developer/Linux/LinuxTargetPlatform/Private/LinuxTargetPlatformClasses.cpp b/Engine/Source/Developer/Linux/LinuxTargetPlatform/Private/LinuxTargetPlatformClasses.cpp
--- a/Engine/Source/Developer/Linux/LinuxTargetPlatform/Private/LinuxTargetPlatformClasses.cpp
+++ b/Engine/Source/Developer/Linux/LinuxTargetPlatform/Private/LinuxTargetPlatformClasses.cpp
@@ -10,7 +10,7 @@
 #include "LinuxTargetPlatform.h"
 
 
-/* UAndroidTargetSettings structors
+/* ULinuxTargetSettings structors
  *****************************************************************************/
 
 ULinuxTargetSettings::ULinuxTargetSettings( const FObjectInitializer& ObjectInitializer )
diff --git a/Engine/Source/Developer/Linux/LinuxTargetPlatform/Private/LinuxTargetPlatformModule.cpp b/Engine/Source/Developer/Linux/LinuxTargetPlatform/Private/LinuxTargetPlatformModule.cpp
--- a/Engine/Source/Developer/Linux/LinuxTargetPlatform/Private/LinuxTargetPlatformModule.cpp
+++ b/Engine/Source/Developer/Linux/LinuxTargetPlatform/Private/LinuxTargetPlatformModule.cpp
@@ -24,7 +24,7 @@ static ITargetPlatform* Singleton = NULL;
 
 
 /**
- * Module for the Android target platform.
+ * Module for the Linux target platform.
  */
 class FLinuxTargetPlatformModule
 	: public ITargetPlatformModule
diff --git a/Engine/Source/Developer/TargetPlatform/Private/TargetPlatformManagerModule.cpp b/Engine/Source/Developer/TargetPlatform/Private/TargetPlatformManagerModule.cpp
--- a/Engine/Source/Developer/TargetPlatform/Private/TargetPlatformManagerModule.cpp
+++ b/Engine/Source/Developer/TargetPlatform/Private/TargetPlatformManagerModule.cpp
@@ -1055,6 +1055,12 @@ protected:
 				else if (PlatformName == TEXT("LinuxAArch64"))
 				{
 					PlatformInfo::UpdatePlatformSDKStatus(PlatformName, Status);
+					PlatformName = TEXT("LinuxAArch64NoEditor");
+					PlatformInfo::UpdatePlatformSDKStatus(PlatformName, Status);
+					PlatformName = TEXT("LinuxAArch64Client");
+					PlatformInfo::UpdatePlatformSDKStatus(PlatformName, Status);
+					PlatformName = TEXT("LinuxAArch64Server");
+					PlatformInfo::UpdatePlatformSDKStatus(PlatformName, Status);
 				}
 				else if (PlatformName == TEXT("Desktop"))
 				{
diff --git a/Engine/Source/Programs/AutomationTool/Linux/LinuxPlatform.Automation.cs b/Engine/Source/Programs/AutomationTool/Linux/LinuxPlatform.Automation.cs
--- a/Engine/Source/Programs/AutomationTool/Linux/LinuxPlatform.Automation.cs
+++ b/Engine/Source/Programs/AutomationTool/Linux/LinuxPlatform.Automation.cs
@@ -147,28 +147,21 @@ public abstract class BaseLinuxPlatform : Platform
 
 	public override string GetCookPlatform(bool bDedicatedServer, bool bIsClientOnly)
 	{
-		const string NoEditorCookPlatform = "LinuxNoEditor";
-		const string NoEditorAArch64CookPlatform = "LinuxAArch64NoEditor";
-		const string ServerCookPlatform = "LinuxServer";
-		const string ClientCookPlatform = "LinuxClient";
+		const string NoEditorCookPlatform = "NoEditor";
+		const string ServerCookPlatform = "Server";
+		const string ClientCookPlatform = "Client";
+		string PlatformStr = (TargetPlatformType == UnrealTargetPlatform.LinuxAArch64) ? "LinuxAArch64" : "Linux";
 
 		if (bDedicatedServer)
 		{
-			return ServerCookPlatform;
+			return PlatformStr + ServerCookPlatform;
 		}
 		else if (bIsClientOnly)
 		{
-			return ClientCookPlatform;
+			return PlatformStr + ClientCookPlatform;
 		}
 
-		if (TargetPlatformType == UnrealTargetPlatform.LinuxAArch64)
-		{
-			return NoEditorAArch64CookPlatform;
-		}
-		else
-		{
-			return NoEditorCookPlatform;
-		}
+		return PlatformStr + NoEditorCookPlatform;
 	}
 
 	public override string GetEditorCookPlatform()
diff --git a/Engine/Source/Programs/UnrealBuildTool/Platform/Linux/UEBuildLinux.cs b/Engine/Source/Programs/UnrealBuildTool/Platform/Linux/UEBuildLinux.cs
--- a/Engine/Source/Programs/UnrealBuildTool/Platform/Linux/UEBuildLinux.cs
+++ b/Engine/Source/Programs/UnrealBuildTool/Platform/Linux/UEBuildLinux.cs
@@ -357,7 +357,9 @@ namespace UnrealBuildTool
 							Rules.DynamicallyLoadedModuleNames.Add("LinuxNoEditorTargetPlatform");
 							Rules.DynamicallyLoadedModuleNames.Add("LinuxAArch64NoEditorTargetPlatform");
 							Rules.DynamicallyLoadedModuleNames.Add("LinuxClientTargetPlatform");
+							Rules.DynamicallyLoadedModuleNames.Add("LinuxAArch64ClientTargetPlatform");
 							Rules.DynamicallyLoadedModuleNames.Add("LinuxServerTargetPlatform");
+							Rules.DynamicallyLoadedModuleNames.Add("LinuxAArch64ServerTargetPlatform");
 						}
 					}
 				}
@@ -369,7 +371,9 @@ namespace UnrealBuildTool
 					Rules.DynamicallyLoadedModuleNames.Add("LinuxNoEditorTargetPlatform");
 					Rules.DynamicallyLoadedModuleNames.Add("LinuxAArch64NoEditorTargetPlatform");
 					Rules.DynamicallyLoadedModuleNames.Add("LinuxClientTargetPlatform");
+					Rules.DynamicallyLoadedModuleNames.Add("LinuxAArch64ClientTargetPlatform");
 					Rules.DynamicallyLoadedModuleNames.Add("LinuxServerTargetPlatform");
+					Rules.DynamicallyLoadedModuleNames.Add("LinuxAArch64ServerTargetPlatform");
 				}
 			}
 		}
@@ -402,7 +406,9 @@ namespace UnrealBuildTool
 					Rules.DynamicallyLoadedModuleNames.Add("LinuxNoEditorTargetPlatform");
 					Rules.DynamicallyLoadedModuleNames.Add("LinuxAArch64NoEditorTargetPlatform");
 					Rules.DynamicallyLoadedModuleNames.Add("LinuxClientTargetPlatform");
+					Rules.DynamicallyLoadedModuleNames.Add("LinuxAArch64ClientTargetPlatform");
 					Rules.DynamicallyLoadedModuleNames.Add("LinuxServerTargetPlatform");
+					Rules.DynamicallyLoadedModuleNames.Add("LinuxAArch64ServerTargetPlatform");
 					Rules.DynamicallyLoadedModuleNames.Add("AllDesktopTargetPlatform");
 				}
 
@@ -410,11 +416,12 @@ namespace UnrealBuildTool
 				{
 					Rules.DynamicallyLoadedModuleNames.Add("ShaderFormatOpenGL");
 					Rules.DynamicallyLoadedModuleNames.Add("VulkanShaderFormat");
+					Rules.DynamicallyLoadedModuleNames.Add("ShaderFormatVectorVM");
 				}
 			}
 		}
 
-		public virtual void SetUpSpecificEnvironment(CppCompileEnvironment CompileEnvironment, LinkEnvironment LinkEnvironment)
+		public virtual void SetUpSpecificEnvironment(ReadOnlyTargetRules Target, CppCompileEnvironment CompileEnvironment, LinkEnvironment LinkEnvironment)
 		{
 			CompileEnvironment.Definitions.Add("PLATFORM_LINUX=1");
 			CompileEnvironment.Definitions.Add("PLATFORM_UNIX=1");
@@ -426,6 +433,9 @@ namespace UnrealBuildTool
 
 			// LinuxAArch64 uses only Linux header files
 			CompileEnvironment.Definitions.Add("OVERRIDE_PLATFORM_HEADER_NAME=Linux");
+
+			CompileEnvironment.Definitions.Add("PLATFORM_LINUXAARCH64=" +
+				(Target.Platform == UnrealTargetPlatform.LinuxAArch64 ? "1" : "0"));
 		}
 
 		/// <summary>
@@ -493,7 +503,7 @@ namespace UnrealBuildTool
 			LinkEnvironment.AdditionalLibraries.Add("pthread");
 
 			// let this class or a sub class do settings specific to that class
-			SetUpSpecificEnvironment(CompileEnvironment, LinkEnvironment);
+			SetUpSpecificEnvironment(Target, CompileEnvironment, LinkEnvironment);
 		}
 
 		/// <summary>
diff --git a/Engine/Source/Programs/UnrealBuildTool/System/Utils.cs b/Engine/Source/Programs/UnrealBuildTool/System/Utils.cs
--- a/Engine/Source/Programs/UnrealBuildTool/System/Utils.cs
+++ b/Engine/Source/Programs/UnrealBuildTool/System/Utils.cs
@@ -482,7 +482,7 @@ namespace UnrealBuildTool
 				case UnrealPlatformClass.Editor:
 					return new UnrealTargetPlatform[] { UnrealTargetPlatform.Win64, UnrealTargetPlatform.Linux, UnrealTargetPlatform.Mac };
 				case UnrealPlatformClass.Server:
-					return new UnrealTargetPlatform[] { UnrealTargetPlatform.Win32, UnrealTargetPlatform.Win64, UnrealTargetPlatform.Linux, UnrealTargetPlatform.Mac };
+					return new UnrealTargetPlatform[] { UnrealTargetPlatform.Win32, UnrealTargetPlatform.Win64, UnrealTargetPlatform.Linux, UnrealTargetPlatform.LinuxAArch64, UnrealTargetPlatform.Mac };
 			}
 			throw new ArgumentException(String.Format("'{0}' is not a valid value for UnrealPlatformClass", (int)Class));
 		}
diff --git a/Engine/Source/Runtime/Core/Public/Linux/LinuxPlatformProperties.h b/Engine/Source/Runtime/Core/Public/Linux/LinuxPlatformProperties.h
--- a/Engine/Source/Runtime/Core/Public/Linux/LinuxPlatformProperties.h
+++ b/Engine/Source/Runtime/Core/Public/Linux/LinuxPlatformProperties.h
@@ -27,12 +27,7 @@ struct FLinuxPlatformProperties
 
 	static FORCEINLINE const char* IniPlatformName( )
 	{
-		if (IS_AARCH64)
-		{
-			return "LinuxAArch64";
-		}
-
-		return "Linux";
+		return IS_AARCH64 ? "LinuxAArch64" : "Linux";
 	}
 
 	static FORCEINLINE const TCHAR* GetRuntimeSettingsClassName()
@@ -55,11 +50,16 @@ struct FLinuxPlatformProperties
 		return IS_CLIENT_ONLY;
 	}
 
+	static FORCEINLINE bool IsAArch64()
+	{
+		return IS_AARCH64;
+	}
+
 	static FORCEINLINE const char* PlatformName( )
 	{
 		if (IS_DEDICATED_SERVER)
 		{
-			return "LinuxServer";
+			return IS_AARCH64 ? "LinuxAArch64Server" : "LinuxServer";
 		}
 
 		if (HAS_EDITOR_DATA)
@@ -69,15 +69,10 @@ struct FLinuxPlatformProperties
 
 		if (IS_CLIENT_ONLY)
 		{
-			return "LinuxClient";
-		}
-
-		if (IS_AARCH64)
-		{
-			return "LinuxAArch64NoEditor";
+			return IS_AARCH64 ? "LinuxAArch64Client" : "LinuxClient";
 		}
 
-		return "LinuxNoEditor";
+		return IS_AARCH64 ? "LinuxAArch64NoEditor" : "LinuxNoEditor";
 	}
 
 	static FORCEINLINE bool RequiresCookedData( )
diff --git a/Engine/Source/Runtime/Engine/Engine.Build.cs b/Engine/Source/Runtime/Engine/Engine.Build.cs
--- a/Engine/Source/Runtime/Engine/Engine.Build.cs
+++ b/Engine/Source/Runtime/Engine/Engine.Build.cs
@@ -292,7 +292,9 @@ public class Engine : ModuleRules
 						"LinuxNoEditorTargetPlatform",
 						"LinuxAArch64NoEditorTargetPlatform",
 						"LinuxServerTargetPlatform",
+						"LinuxAArch64ServerTargetPlatform",
 						"LinuxClientTargetPlatform",
+						"LinuxAArch64ClientTargetPlatform",
 						"AllDesktopTargetPlatform",
 						"LinuxPlatformEditor",
 					}
diff --git a/Engine/Source/Runtime/RSA/RSA.Build.cs b/Engine/Source/Runtime/RSA/RSA.Build.cs
--- a/Engine/Source/Runtime/RSA/RSA.Build.cs
+++ b/Engine/Source/Runtime/RSA/RSA.Build.cs
@@ -8,8 +8,8 @@ public class RSA : ModuleRules
 	{
 		PrivateDependencyModuleNames.Add("Core");
 
-        if (Target.Platform == UnrealTargetPlatform.Win64 || Target.Platform == UnrealTargetPlatform.Win32 ||
-			Target.Platform == UnrealTargetPlatform.Mac || Target.Platform == UnrealTargetPlatform.Linux)				
+		if (Target.Platform == UnrealTargetPlatform.Win64 || Target.Platform == UnrealTargetPlatform.Win32 ||
+			Target.Platform == UnrealTargetPlatform.Mac || Target.IsInPlatformGroup(UnrealPlatformGroup.Linux))
 		{
 			PrivateDependencyModuleNames.Add("OpenSSL");
 			PrivateDefinitions.Add("RSA_USE_OPENSSL=1");
