Remove falling back to system compiler by default on Linux

We now only check for and use the system compiler if "-ForceUseSystemCompiler" command line is passed
Also don't default to epic compiler if -ForceUseSystemCompiler specified and we can't find the system compiler

More descriptive error messages when linux sdk toolchain not found

  $ LINUX_MULTIARCH_ROOT=/tmp/foobar make BlankProgram
  bash "/epic/UE4-23.git/Engine/Build/BatchFiles/Linux/Build.sh" BlankProgram Linux Development
  Fixing inconsistent case in filenames.
  Setting up Mono
  Building BlankProgram...
  Unable to locate Linux SDK toolchain at /tmp/foobar/x86_64-unknown-linux-gnu.
  ERROR: GetBuildPlatform: No BuildPlatform found for Linux
  make: *** [Makefile:183: BlankProgram] Error 5

  $ make BlankProgram ARGS="-forceusesystemcompiler"
  bash "/epic/UE4-23.git/Engine/Build/BatchFiles/Linux/Build.sh" BlankProgram Linux Development  -forceusesystemcompiler
  Fixing inconsistent case in filenames.
  Setting up Mono
  Building BlankProgram...
  Unable to locate system compiler (-ForceUseSystemCompiler specified).
  ERROR: GetBuildPlatform: No BuildPlatform found for Linux
  make: *** [Makefile:183: BlankProgram] Error 5

  $ LINUX_MULTIARCH_ROOT= make BlankProgram
  bash "/epic/UE4-23.git/Engine/Build/BatchFiles/Linux/Build.sh" BlankProgram Linux Development
  Fixing inconsistent case in filenames.
  Setting up Mono
  Building BlankProgram...
  Unable to locate Linux SDK toolchain. Please run Setup.sh.
  ERROR: GetBuildPlatform: No BuildPlatform found for Linux
  make: *** [Makefile:183: BlankProgram] Error 5

Cherry pick 9465561, 9472982, 9602106

diff --git a/Engine/Source/Programs/UnrealBuildTool/Platform/Linux/LinuxToolChain.cs b/Engine/Source/Programs/UnrealBuildTool/Platform/Linux/LinuxToolChain.cs
--- a/Engine/Source/Programs/UnrealBuildTool/Platform/Linux/LinuxToolChain.cs
+++ b/Engine/Source/Programs/UnrealBuildTool/Platform/Linux/LinuxToolChain.cs
@@ -74,60 +74,16 @@
 			MultiArchRoot = PlatformSDK.GetSDKLocation();
 			BaseLinuxPath = PlatformSDK.GetBaseLinuxPathForArchitecture(InArchitecture);
 
-			bool bCanUseSystemCompiler = PlatformSDK.CanUseSystemCompiler();
+			bool bForceUseSystemCompiler = PlatformSDK.ForceUseSystemCompiler();
 			bool bHasValidCompiler = false;
 
-			// First validate the BaseLinuxPath toolchain.
-
-			if (!bCanUseSystemCompiler)
-			{
-				// don't register if BaseLinuxPath is not specified and cannot use the system compiler
-				if (String.IsNullOrEmpty(BaseLinuxPath))
-				{
-					throw new BuildException("LINUX_MULTIARCH_ROOT environment variable is not set; cannot instantiate Linux toolchain");
-				}
-			}
-
 			// these are supplied by the engine and do not change depending on the circumstances
 			DumpSymsPath = Path.Combine(UnrealBuildTool.EngineDirectory.FullName, "Binaries", "Linux", "dump_syms");
 			BreakpadEncoderPath = Path.Combine(UnrealBuildTool.EngineDirectory.FullName, "Binaries", "Linux", "BreakpadSymbolEncoder");
 
-			if (!String.IsNullOrEmpty(BaseLinuxPath))
+			if (bForceUseSystemCompiler)
 			{
-				if (String.IsNullOrEmpty(MultiArchRoot)) 
-				{
-					MultiArchRoot = BaseLinuxPath;
-					Log.TraceInformation("Using LINUX_ROOT (deprecated, consider LINUX_MULTIARCH_ROOT)");
-				}
-
-				BaseLinuxPath = BaseLinuxPath.Replace("\"", "").Replace('\\', '/');
-
-				// set up the path to our toolchain
-				GCCPath = "";
-				ClangPath = Path.Combine(BaseLinuxPath, @"bin", "clang++" + GetHostPlatformBinarySuffix());
-				ArPath = Path.Combine(Path.Combine(BaseLinuxPath, String.Format("bin/{0}-{1}", Architecture, "ar" + GetHostPlatformBinarySuffix())));
-				LlvmArPath = Path.Combine(Path.Combine(BaseLinuxPath, String.Format("bin/{0}", "llvm-ar" + GetHostPlatformBinarySuffix())));
-				RanlibPath = Path.Combine(Path.Combine(BaseLinuxPath, String.Format("bin/{0}-{1}", Architecture, "ranlib" + GetHostPlatformBinarySuffix())));
-				StripPath = Path.Combine(Path.Combine(BaseLinuxPath, String.Format("bin/{0}-{1}", Architecture, "strip" + GetHostPlatformBinarySuffix())));
-				ObjcopyPath = Path.Combine(Path.Combine(BaseLinuxPath, String.Format("bin/{0}-{1}", Architecture, "objcopy" + GetHostPlatformBinarySuffix())));
-
-				// When cross-compiling on Windows, use old FixDeps. It is slow, but it does not have timing issues
-				bUseFixdeps = (BuildHostPlatform.Current.Platform == UnrealTargetPlatform.Win64 || BuildHostPlatform.Current.Platform == UnrealTargetPlatform.Win32);
-
-				if (BuildHostPlatform.Current.Platform == UnrealTargetPlatform.Linux)
-				{
-					Environment.SetEnvironmentVariable("LC_ALL", "C");
-				}
-
-				bIsCrossCompiling = true;
-
-				bHasValidCompiler = DetermineCompilerVersion();
-			}
-
-			// Now validate the system toolchain.
-
-			if (bCanUseSystemCompiler && !bHasValidCompiler)
-			{
+				// Validate the system toolchain.
 				BaseLinuxPath = "";
 				MultiArchRoot = "";
 
@@ -158,7 +114,39 @@
 			}
 			else
 			{
+				if (String.IsNullOrEmpty(BaseLinuxPath))
+				{
+					throw new BuildException("LINUX_MULTIARCH_ROOT environment variable is not set; cannot instantiate Linux toolchain");
+				}
+				if (String.IsNullOrEmpty(MultiArchRoot)) 
+				{
+					MultiArchRoot = BaseLinuxPath;
+					Log.TraceInformation("Using LINUX_ROOT (deprecated, consider LINUX_MULTIARCH_ROOT)");
+				}
+
+				BaseLinuxPath = BaseLinuxPath.Replace("\"", "").Replace('\\', '/');
 				ToolchainInfo = String.Format("toolchain located at '{0}'", BaseLinuxPath);
+
+				// set up the path to our toolchain
+				GCCPath = "";
+				ClangPath = Path.Combine(BaseLinuxPath, @"bin", "clang++" + GetHostPlatformBinarySuffix());
+				ArPath = Path.Combine(Path.Combine(BaseLinuxPath, String.Format("bin/{0}-{1}", Architecture, "ar" + GetHostPlatformBinarySuffix())));
+				LlvmArPath = Path.Combine(Path.Combine(BaseLinuxPath, String.Format("bin/{0}", "llvm-ar" + GetHostPlatformBinarySuffix())));
+				RanlibPath = Path.Combine(Path.Combine(BaseLinuxPath, String.Format("bin/{0}-{1}", Architecture, "ranlib" + GetHostPlatformBinarySuffix())));
+				StripPath = Path.Combine(Path.Combine(BaseLinuxPath, String.Format("bin/{0}-{1}", Architecture, "strip" + GetHostPlatformBinarySuffix())));
+				ObjcopyPath = Path.Combine(Path.Combine(BaseLinuxPath, String.Format("bin/{0}-{1}", Architecture, "objcopy" + GetHostPlatformBinarySuffix())));
+
+				// When cross-compiling on Windows, use old FixDeps. It is slow, but it does not have timing issues
+				bUseFixdeps = (BuildHostPlatform.Current.Platform == UnrealTargetPlatform.Win64 || BuildHostPlatform.Current.Platform == UnrealTargetPlatform.Win32);
+
+				if (BuildHostPlatform.Current.Platform == UnrealTargetPlatform.Linux)
+				{
+					Environment.SetEnvironmentVariable("LC_ALL", "C");
+				}
+
+				bIsCrossCompiling = true;
+
+				bHasValidCompiler = DetermineCompilerVersion();
 			}
 
 			if (!bHasValidCompiler)
diff --git a/Engine/Source/Programs/UnrealBuildTool/Platform/Linux/UEBuildLinux.cs b/Engine/Source/Programs/UnrealBuildTool/Platform/Linux/UEBuildLinux.cs
--- a/Engine/Source/Programs/UnrealBuildTool/Platform/Linux/UEBuildLinux.cs
+++ b/Engine/Source/Programs/UnrealBuildTool/Platform/Linux/UEBuildLinux.cs
@@ -574,6 +574,16 @@ namespace UnrealBuildTool
 		/// </summary>
 		static private string TargetPlatformName = "Linux_x64";
 
+		/// <summary>
+		/// Erorr string on why we failed to locate Linux SDK
+		/// </summary>
+		static private string SDKLocationErrorString = null;
+
+		/// <summary>
+		/// Force using system compiler and error out if not possible
+		/// </summary>
+		private int bForceUseSystemCompiler = -1;
+
 		/// <summary>
 		/// Whether to compile with the verbose flag
 		/// </summary>
@@ -694,9 +704,25 @@ namespace UnrealBuildTool
 		/// <summary>
 		/// Whether a host can use its system sdk for this platform
 		/// </summary>
-		public virtual bool CanUseSystemCompiler()
+		public virtual bool ForceUseSystemCompiler()
 		{
-			return (BuildHostPlatform.Current.Platform == UnrealTargetPlatform.Linux);
+			// by default tools chains don't parse arguments, but we want to be able to check the -bForceUseSystemCompiler flag.
+			if (bForceUseSystemCompiler == -1)
+			{
+				bForceUseSystemCompiler = 0;
+				string[] CmdLine = Environment.GetCommandLineArgs();
+
+				foreach (string CmdLineArg in CmdLine)
+				{
+					if (CmdLineArg.Equals("-ForceUseSystemCompiler", StringComparison.OrdinalIgnoreCase))
+					{
+						bForceUseSystemCompiler = 1;
+						break;
+					}
+				}
+			}
+
+			return bForceUseSystemCompiler == 1;
 		}
 
 		/// <summary>
@@ -766,25 +792,41 @@ namespace UnrealBuildTool
 
 			// do not cache this value - it may be changed after sourcing OutputEnvVars.txt
 			string BaseLinuxPath = GetBaseLinuxPathForArchitecture(LinuxPlatform.DefaultHostArchitecture);
+			bool bPrintedSDKLocationErrorString = (SDKLocationErrorString != null);
 
-			// BaseLinuxPath is specified
-			if (!String.IsNullOrEmpty(BaseLinuxPath))
+			if (ForceUseSystemCompiler())
 			{
-				// paths to our toolchains
+				if (!String.IsNullOrEmpty(LinuxCommon.WhichClang()) || !String.IsNullOrEmpty(LinuxCommon.WhichGcc()))
+				{
+					return SDKStatus.Valid;
+				}
+
+				SDKLocationErrorString = "Unable to locate system compiler (-ForceUseSystemCompiler specified).";
+			}
+			else if (!String.IsNullOrEmpty(BaseLinuxPath))
+			{
+				// paths to our toolchains if BaseLinuxPath is specified
 				BaseLinuxPath = BaseLinuxPath.Replace("\"", "");
 
 				if (IsValidClangPath(new DirectoryReference(BaseLinuxPath)))
 				{
 					return SDKStatus.Valid;
 				}
+
+				SDKLocationErrorString = String.Format("Unable to locate Linux SDK toolchain at {0}.", BaseLinuxPath);
+			}
+			else
+			{
+				SDKLocationErrorString = "Unable to locate Linux SDK toolchain. Please run Setup.sh.";
 			}
 
-			if (CanUseSystemCompiler())
+			// It appears that the HasRequiredManualSDKInternal function is called on multiple platforms, so
+			// to prevent this warning from showing up on the Mac (for example), we check that we're on a
+			// platform that can build Linux before spewing the warning message.
+			if (!bPrintedSDKLocationErrorString &&
+				(BuildHostPlatform.Current.Platform == UnrealTargetPlatform.Linux || BuildHostPlatform.Current.Platform == UnrealTargetPlatform.Win64))
 			{
-				if (!String.IsNullOrEmpty(LinuxCommon.WhichClang()) || !String.IsNullOrEmpty(LinuxCommon.WhichGcc()))
-				{
-					return SDKStatus.Valid;
-				}
+				System.Console.WriteLine(SDKLocationErrorString);
 			}
 
 			return SDKStatus.Invalid;
