Fix Linux address sanitizer builds

When linking shared libs (DSOs) under Linux, we need the -shared-libsan option. This will cause us to load the libclang_rt.asan-x86_64.so, etc and not static link so we need to add the path to the executable rpath.

	[REDACTED_PROFILE_URL]/sanitizers/issues/1086
	https://www.gitmemory.com/issue/google/sanitizers/1086/509784249

Add linux sanitizer information to build details if set. Should look something like this:

------- Build details --------
Using system toolchain.
Using clang (/usr/bin/clang++) version '8.0.1' (string), 8 (major), 0 (minor), 1 (patch)
Using bundled libc++ standard C++ library.
Using lld linker
Using llvm-ar : /usr/bin/llvm-ar
Building with: AddressSanitizer
Using fast way to relink  circularly dependent libraries (no FixDeps).
------------------------------

Cherry pick 9953880, 2539328

diff --git a/Engine/Source/Programs/UnrealBuildTool/Platform/Linux/LinuxToolChain.cs b/Engine/Source/Programs/UnrealBuildTool/Platform/Linux/LinuxToolChain.cs
--- a/Engine/Source/Programs/UnrealBuildTool/Platform/Linux/LinuxToolChain.cs
+++ b/Engine/Source/Programs/UnrealBuildTool/Platform/Linux/LinuxToolChain.cs
@@ -538,13 +538,15 @@ namespace UnrealBuildTool
 			// ASan
 			if (Options.HasFlag(LinuxToolChainOptions.EnableAddressSanitizer))
 			{
-				Result += " -fsanitize=address";
+				// Force using the ANSI allocator if ASan is enabled
+				Result += " -fsanitize=address -DFORCE_ANSI_ALLOCATOR=1";
 			}
 
 			// TSan
 			if (Options.HasFlag(LinuxToolChainOptions.EnableThreadSanitizer))
 			{
-				Result += " -fsanitize=thread";
+				// Force using the ANSI allocator if TSan is enabled
+				Result += " -fsanitize=thread -DFORCE_ANSI_ALLOCATOR=1";
 			}
 
 			// UBSan
@@ -918,7 +920,7 @@ namespace UnrealBuildTool
 
 			if (Options.HasFlag(LinuxToolChainOptions.EnableAddressSanitizer) || Options.HasFlag(LinuxToolChainOptions.EnableThreadSanitizer) || Options.HasFlag(LinuxToolChainOptions.EnableUndefinedBehaviorSanitizer))
 			{
-				Result += " -g";
+				Result += " -g -shared-libsan";
 				if (Options.HasFlag(LinuxToolChainOptions.EnableAddressSanitizer))
 				{
 					Result += " -fsanitize=address";
@@ -931,6 +933,12 @@ namespace UnrealBuildTool
 				{
 					Result += " -fsanitize=undefined";
 				}
+
+				if (CrossCompiling())
+				{
+					Result += string.Format(" -Wl,-rpath=\"{0}/lib/clang/{1}.{2}.{3}/lib/linux\"",
+							BaseLinuxPath, CompilerVersionMajor, CompilerVersionMinor, CompilerVersionPatch);
+				}
 			}
 
 			// RPATH for third party libs
@@ -1142,6 +1150,17 @@ namespace UnrealBuildTool
 				Log.TraceInformation("Using {0}", !String.IsNullOrEmpty(LlvmArPath) ? String.Format("llvm-ar : {0}", LlvmArPath) : String.Format("ar and ranlib: {0}, {1}", GetArPath(CompileEnvironment.Architecture), GetRanlibPath(CompileEnvironment.Architecture)));
 			}
 
+			if (Options.HasFlag(LinuxToolChainOptions.EnableAddressSanitizer) || Options.HasFlag(LinuxToolChainOptions.EnableThreadSanitizer) || Options.HasFlag(LinuxToolChainOptions.EnableUndefinedBehaviorSanitizer))
+			{
+				string SanitizerInfo = "Building with:";
+
+				SanitizerInfo += Options.HasFlag(LinuxToolChainOptions.EnableAddressSanitizer) ? " AddressSanitizer" : "";
+				SanitizerInfo += Options.HasFlag(LinuxToolChainOptions.EnableThreadSanitizer) ? " ThreadSanitizer" : "";
+				SanitizerInfo += Options.HasFlag(LinuxToolChainOptions.EnableUndefinedBehaviorSanitizer) ? " UndefinedBehaviorSanitizer" : "";
+
+				Log.TraceInformation(SanitizerInfo);
+			}
+
 			// Also print other once-per-build information
 			if (bUseFixdeps)
 			{
diff --git a/Engine/Source/Programs/UnrealBuildTool/Platform/Linux/UEBuildLinux.cs b/Engine/Source/Programs/UnrealBuildTool/Platform/Linux/UEBuildLinux.cs
--- a/Engine/Source/Programs/UnrealBuildTool/Platform/Linux/UEBuildLinux.cs
+++ b/Engine/Source/Programs/UnrealBuildTool/Platform/Linux/UEBuildLinux.cs
@@ -223,13 +223,6 @@ namespace UnrealBuildTool
 				Target.bCompileCEF3 = false;
 			}
 
-			// Force using the ANSI allocator if ASan or TSan is enabled
-			if (Target.LinuxPlatform.bEnableAddressSanitizer ||
-				Target.LinuxPlatform.bEnableThreadSanitizer)
-			{
-				Target.GlobalDefinitions.Add("FORCE_ANSI_ALLOCATOR=1");
-			}
-
 			// check if OS update invalidated our build
 			Target.bCheckSystemHeadersForModification = (BuildHostPlatform.Current.Platform == UnrealTargetPlatform.Linux);
 
