Remove accidental cherry-pick DynamicallyLoadedModuleNames ShaderFormatVectorVM line

Came in with CL 9618959. Should get added back if CL 9457635 is pulled.

diff --git a/Engine/Source/Programs/UnrealBuildTool/Platform/Linux/UEBuildLinux.cs b/Engine/Source/Programs/UnrealBuildTool/Platform/Linux/UEBuildLinux.cs
--- a/Engine/Source/Programs/UnrealBuildTool/Platform/Linux/UEBuildLinux.cs
+++ b/Engine/Source/Programs/UnrealBuildTool/Platform/Linux/UEBuildLinux.cs
@@ -409,7 +409,6 @@ namespace UnrealBuildTool
 				{
 					Rules.DynamicallyLoadedModuleNames.Add("ShaderFormatOpenGL");
 					Rules.DynamicallyLoadedModuleNames.Add("VulkanShaderFormat");
-					Rules.DynamicallyLoadedModuleNames.Add("ShaderFormatVectorVM");
 				}
 			}
 		}
