UBT: Add a command line option to enable RTTI (requires -OverrideBuildEnvironment also)

#rnx

diff --git a/Engine/Source/Programs/UnrealBuildTool/Configuration/TargetRules.cs b/Engine/Source/Programs/UnrealBuildTool/Configuration/TargetRules.cs
--- a/Engine/Source/Programs/UnrealBuildTool/Configuration/TargetRules.cs
+++ b/Engine/Source/Programs/UnrealBuildTool/Configuration/TargetRules.cs
@@ -555,6 +555,7 @@ namespace UnrealBuildTool
 		/// <summary>
 		/// Enable RTTI for all modules.
 		/// </summary>
+		[CommandLine("-rtti")]
 		[RequiresUniqueBuildEnvironment]
 		public bool bForceEnableRTTI = false;
 
