[IOS] Only set the BundleID to the provisioning profile's identifier, if it does not contain a wildcard.

diff --git a/Engine/Source/Programs/UnrealBuildTool/Platform/IOS/IOSToolChain.cs b/Engine/Source/Programs/UnrealBuildTool/Platform/IOS/IOSToolChain.cs
--- a/Engine/Source/Programs/UnrealBuildTool/Platform/IOS/IOSToolChain.cs
+++ b/Engine/Source/Programs/UnrealBuildTool/Platform/IOS/IOSToolChain.cs
@@ -1,4 +1,4 @@
-﻿// Copyright 1998-2019 Epic Games, Inc. All Rights Reserved.
+// Copyright 1998-2019 Epic Games, Inc. All Rights Reserved.
 
 
 using System;
@@ -1895,7 +1895,11 @@ namespace UnrealBuildTool
 						MobileProvisionFile = ProvisioningData.MobileProvisionFile;
 						MobileProvisionUUID = ProvisioningData.MobileProvisionUUID;
 						TeamUUID = ProvisioningData.TeamUUID;
-						BundleID = ProvisioningData.BundleIdentifier;
+						if (!ProvisioningData.BundleIdentifier.Contains("*"))
+						{
+							// If the BundleIndentifer contains a wild card it will not be valid to use in the plist.
+							BundleID = ProvisioningData.BundleIdentifier;
+						}
 					}
 					else
 					{
@@ -1904,7 +1908,12 @@ namespace UnrealBuildTool
 						MobileProvisionContents MobileProvision = MobileProvisionContents.Read(MobileProvisionFile);
 						MobileProvisionUUID = MobileProvision.GetUniqueId();
 						MobileProvision.TryGetTeamUniqueId(out TeamUUID);
-						BundleID = MobileProvision.GetBundleIdentifier();
+						string BundleIdentifier = MobileProvision.GetBundleIdentifier();
+						if (!BundleIdentifier.Contains("*"))
+						{
+							// If the BundleIndentifer contains a wild card it will not be valid to use in the plist.
+							BundleID = BundleIdentifier;
+						}
 					}
 
 					if(MobileProvisionFile == null)
@@ -1935,7 +1944,7 @@ namespace UnrealBuildTool
 						SchemeName = Target.ProjectFile.GetFileNameWithoutExtension();
 					}
 
-					Console.WriteLine("Provisioning: {0}, {1}, {2}", MobileProvisionFile, MobileProvisionFile.GetFileName(), MobileProvisionUUID);
+					Console.WriteLine("Provisioning: {0}, {1}, {2}, {3}", MobileProvisionFile, MobileProvisionFile.GetFileName(), MobileProvisionUUID, BundleID);
 
 					string CmdLine;
 					if (bBuildAsFramework)
