Use standard UBT logic for finding SignTool.exe instances, so we can use AutoSDK versions distributed with AutoSDK.

diff --git a/Engine/Source/Programs/AutomationTool/AutomationUtils/CommandUtils.cs b/Engine/Source/Programs/AutomationTool/AutomationUtils/CommandUtils.cs
--- a/Engine/Source/Programs/AutomationTool/AutomationUtils/CommandUtils.cs
+++ b/Engine/Source/Programs/AutomationTool/AutomationUtils/CommandUtils.cs
@@ -2748,6 +2748,16 @@ namespace AutomationTool
 		/// <returns>Path to signtool.exe</returns>
 		static string GetSignToolPath()
 		{
+			List<KeyValuePair<string, DirectoryReference>> WindowsSdkDirs = WindowsExports.GetWindowsSdkDirs();
+			foreach (KeyValuePair<string, DirectoryReference> Pair in WindowsSdkDirs)
+			{
+				FileReference SignToolFile = FileReference.Combine(Pair.Value, "bin", Pair.Key, "x64", "SignTool.exe");
+				if (FileReference.Exists(SignToolFile))
+				{
+					return SignToolFile.FullName;
+				}
+			}
+
 			string[] PossibleSignToolNames =
 			{
 				"C:/Program Files (x86)/Windows Kits/8.1/bin/x86/SignTool.exe",
diff --git a/Engine/Source/Programs/AutomationTool/Win/Tasks/SrcSrvTask.cs b/Engine/Source/Programs/AutomationTool/Win/Tasks/SrcSrvTask.cs
--- a/Engine/Source/Programs/AutomationTool/Win/Tasks/SrcSrvTask.cs
+++ b/Engine/Source/Programs/AutomationTool/Win/Tasks/SrcSrvTask.cs
@@ -202,8 +202,8 @@ namespace Win.Automation
 		/// <returns>Path to PDBSTR.EXE</returns>
 		public static FileReference GetPdbStrExe()
 		{
-			List<DirectoryReference> WindowsSdkDirs = WindowsExports.GetWindowsSdkDirs();
-			foreach (DirectoryReference WindowsSdkDir in WindowsSdkDirs)
+			List<KeyValuePair<string, DirectoryReference>> WindowsSdkDirs = WindowsExports.GetWindowsSdkDirs();
+			foreach (DirectoryReference WindowsSdkDir in WindowsSdkDirs.Select(x => x.Value))
 			{
 				FileReference CheckPdbStrExe64 = FileReference.Combine(WindowsSdkDir, "Debuggers", "x64", "SrcSrv", "PdbStr.exe");
 				if (FileReference.Exists(CheckPdbStrExe64))
diff --git a/Engine/Source/Programs/UnrealBuildTool/Platform/Windows/WindowsExports.cs b/Engine/Source/Programs/UnrealBuildTool/Platform/Windows/WindowsExports.cs
--- a/Engine/Source/Programs/UnrealBuildTool/Platform/Windows/WindowsExports.cs
+++ b/Engine/Source/Programs/UnrealBuildTool/Platform/Windows/WindowsExports.cs
@@ -67,25 +67,25 @@ namespace UnrealBuildTool
 		/// Gets a list of Windows Sdk installation directories, ordered by preference
 		/// </summary>
 		/// <returns>String with the name</returns>
-		public static List<DirectoryReference> GetWindowsSdkDirs()
+		public static List<KeyValuePair<string, DirectoryReference>> GetWindowsSdkDirs()
 		{
-			List<DirectoryReference> WindowsSdkDirs = new List<DirectoryReference>();
+			List<KeyValuePair<string, DirectoryReference>> WindowsSdkDirs = new List<KeyValuePair<string, DirectoryReference>>();
 
 			// Add the default directory first
 			VersionNumber Version;
 			DirectoryReference DefaultWindowsSdkDir;
 			if (WindowsPlatform.TryGetWindowsSdkDir(null, out Version, out DefaultWindowsSdkDir))
 			{
-				WindowsSdkDirs.Add(DefaultWindowsSdkDir);
+				WindowsSdkDirs.Add(new KeyValuePair<string, DirectoryReference>(Version.ToString(), DefaultWindowsSdkDir));
 			}
 
 			// Add all the other directories sorted in reverse order
 			IReadOnlyDictionary<VersionNumber, DirectoryReference> WindowsSdkDirPairs = WindowsPlatform.FindWindowsSdkDirs();
 			foreach(KeyValuePair<VersionNumber, DirectoryReference> Pair in WindowsSdkDirPairs.OrderByDescending(x => x.Key))
 			{
-				if(!WindowsSdkDirs.Contains(Pair.Value))
+				if(!WindowsSdkDirs.Any(x => x.Value == Pair.Value))
 				{
-					WindowsSdkDirs.Add(Pair.Value);
+					WindowsSdkDirs.Add(new KeyValuePair<string, DirectoryReference>(Pair.Key.ToString(), Pair.Value));
 				}
 			}
 
