PR #6297: Fixing Rider project files generator (Contributed by [REDACTED_PERSON])

#rnx

diff --git a/Engine/Source/Programs/UnrealBuildTool/ProjectFiles/Rider/RiderProjectFile.cs b/Engine/Source/Programs/UnrealBuildTool/ProjectFiles/Rider/RiderProjectFile.cs
--- a/Engine/Source/Programs/UnrealBuildTool/ProjectFiles/Rider/RiderProjectFile.cs
+++ b/Engine/Source/Programs/UnrealBuildTool/ProjectFiles/Rider/RiderProjectFile.cs
@@ -1,5 +1,6 @@
 // Copyright 1998-2019 Epic Games, Inc. All Rights Reserved.
 
+using System;
 using System.Collections.Generic;
 using System.Linq;
 using Tools.DotNETCommon;
@@ -39,12 +40,11 @@ namespace UnrealBuildTool
 		{
 			string ProjectName = ProjectFilePath.GetFileNameWithoutAnyExtensions();
 			DirectoryReference projectRootFolder = DirectoryReference.Combine(RootPath, ".Rider");
+			List<Tuple<FileReference, UEBuildTarget>> fileToTarget = new List<Tuple<FileReference, UEBuildTarget>>();
 			foreach (UnrealTargetPlatform Platform in InPlatforms.Where(it => it != UnrealTargetPlatform.Win32))
 			{
-				DirectoryReference PlatformFolder = DirectoryReference.Combine(projectRootFolder, Platform.ToString());
 				foreach (UnrealTargetConfiguration Configuration in InConfigurations)
 				{
-					DirectoryReference configurationFolder = DirectoryReference.Combine(PlatformFolder, Configuration.ToString());
 					foreach (ProjectTarget ProjectTarget in ProjectTargets)
 					{
 						if (TargetTypes.Any() && !TargetTypes.Contains(ProjectTarget.TargetRules.Type)) continue;
@@ -60,9 +60,11 @@ namespace UnrealBuildTool
 						{
 							continue;
 						}
+						
+						DirectoryReference ConfigurationFolder = DirectoryReference.Combine(projectRootFolder, Platform.ToString(), Configuration.ToString());
 
 						DirectoryReference TargetFolder =
-							DirectoryReference.Combine(configurationFolder, ProjectTarget.TargetRules.Type.ToString());
+							DirectoryReference.Combine(ConfigurationFolder, ProjectTarget.TargetRules.Type.ToString());
 
 						string DefaultArchitecture = UEBuildPlatform
 							.GetBuildPlatform(BuildHostPlatform.Current.Platform)
@@ -72,18 +74,27 @@ namespace UnrealBuildTool
 							DefaultArchitecture, Arguments);
 						UEBuildTarget BuildTarget = UEBuildTarget.Create(TargetDesc, false, false);
 						FileReference OutputFile = FileReference.Combine(TargetFolder, $"{ProjectName}.json");
-						DirectoryReference.CreateDirectory(OutputFile.Directory);
-						using (JsonWriter Writer = new JsonWriter(OutputFile))
-						{
-							ExportTarget(BuildTarget, Writer);	
-						}
+						fileToTarget.Add(Tuple.Create(OutputFile, BuildTarget));
 					}
 				}
 			}
-
+			foreach (Tuple<FileReference,UEBuildTarget> tuple in fileToTarget)
+			{
+				SerializeTarget(tuple.Item1, tuple.Item2);
+			}
+			
 			return true;
 		}
 
+		private static void SerializeTarget(FileReference OutputFile, UEBuildTarget BuildTarget)
+		{
+			DirectoryReference.CreateDirectory(OutputFile.Directory);
+			using (JsonWriter Writer = new JsonWriter(OutputFile))
+			{
+				ExportTarget(BuildTarget, Writer);
+			}
+		}
+
 		/// <summary>
 		/// Write a Target to a JSON writer. Is array is empty, don't write anything
 		/// </summary>
@@ -116,18 +127,21 @@ namespace UnrealBuildTool
 				Writer.WriteArrayEnd();
 			}
 
-			List<UEBuildModule> TargetModules = Target.Binaries.SelectMany(x => x.Modules).ToList();
-			if(TargetModules.Any())
+			HashSet<string> moduleNames = new HashSet<string>();
+			Writer.WriteObjectStart("Modules");
+			foreach (UEBuildBinary Binary in Target.Binaries)
 			{
-				Writer.WriteObjectStart("Modules");
-				foreach (UEBuildModule Module in TargetModules)
+				foreach (UEBuildModule Module in Binary.Modules)
 				{
-					Writer.WriteObjectStart(Module.Name);
-					ExportModule(Module, Target.GetExecutableDir(), Writer);
-					Writer.WriteObjectEnd();
+					if(moduleNames.Add(Module.Name))
+					{
+						Writer.WriteObjectStart(Module.Name);
+						ExportModule(Module, Binary.OutputDir, Target.GetExecutableDir(), Writer);
+						Writer.WriteObjectEnd();
+					}
 				}
-				Writer.WriteObjectEnd();
 			}
+			Writer.WriteObjectEnd();
 			
 			ExportPluginsFromTarget(Target, Writer);
 			
@@ -137,12 +151,12 @@ namespace UnrealBuildTool
 		/// <summary>
 		/// Write a Module to a JSON writer. If array is empty, don't write anything
 		/// </summary>
+		/// <param name="BinaryOutputDir"></param>
 		/// <param name="TargetOutputDir"></param>
 		/// <param name="Writer">Writer for the array data</param>
 		/// <param name="Module"></param>
-		private static void ExportModule(UEBuildModule Module, DirectoryReference TargetOutputDir, JsonWriter Writer)
+		private static void ExportModule(UEBuildModule Module, DirectoryReference BinaryOutputDir, DirectoryReference TargetOutputDir, JsonWriter Writer)
 		{
-			DirectoryReference BinaryOutputDir = Module.Binary?.OutputDir;
 			Writer.WriteValue("Name", Module.Name);
 			Writer.WriteValue("Directory", Module.ModuleDirectory.FullName);
 			Writer.WriteValue("Rules", Module.RulesFile.FullName);
@@ -200,17 +214,35 @@ namespace UnrealBuildTool
 			
 			if(Module.Rules.RuntimeDependencies.Inner.Any())
 			{
+				// We don't use info from RuntimeDependencies for code analyzes (at the moment)
+				// So we're OK with skipping some values if they are not presented
 				Writer.WriteArrayStart("RuntimeDependencies");
 				foreach (ModuleRules.RuntimeDependency RuntimeDependency in Module.Rules.RuntimeDependencies.Inner)
 				{
 					Writer.WriteObjectStart();
+
+					try
+					{
+						Writer.WriteValue("Path",
+							Module.ExpandPathVariables(RuntimeDependency.Path, BinaryOutputDir, TargetOutputDir));
+					}
+					catch(BuildException buildException)
+					{
+						Log.TraceVerbose("Value {0} for module {1} will not be stored. Reason: {2}", "Path", Module.Name, buildException);	
+					}
 					
-					Writer.WriteValue("Path",
-						Module.ExpandPathVariables(RuntimeDependency.Path, BinaryOutputDir, TargetOutputDir));
 					if (RuntimeDependency.SourcePath != null)
 					{
-						Writer.WriteValue("SourcePath",
-							Module.ExpandPathVariables(RuntimeDependency.SourcePath, BinaryOutputDir, TargetOutputDir));
+						try
+						{
+							Writer.WriteValue("SourcePath",
+								Module.ExpandPathVariables(RuntimeDependency.SourcePath, BinaryOutputDir,
+									TargetOutputDir));
+						}
+						catch(BuildException buildException)
+						{
+							Log.TraceVerbose("Value {0} for module {1} will not be stored. Reason: {2}", "SourcePath", Module.Name, buildException);	
+						}
 					}
 
 					Writer.WriteValue("Type", RuntimeDependency.Type.ToString());
