Fix race condition when -nodebuginfo is set. Was trying to strip the debug info twice on the same so

diff --git a/Engine/Source/Programs/UnrealBuildTool/Platform/Linux/LinuxToolChain.cs b/Engine/Source/Programs/UnrealBuildTool/Platform/Linux/LinuxToolChain.cs
--- a/Engine/Source/Programs/UnrealBuildTool/Platform/Linux/LinuxToolChain.cs
+++ b/Engine/Source/Programs/UnrealBuildTool/Platform/Linux/LinuxToolChain.cs
@@ -318,14 +318,6 @@ namespace UnrealBuildTool
 					);
 				}
 			}
-			else
-			{
-				// strip the final elf file if we are not producing debug info
-				Out += string.Format("\"{0}\" \"{1}\"\n",
-					GetStripPath(LinkEnvironment.Architecture),
-					OutputFile.AbsolutePath
-				);
-			}
 
 			return Out;
 		}
