Remove "Unable to locate Linux SDK toolchain" Linux build warnings

Too many one-off broken cases. We'll work on a solution for 4.25.

diff --git a/Engine/Source/Programs/UnrealBuildTool/Platform/Linux/UEBuildLinux.cs b/Engine/Source/Programs/UnrealBuildTool/Platform/Linux/UEBuildLinux.cs
--- a/Engine/Source/Programs/UnrealBuildTool/Platform/Linux/UEBuildLinux.cs
+++ b/Engine/Source/Programs/UnrealBuildTool/Platform/Linux/UEBuildLinux.cs
@@ -566,11 +566,6 @@ namespace UnrealBuildTool
 		/// </summary>
 		static private string TargetPlatformName = "Linux_x64";
 
-		/// <summary>
-		/// Erorr string on why we failed to locate Linux SDK
-		/// </summary>
-		static private string SDKLocationErrorString = null;
-
 		/// <summary>
 		/// Force using system compiler and error out if not possible
 		/// </summary>
@@ -784,7 +779,6 @@ namespace UnrealBuildTool
 
 			// do not cache this value - it may be changed after sourcing OutputEnvVars.txt
 			string BaseLinuxPath = GetBaseLinuxPathForArchitecture(LinuxPlatform.DefaultHostArchitecture);
-			bool bPrintedSDKLocationErrorString = (SDKLocationErrorString != null);
 
 			if (ForceUseSystemCompiler())
 			{
@@ -792,8 +786,6 @@ namespace UnrealBuildTool
 				{
 					return SDKStatus.Valid;
 				}
-
-				SDKLocationErrorString = "Unable to locate system compiler (-ForceUseSystemCompiler specified).";
 			}
 			else if (!String.IsNullOrEmpty(BaseLinuxPath))
 			{
@@ -804,21 +796,6 @@ namespace UnrealBuildTool
 				{
 					return SDKStatus.Valid;
 				}
-
-				SDKLocationErrorString = String.Format("Unable to locate Linux SDK toolchain at {0}.", BaseLinuxPath);
-			}
-			else
-			{
-				SDKLocationErrorString = "Unable to locate Linux SDK toolchain. Please run Setup.sh.";
-			}
-
-			// It appears that the HasRequiredManualSDKInternal function is called on multiple platforms, so
-			// to prevent this warning from showing up on the Mac (for example), we check that we're on a
-			// platform that can build Linux before spewing the warning message.
-			if (!bPrintedSDKLocationErrorString &&
-				(BuildHostPlatform.Current.Platform == UnrealTargetPlatform.Linux || BuildHostPlatform.Current.Platform == UnrealTargetPlatform.Win64))
-			{
-				System.Console.WriteLine(SDKLocationErrorString);
 			}
 
 			return SDKStatus.Invalid;
