Some command line arguments are now passed on to the target creation (visual studio version specifically). This fixes issues were overriding the visual studio version from command line would not be used in targets, resulting in errors when option like setting the compiler version were also used.

Reported in UDN:
https://udn.unrealengine.com/questions/528859/view.html

Also fixed so that HololensTarget is actually getting configured from setting sources, including command line and BuildConfiguration.xml which I added attributes for to match how Windows targets can be configured.

diff --git a/Engine/Plugins/Developer/VisualStudioCodeSourceCodeAccess/Source/VisualStudioCodeSourceCodeAccess/VisualStudioCodeSourceCodeAccess.Build.cs b/Engine/Plugins/Developer/VisualStudioCodeSourceCodeAccess/Source/VisualStudioCodeSourceCodeAccess/VisualStudioCodeSourceCodeAccess.Build.cs
--- a/Engine/Plugins/Developer/VisualStudioCodeSourceCodeAccess/Source/VisualStudioCodeSourceCodeAccess/VisualStudioCodeSourceCodeAccess.Build.cs
+++ b/Engine/Plugins/Developer/VisualStudioCodeSourceCodeAccess/Source/VisualStudioCodeSourceCodeAccess/VisualStudioCodeSourceCodeAccess.Build.cs
@@ -1,4 +1,5 @@
 // Copyright 1998-2019 Epic Games, Inc. All Rights Reserved.
+using System;
 using Microsoft.Win32;
 
 namespace UnrealBuildTool.Rules
@@ -28,12 +29,18 @@ namespace UnrealBuildTool.Rules
 				string DTEKey = null;
 				switch (Target.WindowsPlatform.Compiler)
 				{
+					case WindowsCompiler.VisualStudio2019:
+						DTEKey = "VisualStudio.DTE.16.0";
+						break;
 					case WindowsCompiler.VisualStudio2017:
 						DTEKey = "VisualStudio.DTE.15.0";
 						break;
 					case WindowsCompiler.VisualStudio2015_DEPRECATED:
 						DTEKey = "VisualStudio.DTE.14.0";
 						break;
+					default:
+						throw new Exception("Unknown visual studio version when mapping to DTEKey: " +
+						                    Target.WindowsPlatform.Compiler.ToString());
 				}
 				bHasVisualStudioDTE = RegistryKey.OpenBaseKey(RegistryHive.ClassesRoot, RegistryView.Registry32).OpenSubKey(DTEKey) != null;
 			}
diff --git a/Engine/Source/Programs/UnrealBuildTool/Configuration/TargetRules.cs b/Engine/Source/Programs/UnrealBuildTool/Configuration/TargetRules.cs
--- a/Engine/Source/Programs/UnrealBuildTool/Configuration/TargetRules.cs
+++ b/Engine/Source/Programs/UnrealBuildTool/Configuration/TargetRules.cs
@@ -1419,6 +1419,7 @@ namespace UnrealBuildTool
 		/// <summary>
 		/// HoloLens-specific target settings.
 		/// </summary>
+		[ConfigSubObject]
 		public HoloLensTargetRules HoloLensPlatform = new HoloLensTargetRules();
 
 		/// <summary>
diff --git a/Engine/Source/Programs/UnrealBuildTool/Configuration/UEBuildTarget.cs b/Engine/Source/Programs/UnrealBuildTool/Configuration/UEBuildTarget.cs
--- a/Engine/Source/Programs/UnrealBuildTool/Configuration/UEBuildTarget.cs
+++ b/Engine/Source/Programs/UnrealBuildTool/Configuration/UEBuildTarget.cs
@@ -646,7 +646,7 @@ namespace UnrealBuildTool
 			// If we're using the shared build environment, make sure all the settings are valid
 			if (RulesObject.BuildEnvironment == TargetBuildEnvironment.Shared)
 			{
-				ValidateSharedEnvironment(RulesAssembly, Descriptor.Name, RulesObject);
+				ValidateSharedEnvironment(RulesAssembly, Descriptor.Name, Descriptor.AdditionalArguments, RulesObject);
 			}
 
 			// If we're precompiling, generate a list of all the files that we depend on
@@ -729,7 +729,7 @@ namespace UnrealBuildTool
 		/// <summary>
 		/// Validates that the build environment matches the shared build environment, by comparing the TargetRules instance to the vanilla target rules for the current target type.
 		/// </summary>
-		static void ValidateSharedEnvironment(RulesAssembly RulesAssembly, string ThisTargetName, TargetRules ThisRules)
+		static void ValidateSharedEnvironment(RulesAssembly RulesAssembly, string ThisTargetName, CommandLineArguments Arguments, TargetRules ThisRules)
 		{
 			// Allow disabling these checks
 			if(ThisRules.bOverrideBuildEnvironment)
@@ -758,7 +758,7 @@ namespace UnrealBuildTool
 			}
 
 			// Create the target rules for it
-			TargetRules BaseRules = RulesAssembly.CreateTargetRules(BaseTargetName, ThisRules.Platform, ThisRules.Configuration, ThisRules.Architecture, null, null);
+			TargetRules BaseRules = RulesAssembly.CreateTargetRules(BaseTargetName, ThisRules.Platform, ThisRules.Configuration, ThisRules.Architecture, null, Arguments);
 
 			// Get all the configurable objects
 			object[] BaseObjects = BaseRules.GetConfigurableObjects().ToArray();
diff --git a/Engine/Source/Programs/UnrealBuildTool/Modes/QueryTargetsMode.cs b/Engine/Source/Programs/UnrealBuildTool/Modes/QueryTargetsMode.cs
--- a/Engine/Source/Programs/UnrealBuildTool/Modes/QueryTargetsMode.cs
+++ b/Engine/Source/Programs/UnrealBuildTool/Modes/QueryTargetsMode.cs
@@ -59,7 +59,7 @@ namespace UnrealBuildTool
 			}
 
 			// Write information about these targets
-			WriteTargetInfo(ProjectFile, Assembly, OutputFile);
+			WriteTargetInfo(ProjectFile, Assembly, OutputFile, Arguments);
 			Log.TraceInformation("Written {0}", OutputFile);
 			return 0;
 		}
@@ -87,7 +87,8 @@ namespace UnrealBuildTool
 		/// <param name="ProjectFile">The project file for the targets being built</param>
 		/// <param name="Assembly">The rules assembly for this target</param>
 		/// <param name="OutputFile">Output file to write to</param>
-		public static void WriteTargetInfo(FileReference ProjectFile, RulesAssembly Assembly, FileReference OutputFile)
+		/// <param name="Arguments"></param>
+		public static void WriteTargetInfo(FileReference ProjectFile, RulesAssembly Assembly, FileReference OutputFile, CommandLineArguments Arguments)
 		{
 			// Construct all the targets in this assembly
 			List<string> TargetNames = new List<string>();
@@ -106,7 +107,7 @@ namespace UnrealBuildTool
 					try
 					{
 						string Architecture = UEBuildPlatform.GetBuildPlatform(BuildHostPlatform.Current.Platform).GetDefaultArchitecture(ProjectFile);
-						TargetRules = Assembly.CreateTargetRules(TargetName, BuildHostPlatform.Current.Platform, UnrealTargetConfiguration.Development, Architecture, ProjectFile, new CommandLineArguments(new string[0]));
+						TargetRules = Assembly.CreateTargetRules(TargetName, BuildHostPlatform.Current.Platform, UnrealTargetConfiguration.Development, Architecture, ProjectFile, Arguments);
 					}
 					catch (Exception Ex)
 					{
diff --git a/Engine/Source/Programs/UnrealBuildTool/Platform/HoloLens/UEBuildHoloLens.cs b/Engine/Source/Programs/UnrealBuildTool/Platform/HoloLens/UEBuildHoloLens.cs
--- a/Engine/Source/Programs/UnrealBuildTool/Platform/HoloLens/UEBuildHoloLens.cs
+++ b/Engine/Source/Programs/UnrealBuildTool/Platform/HoloLens/UEBuildHoloLens.cs
@@ -20,6 +20,7 @@ namespace UnrealBuildTool
 		/// Version of the compiler toolchain to use on HoloLens. A value of "default" will be changed to a specific version at UBT startup.
 		/// </summary>
 		[ConfigFile(ConfigHierarchyType.Engine, "/Script/HoloLensPlatformEditor.HoloLensTargetSettings", "CompilerVersion")]
+		[XmlConfigFile(Category = "HoloLensPlatform")]
 		[CommandLine("-2015", Value = "VisualStudio2015")]
 		[CommandLine("-2017", Value = "VisualStudio2017")]
 		[CommandLine("-2019", Value = "VisualStudio2019")]
@@ -28,8 +29,8 @@ namespace UnrealBuildTool
 		/// <summary>
 		/// Architecture of Target.
 		/// </summary>
-		[CommandLine("x64", Value = "x64")]
-		[CommandLine("arm64", Value = "ARM64")]
+		[CommandLine("-x64", Value = "x64")]
+		[CommandLine("-arm64", Value = "ARM64")]
 		public WindowsArchitecture Architecture = WindowsArchitecture.x64;
 
 		/// <summary>
diff --git a/Engine/Source/Programs/UnrealBuildTool/ProjectFiles/ProjectFileGenerator.cs b/Engine/Source/Programs/UnrealBuildTool/ProjectFiles/ProjectFileGenerator.cs
--- a/Engine/Source/Programs/UnrealBuildTool/ProjectFiles/ProjectFileGenerator.cs
+++ b/Engine/Source/Programs/UnrealBuildTool/ProjectFiles/ProjectFileGenerator.cs
@@ -714,7 +714,7 @@ namespace UnrealBuildTool
 			Dictionary<FileReference, ProjectFile> ProgramProjects;
 			{
 				// Setup buildable projects for all targets
-				AddProjectsForAllTargets( PlatformProjectGenerators, AllGameProjects, AllTargetFiles, out EngineProject, out EnterpriseProject, out GameProjects, out ProgramProjects );
+				AddProjectsForAllTargets( PlatformProjectGenerators, AllGameProjects, AllTargetFiles, Arguments, out EngineProject, out EnterpriseProject, out GameProjects, out ProgramProjects );
 
 				// Add projects for mods
 				AddProjectsForMods(GameProjects, out ModProjects);
@@ -1015,7 +1015,7 @@ namespace UnrealBuildTool
 				foreach(FileReference ProjectFile in AllGameProjects)
 				{
 					RulesAssembly RulesAssembly = RulesCompiler.CreateProjectRulesAssembly(ProjectFile, false, false);
-					QueryTargetsMode.WriteTargetInfo(ProjectFile, RulesAssembly, QueryTargetsMode.GetDefaultOutputFile(ProjectFile));
+					QueryTargetsMode.WriteTargetInfo(ProjectFile, RulesAssembly, QueryTargetsMode.GetDefaultOutputFile(ProjectFile), new CommandLineArguments(GetTargetArguments(Arguments)));
 				}
 			}
 
@@ -2138,14 +2138,16 @@ namespace UnrealBuildTool
 		/// <param name="PlatformProjectGenerators">The registered platform project generators</param>
 		/// <param name="AllGames">All game folders</param>
 		/// <param name="AllTargetFiles">All the target files to add</param>
+		/// <param name="Arguments">The commandline arguments used</param>
 		/// <param name="EngineProject">The engine project we created</param>
 		/// <param name="EnterpriseProject">The enterprise project we created</param>
 		/// <param name="GameProjects">Map of game folder name to all of the game projects we created</param>
 		/// <param name="ProgramProjects">Map of program names to all of the program projects we created</param>
-		private void AddProjectsForAllTargets( 
-			PlatformProjectGeneratorCollection PlatformProjectGenerators, 
-			List<FileReference> AllGames, 
+		private void AddProjectsForAllTargets(
+			PlatformProjectGeneratorCollection PlatformProjectGenerators,
+			List<FileReference> AllGames,
 			List<FileReference> AllTargetFiles,
+			String[] Arguments,
 			out ProjectFile EngineProject, 
 			out ProjectFile EnterpriseProject,
 			out List<ProjectFile> GameProjects, 
@@ -2221,7 +2223,7 @@ namespace UnrealBuildTool
 
 					// Create target rules for all of the platforms and configuration combinations that we want to enable support for.
 					// Just use the current platform as we only need to recover the target type and both should be supported for all targets...
-					TargetRules TargetRulesObject = RulesAssembly.CreateTargetRules(TargetName, BuildHostPlatform.Current.Platform, UnrealTargetConfiguration.Development, "", CheckProjectFile, null);
+					TargetRules TargetRulesObject = RulesAssembly.CreateTargetRules(TargetName, BuildHostPlatform.Current.Platform, UnrealTargetConfiguration.Development, "", CheckProjectFile, new CommandLineArguments(GetTargetArguments(Arguments)));
 
 					bool IsProgramTarget = false;
 
@@ -2377,7 +2379,7 @@ namespace UnrealBuildTool
 							ProjectFilePath = ProjectFilePath,
 							UnrealProjectFilePath = CheckProjectFile,
 							SupportedPlatforms = TargetRulesObject.GetSupportedPlatforms().Where(x => UEBuildPlatform.GetBuildPlatform(x, true) != null).ToArray(),
-							CreateRulesDelegate = (Platform, Configuration) => RulesAssembly.CreateTargetRules(TargetName, Platform, Configuration, "", CheckProjectFile, null)
+							CreateRulesDelegate = (Platform, Configuration) => RulesAssembly.CreateTargetRules(TargetName, Platform, Configuration, "", CheckProjectFile, new CommandLineArguments(GetTargetArguments(Arguments)))
                         };
 
 					ProjectFile.ProjectTargets.Add(ProjectTarget);
@@ -2842,6 +2844,12 @@ namespace UnrealBuildTool
 			OtherProjectFiles.Add( InProject );
 		}
 
+		public virtual string[] GetTargetArguments(string[] Arguments)
+		{
+			// by default we do not forward any arguments to the targets
+			return new string[0];
+		}
+
 		/// The default project to be built for the solution.
 		protected ProjectFile DefaultProject;
 
diff --git a/Engine/Source/Programs/UnrealBuildTool/ProjectFiles/VisualStudio/VCProjectFileGenerator.cs b/Engine/Source/Programs/UnrealBuildTool/ProjectFiles/VisualStudio/VCProjectFileGenerator.cs
--- a/Engine/Source/Programs/UnrealBuildTool/ProjectFiles/VisualStudio/VCProjectFileGenerator.cs
+++ b/Engine/Source/Programs/UnrealBuildTool/ProjectFiles/VisualStudio/VCProjectFileGenerator.cs
@@ -123,6 +123,11 @@ namespace UnrealBuildTool
 			}
 		}
 
+		public override string[] GetTargetArguments(string[] Arguments)
+		{
+			return Arguments.Where(s => string.Equals(s, BuildToolOverride, StringComparison.InvariantCultureIgnoreCase)).ToArray();
+		}
+
 		/// File extension for project files we'll be generating (e.g. ".vcxproj")
 		override public string ProjectFileExtension
 		{
