Add option to build list of plugins but not enable them per target

diff --git a/Engine/Source/Programs/UnrealBuildTool/Configuration/TargetRules.cs b/Engine/Source/Programs/UnrealBuildTool/Configuration/TargetRules.cs
--- a/Engine/Source/Programs/UnrealBuildTool/Configuration/TargetRules.cs
+++ b/Engine/Source/Programs/UnrealBuildTool/Configuration/TargetRules.cs
@@ -281,6 +281,12 @@ namespace UnrealBuildTool
 		[CommandLine("-AllModules")]
 		public bool bBuildAllModules = false;
 
+		/// <summary>
+		/// Additional plugins that are built for this target type but not enabled.
+		/// </summary>
+		[CommandLine("-BuildPlugin=", ListSeparator = '+')]
+		public List<string> BuildPlugins = new List<string>();
+
 		/// <summary>
 		/// A list of additional plugins which need to be included in this target. This allows referencing non-optional plugin modules
 		/// which cannot be disabled, and allows building against specific modules in program targets which do not fit the categories
@@ -1854,6 +1860,11 @@ namespace UnrealBuildTool
 			get { return Inner.DisablePlugins; }
 		}
 
+		public IEnumerable<string> BuildPlugins
+		{
+			get { return Inner.BuildPlugins; }
+		}
+
 		public string PakSigningKeysFile
 		{
 			get { return Inner.PakSigningKeysFile; }
diff --git a/Engine/Source/Programs/UnrealBuildTool/Configuration/UEBuildTarget.cs b/Engine/Source/Programs/UnrealBuildTool/Configuration/UEBuildTarget.cs
--- a/Engine/Source/Programs/UnrealBuildTool/Configuration/UEBuildTarget.cs
+++ b/Engine/Source/Programs/UnrealBuildTool/Configuration/UEBuildTarget.cs
@@ -3061,6 +3061,16 @@ namespace UnrealBuildTool
 			// Create the list of enabled plugins
 			EnabledPlugins = new List<UEBuildPlugin>(NameToInstance.Values);
 
+			// Configure plugins explicitly built but not enabled via target settings
+			foreach (string PluginName in Rules.BuildPlugins)
+			{
+				if (ReferencedNames.Add(PluginName))
+				{
+					PluginReferenceDescriptor PluginReference = new PluginReferenceDescriptor(PluginName, null, true);
+					AddPlugin(PluginReference, "target settings", ExcludeFolders, NameToInstance, NameToInfo);
+				}
+			}
+
 			// Set the list of plugins that should be built
 			BuildPlugins = new List<UEBuildPlugin>(NameToInstance.Values);
 		}
