Fix issue with trying to copy over older obb during packaging

#ue4
#android

diff --git a/Engine/Source/Programs/UnrealBuildTool/Platform/Android/UEDeployAndroid.cs b/Engine/Source/Programs/UnrealBuildTool/Platform/Android/UEDeployAndroid.cs
--- a/Engine/Source/Programs/UnrealBuildTool/Platform/Android/UEDeployAndroid.cs
+++ b/Engine/Source/Programs/UnrealBuildTool/Platform/Android/UEDeployAndroid.cs
@@ -3362,6 +3362,7 @@ namespace UnrealBuildTool
 					if (!File.Exists(DestFileName) || File.GetLastWriteTimeUtc(DestFileName) < File.GetLastWriteTimeUtc(SrcFileName))
 					{
 						Log.TraceInformation("Copying {0} to {1}", SrcFileName, DestFileName);
+						SafeDeleteFile(DestFileName);
 						File.Copy(SrcFileName, DestFileName);
 					}
 				}
