Allow compile actions to execute remotely in Mac and iOS toolchains, plus a small Mac toolchain cleanup

diff --git a/Engine/Source/Programs/UnrealBuildTool/Platform/IOS/IOSToolChain.cs b/Engine/Source/Programs/UnrealBuildTool/Platform/IOS/IOSToolChain.cs
--- a/Engine/Source/Programs/UnrealBuildTool/Platform/IOS/IOSToolChain.cs
+++ b/Engine/Source/Programs/UnrealBuildTool/Platform/IOS/IOSToolChain.cs
@@ -775,7 +775,7 @@ namespace UnrealBuildTool
 				CompileAction.StatusDescription = string.Format("{0}", Path.GetFileName(SourceFile.AbsolutePath));
 				CompileAction.bIsGCCCompiler = true;
 				// We're already distributing the command by execution on Mac.
-				CompileAction.bCanExecuteRemotely = false;
+				CompileAction.bCanExecuteRemotely = true;
 				CompileAction.bShouldOutputStatusDescription = true;
 
 				foreach (UEBuildFramework Framework in CompileEnvironment.AdditionalFrameworks)
diff --git a/Engine/Source/Programs/UnrealBuildTool/Platform/Mac/MacToolChain.cs b/Engine/Source/Programs/UnrealBuildTool/Platform/Mac/MacToolChain.cs
--- a/Engine/Source/Programs/UnrealBuildTool/Platform/Mac/MacToolChain.cs
+++ b/Engine/Source/Programs/UnrealBuildTool/Platform/Mac/MacToolChain.cs
@@ -554,7 +554,7 @@ namespace UnrealBuildTool
 				CompileAction.StatusDescription = Path.GetFileName(SourceFile.AbsolutePath);
 				CompileAction.bIsGCCCompiler = true;
 				// We're already distributing the command by execution on Mac.
-				CompileAction.bCanExecuteRemotely = false;
+				CompileAction.bCanExecuteRemotely = true;
 				CompileAction.bShouldOutputStatusDescription = true;
 				Actions.Add(CompileAction);
 			}
@@ -732,7 +732,7 @@ namespace UnrealBuildTool
 					}
 					else if (AdditionalLibrary.Contains(".framework/"))
 					{
-						LinkCommand += string.Format(" \'{0}\'", AdditionalLibrary);
+						LinkCommand += string.Format(" \"{0}\"", AdditionalLibrary);
 					}
 					else  if (string.IsNullOrEmpty(Path.GetDirectoryName(AdditionalLibrary)) && string.IsNullOrEmpty(Path.GetExtension(AdditionalLibrary)))
 					{
@@ -1056,35 +1056,35 @@ namespace UnrealBuildTool
 
 		FileItem FixDylibDependencies(LinkEnvironment LinkEnvironment, FileItem Executable, List<Action> Actions)
 		{
-			Action LinkAction = new Action(ActionType.Link);
-			LinkAction.WorkingDirectory = UnrealBuildTool.EngineSourceDirectory;
-			LinkAction.CommandPath = BuildHostPlatform.Current.Shell;
-			LinkAction.CommandDescription = "";
+			Action FixDylibAction = new Action(ActionType.PostBuildStep);
+			FixDylibAction.WorkingDirectory = UnrealBuildTool.EngineSourceDirectory;
+			FixDylibAction.CommandPath = BuildHostPlatform.Current.Shell;
+			FixDylibAction.CommandDescription = "";
 
 			// Call the FixDylibDependencies.sh script which will link the dylibs and the main executable, this time proper ones, as it's called
 			// once all are already created, so the cross dependency problem no longer prevents linking.
 			// The script is deleted after it's executed so it's empty when we start appending link commands for the next executable.
 			FileItem FixDylibDepsScript = FileItem.GetItemByFileReference(FileReference.Combine(LinkEnvironment.LocalShadowDirectory, "FixDylibDependencies.sh"));
 
-			LinkAction.CommandArguments = "-c 'chmod +x \"" + FixDylibDepsScript.AbsolutePath + "\"; \"" + FixDylibDepsScript.AbsolutePath + "\"; if [[ $? -ne 0 ]]; then exit 1; fi; ";
+			FixDylibAction.CommandArguments = "-c 'chmod +x \"" + FixDylibDepsScript.AbsolutePath + "\"; \"" + FixDylibDepsScript.AbsolutePath + "\"; if [[ $? -ne 0 ]]; then exit 1; fi; ";
 
 			// Make sure this action is executed after all the dylibs and the main executable are created
 
 			foreach (FileItem Dependency in BundleDependencies)
 			{
-				LinkAction.PrerequisiteItems.Add(Dependency);
+				FixDylibAction.PrerequisiteItems.Add(Dependency);
 			}
 
-			LinkAction.StatusDescription = string.Format("Fixing dylib dependencies for {0}", Path.GetFileName(Executable.AbsolutePath));
-			LinkAction.bCanExecuteRemotely = false;
+			FixDylibAction.StatusDescription = string.Format("Fixing dylib dependencies for {0}", Path.GetFileName(Executable.AbsolutePath));
+			FixDylibAction.bCanExecuteRemotely = false;
 
 			FileItem OutputFile = FileItem.GetItemByFileReference(FileReference.Combine(LinkEnvironment.LocalShadowDirectory, Path.GetFileNameWithoutExtension(Executable.AbsolutePath) + ".link"));
 
-			LinkAction.CommandArguments += "echo \"Dummy\" >> \"" + OutputFile.AbsolutePath + "\"";
-			LinkAction.CommandArguments += "'";
+			FixDylibAction.CommandArguments += "echo \"Dummy\" >> \"" + OutputFile.AbsolutePath + "\"";
+			FixDylibAction.CommandArguments += "'";
 
-			LinkAction.ProducedItems.Add(OutputFile);
-			Actions.Add(LinkAction);
+			FixDylibAction.ProducedItems.Add(OutputFile);
+			Actions.Add(FixDylibAction);
 
 			return OutputFile;
 		}
