Output an error when trying to compile UE4 with known codegen issues, and exclude such toolchain versions when sorting compilers to use. Mark the Visual Studio [REDACTED_PHONE] toolchain (14.23.xxxxx) as bad for the time being.

diff --git a/Engine/Source/Programs/UnrealBuildTool/Platform/Windows/UEBuildWindows.cs b/Engine/Source/Programs/UnrealBuildTool/Platform/Windows/UEBuildWindows.cs
--- a/Engine/Source/Programs/UnrealBuildTool/Platform/Windows/UEBuildWindows.cs
+++ b/Engine/Source/Programs/UnrealBuildTool/Platform/Windows/UEBuildWindows.cs
@@ -1146,6 +1146,37 @@ namespace UnrealBuildTool
 			return FindToolChainDirs(Compiler).Count > 0;
 		}
 
+		/// <summary>
+		/// Checks if a given Visual C++ toolchain version is compatible with UE4
+		/// </summary>
+		/// <param name="Version">The version number to check</param>
+		/// <returns>True if the toolchain is compatible with UE4</returns>
+		static bool IsCompatibleVisualCppToolChain(VersionNumber Version)
+		{
+			string Message;
+			return IsCompatibleVisualCppToolChain(Version, out Message);
+		}
+
+		/// <summary>
+		/// Checks if a given Visual C++ toolchain version is compatible with UE4
+		/// </summary>
+		/// <param name="Version">The version number to check</param>
+		/// <param name="Message">Receives a message describing why the toolchain is not compatible</param>
+		/// <returns>True if the toolchain is compatible with UE4</returns>
+		static bool IsCompatibleVisualCppToolChain(VersionNumber Version, out string Message)
+		{
+			if (Version >= new VersionNumber(14, 23) && Version < new VersionNumber(14, 24))
+			{
+				Message = String.Format("The Visual C++ 14.23 toolchain is known to have code-generation issues with UE4. Please install an earlier or later toolchain from the Visual Studio installer. See here for more information: https://developercommunity.visualstudio.com/content/problem/734585/msvc-142328019-compilation-bug.html");
+				return false;
+			}
+			else
+			{
+				Message = null;
+				return true;
+			}
+		}
+
 		/// <summary>
 		/// Determines the directory containing the MSVC toolchain
 		/// </summary>
@@ -1165,7 +1196,7 @@ namespace UnrealBuildTool
 			{
 				if(String.Compare(CompilerVersion, "Latest", StringComparison.InvariantCultureIgnoreCase) == 0 && ToolChainVersionToDir.Count > 0)
 				{
-					ToolChainVersion = ToolChainVersionToDir.OrderBy(x => Has64BitToolChain(x.Value)).ThenBy(x => x.Key).Last().Key;
+					ToolChainVersion = ToolChainVersionToDir.OrderBy(x => IsCompatibleVisualCppToolChain(x.Key)).ThenBy(x => Has64BitToolChain(x.Value)).ThenBy(x => x.Key).Last().Key;
 				}
 				else if(!VersionNumber.TryParse(CompilerVersion, out ToolChainVersion))
 				{
@@ -1197,10 +1228,17 @@ namespace UnrealBuildTool
 				// if we failed to find any of our preferred toolchains we pick the newest (highest version number)
 				if (ToolChainVersion == null && ToolChainVersionToDir.Count > 0)
 				{
-					ToolChainVersion = ToolChainVersionToDir.OrderBy(x => Has64BitToolChain(x.Value)).ThenBy(x => x.Key).Last().Key;
+					ToolChainVersion = ToolChainVersionToDir.OrderBy(x => IsCompatibleVisualCppToolChain(x.Key)).ThenBy(x => Has64BitToolChain(x.Value)).ThenBy(x => x.Key).Last().Key;
 				}
 			}
 
+			// Check it's valid
+			string Message;
+			if (ToolChainVersion != null && !IsCompatibleVisualCppToolChain(ToolChainVersion, out Message))
+			{
+				throw new BuildException(Message);
+			}
+
 			// Get the actual directory for this version
 			if(ToolChainVersion != null && ToolChainVersionToDir.TryGetValue(ToolChainVersion, out OutToolChainDir))
 			{
