PR #6293: UnrealBuildTool fix (Contributed by [REDACTED_PERSON])
Fix issue with iOS builds failing when running some non-English versions of Windows

diff --git a/Engine/Source/Programs/UnrealBuildTool/Platform/IOS/UEBuildIOS.cs b/Engine/Source/Programs/UnrealBuildTool/Platform/IOS/UEBuildIOS.cs
--- a/Engine/Source/Programs/UnrealBuildTool/Platform/IOS/UEBuildIOS.cs
+++ b/Engine/Source/Programs/UnrealBuildTool/Platform/IOS/UEBuildIOS.cs
@@ -160,7 +160,7 @@ namespace UnrealBuildTool
 
 		public float RuntimeVersion
 		{
-			get { return float.Parse(Inner.ProjectSettings.RuntimeVersion); }
+			get { return float.Parse(Inner.ProjectSettings.RuntimeVersion, System.Globalization.CultureInfo.InvariantCulture); }
 		}
 		
 #if !__MonoCS__
diff --git a/Engine/Source/Programs/UnrealBuildTool/System/XmlConfig.cs b/Engine/Source/Programs/UnrealBuildTool/System/XmlConfig.cs
--- a/Engine/Source/Programs/UnrealBuildTool/System/XmlConfig.cs
+++ b/Engine/Source/Programs/UnrealBuildTool/System/XmlConfig.cs
@@ -614,11 +614,11 @@ namespace UnrealBuildTool
 			}
 			else if(FieldType == typeof(float))
 			{
-				return Single.Parse(TrimmedText);
+				return Single.Parse(TrimmedText, System.Globalization.CultureInfo.InvariantCulture);
 			}
 			else if(FieldType == typeof(double))
 			{
-				return Double.Parse(TrimmedText);
+				return Double.Parse(TrimmedText, System.Globalization.CultureInfo.InvariantCulture);
 			}
 			else if(FieldType.IsEnum)
 			{
