PR #6355: Fixed outdated linker comments (Contributed by [REDACTED_PERSON])

#rnx

diff --git a/Engine/Source/Programs/UnrealBuildTool/Platform/Windows/UEBuildWindows.cs b/Engine/Source/Programs/UnrealBuildTool/Platform/Windows/UEBuildWindows.cs
--- a/Engine/Source/Programs/UnrealBuildTool/Platform/Windows/UEBuildWindows.cs
+++ b/Engine/Source/Programs/UnrealBuildTool/Platform/Windows/UEBuildWindows.cs
@@ -597,12 +597,12 @@ namespace UnrealBuildTool
 		private static IReadOnlyDictionary<VersionNumber, DirectoryReference> CachedUniversalCrtDirs;
 
 		/// <summary>
-		/// True if we should use the Clang linker (LLD) when bCompileWithClang is enabled, otherwise we use the MSVC linker
+		/// True if we should use the Clang linker (LLD) when we are compiling with Clang, otherwise we use the MSVC linker
 		/// </summary>
 		public static readonly bool bAllowClangLinker = false;
 
 		/// <summary>
-		/// True if we should use the Intel linker (xilink) when bCompileWithICL is enabled, otherwise we use the MSVC linker
+		/// True if we should use the Intel linker (xilink) when we are compiling with ICL, otherwise we use the MSVC linker
 		/// </summary>
 		public static readonly bool bAllowICLLinker = true;
 
