Fixed Rider projects failing to generate when a program was defined in the game projects folder, this is similar to how vs project generation works.

diff --git a/Engine/Source/Programs/UnrealBuildTool/ProjectFiles/Rider/RiderProjectFileGenerator.cs b/Engine/Source/Programs/UnrealBuildTool/ProjectFiles/Rider/RiderProjectFileGenerator.cs
--- a/Engine/Source/Programs/UnrealBuildTool/ProjectFiles/Rider/RiderProjectFileGenerator.cs
+++ b/Engine/Source/Programs/UnrealBuildTool/ProjectFiles/Rider/RiderProjectFileGenerator.cs
@@ -2,6 +2,7 @@
 
 using System;
 using System.Collections.Generic;
+using System.IO;
 using System.Linq;
 using Tools.DotNETCommon;
 
@@ -140,11 +141,13 @@ namespace UnrealBuildTool
 			PlatformProjectGeneratorCollection PlatformProjectGenerators,
 			List<FileReference> AllGames,
 			out ProjectFile EngineProject,
-			out List<ProjectFile> GameProjects)
+			out List<ProjectFile> GameProjects,
+			out Dictionary<FileReference, ProjectFile> ProgramProjects)
 		{
 			// As we're creating project files, we'll also keep track of whether we created an "engine" project and return that if we have one
 			EngineProject = null;
 			GameProjects = new List<ProjectFile>();
+			ProgramProjects = new Dictionary<FileReference, ProjectFile>();
 
 			// Find all of the target files.  This will filter out any modules or targets that don't
 			// belong to platforms we're generating project files for.
@@ -322,7 +325,11 @@ namespace UnrealBuildTool
 					}
 
 					// Add the project to the right output list
-					if (IsEngineTarget)
+					if (IsProgramTarget)
+					{
+						ProgramProjects[TargetFilePath] = ProjectFile;
+					}
+					else if (IsEngineTarget)
 					{
 						EngineProject = ProjectFile;
 						if (UnrealBuildTool.IsEngineInstalled())
@@ -472,9 +479,11 @@ namespace UnrealBuildTool
 			ProjectFile EngineProject = null;
 			List<ProjectFile> GameProjects = null;
 			List<ProjectFile> ModProjects = null;
+			Dictionary<FileReference, ProjectFile> ProgramProjects = null;
+
 			// Setup buildable projects for all targets
 			AddProjectsForAllTargets(PlatformProjectGenerators, AllGameProjects, out EngineProject,
-				out GameProjects);
+				out GameProjects, out ProgramProjects);
 
 			AddProjectsForMods(GameProjects, out ModProjects);
 			AddAllGameProjects(GameProjects,null,  null);
