PR #6357: Fix exception message (Contributed by [REDACTED_PERSON])

#rnx

diff --git a/Engine/Source/Programs/UnrealBuildTool/Platform/Windows/VCEnvironment.cs b/Engine/Source/Programs/UnrealBuildTool/Platform/Windows/VCEnvironment.cs
--- a/Engine/Source/Programs/UnrealBuildTool/Platform/Windows/VCEnvironment.cs
+++ b/Engine/Source/Programs/UnrealBuildTool/Platform/Windows/VCEnvironment.cs
@@ -519,7 +519,7 @@ namespace UnrealBuildTool
 				}
 				else
 				{
-					throw new BuildException("{0}{1} must be installed in order to build this target.", WindowsPlatform.GetCompilerName(Compiler), String.IsNullOrEmpty(CompilerVersion)? "" : String.Format(" ({0})", CompilerVersion));
+					throw new BuildException("{0} or {1} must be installed in order to build this target.", WindowsPlatform.GetCompilerName(WindowsCompiler.VisualStudio2019), WindowsPlatform.GetCompilerName(WindowsCompiler.VisualStudio2017));
 				}
 			}
 			else
