PR #6317: UnrealBuildTool: Add detection of all .NET Framework versions, and throw BuildException if none are found.. (Contributed by [REDACTED_PERSON])

diff --git a/Engine/Source/Programs/UnrealBuildTool/Platform/Windows/UEBuildWindows.cs b/Engine/Source/Programs/UnrealBuildTool/Platform/Windows/UEBuildWindows.cs
--- a/Engine/Source/Programs/UnrealBuildTool/Platform/Windows/UEBuildWindows.cs
+++ b/Engine/Source/Programs/UnrealBuildTool/Platform/Windows/UEBuildWindows.cs
@@ -1252,6 +1252,13 @@ namespace UnrealBuildTool
 			return false;
 		}
 
+		static readonly KeyValuePair<RegistryKey, string>[] InstallDirRoots = {
+			new KeyValuePair<RegistryKey, string>(Registry.CurrentUser, "SOFTWARE\\"),
+			new KeyValuePair<RegistryKey, string>(Registry.LocalMachine, "SOFTWARE\\"),
+			new KeyValuePair<RegistryKey, string>(Registry.CurrentUser, "SOFTWARE\\Wow6432Node\\"),
+			new KeyValuePair<RegistryKey, string>(Registry.LocalMachine, "SOFTWARE\\Wow6432Node\\")
+		};
+
 		/// <summary>
 		/// Reads an install directory for a 32-bit program from a registry key. This checks for per-user and machine wide settings, and under the Wow64 virtual keys (HKCU\SOFTWARE, HKLM\SOFTWARE, HKCU\SOFTWARE\Wow6432Node, HKLM\SOFTWARE\Wow6432Node).
 		/// </summary>
@@ -1261,23 +1268,45 @@ namespace UnrealBuildTool
 		/// <returns>True if the key was read, false otherwise.</returns>
 		static bool TryReadInstallDirRegistryKey32(string KeySuffix, string ValueName, out DirectoryReference InstallDir)
 		{
-			if (TryReadDirRegistryKey("HKEY_CURRENT_USER\\SOFTWARE\\" + KeySuffix, ValueName, out InstallDir))
-			{
-				return true;
-			}
-			if (TryReadDirRegistryKey("HKEY_LOCAL_MACHINE\\SOFTWARE\\" + KeySuffix, ValueName, out InstallDir))
-			{
-				return true;
-			}
-			if (TryReadDirRegistryKey("HKEY_CURRENT_USER\\SOFTWARE\\Wow6432Node\\" + KeySuffix, ValueName, out InstallDir))
+			foreach (KeyValuePair<RegistryKey, string> InstallRoot in InstallDirRoots)
 			{
-				return true;
+				using (RegistryKey Key = InstallRoot.Key.OpenSubKey(InstallRoot.Value + KeySuffix))
+				{
+					if (Key != null && TryReadDirRegistryKey(Key.Name, ValueName, out InstallDir))
+					{
+						return true;
+					}
+				}
 			}
-			if (TryReadDirRegistryKey("HKEY_LOCAL_MACHINE\\SOFTWARE\\Wow6432Node\\" + KeySuffix, ValueName, out InstallDir))
+			InstallDir = null;
+			return false;
+		}
+
+		/// <summary>
+		/// For each root location relevant to install dirs, look for the given key and add its subkeys to the set of subkeys to return.
+		/// This checks for per-user and machine wide settings, and under the Wow64 virtual keys (HKCU\SOFTWARE, HKLM\SOFTWARE, HKCU\SOFTWARE\Wow6432Node, HKLM\SOFTWARE\Wow6432Node).
+		/// </summary>
+		/// <param name="KeyName">The subkey to look for under each root location</param>
+		/// <returns>A list of unique subkeys found under any of the existing subkeys</returns>
+		static string[] ReadInstallDirSubKeys32(string KeyName)
+		{
+			HashSet<string> AllSubKeys = new HashSet<string>(StringComparer.Ordinal);
+			foreach (KeyValuePair<RegistryKey, string> Root in InstallDirRoots)
 			{
-				return true;
+				using (RegistryKey Key = Root.Key.OpenSubKey(Root.Value + KeyName))
+				{
+					if (Key == null)
+					{
+						continue;
+					}
+
+					foreach (string SubKey in Key.GetSubKeyNames())
+					{
+						AllSubKeys.Add(SubKey);
+					}
+				}
 			}
-			return false;
+			return AllSubKeys.ToArray();
 		}
 
 		/// <summary>
@@ -1665,9 +1694,40 @@ namespace UnrealBuildTool
 				}
 			}
 
-			return TryReadInstallDirRegistryKey32("Microsoft\\Microsoft SDKs\\NETFXSDK\\4.6", "KitsInstallationFolder", out OutInstallDir) ||
-			       TryReadInstallDirRegistryKey32("Microsoft\\Microsoft SDKs\\NETFXSDK\\4.6.1", "KitsInstallationFolder", out OutInstallDir) ||
-				   TryReadInstallDirRegistryKey32("Microsoft\\Microsoft SDKs\\NETFXSDK\\4.6.2", "KitsInstallationFolder", out OutInstallDir);
+			string NetFxSDKKeyName = "Microsoft\\Microsoft SDKs\\NETFXSDK";
+			string[] PreferredVersions = new string[] { "4.6.2", "4.6.1", "4.6" };
+			foreach (string PreferredVersion in PreferredVersions)
+			{
+				if (TryReadInstallDirRegistryKey32(NetFxSDKKeyName + "\\" + PreferredVersion, "KitsInstallationFolder", out OutInstallDir))
+				{
+					return true;
+				}
+			}
+
+			// If we didn't find one of our preferred versions for NetFXSDK, use the max version present on the system
+			Version MaxVersion = null;
+			string MaxVersionString = null;
+			foreach (string ExistingVersionString in ReadInstallDirSubKeys32(NetFxSDKKeyName))
+			{
+				Version ExistingVersion;
+				if (!Version.TryParse(ExistingVersionString, out ExistingVersion))
+				{
+					continue;
+				}
+				if (MaxVersion == null || ExistingVersion.CompareTo(MaxVersion) > 0)
+				{
+					MaxVersion = ExistingVersion;
+					MaxVersionString = ExistingVersionString;
+				}
+			}
+
+			if (MaxVersionString != null)
+			{
+				return TryReadInstallDirRegistryKey32(NetFxSDKKeyName + "\\" + MaxVersionString, "KitsInstallationFolder", out OutInstallDir);
+			}
+
+			OutInstallDir = null;
+			return false;
 		}
 
 		/// <summary>
diff --git a/Engine/Source/Programs/UnrealBuildTool/Platform/Windows/VCEnvironment.cs b/Engine/Source/Programs/UnrealBuildTool/Platform/Windows/VCEnvironment.cs
--- a/Engine/Source/Programs/UnrealBuildTool/Platform/Windows/VCEnvironment.cs
+++ b/Engine/Source/Programs/UnrealBuildTool/Platform/Windows/VCEnvironment.cs
@@ -458,6 +458,10 @@ namespace UnrealBuildTool
 				IncludePaths.Add(DirectoryReference.Combine(NetFxSdkDir, "include", "um"));
 				LibraryPaths.Add(DirectoryReference.Combine(NetFxSdkDir, "lib", "um", ArchFolder));
 			}
+			else
+			{
+				throw new BuildException("Could not find NetFxSDK install dir; this will prevent SwarmInterface from installing.  Install a version of .NET Framework SDK at 4.6.0 or higher.");
+			}
 
 			// Add the Windows SDK paths
 			if (WindowsSdkVersion >= new VersionNumber(10))
