Added EngineVersion as criteria to rebuild the ModuleRules assembly, fixing issues were a new UBT version can attempt to use a old version of the modules rules assembly.

diff --git a/Engine/Source/Programs/UnrealBuildTool/System/DynamicCompilation.cs b/Engine/Source/Programs/UnrealBuildTool/System/DynamicCompilation.cs
--- a/Engine/Source/Programs/UnrealBuildTool/System/DynamicCompilation.cs
+++ b/Engine/Source/Programs/UnrealBuildTool/System/DynamicCompilation.cs
@@ -30,10 +30,10 @@ namespace UnrealBuildTool
 		/// Checks to see if the assembly needs compilation
 		/// </summary>
 		/// <param name="SourceFiles">Set of source files</param>
-		/// <param name="AssemblySourceListFilePath">File to use to cache source file names</param>
+		/// <param name="AssemblyManifestFilePath">File containing information about this assembly, like which source files it was built with and engine version</param>
 		/// <param name="OutputAssemblyPath">Output path for the assembly</param>
 		/// <returns>True if the assembly needs to be built</returns>
-		private static bool RequiresCompilation(HashSet<FileReference> SourceFiles, FileReference AssemblySourceListFilePath, FileReference OutputAssemblyPath)
+		private static bool RequiresCompilation(HashSet<FileReference> SourceFiles, FileReference AssemblyManifestFilePath, FileReference OutputAssemblyPath)
 		{
 			// Check to see if we already have a compiled assembly file on disk
 			FileItem OutputAssemblyInfo = FileItem.GetItemByFileReference(OutputAssemblyPath);
@@ -54,19 +54,32 @@ namespace UnrealBuildTool
 				return true;
 			}
 
+
 			// Make sure we have a manifest of source files used to compile the output assembly.  If it doesn't exist
 			// for some reason (not an expected case) then we'll need to recompile.
-			FileItem AssemblySourceListFile = FileItem.GetItemByFileReference(AssemblySourceListFilePath);
+			FileItem AssemblySourceListFile = FileItem.GetItemByFileReference(AssemblyManifestFilePath);
 			if (!AssemblySourceListFile.Exists)
 			{
-				Log.TraceLog("Compiling {0}: Missing source file list ({1})", OutputAssemblyPath, AssemblySourceListFilePath);
+				Log.TraceLog("Compiling {0}: Missing source file list ({1})", OutputAssemblyPath, AssemblyManifestFilePath);
 				return true;
 			}
 
+			JsonObject Manifest = JsonObject.Read(AssemblyManifestFilePath);
+
+			// check if the engine version is different
+			string EngineVersionManifest = Manifest.GetStringField("EngineVersion");
+			string EngineVersionCurrent = FormatVersionNumber(ReadOnlyBuildVersion.Current);
+			if (EngineVersionManifest != EngineVersionCurrent)
+			{
+				Log.TraceLog("Compiling {0}: Engine Version changed from {1} to {2}", OutputAssemblyPath, EngineVersionManifest, EngineVersionCurrent);
+				return true;
+			}
+
+
 			// Make sure the source files we're compiling are the same as the source files that were compiled
 			// for the assembly that we want to load
 			HashSet<FileItem> CurrentSourceFileItems = new HashSet<FileItem>();
-			foreach(string Line in FileReference.ReadAllLines(AssemblySourceListFile.Location))
+			foreach(string Line in Manifest.GetStringArrayField("SourceFiles"))
 			{
 				CurrentSourceFileItems.Add(FileItem.GetItemByPath(Line));
 			}
@@ -83,7 +96,7 @@ namespace UnrealBuildTool
 			{
 				if(!SourceFileItems.Contains(CurrentSourceFileItem))
 				{
-					Log.TraceLog("Compiling {0}: Removed source file ({1})", OutputAssemblyPath, AssemblySourceListFilePath);
+					Log.TraceLog("Compiling {0}: Removed source file ({1})", OutputAssemblyPath, CurrentSourceFileItem);
 					return true;
 				}
 			}
@@ -91,7 +104,7 @@ namespace UnrealBuildTool
 			{
 				if(!CurrentSourceFileItems.Contains(SourceFileItem))
 				{
-					Log.TraceLog("Compiling {0}: Added source file ({1})", OutputAssemblyPath, AssemblySourceListFilePath);
+					Log.TraceLog("Compiling {0}: Added source file ({1})", OutputAssemblyPath, SourceFileItem);
 					return true;
 				}
 			}
@@ -384,11 +397,12 @@ namespace UnrealBuildTool
 		public static Assembly CompileAndLoadAssembly(FileReference OutputAssemblyPath, HashSet<FileReference> SourceFileNames, List<string> ReferencedAssembies = null, List<string> PreprocessorDefines = null, bool DoNotCompile = false, bool TreatWarningsAsErrors = false)
 		{
 			// Check to see if the resulting assembly is compiled and up to date
-			FileReference AssemblySourcesListFilePath = FileReference.Combine(OutputAssemblyPath.Directory, Path.GetFileNameWithoutExtension(OutputAssemblyPath.FullName) + "SourceFiles.txt");
+			FileReference AssemblyManifestFilePath = FileReference.Combine(OutputAssemblyPath.Directory, Path.GetFileNameWithoutExtension(OutputAssemblyPath.FullName) + "Manifest.json");
+
 			bool bNeedsCompilation = false;
 			if (!DoNotCompile)
 			{
-				bNeedsCompilation = RequiresCompilation(SourceFileNames, AssemblySourcesListFilePath, OutputAssemblyPath);
+				bNeedsCompilation = RequiresCompilation(SourceFileNames, AssemblyManifestFilePath, OutputAssemblyPath);
 			}
 
 			// Load the assembly to ensure it is correct
@@ -428,9 +442,17 @@ namespace UnrealBuildTool
 					CompiledAssembly = CompileAssembly(OutputAssemblyPath, SourceFileNames, ReferencedAssembies, PreprocessorDefines, TreatWarningsAsErrors);
 				}
 
-				// Save out a list of all the source files we compiled.  This is so that we can tell if whole files were added or removed
-				// since the previous time we compiled the assembly.  In that case, we'll always want to recompile it!
-				FileReference.WriteAllLines(AssemblySourcesListFilePath, SourceFileNames.Select(x => x.FullName));
+				using (JsonWriter Writer = new JsonWriter(AssemblyManifestFilePath))
+				{
+					ReadOnlyBuildVersion Version = ReadOnlyBuildVersion.Current;
+
+					Writer.WriteObjectStart();
+					// Save out a list of all the source files we compiled.  This is so that we can tell if whole files were added or removed
+					// since the previous time we compiled the assembly.  In that case, we'll always want to recompile it!
+					Writer.WriteStringArrayField("SourceFiles", SourceFileNames.Select(x => x.FullName));
+					Writer.WriteValue("EngineVersion", FormatVersionNumber(Version));
+					Writer.WriteObjectEnd();
+				}
 			}
 
 #if !NET_CORE
@@ -447,5 +469,10 @@ namespace UnrealBuildTool
 
 			return CompiledAssembly;
 		}
+
+		private static string FormatVersionNumber(ReadOnlyBuildVersion Version)
+		{
+			return string.Format("{0}.{1}.{2}", Version.MajorVersion, Version.MinorVersion, Version.PatchVersion);
+		}
 	}
 }
