Make sure to cache the project settings by bundle id if it is overriden.
This fixes copying the correct provisioning profile and cert when remote building a blueprint project.

diff --git a/Engine/Source/Programs/UnrealBuildTool/Platform/IOS/UEBuildIOS.cs b/Engine/Source/Programs/UnrealBuildTool/Platform/IOS/UEBuildIOS.cs
--- a/Engine/Source/Programs/UnrealBuildTool/Platform/IOS/UEBuildIOS.cs
+++ b/Engine/Source/Programs/UnrealBuildTool/Platform/IOS/UEBuildIOS.cs
@@ -708,6 +708,7 @@ namespace UnrealBuildTool
 	{
 		IOSPlatformSDK SDK;
 		List<IOSProjectSettings> CachedProjectSettings = new List<IOSProjectSettings>();
+		List<IOSProjectSettings> CachedProjectSettingsByBundle = new List<IOSProjectSettings>();
 		Dictionary<string, IOSProvisioningData> ProvisionCache = new Dictionary<string, IOSProvisioningData>();
 
 		// by default, use an empty architecture (which is really just a modifer to the platform for some paths/names)
@@ -854,11 +855,30 @@ namespace UnrealBuildTool
 
 		public IOSProjectSettings ReadProjectSettings(FileReference ProjectFile, string Bundle = "")
 		{
-			IOSProjectSettings ProjectSettings = CachedProjectSettings.FirstOrDefault(x => x.ProjectFile == ProjectFile);
+			IOSProjectSettings ProjectSettings = null;
+
+			// Use separate lists to prevent an overridden Bundle id polluting the standard project file. 
+			bool bCacheByBundle = !string.IsNullOrEmpty(Bundle);
+			if (bCacheByBundle)
+			{
+				ProjectSettings = CachedProjectSettingsByBundle.FirstOrDefault(x => x.ProjectFile == ProjectFile && x.BundleIdentifier == Bundle);
+			}
+			else
+			{
+				ProjectSettings = CachedProjectSettings.FirstOrDefault(x => x.ProjectFile == ProjectFile);
+			}
+
 			if(ProjectSettings == null)
 			{
 				ProjectSettings = CreateProjectSettings(ProjectFile, Bundle);
-				CachedProjectSettings.Add(ProjectSettings);
+				if (bCacheByBundle)
+				{
+					CachedProjectSettingsByBundle.Add(ProjectSettings);
+				}
+				else
+				{
+					CachedProjectSettings.Add(ProjectSettings);
+				}
 			}
 			return ProjectSettings;
 		}
