Xcode 11 is required when using the binary engine, so added a check in AppleToolChain in UBT to detect older versions

diff --git a/Engine/Source/Programs/UnrealBuildTool/ToolChain/AppleToolChain.cs b/Engine/Source/Programs/UnrealBuildTool/ToolChain/AppleToolChain.cs
--- a/Engine/Source/Programs/UnrealBuildTool/ToolChain/AppleToolChain.cs
+++ b/Engine/Source/Programs/UnrealBuildTool/ToolChain/AppleToolChain.cs
@@ -51,6 +51,30 @@ namespace UnrealBuildTool
 				Log.TraceInformationOnce("Compiling with non-standard Xcode ({0}): {1}", Reason, DeveloperDir);
 			}
 
+			// Installed engine requires Xcode 11
+			if (UnrealBuildTool.IsEngineInstalled())
+			{
+				string XcodeBuilderVersionOutput = Utils.RunLocalProcessAndReturnStdOut("xcodebuild", "-version");
+				if (XcodeBuilderVersionOutput.Length > 10)
+				{
+					string[] Version = XcodeBuilderVersionOutput.Substring(6, 4).Split('.');
+					if (Version.Length == 2)
+					{
+						if (int.Parse(Version[0]) < 11)
+						{
+							throw new BuildException("Building for macOS, iOS and tvOS requires Xcode 11 or newer, Xcode " + Version[0] + "." + Version[1] + " detected");
+						}
+					}
+					else
+					{
+						Log.TraceWarning("Failed to query Xcode version");
+					}
+				}
+				else
+				{
+					Log.TraceWarning("Failed to query Xcode version");
+				}
+			}
 		}
 
 		protected void SelectSDK(string BaseSDKDir, string OSPrefix, ref string PlatformSDKVersion, bool bVerbose)
