Fix issue with PLATFORM_USED_NDK_VERSION_INTEGER and show proper NDK level used for compile

#android

diff --git a/Engine/Source/Programs/UnrealBuildTool/Platform/Android/AndroidToolChain.cs b/Engine/Source/Programs/UnrealBuildTool/Platform/Android/AndroidToolChain.cs
--- a/Engine/Source/Programs/UnrealBuildTool/Platform/Android/AndroidToolChain.cs
+++ b/Engine/Source/Programs/UnrealBuildTool/Platform/Android/AndroidToolChain.cs
@@ -1250,9 +1250,29 @@ namespace UnrealBuildTool
 			}
 			*/
 
+			// NDK setup (use no less than 21 for 64-bit targets)
+			int NDKApiLevel32Int = GetNdkApiLevelInt();
+			int NDKApiLevel64Int = NDKApiLevel32Int;
+			string NDKApiLevel32Bit = GetNdkApiLevel();
+			string NDKApiLevel64Bit = NDKApiLevel32Bit;
+			if (NDKApiLevel64Int < 21)
+			{
+				NDKApiLevel64Int = 21;
+				NDKApiLevel64Bit = "android-21";
+			}
+
 			if (!bHasPrintedApiLevel)
 			{
-				Log.TraceInformation("Compiling Native code with NDK API '{0}'", GetNdkApiLevel());
+				bool Has32Bit = Arches.Contains("-armv7") || Arches.Contains("-x86");
+				bool Has64Bit = Arches.Contains("-arm64") || Arches.Contains("-x64");
+				if (Has32Bit)
+				{
+					Log.TraceInformation("Compiling Native 32-bit code with NDK API '{0}'", NDKApiLevel32Bit);
+				}
+				if (Has64Bit)
+				{
+					Log.TraceInformation("Compiling Native 64-bit code with NDK API '{0}'", NDKApiLevel64Bit);
+				}
 				bHasPrintedApiLevel = true;
 			}
 
@@ -1309,11 +1329,11 @@ namespace UnrealBuildTool
 
 					switch (Arch)
 					{
-						case "-armv7": Arguments += " -DPLATFORM_64BITS=0 -DPLATFORM_ANDROID_ARM=1"; break;
-						case "-arm64": Arguments += " -DPLATFORM_64BITS=1 -DPLATFORM_ANDROID_ARM64=1"; break;
-						case "-x86": Arguments += " -DPLATFORM_64BITS=0 -DPLATFORM_ANDROID_X86=1"; break;
-						case "-x64": Arguments += " -DPLATFORM_64BITS=1 -DPLATFORM_ANDROID_X64=1"; break;
-						default: Arguments += " -DPLATFORM_64BITS=0 -DPLATFORM_ANDROID_ARM=1"; break;
+						case "-armv7": Arguments += " -DPLATFORM_64BITS=0 -DPLATFORM_ANDROID_ARM=1 -DPLATFORM_USED_NDK_VERSION_INTEGER=" + NDKApiLevel32Int.ToString(); break;
+						case "-arm64": Arguments += " -DPLATFORM_64BITS=1 -DPLATFORM_ANDROID_ARM64=1 -DPLATFORM_USED_NDK_VERSION_INTEGER=" + NDKApiLevel64Int.ToString(); break;
+						case "-x86": Arguments += " -DPLATFORM_64BITS=0 -DPLATFORM_ANDROID_X86=1 -DPLATFORM_USED_NDK_VERSION_INTEGER=" + NDKApiLevel32Int.ToString(); break;
+						case "-x64": Arguments += " -DPLATFORM_64BITS=1 -DPLATFORM_ANDROID_X64=1 -DPLATFORM_USED_NDK_VERSION_INTEGER=" + NDKApiLevel64Int.ToString(); break;
+						default: Arguments += " -DPLATFORM_64BITS=0 -DPLATFORM_ANDROID_ARM=1 -DPLATFORM_USED_NDK_VERSION_INTEGER=" + NDKApiLevel32Int.ToString(); break;
 					}
 
 					if (Arch == "-arm64" || (Arch == "-armv7" && bUseNEONForArmV7))
diff --git a/Engine/Source/Programs/UnrealBuildTool/Platform/Android/UEBuildAndroid.cs b/Engine/Source/Programs/UnrealBuildTool/Platform/Android/UEBuildAndroid.cs
--- a/Engine/Source/Programs/UnrealBuildTool/Platform/Android/UEBuildAndroid.cs
+++ b/Engine/Source/Programs/UnrealBuildTool/Platform/Android/UEBuildAndroid.cs
@@ -342,10 +342,6 @@ namespace UnrealBuildTool
 
 			Log.TraceInformation("NDK toolchain: {0}, NDK version: {1}, GccVersion: {2}, ClangVersion: {3}", NDKToolchainVersion, NDKVersionInt.ToString(), GccVersion, ToolChain.GetClangVersionString());
 
-
-			CompileEnvironment.Definitions.Add("PLATFORM_USED_NDK_VERSION_INTEGER=" + NDKVersionInt);
-
-
 			CompileEnvironment.Definitions.Add("PLATFORM_DESKTOP=0");
 			CompileEnvironment.Definitions.Add("PLATFORM_CAN_SUPPORT_EDITORONLY_DATA=0");
 
