seems to be the last remains of HTML5 code from engine

#rnx

diff --git a/Engine/Source/Programs/UnrealBuildTool/UnrealBuildTool.csproj b/Engine/Source/Programs/UnrealBuildTool/UnrealBuildTool.csproj
--- a/Engine/Source/Programs/UnrealBuildTool/UnrealBuildTool.csproj
+++ b/Engine/Source/Programs/UnrealBuildTool/UnrealBuildTool.csproj
@@ -335,16 +335,6 @@
       </ItemGroup>
     </When>
   </Choose>
-  <Choose>
-    <When Condition="Exists('Platform\HTML5\UEBuildHTML5.cs')">
-      <ItemGroup>
-        <Compile Include="Platform\HTML5\HTML5ProjectGenerator.cs" />
-        <Compile Include="Platform\HTML5\HTML5ToolChain.cs" />
-        <Compile Include="Platform\HTML5\UEBuildHTML5.cs" />
-        <Compile Include="Platform\HTML5\HTML5SDKInfo.cs" />
-      </ItemGroup>
-    </When>
-  </Choose>
   <Choose>
     <When Condition="Exists('Platform\Linux\UEBuildLinux.cs')">
       <ItemGroup>
@@ -470,4 +460,4 @@
     </ProjectReference>
   </ItemGroup>
   <ItemGroup />
-</Project>
\ No newline at end of file
+</Project>
