Determine the exact compiler version number from the executable, and whitelist the 14.23.28107 compiler for building UE4.

#rnx

diff --git a/Engine/Source/Programs/UnrealBuildTool/Platform/Windows/UEBuildWindows.cs b/Engine/Source/Programs/UnrealBuildTool/Platform/Windows/UEBuildWindows.cs
--- a/Engine/Source/Programs/UnrealBuildTool/Platform/Windows/UEBuildWindows.cs
+++ b/Engine/Source/Programs/UnrealBuildTool/Platform/Windows/UEBuildWindows.cs
@@ -1041,7 +1041,7 @@ namespace UnrealBuildTool
 							    foreach(DirectoryReference ToolChainDir in DirectoryReference.EnumerateDirectories(ToolChainBaseDir))
 							    {
 								    VersionNumber Version;
-								    if(VersionNumber.TryParse(ToolChainDir.GetDirectoryName(), out Version) && IsValidToolChainDir2017or2019(ToolChainDir))
+								    if(IsValidToolChainDir2017or2019(ToolChainDir, out Version))
 									{
 										Log.TraceLog("Found Visual Studio toolchain: {0} (Version={1})", ToolChainDir, Version);
 										if(!ToolChainVersionToDir.ContainsKey(Version))
@@ -1063,7 +1063,7 @@ namespace UnrealBuildTool
 								foreach(DirectoryReference ToolChainDir in DirectoryReference.EnumerateDirectories(ToolChainBaseDir))
 								{
 									VersionNumber Version;
-									if(VersionNumber.TryParse(ToolChainDir.GetDirectoryName(), out Version) && IsValidToolChainDir2017or2019(ToolChainDir))
+									if(IsValidToolChainDir2017or2019(ToolChainDir, out Version))
 									{
 										Log.TraceLog("Found Visual Studio toolchain: {0} (Version={1})", ToolChainDir, Version);
 										if (!ToolChainVersionToDir.ContainsKey(Version))
@@ -1106,15 +1106,33 @@ namespace UnrealBuildTool
 		}
 
 		/// <summary>
-		/// Checks if the given directory contains a valid Visual Studio 2017 toolchain
+		/// Determines if the given path is a valid Visual C++ version number
 		/// </summary>
-		/// <param name="ToolChainDir">Directory to check</param>
-		/// <returns>True if the given directory is valid</returns>
-		static bool IsValidToolChainDir2017or2019(DirectoryReference ToolChainDir)
+		/// <param name="ToolChainDir">The toolchain directory</param>
+		/// <param name="Version">The version number for the toolchain</param>
+		/// <returns>True if the path is a valid version</returns>
+		static bool IsValidToolChainDir2017or2019(DirectoryReference ToolChainDir, out VersionNumber Version)
 		{
-			return FileReference.Exists(FileReference.Combine(ToolChainDir, "bin", "Hostx86", "x64", "cl.exe")) || FileReference.Exists(FileReference.Combine(ToolChainDir, "bin", "Hostx64", "x64", "cl.exe"));
-		}
+			FileReference CompilerExe = FileReference.Combine(ToolChainDir, "bin", "Hostx86", "x64", "cl.exe");
+			if(!FileReference.Exists(CompilerExe))
+			{
+				CompilerExe = FileReference.Combine(ToolChainDir, "bin", "Hostx64", "x64", "cl.exe");
+				if(!FileReference.Exists(CompilerExe))
+				{
+					Version = null;
+					return false;
+				}
+			}
+
+			FileVersionInfo VersionInfo = FileVersionInfo.GetVersionInfo(CompilerExe.FullName);
+			if (VersionInfo.ProductMajorPart != 0)
+			{
+				Version = new VersionNumber(VersionInfo.ProductMajorPart, VersionInfo.ProductMinorPart, VersionInfo.ProductBuildPart);
+				return true;
+			}
 
+			return VersionNumber.TryParse(ToolChainDir.GetDirectoryName(), out Version);
+		}
 
 		/// <summary>
 		/// Checks if the given directory contains a 64-bit toolchain. Used to prefer regular Visual Studio versions over express editions.
@@ -1165,7 +1183,7 @@ namespace UnrealBuildTool
 		/// <returns>True if the toolchain is compatible with UE4</returns>
 		static bool IsCompatibleVisualCppToolChain(VersionNumber Version, out string Message)
 		{
-			if (Version >= new VersionNumber(14, 23) && Version < new VersionNumber(14, 24))
+			if (Version >= new VersionNumber(14, 23, 0) && Version < new VersionNumber(14, 23, 28107))
 			{
 				Message = String.Format("The Visual C++ 14.23 toolchain is known to have code-generation issues with UE4. Please install an earlier or later toolchain from the Visual Studio installer. See here for more information: https://developercommunity.visualstudio.com/content/problem/734585/msvc-142328019-compilation-bug.html");
 				return false;
