Updated location service privilege key.

diff --git a/Engine/Source/Programs/UnrealBuildTool/Platform/IOS/UEDeployIOS.cs b/Engine/Source/Programs/UnrealBuildTool/Platform/IOS/UEDeployIOS.cs
--- a/Engine/Source/Programs/UnrealBuildTool/Platform/IOS/UEDeployIOS.cs
+++ b/Engine/Source/Programs/UnrealBuildTool/Platform/IOS/UEDeployIOS.cs
@@ -1,4 +1,4 @@
-﻿// Copyright 1998-2019 Epic Games, Inc. All Rights Reserved.
+// Copyright 1998-2019 Epic Games, Inc. All Rights Reserved.
 
 using System;
 using System.Collections;
@@ -651,7 +651,7 @@ namespace UnrealBuildTool
 			Text.AppendLine("\t<key>ITSAppUsesNonExemptEncryption</key>");
 			Text.AppendLine("\t<false/>");
 			// add location services descriptions if used
-			Text.AppendLine("\t<key>NSLocationAlwaysUsageDescription</key>");
+			Text.AppendLine("\t<key>NSLocationAlwaysAndWhenInUseUsageDescription</key>");
 			Text.AppendLine(string.Format("\t<string>{0}</string>", LocationAlwaysUsageDescription));
 			Text.AppendLine("\t<key>NSLocationWhenInUseUsageDescription</key>");
 			Text.AppendLine(string.Format("\t<string>{0}></string>", LocationWhenInUseDescription));
