Fix codesigning frameworks when there is a space in the path.

diff --git a/Engine/Source/Programs/UnrealBuildTool/ProjectFiles/Xcode/XcodeProject.cs b/Engine/Source/Programs/UnrealBuildTool/ProjectFiles/Xcode/XcodeProject.cs
--- a/Engine/Source/Programs/UnrealBuildTool/ProjectFiles/Xcode/XcodeProject.cs
+++ b/Engine/Source/Programs/UnrealBuildTool/ProjectFiles/Xcode/XcodeProject.cs
@@ -777,7 +777,7 @@ namespace UnrealBuildTool
 				}
 			}
 
-			string ShellScript = "set -e\\n\\n" +
+			string ShellScript = "set -e\\n\\nIFS=$'\\\\n'\\n\\n" +
 				"if [ $PLATFORM_NAME = iphoneos ] || [ $PLATFORM_NAME = tvos ]; then \\n" +
 				"\\tFRAMEWORK_DIR=$TARGET_BUILD_DIR/$EXECUTABLE_FOLDER_PATH/Frameworks\\n" +
 				FrameworkScript.ToString() + 
