Expose Target Configurations to BuildConfiguration.xml

diff --git a/Engine/Source/Programs/UnrealBuildTool/ProjectFiles/ProjectFileGenerator.cs b/Engine/Source/Programs/UnrealBuildTool/ProjectFiles/ProjectFileGenerator.cs
--- a/Engine/Source/Programs/UnrealBuildTool/ProjectFiles/ProjectFileGenerator.cs
+++ b/Engine/Source/Programs/UnrealBuildTool/ProjectFiles/ProjectFileGenerator.cs
@@ -309,6 +309,13 @@ namespace UnrealBuildTool
 		[XmlConfigFile(Name = "Platforms")]
 		string[] PlatformNames = null;
 
+		/// <summary>
+		/// Names of configurations to include in the generated project files.
+		/// See UnrealTargetConfiguration for valid entries
+		/// </summary>
+		[XmlConfigFile(Name = "Configurations")]
+		string[] ConfigurationNames = null;
+
 		/// <summary>
 		/// Relative path to the directory where the master project file will be saved to
 		/// </summary>
@@ -1882,8 +1889,31 @@ namespace UnrealBuildTool
 				}
 			}
 
+			List<UnrealTargetConfiguration> AllowedTargetConfigurations = new List<UnrealTargetConfiguration>();
+
+			if (ConfigurationNames == null)
+			{
+				AllowedTargetConfigurations = Enum.GetValues(typeof(UnrealTargetConfiguration)).Cast<UnrealTargetConfiguration>().ToList();
+			}
+			else
+			{
+				foreach (string ConfigName in ConfigurationNames)
+				{
+					try
+					{
+						UnrealTargetConfiguration AllowedConfiguration = (UnrealTargetConfiguration)Enum.Parse(typeof(UnrealTargetConfiguration), ConfigName);
+						AllowedTargetConfigurations.Add(AllowedConfiguration);
+					}
+					catch (Exception)
+					{
+						Log.TraceWarning("Invalid entry found in Configurations: {0}. Must be a member of UnrealTargetConfiguration", ConfigName);
+						continue;
+					}
+				}
+			}
+
 			// Add all configurations
-			foreach( UnrealTargetConfiguration CurConfiguration in Enum.GetValues( typeof(UnrealTargetConfiguration) ) )
+			foreach( UnrealTargetConfiguration CurConfiguration in AllowedTargetConfigurations)
 			{
 				if( CurConfiguration != UnrealTargetConfiguration.Unknown )
 				{
