Fix PVSToolChain to work with SingleFileCompile mode.

diff --git a/Engine/Source/Programs/UnrealBuildTool/Platform/Windows/PVSToolChain.cs b/Engine/Source/Programs/UnrealBuildTool/Platform/Windows/PVSToolChain.cs
--- a/Engine/Source/Programs/UnrealBuildTool/Platform/Windows/PVSToolChain.cs
+++ b/Engine/Source/Programs/UnrealBuildTool/Platform/Windows/PVSToolChain.cs
@@ -295,6 +295,7 @@ namespace UnrealBuildTool
 				}
 				AnalyzeAction.PrerequisiteItems.Add(ConfigFileItem);
 				AnalyzeAction.PrerequisiteItems.Add(PreprocessedFileItem);
+				AnalyzeAction.PrerequisiteItems.AddRange(InputFiles); // Add the InputFiles as PrerequisiteItems so that in SingleFileCompile mode the PVSAnalyze step is not filtered out
 				AnalyzeAction.ProducedItems.Add(OutputFileItem);
 				AnalyzeAction.DeleteItems.Add(OutputFileItem); // PVS Studio will append by default, so need to delete produced items
 				Actions.Add(AnalyzeAction);
@@ -323,6 +324,9 @@ namespace UnrealBuildTool
 
 			List<FileReference> InputFiles = Makefile.OutputItems.Select(x => x.Location).Where(x => x.HasExtension(".pvslog")).ToList();
 
+			// Collect the prerequisite items off of the Compile action added in CompileCPPFiles so that in SingleFileCompile mode the PVSGather step is also not filtered out
+			List<FileItem> AnalyzeActionPrerequisiteItems = Makefile.Actions.Where(x => x.ActionType == ActionType.Compile).SelectMany(x => x.PrerequisiteItems).ToList();
+
 			FileItem InputFileListItem = FileItem.CreateIntermediateTextFile(OutputFile.ChangeExtension(".input"), InputFiles.Select(x => x.FullName));
 
 			Action AnalyzeAction = new Action(ActionType.Compile);
@@ -331,6 +335,7 @@ namespace UnrealBuildTool
 			AnalyzeAction.WorkingDirectory = UnrealBuildTool.EngineSourceDirectory;
 			AnalyzeAction.PrerequisiteItems.Add(InputFileListItem);
 			AnalyzeAction.PrerequisiteItems.AddRange(Makefile.OutputItems);
+			AnalyzeAction.PrerequisiteItems.AddRange(AnalyzeActionPrerequisiteItems);
 			AnalyzeAction.ProducedItems.Add(FileItem.GetItemByFileReference(OutputFile));
 			AnalyzeAction.DeleteItems.AddRange(AnalyzeAction.ProducedItems);
 			Makefile.Actions.Add(AnalyzeAction);
