Submit of Josh Adams' restructuring of platform handling in UBT

diff --git a/Engine/Platforms/XXX/Source/Programs/UnrealBuildTool/UEBuildXXX.cs b/Engine/Platforms/XXX/Source/Programs/UnrealBuildTool/UEBuildXXX.cs
--- a/Engine/Platforms/XXX/Source/Programs/UnrealBuildTool/UEBuildXXX.cs
+++ b/Engine/Platforms/XXX/Source/Programs/UnrealBuildTool/UEBuildXXX.cs
@@ -206,6 +206,12 @@ class XXXPlatform : UEBuildPlatform
 			Rules.AppendStringToPublicDefinition("DDPI_SHADER_PLATFORM_NAME_MAP", "{ TEXT(\"XXX\"), SP_XXX },");
 		}
 
+		// don't do any target platform stuff if SDK is not available
+		if (!UEBuildPlatform.IsPlatformAvailable(Platform))
+		{
+			return;
+		}
+
 		if (Target.Platform == UnrealTargetPlatform.Win64)
 		{
 			bool bIsMonolithic = true;
@@ -435,13 +441,8 @@ class XXXPlatformFactory : UEBuildPlatformFactory
  		XXXPlatformSDK SDK = new XXXPlatformSDK();
  		SDK.ManageAndValidateSDK();
 
-		// Make sure the SDK is installed
-		if (SDK.HasRequiredSDKsInstalled() == SDKStatus.Valid)
-		{
-			// Register this build platform for XXX
-			Log.TraceVerbose("        Registering for {0}", UnrealTargetPlatform.XXX.ToString());
-			UEBuildPlatform.RegisterBuildPlatform(new XXXPlatform(SDK));
-			UEBuildPlatform.RegisterPlatformWithGroup(UnrealTargetPlatform.XXX, UnrealPlatformGroup.Fake);
-		}
+		// Register this build platform for XXX
+		UEBuildPlatform.RegisterBuildPlatform(new XXXPlatform(SDK));
+		UEBuildPlatform.RegisterPlatformWithGroup(UnrealTargetPlatform.XXX, UnrealPlatformGroup.Fake);
 	}
 }
\ No newline at end of file
diff --git a/Engine/Source/Programs/AutomationTool/Scripts/ListThirdPartySoftware.Automation.cs b/Engine/Source/Programs/AutomationTool/Scripts/ListThirdPartySoftware.Automation.cs
--- a/Engine/Source/Programs/AutomationTool/Scripts/ListThirdPartySoftware.Automation.cs
+++ b/Engine/Source/Programs/AutomationTool/Scripts/ListThirdPartySoftware.Automation.cs
@@ -89,7 +89,7 @@ class ListThirdPartySoftware : BuildCommand
 
 			// Get the platforms to exclude
 			List<UnrealTargetPlatform> SupportedPlatforms = new List<UnrealTargetPlatform> { UnrealTargetPlatform.Parse(Object.GetStringField("Platform")) };
-			string[] ExcludePlatformNames = Utils.MakeListOfUnsupportedPlatforms(SupportedPlatforms).ToArray();
+			string[] ExcludePlatformNames = Utils.MakeListOfUnsupportedPlatforms(SupportedPlatforms, bIncludeUnbuildablePlatforms: true).ToArray();
 
 			// Find all the TPS files under the engine directory which match
 			foreach(DirectoryReference DirectoryToScan in SortedDirectoriesToScan)
diff --git a/Engine/Source/Programs/UnrealBuildTool/Configuration/UEBuildPlatform.cs b/Engine/Source/Programs/UnrealBuildTool/Configuration/UEBuildPlatform.cs
--- a/Engine/Source/Programs/UnrealBuildTool/Configuration/UEBuildPlatform.cs
+++ b/Engine/Source/Programs/UnrealBuildTool/Configuration/UEBuildPlatform.cs
@@ -351,7 +351,7 @@ namespace UnrealBuildTool
 		/// <returns>True if it's available, false otherwise</returns>
 		public static bool IsPlatformAvailable(UnrealTargetPlatform Platform)
 		{
-			return BuildPlatformDictionary.ContainsKey(Platform);
+			return BuildPlatformDictionary.ContainsKey(Platform) && BuildPlatformDictionary[Platform].HasRequiredSDKsInstalled() == SDKStatus.Valid;
 		}
 
 		/// <summary>
@@ -360,6 +360,8 @@ namespace UnrealBuildTool
 		/// <param name="InBuildPlatform"> The UEBuildPlatform instance to use for the InPlatform</param>
 		public static void RegisterBuildPlatform(UEBuildPlatform InBuildPlatform)
 		{
+			Log.TraceVerbose("        Registering build platform: {0} - buildable: {1}", InBuildPlatform.Platform, InBuildPlatform.HasRequiredSDKsInstalled() == SDKStatus.Valid);
+
 			if (BuildPlatformDictionary.ContainsKey(InBuildPlatform.Platform) == true)
 			{
 				Log.TraceWarning("RegisterBuildPlatform Warning: Registering build platform {0} for {1} when it is already set to {2}",
diff --git a/Engine/Source/Programs/UnrealBuildTool/Configuration/UEBuildTarget.cs b/Engine/Source/Programs/UnrealBuildTool/Configuration/UEBuildTarget.cs
--- a/Engine/Source/Programs/UnrealBuildTool/Configuration/UEBuildTarget.cs
+++ b/Engine/Source/Programs/UnrealBuildTool/Configuration/UEBuildTarget.cs
@@ -600,6 +600,12 @@ namespace UnrealBuildTool
 		/// <returns>The build target object for the specified build rules source file</returns>
 		public static UEBuildTarget Create(TargetDescriptor Descriptor, bool bSkipRulesCompile, bool bUsePrecompiled)
 		{
+			// make sure we are allowed to build this platform
+			if (!UEBuildPlatform.IsPlatformAvailable(Descriptor.Platform))
+			{
+				throw new BuildException("Platform {0} is not a valid platform to build. Check that the SDK is installed properly.", Descriptor.Platform);
+			}
+
 			RulesAssembly RulesAssembly;
 			using(Timeline.ScopeEvent("RulesCompiler.CreateTargetRulesAssembly()"))
 			{
diff --git a/Engine/Source/Programs/UnrealBuildTool/Platform/Android/UEBuildAndroid.cs b/Engine/Source/Programs/UnrealBuildTool/Platform/Android/UEBuildAndroid.cs
--- a/Engine/Source/Programs/UnrealBuildTool/Platform/Android/UEBuildAndroid.cs
+++ b/Engine/Source/Programs/UnrealBuildTool/Platform/Android/UEBuildAndroid.cs
@@ -222,6 +222,12 @@ namespace UnrealBuildTool
 		/// <param name="Target">The target being build</param>
 		public override void ModifyModuleRulesForOtherPlatform(string ModuleName, ModuleRules Rules, ReadOnlyTargetRules Target)
 		{
+			// don't do any target platform stuff if SDK is not available
+			if (!UEBuildPlatform.IsPlatformAvailable(Platform))
+			{
+				return;
+			}
+
 			if ((Target.Platform == UnrealTargetPlatform.Win32) || (Target.Platform == UnrealTargetPlatform.Win64) || (Target.Platform == UnrealTargetPlatform.Mac) || (Target.Platform == UnrealTargetPlatform.Linux))
 			{
 				bool bBuildShaderFormats = Target.bForceBuildShaderFormats;
@@ -645,24 +651,9 @@ namespace UnrealBuildTool
 			AndroidPlatformSDK SDK = new AndroidPlatformSDK();
 			SDK.ManageAndValidateSDK();
 
-			if ((ProjectFileGenerator.bGenerateProjectFiles == true) || (SDK.HasRequiredSDKsInstalled() == SDKStatus.Valid) || Environment.GetEnvironmentVariable("IsBuildMachine") == "1")
-			{
-				bool bRegisterBuildPlatform = true;
-
-				FileReference AndroidTargetPlatformFile = FileReference.Combine(UnrealBuildTool.EngineSourceDirectory, "Developer", "Android", "AndroidTargetPlatform", "AndroidTargetPlatform.Build.cs");
-				if (FileReference.Exists(AndroidTargetPlatformFile) == false)
-				{
-					bRegisterBuildPlatform = false;
-				}
-
-				if (bRegisterBuildPlatform == true)
-				{
-					// Register this build platform
-					Log.TraceVerbose("        Registering for {0}", UnrealTargetPlatform.Android.ToString());
-					UEBuildPlatform.RegisterBuildPlatform(new AndroidPlatform(SDK));
-					UEBuildPlatform.RegisterPlatformWithGroup(UnrealTargetPlatform.Android, UnrealPlatformGroup.Android);
-				}
-			}
+			// Register this build platform
+			UEBuildPlatform.RegisterBuildPlatform(new AndroidPlatform(SDK));
+			UEBuildPlatform.RegisterPlatformWithGroup(UnrealTargetPlatform.Android, UnrealPlatformGroup.Android);
 		}
 	}
 }
diff --git a/Engine/Source/Programs/UnrealBuildTool/Platform/HoloLens/UEBuildHoloLens.cs b/Engine/Source/Programs/UnrealBuildTool/Platform/HoloLens/UEBuildHoloLens.cs
--- a/Engine/Source/Programs/UnrealBuildTool/Platform/HoloLens/UEBuildHoloLens.cs
+++ b/Engine/Source/Programs/UnrealBuildTool/Platform/HoloLens/UEBuildHoloLens.cs
@@ -784,14 +784,9 @@ namespace UnrealBuildTool
 			HoloLensPlatformSDK SDK = new HoloLensPlatformSDK();
 			SDK.ManageAndValidateSDK();
 
-			// Register this build platform for HoloLens
-			if (SDK.HasRequiredSDKsInstalled() == SDKStatus.Valid)
-			{
-				Log.TraceVerbose("		Registering for {0}", UnrealTargetPlatform.HoloLens.ToString());
-				UEBuildPlatform.RegisterBuildPlatform(new HoloLens(UnrealTargetPlatform.HoloLens, SDK));
-				UEBuildPlatform.RegisterPlatformWithGroup(UnrealTargetPlatform.HoloLens, UnrealPlatformGroup.Microsoft);
-				UEBuildPlatform.RegisterPlatformWithGroup(UnrealTargetPlatform.HoloLens, UnrealPlatformGroup.HoloLens);
-			}
+			UEBuildPlatform.RegisterBuildPlatform(new HoloLens(UnrealTargetPlatform.HoloLens, SDK));
+			UEBuildPlatform.RegisterPlatformWithGroup(UnrealTargetPlatform.HoloLens, UnrealPlatformGroup.Microsoft);
+			UEBuildPlatform.RegisterPlatformWithGroup(UnrealTargetPlatform.HoloLens, UnrealPlatformGroup.HoloLens);
 		}
 	}
 }
diff --git a/Engine/Source/Programs/UnrealBuildTool/Platform/IOS/UEBuildIOS.cs b/Engine/Source/Programs/UnrealBuildTool/Platform/IOS/UEBuildIOS.cs
--- a/Engine/Source/Programs/UnrealBuildTool/Platform/IOS/UEBuildIOS.cs
+++ b/Engine/Source/Programs/UnrealBuildTool/Platform/IOS/UEBuildIOS.cs
@@ -1019,6 +1019,12 @@ namespace UnrealBuildTool
 		/// <param name="Target">The target being build</param>
 		public override void ModifyModuleRulesForOtherPlatform(string ModuleName, ModuleRules Rules, ReadOnlyTargetRules Target)
 		{
+			// don't do any target platform stuff if SDK is not available
+			if (!UEBuildPlatform.IsPlatformAvailable(Platform))
+			{
+				return;
+			}
+
 			if ((Target.Platform == UnrealTargetPlatform.Win32) || (Target.Platform == UnrealTargetPlatform.Win64) || (Target.Platform == UnrealTargetPlatform.Mac))
 			{
 				bool bBuildShaderFormats = Target.bForceBuildShaderFormats;
@@ -1272,7 +1278,6 @@ namespace UnrealBuildTool
 			SDK.ManageAndValidateSDK();
 
 			// Register this build platform for IOS
-			Log.TraceVerbose("        Registering for {0}", UnrealTargetPlatform.IOS.ToString());
 			UEBuildPlatform.RegisterBuildPlatform(new IOSPlatform(SDK));
 			UEBuildPlatform.RegisterPlatformWithGroup(UnrealTargetPlatform.IOS, UnrealPlatformGroup.Apple);
 			UEBuildPlatform.RegisterPlatformWithGroup(UnrealTargetPlatform.IOS, UnrealPlatformGroup.IOS);
diff --git a/Engine/Source/Programs/UnrealBuildTool/Platform/Linux/UEBuildLinux.cs b/Engine/Source/Programs/UnrealBuildTool/Platform/Linux/UEBuildLinux.cs
--- a/Engine/Source/Programs/UnrealBuildTool/Platform/Linux/UEBuildLinux.cs
+++ b/Engine/Source/Programs/UnrealBuildTool/Platform/Linux/UEBuildLinux.cs
@@ -338,6 +338,12 @@ namespace UnrealBuildTool
 		/// <param name="Target">The target being build</param>
 		public override void ModifyModuleRulesForOtherPlatform(string ModuleName, ModuleRules Rules, ReadOnlyTargetRules Target)
 		{
+			// don't do any target platform stuff if SDK is not available
+			if (!UEBuildPlatform.IsPlatformAvailable(Platform))
+			{
+				return;
+			}
+
 			if ((Target.Platform == UnrealTargetPlatform.Win32) || (Target.Platform == UnrealTargetPlatform.Win64))
 			{
 				if (!Target.bBuildRequiresCookedData)
@@ -817,22 +823,14 @@ namespace UnrealBuildTool
 			LinuxPlatformSDK SDK = new LinuxPlatformSDK();
 			SDK.ManageAndValidateSDK();
 
-			if ((ProjectFileGenerator.bGenerateProjectFiles == true) || (SDK.HasRequiredSDKsInstalled() == SDKStatus.Valid))
-			{
-				FileReference LinuxTargetPlatformFile = FileReference.Combine(UnrealBuildTool.EngineSourceDirectory, "Developer", "Linux", "LinuxTargetPlatform", "LinuxTargetPlatform.Build.cs");
-				if (FileReference.Exists(LinuxTargetPlatformFile))
-				{
-					// Register this build platform for Linux x86-64 and AArch64
-					Log.TraceVerbose("        Registering for {0}", UnrealTargetPlatform.Linux.ToString());
-					UEBuildPlatform.RegisterBuildPlatform(new LinuxPlatform(UnrealTargetPlatform.Linux, SDK));
-					UEBuildPlatform.RegisterPlatformWithGroup(UnrealTargetPlatform.Linux, UnrealPlatformGroup.Linux);
-					UEBuildPlatform.RegisterPlatformWithGroup(UnrealTargetPlatform.Linux, UnrealPlatformGroup.Unix);
-
-					UEBuildPlatform.RegisterBuildPlatform(new LinuxPlatform(UnrealTargetPlatform.LinuxAArch64, SDK));
-					UEBuildPlatform.RegisterPlatformWithGroup(UnrealTargetPlatform.LinuxAArch64, UnrealPlatformGroup.Linux);
-					UEBuildPlatform.RegisterPlatformWithGroup(UnrealTargetPlatform.LinuxAArch64, UnrealPlatformGroup.Unix);
-				}
-			}
+			// Register this build platform for Linux x86-64 and AArch64
+			UEBuildPlatform.RegisterBuildPlatform(new LinuxPlatform(UnrealTargetPlatform.Linux, SDK));
+			UEBuildPlatform.RegisterPlatformWithGroup(UnrealTargetPlatform.Linux, UnrealPlatformGroup.Linux);
+			UEBuildPlatform.RegisterPlatformWithGroup(UnrealTargetPlatform.Linux, UnrealPlatformGroup.Unix);
+
+			UEBuildPlatform.RegisterBuildPlatform(new LinuxPlatform(UnrealTargetPlatform.LinuxAArch64, SDK));
+			UEBuildPlatform.RegisterPlatformWithGroup(UnrealTargetPlatform.LinuxAArch64, UnrealPlatformGroup.Linux);
+			UEBuildPlatform.RegisterPlatformWithGroup(UnrealTargetPlatform.LinuxAArch64, UnrealPlatformGroup.Unix);
 		}
 	}
 }
diff --git a/Engine/Source/Programs/UnrealBuildTool/Platform/Lumin/UEBuildLumin.cs b/Engine/Source/Programs/UnrealBuildTool/Platform/Lumin/UEBuildLumin.cs
--- a/Engine/Source/Programs/UnrealBuildTool/Platform/Lumin/UEBuildLumin.cs
+++ b/Engine/Source/Programs/UnrealBuildTool/Platform/Lumin/UEBuildLumin.cs
@@ -195,6 +195,12 @@ namespace UnrealBuildTool
 		/// <param name="Target">The target being build</param>
 		public override void ModifyModuleRulesForOtherPlatform(string ModuleName, ModuleRules Rules, ReadOnlyTargetRules Target)
 		{
+			// don't do any target platform stuff if SDK is not available
+			if (!UEBuildPlatform.IsPlatformAvailable(Platform))
+			{
+				return;
+			}
+
 			if ((Target.Platform == UnrealTargetPlatform.Win32) || (Target.Platform == UnrealTargetPlatform.Win64) || (Target.Platform == UnrealTargetPlatform.Mac) || (Target.Platform == UnrealTargetPlatform.Linux))
 			{
 				bool bBuildShaderFormats = Target.bForceBuildShaderFormats;
@@ -378,24 +384,9 @@ namespace UnrealBuildTool
 			LuminPlatformSDK SDK = new LuminPlatformSDK();
 			SDK.ManageAndValidateSDK();
 
-			if ((ProjectFileGenerator.bGenerateProjectFiles == true) || (SDK.HasRequiredSDKsInstalled() == SDKStatus.Valid) || Environment.GetEnvironmentVariable("IsBuildMachine") == "1")
-			{
-				bool bRegisterBuildPlatform = true;
-
-				FileReference TargetPlatformFile = FileReference.Combine(UnrealBuildTool.EngineSourceDirectory, "Developer", "Lumin", "LuminTargetPlatform", "LuminTargetPlatform.Build.cs");
-				if (FileReference.Exists(TargetPlatformFile) == false)
-				{
-					bRegisterBuildPlatform = false;
-				}
-
-				if (bRegisterBuildPlatform == true)
-				{
-					// Register this build platform
-					Log.TraceVerbose("        Registering for {0}", UnrealTargetPlatform.Lumin.ToString());
-					UEBuildPlatform.RegisterBuildPlatform(new LuminPlatform(SDK));
-					UEBuildPlatform.RegisterPlatformWithGroup(UnrealTargetPlatform.Lumin, UnrealPlatformGroup.Android);
-				}
-			}
+			// Register this build platform
+			UEBuildPlatform.RegisterBuildPlatform(new LuminPlatform(SDK));
+			UEBuildPlatform.RegisterPlatformWithGroup(UnrealTargetPlatform.Lumin, UnrealPlatformGroup.Android);
 		}
 	}
 
diff --git a/Engine/Source/Programs/UnrealBuildTool/Platform/Mac/UEBuildMac.cs b/Engine/Source/Programs/UnrealBuildTool/Platform/Mac/UEBuildMac.cs
--- a/Engine/Source/Programs/UnrealBuildTool/Platform/Mac/UEBuildMac.cs
+++ b/Engine/Source/Programs/UnrealBuildTool/Platform/Mac/UEBuildMac.cs
@@ -388,7 +388,6 @@ namespace UnrealBuildTool
 			SDK.ManageAndValidateSDK();
 
 			// Register this build platform for Mac
-			Log.TraceVerbose("        Registering for {0}", UnrealTargetPlatform.Mac.ToString());
 			UEBuildPlatform.RegisterBuildPlatform(new MacPlatform(SDK));
 			UEBuildPlatform.RegisterPlatformWithGroup(UnrealTargetPlatform.Mac, UnrealPlatformGroup.Apple);
 		}
diff --git a/Engine/Source/Programs/UnrealBuildTool/Platform/TVOS/UEBuildTVOS.cs b/Engine/Source/Programs/UnrealBuildTool/Platform/TVOS/UEBuildTVOS.cs
--- a/Engine/Source/Programs/UnrealBuildTool/Platform/TVOS/UEBuildTVOS.cs
+++ b/Engine/Source/Programs/UnrealBuildTool/Platform/TVOS/UEBuildTVOS.cs
@@ -95,6 +95,12 @@ namespace UnrealBuildTool
 		{
 			base.ModifyModuleRulesForOtherPlatform(ModuleName, Rules, Target);
 
+			// don't do any target platform stuff if SDK is not available
+			if (!UEBuildPlatform.IsPlatformAvailable(Platform))
+			{
+				return;
+			}
+
 			if ((Target.Platform == UnrealTargetPlatform.Win32) || (Target.Platform == UnrealTargetPlatform.Win64) || (Target.Platform == UnrealTargetPlatform.Mac))
 			{
 				// allow standalone tools to use targetplatform modules, without needing Engine
@@ -155,7 +161,6 @@ namespace UnrealBuildTool
 			SDK.ManageAndValidateSDK();
 
 			// Register this build platform for IOS
-			Log.TraceVerbose("        Registering for {0}", UnrealTargetPlatform.TVOS.ToString());
 			UEBuildPlatform.RegisterBuildPlatform(new TVOSPlatform(SDK));
 			UEBuildPlatform.RegisterPlatformWithGroup(UnrealTargetPlatform.TVOS, UnrealPlatformGroup.Apple);
 			UEBuildPlatform.RegisterPlatformWithGroup(UnrealTargetPlatform.TVOS, UnrealPlatformGroup.IOS);
diff --git a/Engine/Source/Programs/UnrealBuildTool/Platform/Windows/UEBuildWindows.cs b/Engine/Source/Programs/UnrealBuildTool/Platform/Windows/UEBuildWindows.cs
--- a/Engine/Source/Programs/UnrealBuildTool/Platform/Windows/UEBuildWindows.cs
+++ b/Engine/Source/Programs/UnrealBuildTool/Platform/Windows/UEBuildWindows.cs
@@ -2210,12 +2210,10 @@ namespace UnrealBuildTool
 			SDK.ManageAndValidateSDK();
 
 			// Register this build platform for both Win64 and Win32
-			Log.TraceVerbose("        Registering for {0}", UnrealTargetPlatform.Win64.ToString());
 			UEBuildPlatform.RegisterBuildPlatform(new WindowsPlatform(UnrealTargetPlatform.Win64, SDK));
 			UEBuildPlatform.RegisterPlatformWithGroup(UnrealTargetPlatform.Win64, UnrealPlatformGroup.Windows);
 			UEBuildPlatform.RegisterPlatformWithGroup(UnrealTargetPlatform.Win64, UnrealPlatformGroup.Microsoft);
 
-			Log.TraceVerbose("        Registering for {0}", UnrealTargetPlatform.Win32.ToString());
 			UEBuildPlatform.RegisterBuildPlatform(new WindowsPlatform(UnrealTargetPlatform.Win32, SDK));
 			UEBuildPlatform.RegisterPlatformWithGroup(UnrealTargetPlatform.Win32, UnrealPlatformGroup.Windows);
 			UEBuildPlatform.RegisterPlatformWithGroup(UnrealTargetPlatform.Win32, UnrealPlatformGroup.Microsoft);
diff --git a/Engine/Source/Programs/UnrealBuildTool/ProjectFiles/ProjectFileGenerator.cs b/Engine/Source/Programs/UnrealBuildTool/ProjectFiles/ProjectFileGenerator.cs
--- a/Engine/Source/Programs/UnrealBuildTool/ProjectFiles/ProjectFileGenerator.cs
+++ b/Engine/Source/Programs/UnrealBuildTool/ProjectFiles/ProjectFileGenerator.cs
@@ -444,7 +444,7 @@ namespace UnrealBuildTool
 		/// </summary>
 		void DiscoverCSharpProgramProjects(MasterProjectFolder ProgramsFolder)
 		{
-			string[] UnsupportedPlatformNames = Utils.MakeListOfUnsupportedPlatforms(SupportedPlatforms).ToArray();
+			string[] UnsupportedPlatformNames = Utils.MakeListOfUnsupportedPlatforms(SupportedPlatforms, bIncludeUnbuildablePlatforms:true).ToArray();
 
 			List<FileReference> FoundProjects = new List<FileReference>();
 
@@ -1414,7 +1414,7 @@ namespace UnrealBuildTool
 			List<FileReference> AllTargetFiles = new List<FileReference>();
 
 			// Make a list of all platform name strings that we're *not* including in the project files
-			List<string> UnsupportedPlatformNameStrings = Utils.MakeListOfUnsupportedPlatforms( SupportedPlatforms );
+			List<string> UnsupportedPlatformNameStrings = Utils.MakeListOfUnsupportedPlatforms(SupportedPlatforms, bIncludeUnbuildablePlatforms:true);
 
 			// Locate all targets (*.Target.cs files)
 			List<FileReference> FoundTargetFiles = RulesCompiler.FindAllRulesSourceFiles( RulesCompiler.RulesFileType.Target, AllGameProjects.Select(x => x.Directory).ToList(), ForeignPlugins: DiscoverExtraPlugins(AllGameProjects), AdditionalSearchPaths:null, bIncludeTempTargets: bIncludeTempTargets );
diff --git a/Engine/Source/Programs/UnrealBuildTool/System/Utils.cs b/Engine/Source/Programs/UnrealBuildTool/System/Utils.cs
--- a/Engine/Source/Programs/UnrealBuildTool/System/Utils.cs
+++ b/Engine/Source/Programs/UnrealBuildTool/System/Utils.cs
@@ -491,8 +491,10 @@ namespace UnrealBuildTool
 		/// Given a list of supported platforms, returns a list of names of platforms that should not be supported
 		/// </summary>
 		/// <param name="SupportedPlatforms">List of supported platforms</param>
+		/// <param name="bIncludeUnbuildablePlatforms">If true, add platforms that are present but not available for compiling</param>
+		/// 
 		/// <returns>List of unsupported platforms in string format</returns>
-		public static List<string> MakeListOfUnsupportedPlatforms(List<UnrealTargetPlatform> SupportedPlatforms)
+		public static List<string> MakeListOfUnsupportedPlatforms(List<UnrealTargetPlatform> SupportedPlatforms, bool bIncludeUnbuildablePlatforms)
 		{
 			// Make a list of all platform name strings that we're *not* currently compiling, to speed
 			// up file path comparisons later on
@@ -551,6 +553,11 @@ namespace UnrealBuildTool
 					{
 						OtherPlatformNameStrings.Add(CurPlatform.ToString());
 					}
+					// if a platform isn't available to build, then return it 
+					else if (bIncludeUnbuildablePlatforms && !UEBuildPlatform.IsPlatformAvailable(CurPlatform))
+					{
+						OtherPlatformNameStrings.Add(CurPlatform.ToString());
+					}
 				}
 
 				return OtherPlatformNameStrings;
