#UE4 Add bExplicitlyLoaded to UBT's PluginDescriptor as well to match the c++ code

diff --git a/Engine/Source/Programs/UnrealBuildTool/System/PluginDescriptor.cs b/Engine/Source/Programs/UnrealBuildTool/System/PluginDescriptor.cs
--- a/Engine/Source/Programs/UnrealBuildTool/System/PluginDescriptor.cs
+++ b/Engine/Source/Programs/UnrealBuildTool/System/PluginDescriptor.cs
@@ -169,6 +169,11 @@ namespace UnrealBuildTool
 		/// </summary>
 		public bool bRequiresBuildPlatform;
 
+		/// <summary>
+		/// When true, this plugin's modules will not be loaded automatically nor will it's content be mounted automatically. It will load/mount when explicitly requested and LoadingPhases will be ignored
+		/// </summary>
+		public bool bExplicitlyLoaded;
+
 		/// <summary>
 		/// Set of pre-build steps to execute, keyed by host platform name.
 		/// </summary>
@@ -288,6 +293,7 @@ namespace UnrealBuildTool
 			}
 
 			RawObject.TryGetBoolField("RequiresBuildPlatform", out bRequiresBuildPlatform);
+			RawObject.TryGetBoolField("ExplicitlyLoaded", out bExplicitlyLoaded);
 
 			CustomBuildSteps.TryRead(RawObject, "PreBuildSteps", out PreBuildSteps);
 			CustomBuildSteps.TryRead(RawObject, "PostBuildSteps", out PostBuildSteps);
@@ -383,6 +389,11 @@ namespace UnrealBuildTool
 				Writer.WriteValue("RequiresBuildPlatform", bRequiresBuildPlatform);
 			}
 
+			if (bExplicitlyLoaded)
+			{
+				Writer.WriteValue("ExplicitlyLoaded", bExplicitlyLoaded);
+			}
+
 			if(SupportedTargetPlatforms != null && SupportedTargetPlatforms.Length > 0)
 			{
 				Writer.WriteStringArrayField("SupportedTargetPlatforms", SupportedTargetPlatforms.Select<UnrealTargetPlatform, string>(x => x.ToString()).ToArray());
