UBT: Don't reference FASTBuild class in ActionGraph.cs if FASTBuild is unavailable

diff --git a/Engine/Source/Programs/UnrealBuildTool/System/ActionGraph.cs b/Engine/Source/Programs/UnrealBuildTool/System/ActionGraph.cs
--- a/Engine/Source/Programs/UnrealBuildTool/System/ActionGraph.cs
+++ b/Engine/Source/Programs/UnrealBuildTool/System/ActionGraph.cs
@@ -129,10 +129,12 @@ namespace UnrealBuildTool
 				{
 					Executor = new XGE();
 				}
+#if __FASTBUILD_AVAILABLE__
 				else if (BuildConfiguration.bAllowFASTBuild && FASTBuild.IsAvailable())
 				{
 					Executor = new FASTBuild();
 				}
+#endif
 				else if (BuildConfiguration.bAllowDistcc)
 				{
 					Executor = new Distcc();
diff --git a/Engine/Source/Programs/UnrealBuildTool/UnrealBuildTool.csproj b/Engine/Source/Programs/UnrealBuildTool/UnrealBuildTool.csproj
--- a/Engine/Source/Programs/UnrealBuildTool/UnrealBuildTool.csproj
+++ b/Engine/Source/Programs/UnrealBuildTool/UnrealBuildTool.csproj
@@ -307,6 +307,9 @@
       <ItemGroup>
           <Compile Include="System\NotForLicensees\FASTBuild.cs" />
       </ItemGroup>
+      <PropertyGroup>
+        <DefineConstants>__FASTBUILD_AVAILABLE__</DefineConstants>
+      </PropertyGroup>
     </When>
   </Choose>
   <Choose>
