Null check fix for CL 10611459

diff --git a/Engine/Source/Programs/UnrealBuildTool/Platform/IOS/UEDeployIOS.cs b/Engine/Source/Programs/UnrealBuildTool/Platform/IOS/UEDeployIOS.cs
--- a/Engine/Source/Programs/UnrealBuildTool/Platform/IOS/UEDeployIOS.cs
+++ b/Engine/Source/Programs/UnrealBuildTool/Platform/IOS/UEDeployIOS.cs
@@ -797,12 +797,24 @@ namespace UnrealBuildTool
 			return SdkVersion;
         }
 
+		public static bool GetCompileAsDll(TargetReceipt Receipt)
+		{
+			if (Receipt != null)
+			{
+				ReceiptProperty CompileAsDllProperty = Receipt.AdditionalProperties.FirstOrDefault(x => x.Name == "CompileAsDll");
+				if (CompileAsDllProperty != null && CompileAsDllProperty.Value == "true")
+				{
+					return true;
+				}
+			}
+			return false;
+		}
+
 		public bool GeneratePList(FileReference ProjectFile, UnrealTargetConfiguration Config, string ProjectDirectory, bool bIsUE4Game, string GameName, string ProjectName, string InEngineDir, string AppDirectory, TargetReceipt Receipt, out bool bSupportsPortrait, out bool bSupportsLandscape, out bool bSkipIcons)
 		{
 			List<string> UPLScripts = CollectPluginDataPaths(Receipt.AdditionalProperties);
 			VersionNumber SdkVersion = GetSdkVersion(Receipt);
-			ReceiptProperty CompileAsDllProperty = Receipt.AdditionalProperties.FirstOrDefault(x => x.Name == "CompileAsDll");
-			bool bBuildAsFramework = CompileAsDllProperty != null && CompileAsDllProperty.Value == "true";
+			bool bBuildAsFramework = GetCompileAsDll(Receipt);
 			return GeneratePList(ProjectFile, Config, ProjectDirectory, bIsUE4Game, GameName, ProjectName, InEngineDir, AppDirectory, UPLScripts, SdkVersion, "", bBuildAsFramework, out bSupportsPortrait, out bSupportsLandscape, out bSkipIcons);
 		}
 
@@ -943,8 +955,7 @@ namespace UnrealBuildTool
 		{
 			List<string> UPLScripts = CollectPluginDataPaths(Receipt.AdditionalProperties);
 			VersionNumber SdkVersion = GetSdkVersion(Receipt);
-			ReceiptProperty CompileAsDllProperty = Receipt.AdditionalProperties.FirstOrDefault(x => x.Name == "CompileAsDll");
-			bool bBuildAsFramework = CompileAsDllProperty != null && CompileAsDllProperty.Value == "true";
+			bool bBuildAsFramework = GetCompileAsDll(Receipt);
 			return PrepForUATPackageOrDeploy(Config, ProjectFile, InProjectName, InProjectDirectory, InExecutablePath, InEngineDir, bForDistribution, CookFlavor, bIsDataDeploy, bCreateStubIPA, UPLScripts, SdkVersion, "", bBuildAsFramework);
 		}
 
@@ -1160,8 +1171,7 @@ namespace UnrealBuildTool
 		{
 			List<string> UPLScripts = CollectPluginDataPaths(Receipt.AdditionalProperties);
 			VersionNumber SdkVersion = GetSdkVersion(Receipt);
-			ReceiptProperty CompileAsDllProperty = Receipt.AdditionalProperties.FirstOrDefault(x => x.Name == "CompileAsDll");
-			bool bBuildAsFramework = CompileAsDllProperty != null && CompileAsDllProperty.Value == "true";
+			bool bBuildAsFramework = GetCompileAsDll(Receipt);
 			return PrepTargetForDeployment(Receipt.ProjectFile, Receipt.TargetName, Receipt.Platform, Receipt.Configuration, UPLScripts, SdkVersion, false, "", bBuildAsFramework);
 		}
 
diff --git a/Engine/Source/Programs/UnrealBuildTool/ProjectFiles/Xcode/XcodeProject.cs b/Engine/Source/Programs/UnrealBuildTool/ProjectFiles/Xcode/XcodeProject.cs
--- a/Engine/Source/Programs/UnrealBuildTool/ProjectFiles/Xcode/XcodeProject.cs
+++ b/Engine/Source/Programs/UnrealBuildTool/ProjectFiles/Xcode/XcodeProject.cs
@@ -1264,8 +1264,7 @@ namespace UnrealBuildTool
 							TargetReceipt Receipt;
 							TargetReceipt.TryRead(ReceiptFilename, out Receipt);
 							VersionNumber SdkVersion = UEDeployIOS.GetSdkVersion(Receipt);
-							ReceiptProperty CompileAsDllProperty = Receipt.AdditionalProperties.FirstOrDefault(x => x.Name == "CompileAsDll");
-							bool bBuildAsFramework = CompileAsDllProperty != null && CompileAsDllProperty.Value == "true";
+							bool bBuildAsFramework = UEDeployIOS.GetCompileAsDll(Receipt);
 							UEDeployIOS.GenerateIOSPList(ProjectFile, Config.BuildConfig, ProjectPath.FullName, bIsUE4Game, GameName, Config.BuildTarget, EngineDir.FullName, ProjectPath + "/Binaries/IOS/Payload", SdkVersion, null, BundleIdentifier, bBuildAsFramework, out bSupportPortrait, out bSupportLandscape, out bSkipIcons);
 						}
 						if (bCreateTVOSInfoPlist)
