Enable chaos and its plugins by default
#rnx

diff --git a/Engine/Plugins/Experimental/ChaosEditor/ChaosEditor.uplugin b/Engine/Plugins/Experimental/ChaosEditor/ChaosEditor.uplugin
--- a/Engine/Plugins/Experimental/ChaosEditor/ChaosEditor.uplugin
+++ b/Engine/Plugins/Experimental/ChaosEditor/ChaosEditor.uplugin
@@ -13,7 +13,7 @@
 	"CanContainContent": false,
 	"IsBetaVersion": true,
 	"Installed": false,
-	"EnabledByDefault": false,
+	"EnabledByDefault": true,
 	"Modules": [
 		{
 			"Name": "FractureEditor",
diff --git a/Engine/Plugins/Experimental/ChaosNiagara/ChaosNiagara.uplugin b/Engine/Plugins/Experimental/ChaosNiagara/ChaosNiagara.uplugin
--- a/Engine/Plugins/Experimental/ChaosNiagara/ChaosNiagara.uplugin
+++ b/Engine/Plugins/Experimental/ChaosNiagara/ChaosNiagara.uplugin
@@ -10,7 +10,7 @@
 	"DocsURL": "",
 	"MarketplaceURL": "",
 	"SupportURL": "",
-	"EnabledByDefault": false,
+	"EnabledByDefault": true,
 	"CanContainContent": true,
 	"IsBetaVersion": true,
 	"Installed": false,
diff --git a/Engine/Plugins/Experimental/ChaosSolverPlugin/ChaosSolverPlugin.uplugin b/Engine/Plugins/Experimental/ChaosSolverPlugin/ChaosSolverPlugin.uplugin
--- a/Engine/Plugins/Experimental/ChaosSolverPlugin/ChaosSolverPlugin.uplugin
+++ b/Engine/Plugins/Experimental/ChaosSolverPlugin/ChaosSolverPlugin.uplugin
@@ -10,7 +10,7 @@
     "DocsURL" : "",
     "MarketplaceURL" : "",
     "SupportURL" : "",
-    "EnabledByDefault" : false,
+    "EnabledByDefault" : true,
     "CanContainContent" : true,
     "IsBetaVersion" : true,
     "Installed" : false,
diff --git a/Engine/Source/Programs/UnrealBuildTool/Configuration/TargetRules.cs b/Engine/Source/Programs/UnrealBuildTool/Configuration/TargetRules.cs
--- a/Engine/Source/Programs/UnrealBuildTool/Configuration/TargetRules.cs
+++ b/Engine/Source/Programs/UnrealBuildTool/Configuration/TargetRules.cs
@@ -359,13 +359,13 @@ namespace UnrealBuildTool
 		/// Whether to compile the Chaos physics plugin.
 		/// </summary>
 		[RequiresUniqueBuildEnvironment]
-		public bool bCompileChaos = false;
+		public bool bCompileChaos = true;
 
 		/// <summary>
 		/// Whether to use the Chaos physics interface. This overrides the physx flags to disable APEX and NvCloth
 		/// </summary>
 		[RequiresUniqueBuildEnvironment]
-		public bool bUseChaos = false;
+		public bool bUseChaos = true;
 
 		/// <summary>
 		/// Whether to compile in checked chaos features for debugging
