#UE4 Remove SDK check from UEBuildIOS since it was causing a complication related to loading shader formats in SCW

diff --git a/Engine/Source/Programs/UnrealBuildTool/Platform/IOS/UEBuildIOS.cs b/Engine/Source/Programs/UnrealBuildTool/Platform/IOS/UEBuildIOS.cs
--- a/Engine/Source/Programs/UnrealBuildTool/Platform/IOS/UEBuildIOS.cs
+++ b/Engine/Source/Programs/UnrealBuildTool/Platform/IOS/UEBuildIOS.cs
@@ -1006,12 +1006,6 @@ namespace UnrealBuildTool
 		/// <param name="Target">The target being build</param>
 		public override void ModifyModuleRulesForOtherPlatform(string ModuleName, ModuleRules Rules, ReadOnlyTargetRules Target)
 		{
-			// don't do any target platform stuff if SDK is not available
-			if (!UEBuildPlatform.IsPlatformAvailable(Platform))
-			{
-				return;
-			}
-
 			if ((Target.Platform == UnrealTargetPlatform.Win32) || (Target.Platform == UnrealTargetPlatform.Win64) || (Target.Platform == UnrealTargetPlatform.Mac))
 			{
 				bool bBuildShaderFormats = Target.bForceBuildShaderFormats;
