Copy of CL 10625418 by [REDACTED_PERSON]

Fix missing UAT/UBT log output on Mac (fallout from CL 6849473). Standard macros set through the DefineConstants property (eg. "TRACE") were being overwritten by __MonoCS__, causing TraceLog() function calls to be compiled out.

diff --git a/Engine/Source/Programs/AutomationTool/HoloLens/HoloLens.Automation.csproj b/Engine/Source/Programs/AutomationTool/HoloLens/HoloLens.Automation.csproj
--- a/Engine/Source/Programs/AutomationTool/HoloLens/HoloLens.Automation.csproj
+++ b/Engine/Source/Programs/AutomationTool/HoloLens/HoloLens.Automation.csproj
@@ -53,7 +53,7 @@
   <Choose>
     <When Condition="'$(OS)' != 'Windows_NT'">
 	  <PropertyGroup>
-		<DefineConstants>__MonoCS__</DefineConstants>
+		<DefineConstants>$(DefineConstants);__MonoCS__</DefineConstants>
 	  </PropertyGroup>
 	</When>
   </Choose>
diff --git a/Engine/Source/Programs/UnrealBuildTool/UnrealBuildTool.csproj b/Engine/Source/Programs/UnrealBuildTool/UnrealBuildTool.csproj
--- a/Engine/Source/Programs/UnrealBuildTool/UnrealBuildTool.csproj
+++ b/Engine/Source/Programs/UnrealBuildTool/UnrealBuildTool.csproj
@@ -308,16 +308,16 @@
           <Compile Include="System\NotForLicensees\FASTBuild.cs" />
       </ItemGroup>
       <PropertyGroup>
-        <DefineConstants>__FASTBUILD_AVAILABLE__</DefineConstants>
+        <DefineConstants>$(DefineConstants);__FASTBUILD_AVAILABLE__</DefineConstants>
       </PropertyGroup>
     </When>
   </Choose>
   <Choose>
     <When Condition="'$(OS)' != 'Windows_NT'">
-      <PropertyGroup>
-        <DefineConstants>__MonoCS__</DefineConstants>
-      </PropertyGroup>
-    </When>
+	  <PropertyGroup>
+		<DefineConstants>$(DefineConstants);__MonoCS__</DefineConstants>
+	  </PropertyGroup>
+	</When>
   </Choose>
   <Choose>
     <When Condition="'$(OS)' == 'Windows_NT'">
