Rework how we generate our RPaths for Linux

Now works like this now:

Given a LibraryPath
/some/path/to/engine/root/Engine/Binaries/Linux

Split the LibraryPath to everything but the root path to the engine:
Engine/Binaries/Linux

Given the OuputFile.Location.Directory for the module being produced
/some/path/to/engine/root/Engine/Plugin/Foo/Binaries/Linux/

Make relative from the root engine dir against the OutputFile.Location.Directory to produce how many ../ are needed:
MakeRelative /some/path/to/engine/root  To -> /some/path/to/engine/root/Engine/Plugin/Foo/Binaries/Linux/ = ../../../../../

Now just combine ../../../../../ + Engine/Binaries/Linux as our new RPath. So now at runtime we can use this relative path to figure out where the *.so we depend on in that dir exists

diff --git a/Engine/Source/Programs/UnrealBuildTool/Platform/Linux/LinuxToolChain.cs b/Engine/Source/Programs/UnrealBuildTool/Platform/Linux/LinuxToolChain.cs
--- a/Engine/Source/Programs/UnrealBuildTool/Platform/Linux/LinuxToolChain.cs
+++ b/Engine/Source/Programs/UnrealBuildTool/Platform/Linux/LinuxToolChain.cs
@@ -936,7 +936,6 @@ namespace UnrealBuildTool
 			// RPATH for third party libs
 			Result += " -Wl,-rpath=${ORIGIN}";
 			Result += " -Wl,-rpath-link=${ORIGIN}";
-			Result += " -Wl,-rpath=${ORIGIN}/../../../Engine/Binaries/Linux";
 			Result += " -Wl,-rpath=${ORIGIN}/..";	// for modules that are in sub-folders of the main Engine/Binary/Linux folder
 			if (LinkEnvironment.Architecture.StartsWith("x86_64"))
 			{
@@ -947,10 +946,6 @@ namespace UnrealBuildTool
 				// x86_64 is now using updated ICU that doesn't need extra .so
 				Result += " -Wl,-rpath=${ORIGIN}/../../../Engine/Binaries/ThirdParty/ICU/icu4c-53_1/Linux/" + LinkEnvironment.Architecture;
 			}
-			Result += " -Wl,-rpath=${ORIGIN}/../../../Engine/Binaries/ThirdParty/OpenVR/OpenVRv1_5_17/linux64";
-
-			// @FIXME: Workaround for generating RPATHs for launching on devices UE-54136
-			Result += " -Wl,-rpath=${ORIGIN}/../../../Engine/Binaries/ThirdParty/PhysX3/Linux/x86_64-unknown-linux-gnu";
 
 			// Some OS ship ld with new ELF dynamic tags, which use DT_RUNPATH vs DT_RPATH. Since DT_RUNPATH do not propagate to dlopen()ed DSOs,
 			// this breaks the editor on such systems. See https://kenai.com/projects/maxine/lists/users/archive/2011-01/message/12 for details
@@ -1627,7 +1622,15 @@ namespace UnrealBuildTool
 
 				if ((AdditionalLibrary.Contains("Plugins") || AdditionalLibrary.Contains("Binaries/ThirdParty") || AdditionalLibrary.Contains("Binaries\\ThirdParty")) && Path.GetDirectoryName(AdditionalLibrary) != Path.GetDirectoryName(OutputFile.AbsolutePath))
 				{
-					string RelativePath = new FileReference(AdditionalLibrary).Directory.MakeRelativeTo(OutputFile.Location.Directory);
+					// Remove the root UnrealBuildTool.RootDirectory from the RuntimeLibaryPath
+					string AdditionalLibraryRootPath = new FileReference(AdditionalLibrary).Directory.MakeRelativeTo(UnrealBuildTool.RootDirectory);
+
+					// Figure out how many dirs we need to go back
+					string RelativeRootPath = UnrealBuildTool.RootDirectory.MakeRelativeTo(OutputFile.Location.Directory);
+
+					// Combine the two together ie. number of ../ + the path after the root
+					string RelativePath = Path.Combine(RelativeRootPath, AdditionalLibraryRootPath);
+
 					// On Windows, MakeRelativeTo can silently fail if the engine and the project are located on different drives
 					if (CrossCompiling() && RelativePath.StartsWith(UnrealBuildTool.RootDirectory.FullName))
 					{
@@ -1648,12 +1651,28 @@ namespace UnrealBuildTool
 			foreach(string RuntimeLibaryPath in LinkEnvironment.RuntimeLibraryPaths)
 			{
 				string RelativePath = RuntimeLibaryPath;
+
 				if(!RelativePath.StartsWith("$"))
 				{
-					string RelativeRootPath = new DirectoryReference(RuntimeLibaryPath).MakeRelativeTo(UnrealBuildTool.RootDirectory);
-					// We're assuming that the binary will be placed according to our ProjectName/Binaries/Platform scheme
-					RelativePath = Path.Combine("..", "..", "..", RelativeRootPath);
+					// Remove the root UnrealBuildTool.RootDirectory from the RuntimeLibaryPath
+					string RuntimeLibraryRootPath = new DirectoryReference(RuntimeLibaryPath).MakeRelativeTo(UnrealBuildTool.RootDirectory);
+
+					// Figure out how many dirs we need to go back
+					string RelativeRootPath = UnrealBuildTool.RootDirectory.MakeRelativeTo(OutputFile.Location.Directory);
+
+					// Combine the two together ie. number of ../ + the path after the root
+					RelativePath = Path.Combine(RelativeRootPath, RuntimeLibraryRootPath);
 				}
+
+				// On Windows, MakeRelativeTo can silently fail if the engine and the project are located on different drives
+				if (CrossCompiling() && RelativePath.StartsWith(UnrealBuildTool.RootDirectory.FullName))
+				{
+					// do not replace directly, but take care to avoid potential double slashes or missed slashes
+					string PathFromRootDir = RelativePath.Replace(UnrealBuildTool.RootDirectory.FullName, "");
+					// Path.Combine doesn't combine these properly
+					RelativePath = ((PathFromRootDir.StartsWith("\\") || PathFromRootDir.StartsWith("/")) ? "..\\..\\.." : "..\\..\\..\\") + PathFromRootDir;
+				}
+
 				if (!RPaths.Contains(RelativePath))
 				{
 					RPaths.Add(RelativePath);
