Do not add iOS stub to build products when building a game framework

diff --git a/Engine/Source/Programs/UnrealBuildTool/Platform/IOS/IOSToolChain.cs b/Engine/Source/Programs/UnrealBuildTool/Platform/IOS/IOSToolChain.cs
--- a/Engine/Source/Programs/UnrealBuildTool/Platform/IOS/IOSToolChain.cs
+++ b/Engine/Source/Programs/UnrealBuildTool/Platform/IOS/IOSToolChain.cs
@@ -131,18 +131,7 @@ namespace UnrealBuildTool
 
 		public override void ModifyBuildProducts(ReadOnlyTargetRules Target, UEBuildBinary Binary, List<string> Libraries, List<UEBuildBundleResource> BundleResources, Dictionary<FileReference, BuildProductType> BuildProducts)
 		{
-			if (Target.bShouldCompileAsDLL)
-			{
-				int Index = Binary.OutputFilePath.GetFileNameWithoutExtension().IndexOf("-");
-				string OutputFile = Binary.OutputFilePath.GetFileNameWithoutExtension().Substring(0, Index > 0 ? Index : Binary.OutputFilePath.GetFileNameWithoutExtension().Length);
-
-				DirectoryReference OutputDir = Binary.OutputDir;
-				if (OutputDir.ToString().Contains(".framework"))
-					OutputDir = OutputDir.ParentDirectory.ParentDirectory;
-				FileReference StubFile = FileReference.Combine(OutputDir, Binary.OutputFilePath.GetFileNameWithoutExtension() + ".stub");
-				BuildProducts.Add(StubFile, BuildProductType.Package);
-			}
-			else if (Target.IOSPlatform.bCreateStubIPA && Binary.Type != UEBuildBinaryType.StaticLibrary)
+			if (Target.IOSPlatform.bCreateStubIPA && Binary.Type != UEBuildBinaryType.StaticLibrary)
 			{
 				FileReference StubFile = FileReference.Combine(Binary.OutputFilePath.Directory, Binary.OutputFilePath.GetFileNameWithoutExtension() + ".stub");
 				BuildProducts.Add(StubFile, BuildProductType.Package);
