If we have our own toolchain make sure we surround the path with quotes for -L

diff --git a/Engine/Source/Programs/UnrealBuildTool/Platform/Linux/LinuxToolChain.cs b/Engine/Source/Programs/UnrealBuildTool/Platform/Linux/LinuxToolChain.cs
--- a/Engine/Source/Programs/UnrealBuildTool/Platform/Linux/LinuxToolChain.cs
+++ b/Engine/Source/Programs/UnrealBuildTool/Platform/Linux/LinuxToolChain.cs
@@ -1010,10 +1010,10 @@ namespace UnrealBuildTool
 				// Linking with the toolchain on linux appears to not search usr/
 				if (BuildHostPlatform.Current.Platform == UnrealTargetPlatform.Linux)
 				{
-					Result += String.Format(" -B{0}/usr/lib/", SysRootPath);
-					Result += String.Format(" -B{0}/usr/lib64/", SysRootPath);
-					Result += String.Format(" -L{0}/usr/lib/", SysRootPath);
-					Result += String.Format(" -L{0}/usr/lib64/", SysRootPath);
+					Result += String.Format(" -B\"{0}/usr/lib/\"", SysRootPath);
+					Result += String.Format(" -B\"{0}/usr/lib64/\"", SysRootPath);
+					Result += String.Format(" -L\"{0}/usr/lib/\"", SysRootPath);
+					Result += String.Format(" -L\"{0}/usr/lib64/\"", SysRootPath);
 				}
 			}
 
