UBT: Remove bAllowFASTBuild from BuildConfiguration if FASTBuild executor file is unavailable

diff --git a/Engine/Source/Programs/UnrealBuildTool/Configuration/BuildConfiguration.cs b/Engine/Source/Programs/UnrealBuildTool/Configuration/BuildConfiguration.cs
--- a/Engine/Source/Programs/UnrealBuildTool/Configuration/BuildConfiguration.cs
+++ b/Engine/Source/Programs/UnrealBuildTool/Configuration/BuildConfiguration.cs
@@ -54,12 +54,14 @@ namespace UnrealBuildTool
 		[CommandLine("-NoXGE", Value = "false")]
 		public bool bAllowXGE = true;
 
+#if __FASTBUILD_AVAILABLE__
 		/// <summary>
 		/// Whether FASTBuild may be used.
 		/// </summary>
 		[XmlConfigFile]
 		[CommandLine("-NoFASTBuild", Value = "false")]
 		public bool bAllowFASTBuild = true;
+#endif
 
 		/// <summary>
 		/// Whether SN-DBS may be used.
