Fix PVS studio warnings.

diff --git a/Engine/Source/Programs/UnrealBuildTool/Platform/HoloLens/HoloLensToolChain.cs b/Engine/Source/Programs/UnrealBuildTool/Platform/HoloLens/HoloLensToolChain.cs
--- a/Engine/Source/Programs/UnrealBuildTool/Platform/HoloLens/HoloLensToolChain.cs
+++ b/Engine/Source/Programs/UnrealBuildTool/Platform/HoloLens/HoloLensToolChain.cs
@@ -442,16 +442,16 @@ namespace UnrealBuildTool
 			// Add include paths to the argument list.
 			foreach (DirectoryReference IncludePath in CompileEnvironment.UserIncludePaths)
 			{
-				VCToolChain.AddIncludePath(SharedArguments, IncludePath, Target.HoloLensPlatform.Compiler);
+				VCToolChain.AddIncludePath(SharedArguments, IncludePath, Target.HoloLensPlatform.Compiler, CompileEnvironment.bPreprocessOnly);
 			}
 			foreach (DirectoryReference IncludePath in CompileEnvironment.SystemIncludePaths)
 			{
-				VCToolChain.AddIncludePath(SharedArguments, IncludePath, Target.HoloLensPlatform.Compiler);
+				VCToolChain.AddIncludePath(SharedArguments, IncludePath, Target.HoloLensPlatform.Compiler, CompileEnvironment.bPreprocessOnly);
 			}
 
 			foreach (DirectoryReference IncludePath in EnvVars.IncludePaths)
 			{
-				VCToolChain.AddIncludePath(SharedArguments, IncludePath, Target.HoloLensPlatform.Compiler);
+				VCToolChain.AddIncludePath(SharedArguments, IncludePath, Target.HoloLensPlatform.Compiler, CompileEnvironment.bPreprocessOnly);
 			}
 
 			// Add preprocessor definitions to the argument list.
diff --git a/Engine/Source/Programs/UnrealBuildTool/Platform/Windows/VCToolChain.cs b/Engine/Source/Programs/UnrealBuildTool/Platform/Windows/VCToolChain.cs
--- a/Engine/Source/Programs/UnrealBuildTool/Platform/Windows/VCToolChain.cs
+++ b/Engine/Source/Programs/UnrealBuildTool/Platform/Windows/VCToolChain.cs
@@ -103,11 +103,11 @@ namespace UnrealBuildTool
 		}
 
 
-		public static void AddIncludePath(List<string> Arguments, DirectoryReference IncludePath, WindowsCompiler Compiler)
+		public static void AddIncludePath(List<string> Arguments, DirectoryReference IncludePath, WindowsCompiler Compiler, bool bPreprocessOnly)
 		{
-			// If the value has a space in it and isn't wrapped in quotes, do that now. Make sure it doesn't include a trailing slash, because that will escape the closing quote.
+			// Try to use a relative path to shorten command line length. Always need the full path when preprocessing because the output file will be in a different place, where include paths cannot be relative.
 			string IncludePathString;
-			if(IncludePath.IsUnderDirectory(UnrealBuildTool.RootDirectory) && Compiler != WindowsCompiler.Clang)
+			if(IncludePath.IsUnderDirectory(UnrealBuildTool.RootDirectory) && Compiler != WindowsCompiler.Clang && !bPreprocessOnly)
 			{
 				IncludePathString = IncludePath.MakeRelativeTo(UnrealBuildTool.EngineSourceDirectory);
 			}
@@ -116,7 +116,8 @@ namespace UnrealBuildTool
 				IncludePathString = IncludePath.FullName;
 			}
 
-			if(IncludePathString.Contains(" "))
+			// If the value has a space in it and isn't wrapped in quotes, do that now. Make sure it doesn't include a trailing slash, because that will escape the closing quote.
+			if (IncludePathString.Contains(" "))
 			{
 				IncludePathString = "\"" + IncludePathString + "\"";
 			}
@@ -124,7 +125,7 @@ namespace UnrealBuildTool
 			Arguments.Add("/I " + IncludePathString);
 		}
 
-		public static void AddSystemIncludePath(List<string> Arguments, DirectoryReference IncludePath, WindowsCompiler Compiler)
+		public static void AddSystemIncludePath(List<string> Arguments, DirectoryReference IncludePath, WindowsCompiler Compiler, bool bPreprocessOnly)
 		{
 			if (Compiler == WindowsCompiler.Clang)
 			{
@@ -134,7 +135,7 @@ namespace UnrealBuildTool
 			}
 			else
 			{
-				AddIncludePath(Arguments, IncludePath, Compiler);
+				AddIncludePath(Arguments, IncludePath, Compiler, bPreprocessOnly);
 			}
 		}
 
@@ -904,17 +905,17 @@ namespace UnrealBuildTool
 			// Add include paths to the argument list.
 			foreach (DirectoryReference IncludePath in CompileEnvironment.UserIncludePaths)
 			{
-				AddIncludePath(SharedArguments, IncludePath, Target.WindowsPlatform.Compiler);
+				AddIncludePath(SharedArguments, IncludePath, Target.WindowsPlatform.Compiler, CompileEnvironment.bPreprocessOnly);
 			}
 
 			foreach (DirectoryReference IncludePath in CompileEnvironment.SystemIncludePaths)
 			{
-				AddSystemIncludePath(SharedArguments, IncludePath, Target.WindowsPlatform.Compiler);
+				AddSystemIncludePath(SharedArguments, IncludePath, Target.WindowsPlatform.Compiler, CompileEnvironment.bPreprocessOnly);
 			}
 
 			foreach (DirectoryReference IncludePath in EnvVars.IncludePaths)
 			{
-				AddSystemIncludePath(SharedArguments, IncludePath, Target.WindowsPlatform.Compiler);
+				AddSystemIncludePath(SharedArguments, IncludePath, Target.WindowsPlatform.Compiler, CompileEnvironment.bPreprocessOnly);
 			}
 
 			if (CompileEnvironment.bPrintTimingInfo || Target.WindowsPlatform.bCompilerTrace)
@@ -974,7 +975,7 @@ namespace UnrealBuildTool
 
 					// Make sure the original source directory the PCH header file existed in is added as an include
 					// path -- it might be a private PCH header and we need to make sure that its found!
-					AddIncludePath(FileArguments, SourceFile.Location.Directory, Target.WindowsPlatform.Compiler);
+					AddIncludePath(FileArguments, SourceFile.Location.Directory, Target.WindowsPlatform.Compiler, CompileEnvironment.bPreprocessOnly);
 
 					// Add the precompiled header file to the produced items list.
 					FileItem PrecompiledHeaderFile = FileItem.GetItemByFileReference(
diff --git a/Engine/Source/Runtime/Engine/Classes/Engine/NetSerialization.h b/Engine/Source/Runtime/Engine/Classes/Engine/NetSerialization.h
--- a/Engine/Source/Runtime/Engine/Classes/Engine/NetSerialization.h
+++ b/Engine/Source/Runtime/Engine/Classes/Engine/NetSerialization.h
@@ -903,7 +903,7 @@ bool FFastArraySerializer::TFastArraySerializeHelper<Type, SerializerType>::Read
 	Reader << Header.ArrayReplicationKey;
 	Reader << Header.BaseReplicationKey;
 
-	int32 NumDeletes;
+	int32 NumDeletes = 0;
 	Reader << NumDeletes;
 
 	UE_LOG(LogNetFastTArray, Log, TEXT("Received [%d/%d]."), Header.ArrayReplicationKey, Header.BaseReplicationKey);
@@ -933,7 +933,7 @@ bool FFastArraySerializer::TFastArraySerializeHelper<Type, SerializerType>::Read
 	{
 		for (int32 i = 0; i < NumDeletes; ++i)
 		{
-			int32 ElementID;
+			int32 ElementID = 0;
 			Reader << ElementID;
 
 			int32* ElementIndexPtr = ArraySerializer.ItemMap.Find(ElementID);
@@ -1338,7 +1338,7 @@ bool FFastArraySerializer::FastArrayDeltaSerialize(TArray<Type> &Items, FNetDelt
 		//---------------
 		for(int32 i = 0; i < Header.NumChanged; ++i)
 		{
-			int32 ElementID;
+			int32 ElementID = 0;
 			Reader << ElementID;
 
 			int32* ElementIndexPtr = ArraySerializer.ItemMap.Find(ElementID);
