#UE4 Disable bCompileISPC on mac for now due to build error

diff --git a/Engine/Source/Programs/UnrealBuildTool/Platform/Mac/UEBuildMac.cs b/Engine/Source/Programs/UnrealBuildTool/Platform/Mac/UEBuildMac.cs
--- a/Engine/Source/Programs/UnrealBuildTool/Platform/Mac/UEBuildMac.cs
+++ b/Engine/Source/Programs/UnrealBuildTool/Platform/Mac/UEBuildMac.cs
@@ -158,7 +158,7 @@ namespace UnrealBuildTool
 
 			Target.bCheckSystemHeadersForModification = BuildHostPlatform.Current.Platform != UnrealTargetPlatform.Mac;
 
-			Target.bCompileISPC = true;
+			//Target.bCompileISPC = true;
 		}
 
 		/// <summary>
