Support AdditionalBundleResources in the UnrealBuildTool iOS toolchain

diff --git a/Engine/Source/Programs/UnrealBuildTool/Platform/IOS/IOSToolChain.cs b/Engine/Source/Programs/UnrealBuildTool/Platform/IOS/IOSToolChain.cs
--- a/Engine/Source/Programs/UnrealBuildTool/Platform/IOS/IOSToolChain.cs
+++ b/Engine/Source/Programs/UnrealBuildTool/Platform/IOS/IOSToolChain.cs
@@ -73,6 +73,8 @@ namespace UnrealBuildTool
 
 	class IOSToolChain : AppleToolChain
 	{
+		private static List<FileItem> BundleDependencies = new List<FileItem>();
+
 		public readonly ReadOnlyTargetRules Target;
 		protected IOSProjectSettings ProjectSettings;
 
@@ -1490,6 +1492,12 @@ namespace UnrealBuildTool
 					}
 				}
 
+				FileReference StagedExecutablePath = GetStagedExecutablePath(Executable.Location, Target.Name);
+				foreach (UEBuildBundleResource Resource in BinaryLinkEnvironment.AdditionalBundleResources)
+				{
+					OutputFiles.Add(CopyBundleResource(Resource, Executable, StagedExecutablePath.Directory, Actions));
+				}
+
 				IOSPostBuildSyncTarget PostBuildSyncTarget = new IOSPostBuildSyncTarget(Target, Executable.Location, BinaryLinkEnvironment.IntermediateDirectory, UPLScripts, SdkVersion, FrameworkNameToSourceDir);
 				FileReference PostBuildSyncFile = FileReference.Combine(BinaryLinkEnvironment.IntermediateDirectory, "PostBuildSync.dat");
 				BinaryFormatterUtils.Save(PostBuildSyncFile, PostBuildSyncTarget);
@@ -2116,5 +2124,40 @@ namespace UnrealBuildTool
 		{
 			StripSymbolsWithXcode(SourceFile, TargetFile, Settings.Value.ToolchainDir);
 		}
+
+		FileItem CopyBundleResource(UEBuildBundleResource Resource, FileItem Executable, DirectoryReference BundleDirectory, List<Action> Actions)
+		{
+			Action CopyAction = new Action(ActionType.CreateAppBundle);
+			CopyAction.WorkingDirectory = GetMacDevSrcRoot(); // Path.GetFullPath(".");
+			CopyAction.CommandPath = BuildHostPlatform.Current.Shell;
+			CopyAction.CommandDescription = "";
+
+			string BundlePath = BundleDirectory.FullName;
+			string SourcePath = Path.Combine(Path.GetFullPath("."), Resource.ResourcePath);
+			string TargetPath = Path.Combine(BundlePath, Path.GetFileName(Resource.ResourcePath));
+
+			FileItem TargetItem = FileItem.GetItemByPath(TargetPath);
+
+			CopyAction.CommandArguments = string.Format("-c 'cp -f -R \"{0}\" \"{1}\"; touch -c \"{2}\"'", SourcePath, Path.GetDirectoryName(TargetPath).Replace('\\', '/') + "/", TargetPath.Replace('\\', '/'));
+			CopyAction.PrerequisiteItems.Add(Executable);
+			CopyAction.ProducedItems.Add(TargetItem);
+			CopyAction.bShouldOutputStatusDescription = Resource.bShouldLog;
+			CopyAction.StatusDescription = string.Format("Copying {0} to app bundle", Path.GetFileName(Resource.ResourcePath));
+			CopyAction.bCanExecuteRemotely = false;
+			Actions.Add(CopyAction);
+
+			return TargetItem;
+		}
+
+		public override void SetupBundleDependencies(List<UEBuildBinary> Binaries, string GameName)
+		{
+			base.SetupBundleDependencies(Binaries, GameName);
+
+			foreach (UEBuildBinary Binary in Binaries)
+			{
+				BundleDependencies.Add(FileItem.GetItemByFileReference(Binary.OutputFilePath));
+			}
+		}
+
 	};
 }
