Launchscreen.Storyboard fix for 4.24.2

diff --git a/Engine/Source/Programs/UnrealBuildTool/Platform/IOS/UEDeployIOS.cs b/Engine/Source/Programs/UnrealBuildTool/Platform/IOS/UEDeployIOS.cs
--- a/Engine/Source/Programs/UnrealBuildTool/Platform/IOS/UEDeployIOS.cs
+++ b/Engine/Source/Programs/UnrealBuildTool/Platform/IOS/UEDeployIOS.cs
@@ -209,9 +209,14 @@
 			// read the old file
 			string OldPListData = File.Exists(PListFile) ? File.ReadAllText(PListFile) : "";
 
-			// determine if there is a launch.xib
+			// determine if there is a Lanchscreen Storyboard THEN a launch.xib
+			string LaunchStoryboard = InEngineDir + "/Build/IOS/Resources/Interface/LaunchScreen.storyboard";
 			string LaunchXib = InEngineDir + "/Build/IOS/Resources/Interface/LaunchScreen.xib";
-			if (File.Exists(BuildDirectory + "/Resources/Interface/LaunchScreen.xib"))
+			if (File.Exists(BuildDirectory + "/Resources/Interface/LaunchScreen.storyboard"))
+			{
+				LaunchStoryboard = BuildDirectory + "/Resources/Interface/LaunchScreen.storyboard";
+			}
+			else if (File.Exists(BuildDirectory + "/Resources/Interface/LaunchScreen.xib"))
 			{
 				LaunchXib = BuildDirectory + "/Resources/Interface/LaunchScreen.xib";
 			}
@@ -560,9 +565,8 @@
 				Text.AppendLine("\t\t</dict>");
 				Text.AppendLine("\t</dict>");
 			}
-			if (File.Exists(LaunchXib))
+			if (File.Exists(LaunchStoryboard) || File.Exists(LaunchXib))
 			{
-				// TODO: compile the xib via remote tool
 				Text.AppendLine("\t<key>UILaunchStoryboardName</key>");
 				Text.AppendLine("\t<string>LaunchScreen</string>");
 				bSkipDefaultPNGs = true;
@@ -862,6 +866,8 @@
 			// copy engine assets in (IOS and TVOS shared in IOS)
 			if (bSkipDefaultPNGs)
 			{
+				//Copy storyboard or xib
+				CopyFiles(InEngineDir + "/Build/IOS/Resources/Interface", AppDirectory, "*", true);
 				// we still want default icons
 				if (!bSkipIcons)
 				{
@@ -1116,14 +1122,6 @@
 				DestFileInfo.Attributes = DestFileInfo.Attributes & ~FileAttributes.ReadOnly;
 			}
 
-			// compile the launch .xib
-			// @todo tvos: Is this needed for IOS, but not TVOS?
-			//			string LaunchXib = InEngineDir + "/Build/IOS/Resources/Interface/LaunchScreen.xib";
-			//			if (File.Exists(BuildDirectory + "/Resources/Interface/LaunchScreen.xib"))
-			//			{
-			//				LaunchXib = BuildDirectory + "/Resources/Interface/LaunchScreen.xib";
-			//			}
-
 			bool bSupportsPortrait = true;
 			bool bSupportsLandscape = false;
 			bool bSkipIcons = false;
