Bundle the ARM Vulkan debug layer shared lib in the APK when the project is configured to use the Mali Graphics Debugger.

Without this, vkCreateInstance() fails with VK_ERROR_LAYER_NOT_PRESENT when trying to launch the app under the debugger.

diff --git a/Engine/Source/Programs/UnrealBuildTool/Platform/Android/UEDeployAndroid.cs b/Engine/Source/Programs/UnrealBuildTool/Platform/Android/UEDeployAndroid.cs
--- a/Engine/Source/Programs/UnrealBuildTool/Platform/Android/UEDeployAndroid.cs
+++ b/Engine/Source/Programs/UnrealBuildTool/Platform/Android/UEDeployAndroid.cs
@@ -1239,7 +1239,15 @@ namespace UnrealBuildTool
 							string MaliLibDstPath = Path.Combine(UE4BuildPath, "libs", NDKArch, "libMGD.so");
 
 							Log.TraceInformation("Copying {0} to {1}", MaliLibSrcPath, MaliLibDstPath);
-							File.Copy(MaliLibSrcPath, MaliLibDstPath, true); 
+							File.Copy(MaliLibSrcPath, MaliLibDstPath, true);
+
+							string MaliVkLayerLibSrcPath = Path.Combine(MaliGraphicsDebuggerPath, "target", "android", "arm", "rooted", NDKArch, "libGLES_aga.so");
+							if (File.Exists(MaliVkLayerLibSrcPath))
+							{
+								string MaliVkLayerLibDstPath = Path.Combine(UE4BuildPath, "libs", NDKArch, "libVkLayerAGA.so");
+								Log.TraceInformation("Copying {0} to {1}", MaliVkLayerLibSrcPath, MaliVkLayerLibDstPath);
+								File.Copy(MaliVkLayerLibSrcPath, MaliVkLayerLibDstPath, true);
+							}
 						}
 					}
 					break;
