Expose an iOS Bitcode compilation via Target.IOSPlatform.bShipForBitcode in UnrealBuildTool
#rnx

diff --git a/Engine/Source/Programs/UnrealBuildTool/Platform/IOS/UEBuildIOS.cs b/Engine/Source/Programs/UnrealBuildTool/Platform/IOS/UEBuildIOS.cs
--- a/Engine/Source/Programs/UnrealBuildTool/Platform/IOS/UEBuildIOS.cs
+++ b/Engine/Source/Programs/UnrealBuildTool/Platform/IOS/UEBuildIOS.cs
@@ -24,6 +24,11 @@ namespace UnrealBuildTool
 		[CommandLine("-stripsymbols", Value = "true")]
 		public bool bStripSymbols = false;
 
+		/// <summary>
+		///
+		/// </summary>
+		public bool bShipForBitcode = false;
+
 		/// <summary>
 		/// If true, then a stub IPA will be generated when compiling is done (minimal files needed for a valid IPA).
 		/// </summary>
@@ -108,6 +113,11 @@ namespace UnrealBuildTool
 			get { return Inner.bStripSymbols; }
 		}
 			
+		public bool bShipForBitcode
+		{
+			get { return Inner.ProjectSettings.bShipForBitcode; }
+		}
+
 		public bool bGenerateFrameworkWrapperProject
 		{
 			get { return Inner.bGenerateFrameworkWrapperProject; }
