#UE4 Add compile options to disable compile/use chaos

diff --git a/Engine/Source/Programs/UnrealBuildTool/Configuration/TargetRules.cs b/Engine/Source/Programs/UnrealBuildTool/Configuration/TargetRules.cs
--- a/Engine/Source/Programs/UnrealBuildTool/Configuration/TargetRules.cs
+++ b/Engine/Source/Programs/UnrealBuildTool/Configuration/TargetRules.cs
@@ -359,12 +359,14 @@ namespace UnrealBuildTool
 		/// Whether to compile the Chaos physics plugin.
 		/// </summary>
 		[RequiresUniqueBuildEnvironment]
+		[CommandLine("-NoCompileChaos", Value = "false")]
 		public bool bCompileChaos = true;
 
 		/// <summary>
 		/// Whether to use the Chaos physics interface. This overrides the physx flags to disable APEX and NvCloth
 		/// </summary>
 		[RequiresUniqueBuildEnvironment]
+		[CommandLine("-NoUseChaos", Value = "false")]
 		public bool bUseChaos = true;
 
 		/// <summary>
