PR #5824: Adjustment to vscode code-workspace generation for foreign projects & plugin suggestions (Contributed by [REDACTED_PERSON])

diff --git a/Engine/Source/Programs/UnrealBuildTool/ProjectFiles/VisualStudioCode/VSCodeProjectFileGenerator.cs b/Engine/Source/Programs/UnrealBuildTool/ProjectFiles/VisualStudioCode/VSCodeProjectFileGenerator.cs
--- a/Engine/Source/Programs/UnrealBuildTool/ProjectFiles/VisualStudioCode/VSCodeProjectFileGenerator.cs
+++ b/Engine/Source/Programs/UnrealBuildTool/ProjectFiles/VisualStudioCode/VSCodeProjectFileGenerator.cs
@@ -1257,20 +1257,26 @@ namespace UnrealBuildTool
 			{
 				WorkspaceFile.BeginArray("folders");
 				{
-					if (bForeignProject)
+					// Add the directory in which which the code-workspace file exists.
+					// This is also known as ${workspaceRoot}
+					WorkspaceFile.BeginObject();
+					{
+						string ProjectName = bForeignProject ? GameProjectName : "UE4";
+						WorkspaceFile.AddField("name", ProjectName);
+						WorkspaceFile.AddField("path", ".");
+					}
+					WorkspaceFile.EndObject();
+
+					// If this project is outside the engine folder, add the root engine directory
+					if (bIncludeEngineSource && bForeignProject)
 					{
 						WorkspaceFile.BeginObject();
 						{
-							WorkspaceFile.AddField("path", ".");
+							WorkspaceFile.AddField("name", "UE4");
+							WorkspaceFile.AddField("path", MakeUnquotedPathString(UnrealBuildTool.RootDirectory, EPathType.Absolute));
 						}
 						WorkspaceFile.EndObject();
 					}
-
-					WorkspaceFile.BeginObject();
-					{
-						WorkspaceFile.AddField("path", MakeUnquotedPathString(UnrealBuildTool.RootDirectory, EPathType.Absolute));
-					}
-					WorkspaceFile.EndObject();
 				}
 				WorkspaceFile.EndArray();
 			}
@@ -1280,6 +1286,26 @@ namespace UnrealBuildTool
 				WorkspaceFile.AddField("typescript.tsc.autoDetect", "off");
 			}
 			WorkspaceFile.EndObject();
+			
+			WorkspaceFile.BeginObject("extensions");
+			{
+				// extensions is a set of recommended extensions that a user should install.
+				// Adding this section aids discovery of extensions which are helpful to have installed for Unreal development.
+				WorkspaceFile.BeginArray("recommendations");
+				{
+					WorkspaceFile.AddUnnamedField("ms-vscode.cpptools");
+					WorkspaceFile.AddUnnamedField("ms-vscode.csharp");
+
+					// If the platform we run the generator on uses mono, there are additional debugging extensions to add.
+					if (Utils.IsRunningOnMono)
+					{
+						WorkspaceFile.AddUnnamedField("vadimcn.vscode-lldb");
+						WorkspaceFile.AddUnnamedField("ms-vscode.mono-debug");
+					}
+				}
+				WorkspaceFile.EndArray();
+			}
+			WorkspaceFile.EndObject();
 
 			WorkspaceFile.EndRootObject();
 
