Ensure that the UHT log is generated in the same place as the UBT log so that the bulid system will gather it up
#rnx

diff --git a/Engine/Source/Programs/UnrealBuildTool/System/ExternalExecution.cs b/Engine/Source/Programs/UnrealBuildTool/System/ExternalExecution.cs
--- a/Engine/Source/Programs/UnrealBuildTool/System/ExternalExecution.cs
+++ b/Engine/Source/Programs/UnrealBuildTool/System/ExternalExecution.cs
@@ -1120,7 +1120,7 @@ namespace UnrealBuildTool
 						TargetDescriptor TargetDescriptor = new TargetDescriptor(ScriptProjectFile, "UnrealHeaderTool", Platform, Configuration, Architecture, null);
 						TargetDescriptor.bQuiet = true;
 
-						using(Timeline.ScopeEvent("Buildng UnrealHeaderTool"))
+						using(Timeline.ScopeEvent("Building UnrealHeaderTool"))
 						{
 							BuildMode.Build(new List<TargetDescriptor>{ TargetDescriptor }, BuildConfiguration, WorkingSet, BuildOptions.None, null);
 						}
@@ -1141,8 +1141,12 @@ namespace UnrealBuildTool
 					Directory.CreateDirectory(ModuleInfoFileName.Directory.FullName);
 					System.IO.File.WriteAllText(ModuleInfoFileName.FullName, fastJSON.JSON.Instance.ToJSON(Manifest, new fastJSON.JSONParameters { UseExtensions = false }));
 
+					string LogFileName = Log.OutputFile.GetFileName();
+					LogFileName = (LogFileName.StartsWith("UBT") ? "UHT" + LogFileName.Substring(3) : LogFileName + "_UHT");
+					LogFileName = FileReference.Combine(Log.OutputFile.Directory, LogFileName).ToString();
+
 					string CmdLine = (ProjectFile != null) ? "\"" + ProjectFile.FullName + "\"" : TargetName;
-					CmdLine += " \"" + ModuleInfoFileName + "\" -LogCmds=\"loginit warning, logexit warning, logdatabase error\" -Unattended -WarningsAsErrors";
+					CmdLine += " \"" + ModuleInfoFileName + "\" -LogCmds=\"loginit warning, logexit warning, logdatabase error\" -Unattended -WarningsAsErrors -abslog=\"" + LogFileName + "\"";
 					if (UnrealBuildTool.IsEngineInstalled())
 					{
 						CmdLine += " -installed";
