Fix null access by only computing a log location for UHT if the current application has an Output file (for now? maybe better fix later, but unblocking programs now)
#rnx

diff --git a/Engine/Source/Programs/UnrealBuildTool/System/ExternalExecution.cs b/Engine/Source/Programs/UnrealBuildTool/System/ExternalExecution.cs
--- a/Engine/Source/Programs/UnrealBuildTool/System/ExternalExecution.cs
+++ b/Engine/Source/Programs/UnrealBuildTool/System/ExternalExecution.cs
@@ -1141,12 +1141,18 @@ namespace UnrealBuildTool
 					Directory.CreateDirectory(ModuleInfoFileName.Directory.FullName);
 					System.IO.File.WriteAllText(ModuleInfoFileName.FullName, fastJSON.JSON.Instance.ToJSON(Manifest, new fastJSON.JSONParameters { UseExtensions = false }));
 
-					string LogFileName = Log.OutputFile.GetFileName();
-					LogFileName = (LogFileName.StartsWith("UBT") ? "UHT" + LogFileName.Substring(3) : LogFileName + "_UHT");
-					LogFileName = FileReference.Combine(Log.OutputFile.Directory, LogFileName).ToString();
-
 					string CmdLine = (ProjectFile != null) ? "\"" + ProjectFile.FullName + "\"" : TargetName;
-					CmdLine += " \"" + ModuleInfoFileName + "\" -LogCmds=\"loginit warning, logexit warning, logdatabase error\" -Unattended -WarningsAsErrors -abslog=\"" + LogFileName + "\"";
+					CmdLine += " \"" + ModuleInfoFileName + "\" -LogCmds=\"loginit warning, logexit warning, logdatabase error\" -Unattended -WarningsAsErrors";
+
+					if (Log.OutputFile != null)
+					{
+						string LogFileName = Log.OutputFile.GetFileName();
+						LogFileName = (LogFileName.StartsWith("UBT") ? "UHT" + LogFileName.Substring(3) : LogFileName + "_UHT");
+						LogFileName = FileReference.Combine(Log.OutputFile.Directory, LogFileName).ToString();
+
+						CmdLine += " -abslog =\"" + LogFileName + "\"";
+					}
+
 					if (UnrealBuildTool.IsEngineInstalled())
 					{
 						CmdLine += " -installed";
