Put the _UHT on to the UHT log when running from UBT before the .txt instead of having .txt_UHT
#rnx

diff --git a/Engine/Source/Programs/UnrealBuildTool/System/ExternalExecution.cs b/Engine/Source/Programs/UnrealBuildTool/System/ExternalExecution.cs
--- a/Engine/Source/Programs/UnrealBuildTool/System/ExternalExecution.cs
+++ b/Engine/Source/Programs/UnrealBuildTool/System/ExternalExecution.cs
@@ -1146,8 +1146,8 @@ namespace UnrealBuildTool
 
 					if (Log.OutputFile != null)
 					{
-						string LogFileName = Log.OutputFile.GetFileName();
-						LogFileName = (LogFileName.StartsWith("UBT") ? "UHT" + LogFileName.Substring(3) : LogFileName + "_UHT");
+						string LogFileName = Log.OutputFile.GetFileNameWithoutExtension();
+						LogFileName = (LogFileName.StartsWith("UBT") ? "UHT" + LogFileName.Substring(3) : LogFileName + "_UHT") + ".txt";
 						LogFileName = FileReference.Combine(Log.OutputFile.Directory, LogFileName).ToString();
 
 						CmdLine += " -abslog =\"" + LogFileName + "\"";
