- Fix ClientOnlyNoCommandlet modules not being included in builds

diff --git a/Engine/Source/Programs/UnrealBuildTool/System/ModuleDescriptor.cs b/Engine/Source/Programs/UnrealBuildTool/System/ModuleDescriptor.cs
--- a/Engine/Source/Programs/UnrealBuildTool/System/ModuleDescriptor.cs
+++ b/Engine/Source/Programs/UnrealBuildTool/System/ModuleDescriptor.cs
@@ -508,6 +508,7 @@ namespace UnrealBuildTool
                 case ModuleHostType.ServerOnly:
                     return TargetType != TargetType.Program && TargetType != TargetType.Client;
                 case ModuleHostType.ClientOnly:
+                case ModuleHostType.ClientOnlyNoCommandlet:
                     return TargetType != TargetType.Program && TargetType != TargetType.Server;
             }
 
