Add RPath workaround back, fix didnt completely handle all edge cases

diff --git a/Engine/Source/Programs/UnrealBuildTool/Platform/Linux/LinuxToolChain.cs b/Engine/Source/Programs/UnrealBuildTool/Platform/Linux/LinuxToolChain.cs
--- a/Engine/Source/Programs/UnrealBuildTool/Platform/Linux/LinuxToolChain.cs
+++ b/Engine/Source/Programs/UnrealBuildTool/Platform/Linux/LinuxToolChain.cs
@@ -947,6 +947,11 @@ namespace UnrealBuildTool
 				Result += " -Wl,-rpath=${ORIGIN}/../../../Engine/Binaries/ThirdParty/ICU/icu4c-53_1/Linux/" + LinkEnvironment.Architecture;
 			}
 
+			Result += " -Wl,-rpath=${ORIGIN}/../../../Engine/Binaries/ThirdParty/OpenVR/OpenVRv1_5_17/linux64";
+
+			// @FIXME: Workaround for generating RPATHs for launching on devices UE-54136
+			Result += " -Wl,-rpath=${ORIGIN}/../../../Engine/Binaries/ThirdParty/PhysX3/Linux/x86_64-unknown-linux-gnu";
+
 			// Some OS ship ld with new ELF dynamic tags, which use DT_RUNPATH vs DT_RPATH. Since DT_RUNPATH do not propagate to dlopen()ed DSOs,
 			// this breaks the editor on such systems. See https://kenai.com/projects/maxine/lists/users/archive/2011-01/message/12 for details
 			Result += " -Wl,--disable-new-dtags";
