Disabling Chaos in 12.00

diff --git a/Engine/Source/Programs/UnrealBuildTool/Configuration/TargetRules.cs b/Engine/Source/Programs/UnrealBuildTool/Configuration/TargetRules.cs
--- a/Engine/Source/Programs/UnrealBuildTool/Configuration/TargetRules.cs
+++ b/Engine/Source/Programs/UnrealBuildTool/Configuration/TargetRules.cs
@@ -365,14 +365,14 @@ namespace UnrealBuildTool
 		/// </summary>
 		[RequiresUniqueBuildEnvironment]
 		[CommandLine("-NoCompileChaos", Value = "false")]
-		public bool bCompileChaos = true;
+		public bool bCompileChaos = false;
 
 		/// <summary>
 		/// Whether to use the Chaos physics interface. This overrides the physx flags to disable APEX and NvCloth
 		/// </summary>
 		[RequiresUniqueBuildEnvironment]
 		[CommandLine("-NoUseChaos", Value = "false")]
-		public bool bUseChaos = true;
+		public bool bUseChaos = false;
 
 		/// <summary>
 		/// Whether to compile in checked chaos features for debugging
