Back out changelist 10983314

re-enabling chaos in 12.00 to avoid impacting downstream ... will redo it and deadend appropriately.

diff --git a/Engine/Source/Programs/UnrealBuildTool/Configuration/TargetRules.cs b/Engine/Source/Programs/UnrealBuildTool/Configuration/TargetRules.cs
--- a/Engine/Source/Programs/UnrealBuildTool/Configuration/TargetRules.cs
+++ b/Engine/Source/Programs/UnrealBuildTool/Configuration/TargetRules.cs
@@ -365,14 +365,14 @@ namespace UnrealBuildTool
 		/// </summary>
 		[RequiresUniqueBuildEnvironment]
 		[CommandLine("-NoCompileChaos", Value = "false")]
-		public bool bCompileChaos = false;
+		public bool bCompileChaos = true;
 
 		/// <summary>
 		/// Whether to use the Chaos physics interface. This overrides the physx flags to disable APEX and NvCloth
 		/// </summary>
 		[RequiresUniqueBuildEnvironment]
 		[CommandLine("-NoUseChaos", Value = "false")]
-		public bool bUseChaos = false;
+		public bool bUseChaos = true;
 
 		/// <summary>
 		/// Whether to compile in checked chaos features for debugging
