RulesAssembly.GetTargetNameByType() will only consider primary targets

diff --git a/Engine/Source/Programs/UnrealBuildTool/System/RulesAssembly.cs b/Engine/Source/Programs/UnrealBuildTool/System/RulesAssembly.cs
--- a/Engine/Source/Programs/UnrealBuildTool/System/RulesAssembly.cs
+++ b/Engine/Source/Programs/UnrealBuildTool/System/RulesAssembly.cs
@@ -638,7 +638,7 @@ namespace UnrealBuildTool
 			foreach(KeyValuePair<string, FileReference> TargetPair in TargetNameToTargetFile)
 			{
 				TargetRules Rules = CreateTargetRulesInstance(TargetPair.Key + "Target", new TargetInfo(TargetPair.Key, Platform, Configuration, Architecture, ProjectFile, null));
-				if(Rules.Type == Type)
+				if(Rules.Type == Type && Rules.IsPrimaryTarget)
 				{
 					Matches.Add(TargetPair.Key);
 				}
