Fix for compilation errors using SN-DBS

diff --git a/Engine/Source/Programs/UnrealBuildTool/Executors/SNDBS.cs b/Engine/Source/Programs/UnrealBuildTool/Executors/SNDBS.cs
--- a/Engine/Source/Programs/UnrealBuildTool/Executors/SNDBS.cs
+++ b/Engine/Source/Programs/UnrealBuildTool/Executors/SNDBS.cs
@@ -536,6 +536,12 @@ namespace UnrealBuildTool
 				IEnumerable<string> PlatformExpansions = PlatformNames.Select(p => String.Format("$1/{0}/{0}$2|$1/{0}$2", p));
 				IncludeRewriteRulesText.Add(String.Format("expansions2={0}", String.Join("|", PlatformExpansions)));
 			}
+			{
+				IncludeRewriteRulesText.Add(@"pattern3=^[A-Z]{5}_STRINGIFY\(([^/]+/[^/]+/[^/]+/)[^)]+\)");
+				IEnumerable<string> PlatformExpansions = PlatformNames.Select(p => String.Format("$1{0}/{0}.h", p));
+				IncludeRewriteRulesText.Add(String.Format("expansions3={0}", String.Join("|", PlatformExpansions)));
+			}
+
 			File.WriteAllText(IncludeRewriteRulesFile.FullName, String.Join(Environment.NewLine, IncludeRewriteRulesText));
 		}
 	}
