- add CompileChaos and -UseChaos flags

diff --git a/Engine/Source/Programs/UnrealBuildTool/Configuration/TargetRules.cs b/Engine/Source/Programs/UnrealBuildTool/Configuration/TargetRules.cs
--- a/Engine/Source/Programs/UnrealBuildTool/Configuration/TargetRules.cs
+++ b/Engine/Source/Programs/UnrealBuildTool/Configuration/TargetRules.cs
@@ -365,6 +365,7 @@ namespace UnrealBuildTool
 		/// </summary>
 		[RequiresUniqueBuildEnvironment]
 		[CommandLine("-NoCompileChaos", Value = "false")]
+		[CommandLine("-CompileChaos", Value = "true")]
 		public bool bCompileChaos = true;
 
 		/// <summary>
@@ -372,6 +373,7 @@ namespace UnrealBuildTool
 		/// </summary>
 		[RequiresUniqueBuildEnvironment]
 		[CommandLine("-NoUseChaos", Value = "false")]
+		[CommandLine("-UseChaos", Value = "true")]
 		public bool bUseChaos = true;
 
 		/// <summary>
