PR #6420: Added code to support Sign in with apple. (Contributed by [REDACTED_PERSON])

diff --git a/Engine/Source/Programs/UnrealBuildTool/Platform/IOS/IOSExports.cs b/Engine/Source/Programs/UnrealBuildTool/Platform/IOS/IOSExports.cs
--- a/Engine/Source/Programs/UnrealBuildTool/Platform/IOS/IOSExports.cs
+++ b/Engine/Source/Programs/UnrealBuildTool/Platform/IOS/IOSExports.cs
@@ -293,6 +293,18 @@ namespace UnrealBuildTool
 					Text.AppendLine("\t<key>aps-environment</key>");
 					Text.AppendLine(string.Format("\t<string>{0}</string>", bForDistribution ? "production" : "development"));
 				}
+
+				// for Sign in with Apple
+				bool bSignInWithAppleSupported = false;
+				PlatformGameConfig.GetBool("/Script/IOSRuntimeSettings.IOSRuntimeSettings", "bEnableSignInWithAppleSupport", out bSignInWithAppleSupported);
+
+				if (bSignInWithAppleSupported)
+				{
+					Text.AppendLine("\t<key>com.apple.developer.applesignin</key>");
+					Text.AppendLine("\t<array><string>Default</string></array>");
+				}
+
+				// End of entitlements
 				Text.AppendLine("</dict>");
 				Text.AppendLine("</plist>");
 
