Core: Disable type conversion warnings on Win32 for now (lots of int64->SIZE_T warnings that will need to be addressed)

diff --git a/Engine/Source/Programs/UnrealBuildTool/Platform/Windows/VCToolChain.cs b/Engine/Source/Programs/UnrealBuildTool/Platform/Windows/VCToolChain.cs
--- a/Engine/Source/Programs/UnrealBuildTool/Platform/Windows/VCToolChain.cs
+++ b/Engine/Source/Programs/UnrealBuildTool/Platform/Windows/VCToolChain.cs
@@ -510,12 +510,13 @@ namespace UnrealBuildTool
 			//@TODO: FLOATPRECISION: Consider doing the following as well:
 			//   4267: 'var' : conversion from 'size_t' to 'type', possible loss of data
 			//   4305: 'identifier' : truncation from 'type1' to 'type2'
-			if (CompileEnvironment.UnsafeTypeCastWarningLevel == WarningLevel.Error)
+			WarningLevel EffectiveCastWarningLevel = (Target.Platform == UnrealTargetPlatform.Win64) ? CompileEnvironment.UnsafeTypeCastWarningLevel : WarningLevel.Off;
+			if (EffectiveCastWarningLevel == WarningLevel.Error)
  			{
  				Arguments.Add("/we4244");
 				Arguments.Add("/we4838");
  			}
- 			else if (CompileEnvironment.UnsafeTypeCastWarningLevel == WarningLevel.Warning)
+ 			else if (EffectiveCastWarningLevel == WarningLevel.Warning)
  			{
 				// Note: The extra 4 is not a typo, /wLXXXX sets warning XXXX to level L
  				Arguments.Add("/w44244");
