enable xge for android builds

diff --git a/Engine/Source/Programs/UnrealBuildTool/Platform/Android/UEBuildAndroid.cs b/Engine/Source/Programs/UnrealBuildTool/Platform/Android/UEBuildAndroid.cs
--- a/Engine/Source/Programs/UnrealBuildTool/Platform/Android/UEBuildAndroid.cs
+++ b/Engine/Source/Programs/UnrealBuildTool/Platform/Android/UEBuildAndroid.cs
@@ -111,7 +111,7 @@ namespace UnrealBuildTool
 		public override bool CanUseXGE()
 		{
 			// Disable when building on Linux
-			return BuildHostPlatform.Current.Platform != UnrealTargetPlatform.Linux && !(Environment.GetEnvironmentVariable("IsBuildMachine") == "1");
+			return BuildHostPlatform.Current.Platform != UnrealTargetPlatform.Linux;
 		}
 
 		public override bool IsBuildProduct(string FileName, string[] NamePrefixes, string[] NameSuffixes)
