When building monolithic statically link LLVM sanitizers

Enables support for LLVMs Memory Sanitizer (-EnableMSan)

diff --git a/Engine/Source/Programs/UnrealBuildTool/Platform/Linux/LinuxToolChain.cs b/Engine/Source/Programs/UnrealBuildTool/Platform/Linux/LinuxToolChain.cs
--- a/Engine/Source/Programs/UnrealBuildTool/Platform/Linux/LinuxToolChain.cs
+++ b/Engine/Source/Programs/UnrealBuildTool/Platform/Linux/LinuxToolChain.cs
@@ -37,10 +37,20 @@ namespace UnrealBuildTool
 		/// </summary>
 		EnableUndefinedBehaviorSanitizer = 0x4,
 
+		/// <summary>
+		/// Enable memory sanitizer
+		/// </summary>
+		EnableMemorySanitizer = 0x8,
+
 		/// <summary>
 		/// Enable thin LTO
 		/// </summary>
-		EnableThinLTO = 0x8,
+		EnableThinLTO = 0x10,
+
+		/// <summary>
+		/// Enable Shared library for the Sanitizers otherwise defaults to Statically linked
+		/// </summary>
+		EnableSharedSanitizer = 0x20,
 	}
 
 	class LinuxToolChain : ISPCToolChain
@@ -547,6 +557,15 @@ namespace UnrealBuildTool
 				Result += " -fsanitize=undefined";
 			}
 
+			// MSan
+			if (Options.HasFlag(LinuxToolChainOptions.EnableMemorySanitizer))
+			{
+				// Force using the ANSI allocator if MSan is enabled
+				// -fsanitize-memory-track-origins adds a 1.5x-2.5x slow down ontop of MSan normal amount of overhead
+				// -fsanitize-memory-track-origins=1 is faster but collects only allocation points but not intermediate stores
+				Result += " -fsanitize=memory -fsanitize-memory-track-origins -DFORCE_ANSI_ALLOCATOR=1";
+			}
+
 			Result += " -Wall -Werror";
 
 			if (!CompileEnvironment.Architecture.StartsWith("x86_64") && !CompileEnvironment.Architecture.StartsWith("i686"))
@@ -910,9 +929,18 @@ namespace UnrealBuildTool
 				Result += string.Format(" -Wl,--unresolved-symbols=ignore-in-shared-libs");
 			}
 
-			if (Options.HasFlag(LinuxToolChainOptions.EnableAddressSanitizer) || Options.HasFlag(LinuxToolChainOptions.EnableThreadSanitizer) || Options.HasFlag(LinuxToolChainOptions.EnableUndefinedBehaviorSanitizer))
+			if (Options.HasFlag(LinuxToolChainOptions.EnableAddressSanitizer) ||
+				Options.HasFlag(LinuxToolChainOptions.EnableThreadSanitizer) ||
+				Options.HasFlag(LinuxToolChainOptions.EnableUndefinedBehaviorSanitizer) ||
+				Options.HasFlag(LinuxToolChainOptions.EnableMemorySanitizer))
 			{
-				Result += " -g -shared-libsan";
+				Result += " -g";
+
+				if (Options.HasFlag(LinuxToolChainOptions.EnableSharedSanitizer))
+				{
+					Result += " -shared-libsan";
+				}
+
 				if (Options.HasFlag(LinuxToolChainOptions.EnableAddressSanitizer))
 				{
 					Result += " -fsanitize=address";
@@ -925,6 +953,12 @@ namespace UnrealBuildTool
 				{
 					Result += " -fsanitize=undefined";
 				}
+				else if (Options.HasFlag(LinuxToolChainOptions.EnableMemorySanitizer))
+				{
+					// -fsanitize-memory-track-origins adds a 1.5x-2.5x slow ontop of MSan normal amount of overhead
+					// -fsanitize-memory-track-origins=1 is faster but collects only allocation points but not intermediate stores
+					Result += " -fsanitize=memory -fsanitize-memory-track-origins";
+				}
 
 				if (CrossCompiling())
 				{
@@ -1142,13 +1176,18 @@ namespace UnrealBuildTool
 				Log.TraceInformation("Using {0}", !String.IsNullOrEmpty(LlvmArPath) ? String.Format("llvm-ar : {0}", LlvmArPath) : String.Format("ar and ranlib: {0}, {1}", GetArPath(CompileEnvironment.Architecture), GetRanlibPath(CompileEnvironment.Architecture)));
 			}
 
-			if (Options.HasFlag(LinuxToolChainOptions.EnableAddressSanitizer) || Options.HasFlag(LinuxToolChainOptions.EnableThreadSanitizer) || Options.HasFlag(LinuxToolChainOptions.EnableUndefinedBehaviorSanitizer))
+			if (Options.HasFlag(LinuxToolChainOptions.EnableAddressSanitizer) ||
+				Options.HasFlag(LinuxToolChainOptions.EnableThreadSanitizer) ||
+				Options.HasFlag(LinuxToolChainOptions.EnableUndefinedBehaviorSanitizer) ||
+				Options.HasFlag(LinuxToolChainOptions.EnableMemorySanitizer))
 			{
 				string SanitizerInfo = "Building with:";
+				string StaticOrShared = Options.HasFlag(LinuxToolChainOptions.EnableSharedSanitizer) ? " dynamically" : " statically";
 
-				SanitizerInfo += Options.HasFlag(LinuxToolChainOptions.EnableAddressSanitizer) ? " AddressSanitizer" : "";
-				SanitizerInfo += Options.HasFlag(LinuxToolChainOptions.EnableThreadSanitizer) ? " ThreadSanitizer" : "";
-				SanitizerInfo += Options.HasFlag(LinuxToolChainOptions.EnableUndefinedBehaviorSanitizer) ? " UndefinedBehaviorSanitizer" : "";
+				SanitizerInfo += Options.HasFlag(LinuxToolChainOptions.EnableAddressSanitizer) ? StaticOrShared + " linked AddressSanitizer" : "";
+				SanitizerInfo += Options.HasFlag(LinuxToolChainOptions.EnableThreadSanitizer) ? StaticOrShared + " linked ThreadSanitizer" : "";
+				SanitizerInfo += Options.HasFlag(LinuxToolChainOptions.EnableUndefinedBehaviorSanitizer) ? StaticOrShared + " linked UndefinedBehaviorSanitizer" : "";
+				SanitizerInfo += Options.HasFlag(LinuxToolChainOptions.EnableMemorySanitizer) ? StaticOrShared + " linked MemorySanitizer" : "";
 
 				Log.TraceInformation(SanitizerInfo);
 			}
diff --git a/Engine/Source/Programs/UnrealBuildTool/Platform/Linux/UEBuildLinux.cs b/Engine/Source/Programs/UnrealBuildTool/Platform/Linux/UEBuildLinux.cs
--- a/Engine/Source/Programs/UnrealBuildTool/Platform/Linux/UEBuildLinux.cs
+++ b/Engine/Source/Programs/UnrealBuildTool/Platform/Linux/UEBuildLinux.cs
@@ -60,6 +60,12 @@ namespace UnrealBuildTool
 		[XmlConfigFile(Category = "BuildConfiguration", Name = "bEnableUndefinedBehaviorSanitizer")]
 		public bool bEnableUndefinedBehaviorSanitizer = false;
 
+		/// Enables memory sanitizer (MSan)
+		/// </summary>
+		[CommandLine("-EnableMSan")]
+		[XmlConfigFile(Category = "BuildConfiguration", Name = "bEnableMemorySanitizer")]
+		public bool bEnableMemorySanitizer = false;
+
 		/// <summary>
 		/// Enables "thin" LTO
 		/// </summary>
@@ -120,6 +126,11 @@ namespace UnrealBuildTool
 			get { return Inner.bEnableUndefinedBehaviorSanitizer; }
 		}
 
+		public bool bEnableMemorySanitizer
+		{
+			get { return Inner.bEnableMemorySanitizer; }
+		}
+
 		public bool bEnableThinLTO
 		{
 			get { return Inner.bEnableThinLTO; }
@@ -537,14 +548,38 @@ namespace UnrealBuildTool
 			if(Target.LinuxPlatform.bEnableAddressSanitizer)
 			{
 				Options |= LinuxToolChainOptions.EnableAddressSanitizer;
+
+				if (Target.LinkType != TargetLinkType.Monolithic)
+				{
+					Options |= LinuxToolChainOptions.EnableSharedSanitizer;
+				}
 			}
 			if(Target.LinuxPlatform.bEnableThreadSanitizer)
 			{
 				Options |= LinuxToolChainOptions.EnableThreadSanitizer;
+
+				if (Target.LinkType != TargetLinkType.Monolithic)
+				{
+					throw new BuildException("Thread Sanitizer (TSan) unsupported for non-monolithic builds");
+				}
 			}
 			if(Target.LinuxPlatform.bEnableUndefinedBehaviorSanitizer)
 			{
 				Options |= LinuxToolChainOptions.EnableUndefinedBehaviorSanitizer;
+
+				if (Target.LinkType != TargetLinkType.Monolithic)
+				{
+					Options |= LinuxToolChainOptions.EnableSharedSanitizer;
+				}
+			}
+			if(Target.LinuxPlatform.bEnableMemorySanitizer)
+			{
+				Options |= LinuxToolChainOptions.EnableMemorySanitizer;
+
+				if (Target.LinkType != TargetLinkType.Monolithic)
+				{
+					throw new BuildException("Memory Sanitizer (MSan) unsupported for non-monolithic builds");
+				}
 			}
 			if(Target.LinuxPlatform.bEnableThinLTO)
 			{
