Fix XML in UEBuildLinux.cs

diff --git a/Engine/Source/Programs/UnrealBuildTool/Platform/Linux/UEBuildLinux.cs b/Engine/Source/Programs/UnrealBuildTool/Platform/Linux/UEBuildLinux.cs
--- a/Engine/Source/Programs/UnrealBuildTool/Platform/Linux/UEBuildLinux.cs
+++ b/Engine/Source/Programs/UnrealBuildTool/Platform/Linux/UEBuildLinux.cs
@@ -60,6 +60,7 @@ namespace UnrealBuildTool
 		[XmlConfigFile(Category = "BuildConfiguration", Name = "bEnableUndefinedBehaviorSanitizer")]
 		public bool bEnableUndefinedBehaviorSanitizer = false;
 
+		/// <summary>
 		/// Enables memory sanitizer (MSan)
 		/// </summary>
 		[CommandLine("-EnableMSan")]
