Fix building of script plugins for the marketplace. Was not running the plugin with a project file that has the plugin enabled.

diff --git a/Engine/Source/Programs/AutomationTool/Scripts/BuildPluginCommand.Automation.cs b/Engine/Source/Programs/AutomationTool/Scripts/BuildPluginCommand.Automation.cs
--- a/Engine/Source/Programs/AutomationTool/Scripts/BuildPluginCommand.Automation.cs
+++ b/Engine/Source/Programs/AutomationTool/Scripts/BuildPluginCommand.Automation.cs
@@ -208,7 +208,7 @@ class BuildPlugin : BuildCommand
 				Arguments += InAdditionalArgs;
 			}
 
-			CommandUtils.RunUBT(CmdEnv, UE4Build.GetUBTExecutable(), String.Format("{0} {1} {2} {3}", TargetName, Platform, Configuration, Arguments));
+			CommandUtils.RunUBT(CmdEnv, UE4Build.GetUBTExecutable(), HostProjectFile, TargetName, Platform, Configuration, Arguments);
 		}
 	}
 
diff --git a/Engine/Source/Programs/UnrealBuildTool/System/RulesAssembly.cs b/Engine/Source/Programs/UnrealBuildTool/System/RulesAssembly.cs
--- a/Engine/Source/Programs/UnrealBuildTool/System/RulesAssembly.cs
+++ b/Engine/Source/Programs/UnrealBuildTool/System/RulesAssembly.cs
@@ -224,7 +224,10 @@ namespace UnrealBuildTool
 			{
 				Parent.GetAllModuleNames(ModuleNames);
 			}
-			ModuleNames.AddRange(CompiledAssembly.GetTypes().Where(x => x.IsClass && x.IsSubclassOf(typeof(ModuleRules)) && ModuleNameToModuleFile.ContainsKey(x.Name)).Select(x => x.Name));
+			if (CompiledAssembly != null)
+			{
+				ModuleNames.AddRange(CompiledAssembly.GetTypes().Where(x => x.IsClass && x.IsSubclassOf(typeof(ModuleRules)) && ModuleNameToModuleFile.ContainsKey(x.Name)).Select(x => x.Name));
+			}
 		}
 
 		/// <summary>
