Merging CL 10078845: Added AdditionalRootDirectories to project file. This allows projects to compile modules from source code stored in shared directories.

#rnx

diff --git a/Engine/Source/Programs/UnrealBuildTool/Configuration/UEBuildBinary.cs b/Engine/Source/Programs/UnrealBuildTool/Configuration/UEBuildBinary.cs
--- a/Engine/Source/Programs/UnrealBuildTool/Configuration/UEBuildBinary.cs
+++ b/Engine/Source/Programs/UnrealBuildTool/Configuration/UEBuildBinary.cs
@@ -587,7 +587,7 @@ namespace UnrealBuildTool
 		/// Checks whether the binary output paths are appropriate for the distribution
 		/// level of its direct module dependencies
 		/// </summary>
-		public bool CheckRestrictedFolders(DirectoryReference ProjectDir, Dictionary<UEBuildModule, Dictionary<RestrictedFolder, DirectoryReference>> ModuleRestrictedFolderCache)
+		public bool CheckRestrictedFolders(List<DirectoryReference> RootDirectories, Dictionary<UEBuildModule, Dictionary<RestrictedFolder, DirectoryReference>> ModuleRestrictedFolderCache)
 		{
 			// Find all the modules we depend on
 			Dictionary<UEBuildModule, UEBuildModule> ModuleReferencedBy = new Dictionary<UEBuildModule, UEBuildModule>();
@@ -598,16 +598,8 @@ namespace UnrealBuildTool
 			foreach (FileReference OutputFilePath in OutputFilePaths)
 			{
 				// Find the base directory for this binary
-				DirectoryReference BaseDir;
-				if(OutputFilePath.IsUnderDirectory(UnrealBuildTool.EngineDirectory))
-				{
-					BaseDir = UnrealBuildTool.EngineDirectory;
-				}
-				else if(ProjectDir != null && OutputFilePath.IsUnderDirectory(ProjectDir))
-				{
-					BaseDir = ProjectDir;
-				}
-				else
+				DirectoryReference BaseDir = RootDirectories.FirstOrDefault(x => OutputFilePath.IsUnderDirectory(x));
+				if (BaseDir == null)
 				{
 					continue;
 				}
@@ -622,7 +614,7 @@ namespace UnrealBuildTool
 					Dictionary<RestrictedFolder, DirectoryReference> ModuleRestrictedFolders;
 					if (!ModuleRestrictedFolderCache.TryGetValue(Module, out ModuleRestrictedFolders))
 					{
-						ModuleRestrictedFolders = Module.FindRestrictedFolderReferences(ProjectDir);
+						ModuleRestrictedFolders = Module.FindRestrictedFolderReferences(RootDirectories);
 						ModuleRestrictedFolderCache.Add(Module, ModuleRestrictedFolders);
 					}
 
diff --git a/Engine/Source/Programs/UnrealBuildTool/Configuration/UEBuildModule.cs b/Engine/Source/Programs/UnrealBuildTool/Configuration/UEBuildModule.cs
--- a/Engine/Source/Programs/UnrealBuildTool/Configuration/UEBuildModule.cs
+++ b/Engine/Source/Programs/UnrealBuildTool/Configuration/UEBuildModule.cs
@@ -337,9 +337,9 @@ namespace UnrealBuildTool
 		/// <summary>
 		/// Determines the distribution level of a module based on its directory and includes.
 		/// </summary>
-		/// <param name="ProjectDir">The project directory, if available</param>
+		/// <param name="RootDirectories">The set of additional paths to check, if available</param>
 		/// <returns>Map of the restricted folder types to the first found instance</returns>
-		public Dictionary<RestrictedFolder, DirectoryReference> FindRestrictedFolderReferences(DirectoryReference ProjectDir)
+		public Dictionary<RestrictedFolder, DirectoryReference> FindRestrictedFolderReferences(List<DirectoryReference> RootDirectories)
 		{
 			Dictionary<RestrictedFolder, DirectoryReference> References = new Dictionary<RestrictedFolder, DirectoryReference>();
 			if (!Rules.bLegalToDistributeObjectCode)
@@ -357,17 +357,9 @@ namespace UnrealBuildTool
 				foreach(DirectoryReference ReferencedDir in ReferencedDirs)
 				{
 					// Find the base directory containing this reference
-					DirectoryReference BaseDir;
+					DirectoryReference BaseDir = RootDirectories.FirstOrDefault(x => ReferencedDir.IsUnderDirectory(x));
 					// @todo platplug does this need to check platform extension engine directories? what are ReferencedDir's here?
-					if(ReferencedDir.IsUnderDirectory(UnrealBuildTool.EngineDirectory))
-					{
-						BaseDir = UnrealBuildTool.EngineDirectory;
-					}
-					else if(ProjectDir != null && ReferencedDir.IsUnderDirectory(ProjectDir))
-					{
-						BaseDir = ProjectDir;
-					}
-					else
+					if (BaseDir == null)
 					{
 						continue;
 					}
diff --git a/Engine/Source/Programs/UnrealBuildTool/Configuration/UEBuildTarget.cs b/Engine/Source/Programs/UnrealBuildTool/Configuration/UEBuildTarget.cs
--- a/Engine/Source/Programs/UnrealBuildTool/Configuration/UEBuildTarget.cs
+++ b/Engine/Source/Programs/UnrealBuildTool/Configuration/UEBuildTarget.cs
@@ -1872,14 +1872,26 @@ namespace UnrealBuildTool
 			if(!Rules.bDisableLinking)
 			{
 				// Check the distribution level of all binaries based on the dependencies they have
-				if(ProjectFile == null && !Rules.bLegalToDistributeBinary)
+				if (ProjectFile == null && !Rules.bLegalToDistributeBinary)
 				{
+					List<DirectoryReference> RootDirectories = new List<DirectoryReference>();
+					RootDirectories.Add(UnrealBuildTool.EngineDirectory);
+					if (ProjectFile != null)
+					{
+						DirectoryReference ProjectDir = DirectoryReference.FromFile(ProjectFile);
+						RootDirectories.Add(ProjectDir);
+						if (ProjectDescriptor != null)
+						{
+							ProjectDescriptor.AddAdditionalPaths(RootDirectories, ProjectDir);
+						}
+					}
+
 					Dictionary<UEBuildModule, Dictionary<RestrictedFolder, DirectoryReference>> ModuleRestrictedFolderCache = new Dictionary<UEBuildModule, Dictionary<RestrictedFolder, DirectoryReference>>();
 
 					bool bResult = true;
 					foreach (UEBuildBinary Binary in Binaries)
 					{
-						bResult &= Binary.CheckRestrictedFolders(DirectoryReference.FromFile(ProjectFile), ModuleRestrictedFolderCache);
+						bResult &= Binary.CheckRestrictedFolders(RootDirectories, ModuleRestrictedFolderCache);
 					}
 					foreach(KeyValuePair<FileReference, FileReference> Pair in RuntimeDependencyTargetFileToSourceFile)
 					{
diff --git a/Engine/Source/Programs/UnrealBuildTool/System/ProjectDescriptor.cs b/Engine/Source/Programs/UnrealBuildTool/System/ProjectDescriptor.cs
--- a/Engine/Source/Programs/UnrealBuildTool/System/ProjectDescriptor.cs
+++ b/Engine/Source/Programs/UnrealBuildTool/System/ProjectDescriptor.cs
@@ -81,9 +81,14 @@ namespace UnrealBuildTool
 		public PluginReferenceDescriptor[] Plugins;
 
 		/// <summary>
-        /// List of additional plugin directories to scan for available plugins
+		/// Array of additional root directories
 		/// </summary>
-        public string[] AdditionalPluginDirectories;
+		public string[] AdditionalRootDirectories;
+
+		/// <summary>
+		/// List of additional plugin directories to scan for available plugins
+		/// </summary>
+		public string[] AdditionalPluginDirectories;
 
 		/// <summary>
 		/// Array of platforms that this project is targeting
@@ -160,14 +165,17 @@ namespace UnrealBuildTool
 				Plugins = Array.ConvertAll(PluginsArray, x => PluginReferenceDescriptor.FromJsonObject(x));
 			}
 
+			// Read the additional root directories
+			RawObject.TryGetStringArrayField("AdditionalRootDirectories", out AdditionalRootDirectories);
+
 			// Read the additional plugin directories
-            RawObject.TryGetStringArrayField("AdditionalPluginDirectories", out AdditionalPluginDirectories);
+			RawObject.TryGetStringArrayField("AdditionalPluginDirectories", out AdditionalPluginDirectories);
 
-            // Read the target platforms
-            RawObject.TryGetStringArrayField("TargetPlatforms", out TargetPlatforms);
+			// Read the target platforms
+			RawObject.TryGetStringArrayField("TargetPlatforms", out TargetPlatforms);
 
-            // Get the sample name hash
-            RawObject.TryGetUnsignedIntegerField("EpicSampleNameHash", out EpicSampleNameHash);
+			// Get the sample name hash
+			RawObject.TryGetUnsignedIntegerField("EpicSampleNameHash", out EpicSampleNameHash);
 
 			// Read the pre and post-build steps
 			CustomBuildSteps.TryRead(RawObject, "PreBuildSteps", out PreBuildSteps);
@@ -200,6 +208,25 @@ namespace UnrealBuildTool
 			}
 		}
 
+		/// <summary>
+		/// If the descriptor has either additional plugin directories or additional root directories
+		/// then it is considered to have additional paths.  The additional paths will be relative
+		/// to the provided directory.
+		/// </summary>
+		/// <param name="RootDirectories">The directory set to add the paths too.</param>
+		/// <param name="ProjectDir">The directory which is used to setup the additional paths</param>
+		public void AddAdditionalPaths(List<DirectoryReference> RootDirectories, DirectoryReference ProjectDir)
+		{
+			if (AdditionalPluginDirectories != null)
+			{
+				RootDirectories.AddRange(AdditionalPluginDirectories.Select(x => DirectoryReference.Combine(ProjectDir, x)));
+			}
+			if (AdditionalRootDirectories != null)
+			{
+				RootDirectories.AddRange(AdditionalRootDirectories.Select(x => DirectoryReference.Combine(ProjectDir, x)));
+			}
+		}
+
 		/// <summary>
 		/// Saves the descriptor to disk
 		/// </summary>
@@ -237,6 +264,12 @@ namespace UnrealBuildTool
 			// Write the plugin list
 			PluginReferenceDescriptor.WriteArray(Writer, "Plugins", Plugins);
 
+			// Write the custom module roots
+			if(AdditionalRootDirectories != null && AdditionalRootDirectories.Length > 0)
+			{
+				Writer.WriteStringArrayField("AdditionalRootDirectories", AdditionalRootDirectories);
+			}
+
 			// Write out the additional plugin directories to scan
 			if(AdditionalPluginDirectories != null && AdditionalPluginDirectories.Length > 0)
 			{
diff --git a/Engine/Source/Programs/UnrealBuildTool/System/RulesCompiler.cs b/Engine/Source/Programs/UnrealBuildTool/System/RulesCompiler.cs
--- a/Engine/Source/Programs/UnrealBuildTool/System/RulesCompiler.cs
+++ b/Engine/Source/Programs/UnrealBuildTool/System/RulesCompiler.cs
@@ -460,7 +460,6 @@ namespace UnrealBuildTool
 					AddEngineModuleRulesWithContext(SourceDirectory, "Editor", DefaultModuleContext, UHTModuleType.EngineEditor, ModuleFileToContext);
 					AddEngineModuleRulesWithContext(SourceDirectory, "ThirdParty", DefaultModuleContext, UHTModuleType.EngineThirdParty, ModuleFileToContext);
 				}
-
 			}
 			// Add all the plugin modules too (don't need to loop over RootDirectories since the plugins come in already found
 			using (Timeline.ScopeEvent("Finding plugin modules"))
@@ -547,6 +546,18 @@ namespace UnrealBuildTool
 				Dictionary<FileReference, ModuleRulesContext> ModuleFiles = new Dictionary<FileReference, ModuleRulesContext>();
 				List<FileReference> TargetFiles = new List<FileReference>();
 
+				if (Project.AdditionalRootDirectories != null)
+				{
+					foreach (string AdditionalRootDirectory in Project.AdditionalRootDirectories)
+					{
+						DirectoryReference AdditionalModuleRootDirectory = DirectoryReference.Combine(MainProjectDirectory, AdditionalRootDirectory);
+						if (DirectoryReference.Exists(AdditionalModuleRootDirectory))
+						{
+							AddModuleRulesWithContext(AdditionalModuleRootDirectory, DefaultModuleContext, ModuleFiles);
+						}
+					}
+				}
+
 				// Find all the rules/plugins under the project source directories
 				foreach (DirectoryReference ProjectDirectory in UnrealBuildTool.GetAllProjectDirectories(ProjectFileName))
 				{
diff --git a/Engine/Source/Programs/UnrealBuildTool/ToolChain/RemoteMac.cs b/Engine/Source/Programs/UnrealBuildTool/ToolChain/RemoteMac.cs
--- a/Engine/Source/Programs/UnrealBuildTool/ToolChain/RemoteMac.cs
+++ b/Engine/Source/Programs/UnrealBuildTool/ToolChain/RemoteMac.cs
@@ -87,9 +87,14 @@ namespace UnrealBuildTool
 		private readonly FileReference ProjectFile;
 
 		/// <summary>
-		/// The directory containing the project being built.
+		/// The project descriptor for the project being built.
 		/// </summary>
-		private readonly DirectoryReference ProjectDirectory;
+		private readonly ProjectDescriptor ProjectDescriptor;
+
+		/// <summary>
+		/// A set of directories containing additional paths to be built.
+		/// </summary>
+		private readonly List<DirectoryReference> AdditionalPaths;
 
 		/// <summary>
 		/// The base directory on the remote machine
@@ -127,7 +132,16 @@ namespace UnrealBuildTool
 			this.RsyncExe = FileReference.Combine(UnrealBuildTool.EngineDirectory, "Extras", "ThirdPartyNotUE", "DeltaCopy", "Binaries", "Rsync.exe");
 			this.SshExe = FileReference.Combine(UnrealBuildTool.EngineDirectory, "Extras", "ThirdPartyNotUE", "DeltaCopy", "Binaries", "Ssh.exe");
 			this.ProjectFile = ProjectFile;
-			this.ProjectDirectory = DirectoryReference.FromFile(ProjectFile);
+			if (ProjectFile != null)
+			{
+				this.ProjectDescriptor = ProjectDescriptor.FromFile(ProjectFile);
+				this.AdditionalPaths = new List<DirectoryReference>();
+				this.ProjectDescriptor.AddAdditionalPaths(this.AdditionalPaths, ProjectFile.Directory);
+				if (this.AdditionalPaths.Count == 0)
+				{
+					this.AdditionalPaths = null;
+				}
+			}
 
 			// Apply settings from the XML file
 			XmlConfig.ApplyTo(this);
@@ -270,6 +284,17 @@ namespace UnrealBuildTool
 			{
 				Mappings.Add(new RemoteMapping(ProjectFile.Directory, GetRemotePath(ProjectFile.Directory)));
 			}
+			if (AdditionalPaths != null && ProjectFile != null)
+			{
+				foreach (DirectoryReference AdditionalPath in AdditionalPaths)
+				{
+					if (!AdditionalPath.IsUnderDirectory(UnrealBuildTool.EngineDirectory) &&
+						!AdditionalPath.IsUnderDirectory(ProjectFile.Directory))
+					{
+						Mappings.Add(new RemoteMapping(AdditionalPath, GetRemotePath(AdditionalPath)));
+					}
+				}
+			}
 		}
 
 		/// <summary>
@@ -899,10 +924,10 @@ namespace UnrealBuildTool
 			{
 				List<FileReference> ProjectFilters = new List<FileReference>();
 
-				FileReference CustomProjectFilter = FileReference.Combine(ProjectDir, "Build", "Rsync", "RsyncProject.txt");
-				if(FileReference.Exists(CustomProjectFilter))
+				FileReference CustomFilter = FileReference.Combine(ProjectDir, "Build", "Rsync", "RsyncProject.txt");
+				if (FileReference.Exists(CustomFilter))
 				{
-					ProjectFilters.Add(CustomProjectFilter);
+					ProjectFilters.Add(CustomFilter);
 				}
 				ProjectFilters.Add(FileReference.Combine(UnrealBuildTool.EngineDirectory, "Build", "Rsync", "RsyncProject.txt"));
 
@@ -910,6 +935,24 @@ namespace UnrealBuildTool
 				UploadDirectory(ProjectDir, GetRemotePath(ProjectDir), ProjectFilters);
 			}
 
+			if (AdditionalPaths != null)
+			{
+				foreach (DirectoryReference AdditionalPath in AdditionalPaths)
+				{
+					List<FileReference> CustomFilters = new List<FileReference>();
+
+					FileReference CustomFilter = FileReference.Combine(AdditionalPath, "Build", "Rsync", "RsyncProject.txt");
+					if (FileReference.Exists(CustomFilter))
+					{
+						CustomFilters.Add(CustomFilter);
+					}
+					CustomFilters.Add(FileReference.Combine(UnrealBuildTool.EngineDirectory, "Build", "Rsync", "RsyncProject.txt"));
+
+					Log.TraceInformation(string.Format("[Remote] Uploading additional path files [{0}]...", AdditionalPath.FullName));
+					UploadDirectory(AdditionalPath, GetRemotePath(AdditionalPath), CustomFilters);
+				}
+			}
+
 			Execute("/", String.Format("rm -rf {0}/Intermediate/IOS/*.plist", GetRemotePath(UnrealBuildTool.EngineDirectory)));
 			Execute("/", String.Format("rm -rf {0}/Intermediate/TVOS/*.plist", GetRemotePath(UnrealBuildTool.EngineDirectory)));
 			if (ProjectFile != null)
diff --git a/Engine/Source/Runtime/Projects/Private/ProjectDescriptor.cpp b/Engine/Source/Runtime/Projects/Private/ProjectDescriptor.cpp
--- a/Engine/Source/Runtime/Projects/Private/ProjectDescriptor.cpp
+++ b/Engine/Source/Runtime/Projects/Private/ProjectDescriptor.cpp
@@ -154,6 +154,30 @@ bool FProjectDescriptor::Read(const FJsonObject& Object, const FString& PathToPr
 		}
 	}
 
+#if WITH_EDITOR
+	// Read the list of additional root directories to scan
+	const TArray< TSharedPtr<FJsonValue> >* AdditionalRootDirectoriesValue;
+	if (Object.TryGetArrayField(TEXT("AdditionalRootDirectories"), AdditionalRootDirectoriesValue))
+	{
+		for (const TSharedPtr<FJsonValue>& AdditionalRootDirectoryValue : *AdditionalRootDirectoriesValue)
+		{
+			FString AdditionalDir;
+			if (AdditionalRootDirectoryValue->TryGetString(AdditionalDir))
+			{
+				if (IsRootedPath(AdditionalDir))
+				{
+					AddRootDirectory(AdditionalDir);
+				}
+				else
+				{
+					// This is path relative to the project, so convert to absolute
+					AddRootDirectory(IFileManager::Get().ConvertToAbsolutePathForExternalAppForWrite(*(PathToProject / AdditionalDir)));
+				}
+			}
+		}
+	}
+#endif
+
 	// Read the target platforms
 	const TArray< TSharedPtr<FJsonValue> > *TargetPlatformsValue;
 	if(Object.TryGetArrayField(TEXT("TargetPlatforms"), TargetPlatformsValue))
@@ -232,6 +256,18 @@ void FProjectDescriptor::Write(TJsonWriter<>& Writer, const FString& PathToProje
 		Writer.WriteArrayEnd();
 	}
 
+	// Write out the additional root directories to scan
+	if (AdditionalRootDirectories.Num() > 0)
+	{
+		Writer.WriteArrayStart(TEXT("AdditionalRootDirectories"));
+		for (const FString& Dir : AdditionalRootDirectories)
+		{
+			// Convert to relative path if possible before writing it out
+			Writer.WriteValue(MakePathRelativeToProject(Dir, PathToProject));
+		}
+		Writer.WriteArrayEnd();
+	}
+
 	// Write the target platforms
 	if(TargetPlatforms.Num() > 0)
 	{
@@ -298,4 +334,24 @@ void FProjectDescriptor::RemovePluginDirectory(const FString& Dir)
 	AdditionalPluginDirectories.RemoveSingle(Dir);
 }
 
+void FProjectDescriptor::AddRootDirectory(const FString& AdditionalDir)
+{
+	check(!AdditionalDir.StartsWith(IFileManager::Get().ConvertToAbsolutePathForExternalAppForWrite(*FPaths::EngineDir())));
+	check(!AdditionalDir.StartsWith(IFileManager::Get().ConvertToAbsolutePathForExternalAppForWrite(*FPaths::ProjectDir())));
+
+	// Detect calls where the path is not absolute
+#if WITH_EDITOR
+	checkf(IsRootedPath(AdditionalDir), TEXT("%s is not rooted"), *AdditionalDir);
+#endif
+	
+	AdditionalRootDirectories.AddUnique(AdditionalDir);
+}
+
+void FProjectDescriptor::RemoveRootDirectory(const FString& Dir)
+{
+	// Detect calls where the path is not absolute
+	checkf(IsRootedPath(Dir), TEXT("%s is not rooted"), *Dir);
+	AdditionalRootDirectories.RemoveSingle(Dir);
+}
+
 #undef LOCTEXT_NAMESPACE
diff --git a/Engine/Source/Runtime/Projects/Public/ProjectDescriptor.h b/Engine/Source/Runtime/Projects/Public/ProjectDescriptor.h
--- a/Engine/Source/Runtime/Projects/Public/ProjectDescriptor.h
+++ b/Engine/Source/Runtime/Projects/Public/ProjectDescriptor.h
@@ -143,6 +143,26 @@ struct PROJECTS_API FProjectDescriptor
 	 */
 	void RemovePluginDirectory(const FString& Dir);
 
+	/** @return - Access to the additional root directories */
+	const TArray<FString>& GetAdditionalRootDirectories() const
+	{
+		return AdditionalRootDirectories;
+	}
+
+	/**
+	 * Adds a directory to the additional root directories list. 
+	 *
+	 * @param Dir - the new directory to add
+	 */
+	void AddRootDirectory(const FString& Dir);
+	/**
+	 * Removes the directory from the list to scan
+	 *
+	 * @param Dir the directory to remove
+	 */
+	void RemoveRootDirectory(const FString& Dir);
+
+
 private:
 	/** @return the path relative to this project if possible */
 	const FString MakePathRelativeToProject(const FString& Dir, const FString& PathToProject) const;
@@ -152,4 +172,10 @@ private:
 	 * Paths are in memory as absolute paths. Conversion to/from path relative happens during Save/Load
 	 */
 	TArray<FString> AdditionalPluginDirectories;
+
+	/**
+	 * List of additional root directories to scan for modules.
+	 * Paths are in memory as absolute paths. Conversion to/from path relative happens during Save/Load
+	 */
+	TArray<FString> AdditionalRootDirectories;
 };
