Simplified PR #6317 a bit and other code review feedback.

diff --git a/Engine/Source/Programs/UnrealBuildTool/Platform/Windows/UEBuildWindows.cs b/Engine/Source/Programs/UnrealBuildTool/Platform/Windows/UEBuildWindows.cs
--- a/Engine/Source/Programs/UnrealBuildTool/Platform/Windows/UEBuildWindows.cs
+++ b/Engine/Source/Programs/UnrealBuildTool/Platform/Windows/UEBuildWindows.cs
@@ -1226,11 +1226,12 @@ namespace UnrealBuildTool
 			return false;
 		}
 
-		static string[] InstallDirRoots = {
-			"HKEY_CURRENT_USER\\SOFTWARE\\",
-			"HKEY_LOCAL_MACHINE\\SOFTWARE\\",
-			"HKEY_CURRENT_USER\\SOFTWARE\\Wow6432Node\\",
-			"HKEY_LOCAL_MACHINE\\SOFTWARE\\Wow6432Node\\"};
+		static readonly KeyValuePair<RegistryKey, string>[] InstallDirRoots = {
+			new KeyValuePair<RegistryKey, string>(Registry.CurrentUser, "SOFTWARE\\"),
+			new KeyValuePair<RegistryKey, string>(Registry.LocalMachine, "SOFTWARE\\"),
+			new KeyValuePair<RegistryKey, string>(Registry.CurrentUser, "SOFTWARE\\Wow6432Node\\"),
+			new KeyValuePair<RegistryKey, string>(Registry.LocalMachine, "SOFTWARE\\Wow6432Node\\")
+		};
 
 		/// <summary>
 		/// Reads an install directory for a 32-bit program from a registry key. This checks for per-user and machine wide settings, and under the Wow64 virtual keys (HKCU\SOFTWARE, HKLM\SOFTWARE, HKCU\SOFTWARE\Wow6432Node, HKLM\SOFTWARE\Wow6432Node).
@@ -1241,11 +1242,14 @@ namespace UnrealBuildTool
 		/// <returns>True if the key was read, false otherwise.</returns>
 		static bool TryReadInstallDirRegistryKey32(string KeySuffix, string ValueName, out DirectoryReference InstallDir)
 		{
-			foreach (string Root in InstallDirRoots)
+			foreach (KeyValuePair<RegistryKey, string> InstallRoot in InstallDirRoots)
 			{
-				if (TryReadDirRegistryKey(Root + KeySuffix, ValueName, out InstallDir))
+				using (RegistryKey Key = InstallRoot.Key.OpenSubKey(InstallRoot.Value + KeySuffix))
 				{
-					return true;
+					if (Key != null && TryReadDirRegistryKey(Key.Name, ValueName, out InstallDir))
+					{
+						return true;
+					}
 				}
 			}
 			InstallDir = null;
@@ -1261,70 +1265,24 @@ namespace UnrealBuildTool
 		static string[] ReadInstallDirSubKeys32(string KeyName)
 		{
 			HashSet<string> AllSubKeys = new HashSet<string>(StringComparer.Ordinal);
-			foreach (string Root in InstallDirRoots)
+			foreach (KeyValuePair<RegistryKey, string> Root in InstallDirRoots)
 			{
-				RegistryKey Key = OpenRegistryKey(Root + KeyName);
-				if (Key != null)
+				using (RegistryKey Key = Root.Key.OpenSubKey(Root.Value + KeyName))
 				{
-					try
-					{
-						foreach (string SubKey in Key.GetSubKeyNames())
-						{
-							AllSubKeys.Add(SubKey);
-						}
-					}
-					catch (Exception) // UnauthorizedAccessException, SecurityException, IOException
+					if (Key == null)
 					{
+						continue;
 					}
-					finally
+
+					foreach (string SubKey in Key.GetSubKeyNames())
 					{
-						Key.Dispose();
+						AllSubKeys.Add(SubKey);
 					}
 				}
 			}
 			return AllSubKeys.ToArray();
 		}
 
-		/// <summary>
-		/// Opens the given registry key using the same syntax for the key as Registry.GetValue
-		/// Returns either null or a RegistryKey.  If a RegistryKey is returned, the caller is responsible for calling .Dispose on the returned key.
-		/// </summary>
-		/// <param name="KeyName">Path to the key, e.g. HKEY_LOCAL_MACHINE\Software</param>
-		/// <returns>A RegistryKey, or null</returns>
-		static RegistryKey OpenRegistryKey(string KeyName)
-		{
-			string SubKey = null;
-			RegistryKey RootKey = null;
-			if (KeyName.StartsWith("HKEY_CURRENT_USER\\"))
-			{
-				SubKey = KeyName.Substring("HKEY_CURRENT_USER\\".Length);
-				RootKey = Registry.CurrentUser;
-			}
-			else if (KeyName.StartsWith("HKEY_LOCAL_MACHINE"))
-			{
-				SubKey = KeyName.Substring("HKEY_LOCAL_MACHINE\\".Length);
-				RootKey = Registry.LocalMachine;
-			}
-			else
-			{
-				throw new NotImplementedException($"Unrecognized root key in {KeyName}");
-			}
-			if (string.IsNullOrEmpty(SubKey))
-			{
-				// We can not support returning a root key, because our contract is that the caller must call Dispose, and Dispose should not be called on one of the root keys
-				throw new ArgumentException($"Invalid attempt to open root key {KeyName}");
-			}
-			try
-			{
-				return RootKey.OpenSubKey(SubKey);
-			}
-			catch (System.Security.SecurityException)
-			{
-				return null;
-			}
-
-		}
-
 		/// <summary>
 		/// Attempts to reads a directory name stored in a registry key
 		/// </summary>
@@ -1726,11 +1684,7 @@ namespace UnrealBuildTool
 			foreach (string ExistingVersionString in ReadInstallDirSubKeys32(NetFxSDKKeyName))
 			{
 				Version ExistingVersion;
-				try
-				{
-					ExistingVersion = new Version(ExistingVersionString);
-				}
-				catch (Exception)
+				if (!Version.TryParse(ExistingVersionString, out ExistingVersion))
 				{
 					continue;
 				}
