Use strongly typed directory reference objects for additional directories.

#rnx

diff --git a/Engine/Source/Programs/UnrealBuildTool/ProjectFiles/Xcode/XcodeProject.cs b/Engine/Source/Programs/UnrealBuildTool/ProjectFiles/Xcode/XcodeProject.cs
--- a/Engine/Source/Programs/UnrealBuildTool/ProjectFiles/Xcode/XcodeProject.cs
+++ b/Engine/Source/Programs/UnrealBuildTool/ProjectFiles/Xcode/XcodeProject.cs
@@ -742,7 +742,7 @@ namespace UnrealBuildTool
 			
 			// @todo: look also in Project/Build/Frameworks directory!
 			ProjectDescriptor Project = ProjectDescriptor.FromFile(UProjectPath);
-			List<PluginInfo> AvailablePlugins = Plugins.ReadAvailablePlugins(UnrealBuildTool.EngineDirectory, UProjectPath, Project.AdditionalPluginDirectories);
+			List<PluginInfo> AvailablePlugins = Plugins.ReadAvailablePlugins(UnrealBuildTool.EngineDirectory, DirectoryReference.FromFile(UProjectPath), Project.AdditionalPluginDirectories);
 
 			// look in each plugin for frameworks
 			// @todo: Cache this kind of things since every target will re-do this work!
diff --git a/Engine/Source/Programs/UnrealBuildTool/System/Plugins.cs b/Engine/Source/Programs/UnrealBuildTool/System/Plugins.cs
--- a/Engine/Source/Programs/UnrealBuildTool/System/Plugins.cs
+++ b/Engine/Source/Programs/UnrealBuildTool/System/Plugins.cs
@@ -189,30 +189,29 @@ namespace UnrealBuildTool
 		/// <summary>
 		/// Read all the plugins available to a given project
 		/// </summary>
-		/// <param name="EngineDirectoryName">Path to the engine directory</param>
-		/// <param name="ProjectFileName">Path to the project file (or null)</param>
-        /// <param name="AdditionalDirectories">List of additional directories to scan for available plugins</param>
+		/// <param name="EngineDir">Path to the engine directory</param>
+		/// <param name="ProjectDir">Path to the project directory (or null)</param>
+		/// <param name="AdditionalDirectories">List of additional directories to scan for available plugins</param>
 		/// <returns>Sequence of PluginInfo objects, one for each discovered plugin</returns>
-		public static List<PluginInfo> ReadAvailablePlugins(DirectoryReference EngineDirectoryName, FileReference ProjectFileName, string[] AdditionalDirectories)
+		public static List<PluginInfo> ReadAvailablePlugins(DirectoryReference EngineDir, DirectoryReference ProjectDir, List<DirectoryReference> AdditionalDirectories)
 		{
 			List<PluginInfo> Plugins = new List<PluginInfo>();
 
 			// Read all the engine plugins
-			Plugins.AddRange(ReadEnginePlugins(EngineDirectoryName));
+			Plugins.AddRange(ReadEnginePlugins(EngineDir));
 
 			// Read all the project plugins
-			if (ProjectFileName != null)
+			if (ProjectDir != null)
 			{
-				Plugins.AddRange(ReadProjectPlugins(ProjectFileName.Directory));
+				Plugins.AddRange(ReadProjectPlugins(ProjectDir));
 			}
 
             // Scan for shared plugins in project specified additional directories
 			if(AdditionalDirectories != null)
 			{
-				foreach (string AdditionalDirectory in AdditionalDirectories)
+				foreach (DirectoryReference AdditionalDirectory in AdditionalDirectories)
 				{
-					DirectoryReference DirRef = DirectoryReference.Combine(ProjectFileName.Directory, AdditionalDirectory);
-					Plugins.AddRange(ReadPluginsFromDirectory(DirRef, "", PluginType.External));
+					Plugins.AddRange(ReadPluginsFromDirectory(AdditionalDirectory, "", PluginType.External));
 				}
 			}
 
@@ -291,11 +290,10 @@ namespace UnrealBuildTool
 		/// Read all of the plugins found in the project specified additional plugin directories
 		/// </summary>
 		/// <param name="AdditionalDirectory">The additional directory to scan</param>
-		/// <param name="Subdirectory">A subdirectory to look under for AdditionalDirectory </param>
 		/// <returns>List of the found PluginInfo objects</returns>
-		public static IReadOnlyList<PluginInfo> ReadAdditionalPlugins(DirectoryReference AdditionalDirectory, string Subdirectory)
+		public static IReadOnlyList<PluginInfo> ReadAdditionalPlugins(DirectoryReference AdditionalDirectory)
 		{
-			return ReadPluginsFromDirectory(AdditionalDirectory, Subdirectory, PluginType.External);
+			return ReadPluginsFromDirectory(AdditionalDirectory, "", PluginType.External);
 		}
 
 		/// <summary>
diff --git a/Engine/Source/Programs/UnrealBuildTool/System/ProjectDescriptor.cs b/Engine/Source/Programs/UnrealBuildTool/System/ProjectDescriptor.cs
--- a/Engine/Source/Programs/UnrealBuildTool/System/ProjectDescriptor.cs
+++ b/Engine/Source/Programs/UnrealBuildTool/System/ProjectDescriptor.cs
@@ -83,12 +83,12 @@ namespace UnrealBuildTool
 		/// <summary>
 		/// Array of additional root directories
 		/// </summary>
-		public string[] AdditionalRootDirectories;
+		public List<DirectoryReference> AdditionalRootDirectories = new List<DirectoryReference>();
 
 		/// <summary>
 		/// List of additional plugin directories to scan for available plugins
 		/// </summary>
-		public string[] AdditionalPluginDirectories;
+		public List<DirectoryReference> AdditionalPluginDirectories = new List<DirectoryReference>();
 
 		/// <summary>
 		/// Array of platforms that this project is targeting
@@ -128,7 +128,8 @@ namespace UnrealBuildTool
 		/// Constructor
 		/// </summary>
 		/// <param name="RawObject">Raw JSON object to parse</param>
-		public ProjectDescriptor(JsonObject RawObject)
+		/// <param name="BaseDir">Base directory for resolving relative paths</param>
+		public ProjectDescriptor(JsonObject RawObject, DirectoryReference BaseDir)
 		{
 			// Read the version
 			if (!RawObject.TryGetIntegerField("FileVersion", out FileVersion))
@@ -166,10 +167,18 @@ namespace UnrealBuildTool
 			}
 
 			// Read the additional root directories
-			RawObject.TryGetStringArrayField("AdditionalRootDirectories", out AdditionalRootDirectories);
+			string[] RootDirectoryStrings;
+			if (RawObject.TryGetStringArrayField("AdditionalRootDirectories", out RootDirectoryStrings))
+			{
+				AdditionalRootDirectories.AddRange(RootDirectoryStrings.Select(x => DirectoryReference.Combine(BaseDir, x)));
+			}
 
 			// Read the additional plugin directories
-			RawObject.TryGetStringArrayField("AdditionalPluginDirectories", out AdditionalPluginDirectories);
+			string[] PluginDirectoryStrings;
+			if (RawObject.TryGetStringArrayField("AdditionalPluginDirectories", out PluginDirectoryStrings))
+			{
+				AdditionalPluginDirectories.AddRange(PluginDirectoryStrings.Select(x => DirectoryReference.Combine(BaseDir, x)));
+			}
 
 			// Read the target platforms
 			RawObject.TryGetStringArrayField("TargetPlatforms", out TargetPlatforms);
@@ -192,7 +201,7 @@ namespace UnrealBuildTool
 			JsonObject RawObject = JsonObject.Read(FileName);
 			try
 			{
-				ProjectDescriptor Descriptor = new ProjectDescriptor(RawObject);
+				ProjectDescriptor Descriptor = new ProjectDescriptor(RawObject, FileName.Directory);
 				if(Descriptor.Modules != null)
 				{
 					foreach (ModuleDescriptor Module in Descriptor.Modules)
@@ -217,26 +226,20 @@ namespace UnrealBuildTool
 		/// <param name="ProjectDir">The directory which is used to setup the additional paths</param>
 		public void AddAdditionalPaths(List<DirectoryReference> RootDirectories, DirectoryReference ProjectDir)
 		{
-			if (AdditionalPluginDirectories != null)
-			{
-				RootDirectories.AddRange(AdditionalPluginDirectories.Select(x => DirectoryReference.Combine(ProjectDir, x)));
-			}
-			if (AdditionalRootDirectories != null)
-			{
-				RootDirectories.AddRange(AdditionalRootDirectories.Select(x => DirectoryReference.Combine(ProjectDir, x)));
-			}
+			RootDirectories.AddRange(AdditionalRootDirectories);
+			RootDirectories.AddRange(AdditionalPluginDirectories);
 		}
 
 		/// <summary>
 		/// Saves the descriptor to disk
 		/// </summary>
 		/// <param name="FileName">The filename to write to</param>
-		public void Save(string FileName)
+		public void Save(FileReference FileName)
 		{
 			using (JsonWriter Writer = new JsonWriter(FileName))
 			{
 				Writer.WriteObjectStart();
-				Write(Writer);
+				Write(Writer, FileName.Directory);
 				Writer.WriteObjectEnd();
 			}
 		}
@@ -245,7 +248,8 @@ namespace UnrealBuildTool
 		/// Writes the plugin descriptor to an existing Json writer
 		/// </summary>
 		/// <param name="Writer">The writer to receive plugin data</param>
-		public void Write(JsonWriter Writer)
+		/// <param name="BaseDir">Base directory to save paths relative to</param>
+		public void Write(JsonWriter Writer, DirectoryReference BaseDir)
 		{
 			Writer.WriteValue("FileVersion", (int)ProjectDescriptorVersion.Latest);
 			Writer.WriteValue("EngineAssociation", EngineAssociation);
@@ -265,15 +269,15 @@ namespace UnrealBuildTool
 			PluginReferenceDescriptor.WriteArray(Writer, "Plugins", Plugins);
 
 			// Write the custom module roots
-			if(AdditionalRootDirectories != null && AdditionalRootDirectories.Length > 0)
+			if(AdditionalRootDirectories.Count > 0)
 			{
-				Writer.WriteStringArrayField("AdditionalRootDirectories", AdditionalRootDirectories);
+				Writer.WriteStringArrayField("AdditionalRootDirectories", AdditionalRootDirectories.Select(x => x.MakeRelativeTo(BaseDir).Replace(Path.DirectorySeparatorChar, '/')));
 			}
 
 			// Write out the additional plugin directories to scan
-			if(AdditionalPluginDirectories != null && AdditionalPluginDirectories.Length > 0)
+			if(AdditionalPluginDirectories.Count > 0)
 			{
-				Writer.WriteStringArrayField("AdditionalPluginDirectories", AdditionalPluginDirectories);
+				Writer.WriteStringArrayField("AdditionalPluginDirectories", AdditionalPluginDirectories.Select(x => x.MakeRelativeTo(BaseDir).Replace(Path.DirectorySeparatorChar, '/')));
 			}
 
 			// Write the target platforms
diff --git a/Engine/Source/Programs/UnrealBuildTool/System/RulesCompiler.cs b/Engine/Source/Programs/UnrealBuildTool/System/RulesCompiler.cs
--- a/Engine/Source/Programs/UnrealBuildTool/System/RulesCompiler.cs
+++ b/Engine/Source/Programs/UnrealBuildTool/System/RulesCompiler.cs
@@ -546,20 +546,15 @@ namespace UnrealBuildTool
 				Dictionary<FileReference, ModuleRulesContext> ModuleFiles = new Dictionary<FileReference, ModuleRulesContext>();
 				List<FileReference> TargetFiles = new List<FileReference>();
 
+				// Find all the project directories
+				List<DirectoryReference> ProjectDirectories = new List<DirectoryReference>(UnrealBuildTool.GetAllProjectDirectories(ProjectFileName));
 				if (Project.AdditionalRootDirectories != null)
 				{
-					foreach (string AdditionalRootDirectory in Project.AdditionalRootDirectories)
-					{
-						DirectoryReference AdditionalModuleRootDirectory = DirectoryReference.Combine(MainProjectDirectory, AdditionalRootDirectory);
-						if (DirectoryReference.Exists(AdditionalModuleRootDirectory))
-						{
-							AddModuleRulesWithContext(AdditionalModuleRootDirectory, DefaultModuleContext, ModuleFiles);
-						}
-					}
+					ProjectDirectories.AddRange(Project.AdditionalRootDirectories);
 				}
 
 				// Find all the rules/plugins under the project source directories
-				foreach (DirectoryReference ProjectDirectory in UnrealBuildTool.GetAllProjectDirectories(ProjectFileName))
+				foreach (DirectoryReference ProjectDirectory in ProjectDirectories)
 				{
 					DirectoryReference ProjectSourceDirectory = DirectoryReference.Combine(ProjectDirectory, "Source");
 
@@ -574,13 +569,12 @@ namespace UnrealBuildTool
 				// Add the project's additional plugin directories plugins too
 				if (Project.AdditionalPluginDirectories != null)
 				{
-					foreach (string AdditionalPluginDirectory in Project.AdditionalPluginDirectories)
+					foreach (DirectoryReference AdditionalPluginDirectory in Project.AdditionalPluginDirectories)
 					{
-						ProjectPlugins.AddRange(Plugins.ReadAdditionalPlugins(MainProjectDirectory, AdditionalPluginDirectory));
+						ProjectPlugins.AddRange(Plugins.ReadAdditionalPlugins(AdditionalPluginDirectory));
 					}
 				}
 
-
 				// Find all the plugin module rules
 				FindModuleRulesForPlugins(ProjectPlugins, DefaultModuleContext, ModuleFiles);
 
