Fixed issue with installed projects trying to use VSCode and not getting working intellisense for engine source code.

diff --git a/Engine/Source/Programs/UnrealBuildTool/ProjectFiles/VisualStudioCode/VSCodeProjectFileGenerator.cs b/Engine/Source/Programs/UnrealBuildTool/ProjectFiles/VisualStudioCode/VSCodeProjectFileGenerator.cs
--- a/Engine/Source/Programs/UnrealBuildTool/ProjectFiles/VisualStudioCode/VSCodeProjectFileGenerator.cs
+++ b/Engine/Source/Programs/UnrealBuildTool/ProjectFiles/VisualStudioCode/VSCodeProjectFileGenerator.cs
@@ -280,10 +280,17 @@ namespace UnrealBuildTool
 			WriteTasksFile(ProjectData);
 			WriteLaunchFile(ProjectData);
 			WriteWorkspaceSettingsFile(Projects);
-			WriteCppPropertiesFile(ProjectData);
+			WriteCppPropertiesFile(VSCodeDir, false, ProjectData);
 			WriteWorkspaceFile(ProjectData);
 			//WriteProjectDataFile(ProjectData);
 
+			if (bForeignProject && bIncludeEngineSource)
+			{
+				// for installed builds we need to write the cpp properties file under the installed engine as well for intellisense to work
+				DirectoryReference Ue4CodeDirectory = DirectoryReference.Combine(UnrealBuildTool.RootDirectory, ".vscode");
+				WriteCppPropertiesFile(Ue4CodeDirectory, true, ProjectData);
+			}
+
 			return true;
 		}
 
@@ -591,7 +598,7 @@ namespace UnrealBuildTool
 			OutFile.Write(FileReference.Combine(VSCodeDir, "unreal.json"));
 		}
 
-		private void WriteCppPropertiesFile(ProjectData Projects)
+		private void WriteCppPropertiesFile(DirectoryReference OutputDirectory, bool RewriteIncludePaths, ProjectData Projects)
 		{
 			JsonFile OutFile = new JsonFile();
 
@@ -605,9 +612,16 @@ namespace UnrealBuildTool
 
 						OutFile.BeginArray("includePath");
 						{
-							foreach (var Path in Projects.CombinedIncludePaths)
+							if (!RewriteIncludePaths)
+							{
+								foreach (var Path in Projects.CombinedIncludePaths)
+								{
+									OutFile.AddUnnamedField(Path);
+								}
+							}
+							else
 							{
-								OutFile.AddUnnamedField(Path);
+								OutFile.AddUnnamedField("${workspaceFolder}/**");
 							}
 						}
 						OutFile.EndArray();
@@ -646,7 +660,7 @@ namespace UnrealBuildTool
 			}
 			OutFile.EndRootObject();
 
-			OutFile.Write(FileReference.Combine(VSCodeDir, "c_cpp_properties.json"));
+			OutFile.Write(FileReference.Combine(OutputDirectory, "c_cpp_properties.json"));
 		}
 
 		private void WriteNativeTask(ProjectData.Project InProject, JsonFile OutFile)
