Fix issue in UBT ini parsing not combing multiple escaped newlines into a single line

As reported in UDN: https://udn.unrealengine.com/questions/540184/config-ini-parsing-line-continuation-is-incorrectl.html

diff --git a/Engine/Source/Programs/UnrealBuildTool/System/ConfigFile.cs b/Engine/Source/Programs/UnrealBuildTool/System/ConfigFile.cs
--- a/Engine/Source/Programs/UnrealBuildTool/System/ConfigFile.cs
+++ b/Engine/Source/Programs/UnrealBuildTool/System/ConfigFile.cs
@@ -165,7 +165,7 @@ namespace UnrealBuildTool
 						{
 							// Find the last non-whitespace character. If it's an escaped newline, merge the following line with it.
 							int EndIdx = Line.Length;
-							for(; EndIdx > StartIdx; EndIdx--)
+							while (EndIdx > StartIdx)
 							{
 								if(Line[EndIdx - 1] == '\\')
 								{
@@ -176,11 +176,14 @@ namespace UnrealBuildTool
 									}
 									Line += NextLine;
 									EndIdx = Line.Length;
+									continue;
 								}
 								if(Line[EndIdx - 1] != ' ' && Line[EndIdx - 1] != '\t')
 								{
 									break;
 								}
+
+								EndIdx--;
 							}
 
 							// Break out if we've got a comment
