UBT: Propagate the C++ standard version to the generated project files, to fix intellisense errors.

diff --git a/Engine/Source/Programs/UnrealBuildTool/ProjectFiles/VisualStudio/VCProject.cs b/Engine/Source/Programs/UnrealBuildTool/ProjectFiles/VisualStudio/VCProject.cs
--- a/Engine/Source/Programs/UnrealBuildTool/ProjectFiles/VisualStudio/VCProject.cs
+++ b/Engine/Source/Programs/UnrealBuildTool/ProjectFiles/VisualStudio/VCProject.cs
@@ -1431,6 +1431,19 @@ namespace UnrealBuildTool
 					VCProjectFileContent.AppendLine("    <NMakeCleanCommandLine>{0} {1}</NMakeCleanCommandLine>", EscapePath(NormalizeProjectPath(FileReference.Combine(BatchFilesDirectory, "Clean.bat"))), BuildArguments.ToString());
 					VCProjectFileContent.AppendLine("    <NMakeOutput>{0}</NMakeOutput>", NormalizeProjectPath(NMakePath.FullName));
 
+					if (TargetRulesObject.CppStandard >= CppStandardVersion.Latest)
+					{
+						VCProjectFileContent.AppendLine("    <AdditionalOptions>/std:c++latest</AdditionalOptions>");
+					}
+					else if (TargetRulesObject.CppStandard >= CppStandardVersion.Cpp17)
+					{
+						VCProjectFileContent.AppendLine("    <AdditionalOptions>/std:c++17</AdditionalOptions>");
+					}
+					else if (TargetRulesObject.CppStandard >= CppStandardVersion.Cpp14)
+					{
+						VCProjectFileContent.AppendLine("    <AdditionalOptions>/std:c++14</AdditionalOptions>");
+					}
+
 					if (TargetRulesObject.Type == TargetType.Game || TargetRulesObject.Type == TargetType.Client || TargetRulesObject.Type == TargetType.Server)
 					{
 						// Allow platforms to add any special properties they require... like aumid override for Xbox One
