Handle situations where targets in foreign projects may not get the correct project file name.

#rnx

diff --git a/Engine/Source/Programs/UnrealBuildTool/ProjectFiles/VisualStudio/VCProject.cs b/Engine/Source/Programs/UnrealBuildTool/ProjectFiles/VisualStudio/VCProject.cs
--- a/Engine/Source/Programs/UnrealBuildTool/ProjectFiles/VisualStudio/VCProject.cs
+++ b/Engine/Source/Programs/UnrealBuildTool/ProjectFiles/VisualStudio/VCProject.cs
@@ -1239,7 +1239,15 @@ namespace UnrealBuildTool
 			string UProjectPath = "";
 			if (IsForeignProject)
 			{
-				UProjectPath = "\"$(SolutionDir)$(ProjectName).uproject\"";
+				FileReference ProjectFile = Combination.ProjectTarget.UnrealProjectFilePath;
+				if (ProjectFile.IsUnderDirectory(ProjectFileGenerator.MasterProjectPath))
+				{
+					UProjectPath = String.Format("\"$(SolutionDir){0}\"", ProjectFile.MakeRelativeTo(ProjectFileGenerator.MasterProjectPath));
+				}
+				else
+				{
+					UProjectPath = String.Format("\"{0}\"", ProjectFile);
+				}
 			}
 
 			string ConditionString = "Condition=\"'$(Configuration)|$(Platform)'=='" + Combination.ProjectConfigurationAndPlatformName + "'\"";
