Copy of CL 10625418 by [REDACTED_PERSON]

Fix missing UAT/UBT log output on Mac (fallout from CL 6849473). Standard macros set through the DefineConstants property (eg. "TRACE") were being overwritten by __MonoCS__, causing TraceLog() function calls to be compiled out.

diff --git a/Engine/Source/Programs/UnrealBuildTool/UnrealBuildTool.csproj b/Engine/Source/Programs/UnrealBuildTool/UnrealBuildTool.csproj
--- a/Engine/Source/Programs/UnrealBuildTool/UnrealBuildTool.csproj
+++ b/Engine/Source/Programs/UnrealBuildTool/UnrealBuildTool.csproj
@@ -308,7 +308,7 @@
           <Compile Include="System\NotForLicensees\FASTBuild.cs" />
       </ItemGroup>
       <PropertyGroup>
-        <DefineConstants>__FASTBUILD_AVAILABLE__</DefineConstants>
+        <DefineConstants>$(DefineConstants);__FASTBUILD_AVAILABLE__</DefineConstants>
       </PropertyGroup>
     </When>
   </Choose>
