Fix reading ISPC dependency lists with paths that contain spaces. Assumes a very fixed file format of one file per line rather than full make syntax.

#rnx

diff --git a/Engine/Source/Programs/UnrealBuildTool/System/CppDependencyCache.cs b/Engine/Source/Programs/UnrealBuildTool/System/CppDependencyCache.cs
--- a/Engine/Source/Programs/UnrealBuildTool/System/CppDependencyCache.cs
+++ b/Engine/Source/Programs/UnrealBuildTool/System/CppDependencyCache.cs
@@ -311,41 +311,40 @@ namespace UnrealBuildTool
 			{
 				string Text = FileReference.ReadAllText(InputFile);
 
-				List<string> Tokens = new List<string>();
-
-				StringBuilder Token = new StringBuilder();
-				for(int Idx = 0; TryReadMakefileToken(Text, ref Idx, Token); )
-				{
-					Tokens.Add(Token.ToString());
-				}
+				List<FileItem> NewDependencyFiles = new List<FileItem>();
 
-				int TokenIdx = 0;
-				while(TokenIdx < Tokens.Count && Tokens[TokenIdx] == "\n")
+				// Currently expect one file per line
+				int Index = Text.IndexOf(':');
+				if(Index != -1)
 				{
-					TokenIdx++;
-				}
+					Index++;
+					for (; ; )
+					{
+						int EndIndex = Text.IndexOf('\n', Index);
+						if(EndIndex == -1)
+						{
+							break;
+						}
 
-				if(TokenIdx + 1 >= Tokens.Count || Tokens[TokenIdx + 1] != ":")
-				{
-					throw new BuildException("Unable to parse dependency file");
-				}
+						int FileEndIndex = EndIndex;
+						if(FileEndIndex > Index && Text[FileEndIndex - 1] == '\r')
+						{
+							FileEndIndex--;
+						}
+						if(FileEndIndex > Index && Text[FileEndIndex - 1] == '\\')
+						{
+							FileEndIndex--;
+						}
 
-				TokenIdx += 2;
-			
-				List<FileItem> NewDependencyFiles = new List<FileItem>();
-				for(; TokenIdx < Tokens.Count && Tokens[TokenIdx] != "\n"; TokenIdx++)
-				{
-					NewDependencyFiles.Add(FileItem.GetItemByPath(Tokens[TokenIdx]));
-				}
+						string File = Text.Substring(Index, FileEndIndex - Index).Trim();
+						if (File.Length > 0)
+						{
+							NewDependencyFiles.Add(FileItem.GetItemByPath(File));
+						}
 
-				while(TokenIdx < Tokens.Count && Tokens[TokenIdx] == "\n")
-				{
-					TokenIdx++;
-				}
+						Index = EndIndex + 1;
+					}
 
-				if(TokenIdx != Tokens.Count)
-				{
-					throw new BuildException("Unable to parse dependency file");
 				}
 
 				return NewDependencyFiles;
@@ -355,87 +354,5 @@ namespace UnrealBuildTool
 				throw new BuildException("Unknown dependency list file type: {0}", InputFile);
 			}
 		}
-
-		/// <summary>
-		/// Attempts to read a single token from a makefile
-		/// </summary>
-		/// <param name="Text">Text to read from</param>
-		/// <param name="RefIdx">Current position within the file</param>
-		/// <param name="Token">Receives the token characters</param>
-		/// <returns>True if a token was read, false if the end of the buffer was reached</returns>
-		static bool TryReadMakefileToken(string Text, ref int RefIdx, StringBuilder Token)
-		{
-			Token.Clear();
-
-			int Idx = RefIdx;
-			for(;;)
-			{
-				if(Idx == Text.Length)
-				{
-					return false;
-				}
-
-				// Skip whitespace
-				while(Text[Idx] == ' ' || Text[Idx] == '\t')
-				{
-					if(++Idx == Text.Length)
-					{
-						return false;
-					}
-				}
-
-				// Colon token
-				if(Text[Idx] == ':')
-				{
-					Token.Append(':');
-					RefIdx = Idx + 1;
-					return true;
-				}
-
-				// Check for a newline
-				if(Text[Idx] == '\r' || Text[Idx] == '\n')
-				{
-					Token.Append('\n');
-					RefIdx = Idx + 1;
-					return true;
-				}
-
-				// Check for an escaped newline
-				if(Text[Idx] == '\\' && Idx + 1 < Text.Length)
-				{
-					if(Text[Idx + 1] == '\n')
-					{
-						Idx += 2;
-						continue;
-					}
-					if(Text[Idx + 1] == '\r' && Idx + 2 < Text.Length && Text[Idx + 2] == '\n')
-					{
-						Idx += 3;
-						continue;
-					}
-				}
-
-				// Read a token. Special handling for drive letters on Windows!
-				for(; Idx < Text.Length; Idx++)
-				{
-					if(Text[Idx] == ' ' || Text[Idx] == '\t' || Text[Idx] == '\r' || Text[Idx] == '\n')
-					{
-						break;
-					}
-					if(Text[Idx] == ':' && Token.Length > 1)
-					{
-						break;
-					}
-					if(Text[Idx] == '\\' && Idx + 1 < Text.Length && Text[Idx + 1] == ' ')
-					{
-						Idx++;
-					}
-					Token.Append(Text[Idx]);
-				}
-
-				RefIdx = Idx;
-				return true;
-			}
-		}
 	}
 }
