UBT: Ensure that when a plugin is disabled explicitly for a target, it's not built due to the bBuildAllModules flag being set.

#rnx

diff --git a/Engine/Source/Programs/UnrealBuildTool/Configuration/UEBuildTarget.cs b/Engine/Source/Programs/UnrealBuildTool/Configuration/UEBuildTarget.cs
--- a/Engine/Source/Programs/UnrealBuildTool/Configuration/UEBuildTarget.cs
+++ b/Engine/Source/Programs/UnrealBuildTool/Configuration/UEBuildTarget.cs
@@ -2699,8 +2699,14 @@ namespace UnrealBuildTool
 			List<PluginInfo> Plugins = RulesAssembly.EnumeratePlugins().ToList();
 			foreach(PluginInfo Plugin in Plugins)
 			{
+				// Ignore plugins which are specifically disabled by this target
+				if (Rules.DisablePlugins.Contains(Plugin.Name))
+				{
+					continue;
+				}
+
 				// Ignore plugins without any modules
-				if(Plugin.Descriptor.Modules == null)
+				if (Plugin.Descriptor.Modules == null)
 				{
 					continue;
 				}
