Switching ISPC to using response files.

#rnx

diff --git a/Engine/Source/Programs/UnrealBuildTool/ToolChain/UEToolChain.cs b/Engine/Source/Programs/UnrealBuildTool/ToolChain/UEToolChain.cs
--- a/Engine/Source/Programs/UnrealBuildTool/ToolChain/UEToolChain.cs
+++ b/Engine/Source/Programs/UnrealBuildTool/ToolChain/UEToolChain.cs
@@ -2,10 +2,9 @@
 
 using System;
 using System.Collections.Generic;
-using System.Text.RegularExpressions;
-using System.Diagnostics;
 using System.IO;
-using Microsoft.Win32;
+using System.Linq;
+using System.Text.RegularExpressions;
 using Tools.DotNETCommon;
 
 namespace UnrealBuildTool
@@ -312,7 +311,7 @@ namespace UnrealBuildTool
 				List<string> Arguments = new List<string>();
 
 				// Add the ISPC obj file as a prerequisite of the action.
-				Arguments.Add(String.Format(" \"{0}\"", ISPCFile.AbsolutePath));
+				CompileAction.CommandArguments = String.Format("\"{0}\" ", ISPCFile.AbsolutePath);
 
 				// Add the ISPC h file to the produced item list.
 				FileItem ISPCIncludeHeaderFile = FileItem.GetItemByFileReference(
@@ -340,9 +339,9 @@ namespace UnrealBuildTool
 				}
 
 				// Build target triplet
-				Arguments.Add(String.Format("--target-os=\"{0}\"", GetISPCOSTarget(CompileEnvironment.Platform)));
-				Arguments.Add(String.Format("--arch=\"{0}\"", GetISPCArchTarget(CompileEnvironment.Platform, null)));
-				Arguments.Add(String.Format("--target=\"{0}\"", TargetString));
+				Arguments.Add(String.Format("--target-os={0}", GetISPCOSTarget(CompileEnvironment.Platform)));
+				Arguments.Add(String.Format("--arch={0}", GetISPCArchTarget(CompileEnvironment.Platform, null)));
+				Arguments.Add(String.Format("--target={0}", TargetString));
 
 				// PIC is needed for modular builds except on Windows
 				if ((CompileEnvironment.bIsBuildingDLL ||
@@ -354,6 +353,7 @@ namespace UnrealBuildTool
 				}
 
 				// Include paths. Don't use AddIncludePath() here, since it uses the full path and exceeds the max command line length.
+				// Because ISPC response files don't support white space in arguments, paths with white space need to be passed to the command line directly.
 				foreach (DirectoryReference IncludePath in CompileEnvironment.UserIncludePaths)
 				{
 					Arguments.Add(String.Format("-I\"{0}\"", IncludePath));
@@ -376,7 +376,10 @@ namespace UnrealBuildTool
 
 				CompileAction.ProducedItems.Add(ISPCIncludeHeaderFile);
 
-				CompileAction.CommandArguments = String.Join(" ", Arguments);
+				FileReference ResponseFileName = new FileReference(ISPCIncludeHeaderFile.AbsolutePath + ".response");
+				FileItem ResponseFileItem = FileItem.CreateIntermediateTextFile(ResponseFileName, Arguments.Select(x => Utils.ExpandVariables(x)));
+				CompileAction.CommandArguments += String.Format("@\"{0}\"", ResponseFileName);
+				CompileAction.PrerequisiteItems.Add(ResponseFileItem);
 
 				// Add the source file and its included files to the prerequisite item list.
 				CompileAction.PrerequisiteItems.Add(ISPCFile);
@@ -564,7 +567,10 @@ namespace UnrealBuildTool
 				CompileAction.ProducedItems.AddRange(CompiledISPCObjFiles);
 				Result.ObjectFiles.AddRange(CompiledISPCObjFiles);
 
-				CompileAction.CommandArguments = String.Join(" ", Arguments);
+				FileReference ResponseFileName = new FileReference(CompiledISPCObjFileNoISA.AbsolutePath + ".response");
+				FileItem ResponseFileItem = FileItem.CreateIntermediateTextFile(ResponseFileName, Arguments.Select(x => Utils.ExpandVariables(x)));
+				CompileAction.CommandArguments = " @\"" + ResponseFileName + "\"";
+				CompileAction.PrerequisiteItems.Add(ResponseFileItem);
 
 				// Add the source file and its included files to the prerequisite item list.
 				CompileAction.PrerequisiteItems.Add(ISPCFile);
