Added some more smarts to how UBT resolves the project argument when generating files. This makes it easier to generate project files manually as the behavior is more inline with other commands that can just take a project

Look for a candidate file either based on the -project argument or if there is an argument that ends in .uproject. Then check if the candidate exists (e.g if it's a full path),  if not use the ueprojectdirs entries to see if it exists in one of those directories.

(Note - it would be good to reconcile this code with some things in the Automation csproj that do the same thing, but there wasn't an obvious project for that to live in...).

diff --git a/Engine/Source/Programs/UnrealBuildTool/Modes/GenerateProjectFilesMode.cs b/Engine/Source/Programs/UnrealBuildTool/Modes/GenerateProjectFilesMode.cs
--- a/Engine/Source/Programs/UnrealBuildTool/Modes/GenerateProjectFilesMode.cs
+++ b/Engine/Source/Programs/UnrealBuildTool/Modes/GenerateProjectFilesMode.cs
@@ -191,23 +191,55 @@ namespace UnrealBuildTool
 		/// <returns>True if the project file was parsed, false otherwise</returns>
 		private static bool TryParseProjectFileArgument(CommandLineArguments Arguments, out FileReference ProjectFile)
 		{
-			FileReference ExplicitProjectFile;
-			if(Arguments.TryGetValue("-Project=", out ExplicitProjectFile))
+			string CandidateProjectPath = null;
+
+			// look for -project=<path>, if it does not exist check arguments for anything that has .uproject in it
+			if (!Arguments.TryGetValue("-Project=", out CandidateProjectPath))
 			{
-				ProjectFile = ExplicitProjectFile;
-				return true;
+				// Go through the argument list and try to match poorly (or well..) formed arguments like 
+				// EngineTest, EngineTest.uproject
+				// Collaboration/FooProject
+				// by checking for those in the native project list
+				for (int Idx = 0; Idx < Arguments.Count; Idx++)
+				{
+					if (Arguments[Idx][0] != '-' && Arguments[Idx].EndsWith(".uproject", StringComparison.OrdinalIgnoreCase))
+					{
+						CandidateProjectPath = Arguments[Idx];
+						Arguments.MarkAsUsed(Idx);
+						break; 
+					}
+				}
 			}
 
-			for(int Idx = 0; Idx < Arguments.Count; Idx++)
+			// We have a project file either via -project= or because there was something called .uproject in the arg list
+			// so now validate it
+			if (!string.IsNullOrEmpty(CandidateProjectPath))
 			{
-				if(Arguments[Idx][0] != '-' && Arguments[Idx].EndsWith(".uproject", StringComparison.OrdinalIgnoreCase))
+				FileReference CandidateProjectFile = new FileReference(CandidateProjectPath);
+
+				// if the path doesn't exist then check native paths (ueprojectdirs)
+				if (!FileReference.Exists(CandidateProjectFile))
 				{
-					Arguments.MarkAsUsed(Idx);
-					ProjectFile = new FileReference(Arguments[Idx]);
-					return true;
+					// clean everything the user provided to just the name and make sure it has the expected extension
+					string ProjectName = CandidateProjectFile.ChangeExtension("uproject").GetFileName();
+
+					// check native project paths (uprojectdirs)
+					IEnumerable<FileReference> NativeProjectFiles = NativeProjects.EnumerateProjectFiles();
+
+					CandidateProjectFile = NativeProjectFiles.Where(F => F.GetFileName().Equals(ProjectName, StringComparison.InvariantCultureIgnoreCase)).FirstOrDefault();
 				}
-			}
 
+				if (CandidateProjectFile == null || !FileReference.Exists(CandidateProjectFile))
+				{
+					// if we didn't find anything then throw an error as the user explicitly provided a uproject
+					throw new Exception(string.Format("Unable to find project file based on argument {0}", CandidateProjectPath));
+				}
+
+				Log.TraceVerbose("Resolved project argument {0} to {1}", CandidateProjectPath, CandidateProjectFile);
+				ProjectFile = CandidateProjectFile;
+				return true;
+			}
+			
 			FileReference InstalledProjectFile = UnrealBuildTool.GetInstalledProjectFile();
 			if(InstalledProjectFile != null)
 			{
