Use VS2019 as the default compiler and IDE on Windows.

- Download and run the VS2019 installer from the editor
- Update the preferred VC++ toolchains to choose a VS2019 toolchain if possible
- Update the AppLocalPrerequisites folder to include redist files from VS2019
- Update the UE4 prerequisites installer to include VS2019 redist
- Update the installed engine build script to build with VS2017 for backwards compatibility, since VS2019 object files are not compatible with it (only vice versa).

diff --git a/Engine/Build/InstalledEngineBuild.xml b/Engine/Build/InstalledEngineBuild.xml
--- a/Engine/Build/InstalledEngineBuild.xml
+++ b/Engine/Build/InstalledEngineBuild.xml
@@ -104,7 +104,7 @@
 	<Property Name="TargetDebugInfoArg" Value="" If="$(WithFullDebugInfo)"/>
 
 	<!-- Compile flags for Windows targets -->
-	<Property Name="VSCompilerArg" Value=""/>
+	<Property Name="VSCompilerArg" Value="-2017"/>
 
 	<!-- Names of the client and server targets to build -->
 	<Property Name="OptionalClientTarget" Value=""/>
diff --git a/Engine/Config/BaseEditor.ini b/Engine/Config/BaseEditor.ini
--- a/Engine/Config/BaseEditor.ini
+++ b/Engine/Config/BaseEditor.ini
@@ -311,7 +311,7 @@ DocumentationURL="http://udn.unrealengine.com/docs/ue4/INT/index.html"
 AskAQuestionURL="http://udn.unrealengine.com/questions/ask.html"
 SearchForAnswersURL="http://udn.unrealengine.com/search.html"
 SupportWebsiteURL="http://udn.unrealengine.com/index.html"
-SourceCodeIDEURL_Windows="https://aka.ms/vs/15/release/vs_Community.exe"
+SourceCodeIDEURL_Windows="https://aka.ms/vs/16/release/vs_Community.exe"
 SourceCodeIDEURL_Mac="https://developer.apple.com/xcode"
 SourceCodeIDEURL_Other="http://en.wikipedia.org/wiki/Integrated_development_environment"
 APIDocsURL="https://docs.unrealengine.com/latest/INT/API/index.html"
diff --git a/Engine/Source/Programs/UnrealBuildTool/Platform/Windows/UEBuildWindows.cs b/Engine/Source/Programs/UnrealBuildTool/Platform/Windows/UEBuildWindows.cs
--- a/Engine/Source/Programs/UnrealBuildTool/Platform/Windows/UEBuildWindows.cs
+++ b/Engine/Source/Programs/UnrealBuildTool/Platform/Windows/UEBuildWindows.cs
@@ -574,6 +574,8 @@ namespace UnrealBuildTool
 		/// </summary>
 		static readonly VersionNumber[] PreferredVisualStudioToolChainVersion = new VersionNumber[]
 		{
+			VersionNumber.Parse("14.24.28315"), // VS2019 v16.4.3 (installed to 14.24.28314 folder)
+			VersionNumber.Parse("14.22.27905"), // VS2019 v16.2.3
 			VersionNumber.Parse("14.16.27023.2"), // VS2017 v15.9.15
 			VersionNumber.Parse("14.16.27023"), // fallback to VS2017 15.9 toolchain, microsoft updates these in places so for local installs only this version number is present
 		};
@@ -797,32 +799,32 @@ namespace UnrealBuildTool
 			}
 
 			// Second, default based on what's installed, test for 2015 first
-			if (HasCompiler(WindowsCompiler.VisualStudio2017))
-			{
-				return WindowsCompiler.VisualStudio2017;
-			}
 			if (HasCompiler(WindowsCompiler.VisualStudio2019))
 			{
 				return WindowsCompiler.VisualStudio2019;
 			}
+			if (HasCompiler(WindowsCompiler.VisualStudio2017))
+			{
+				return WindowsCompiler.VisualStudio2017;
+			}
 
 			// If we do have a Visual Studio installation, but we're missing just the C++ parts, warn about that.
 			DirectoryReference VSInstallDir;
-			if (TryGetVSInstallDir(WindowsCompiler.VisualStudio2017, out VSInstallDir))
+			if (TryGetVSInstallDir(WindowsCompiler.VisualStudio2019, out VSInstallDir))
 			{
-				Log.TraceWarning("Visual Studio 2017 is installed, but is missing the C++ toolchain. Please verify that the \"VC++ 2017 toolset\" component is selected in the Visual Studio 2017 installation options.");
+				Log.TraceWarning("Visual Studio 2019 is installed, but is missing the C++ toolchain. Please verify that the \"VC++ 2019 toolset\" component is selected in the Visual Studio 2019 installation options.");
 			}
-			else if (TryGetVSInstallDir(WindowsCompiler.VisualStudio2019, out VSInstallDir))
+			else if (TryGetVSInstallDir(WindowsCompiler.VisualStudio2017, out VSInstallDir))
 			{
-				Log.TraceWarning("Visual Studio 2019 is installed, but is missing the C++ toolchain. Please verify that the \"VC++ 2019 toolset\" component is selected in the Visual Studio 2019 installation options.");
+				Log.TraceWarning("Visual Studio 2017 is installed, but is missing the C++ toolchain. Please verify that the \"VC++ 2017 toolset\" component is selected in the Visual Studio 2017 installation options.");
 			}
 			else
 			{
 				Log.TraceWarning("No Visual C++ installation was found. Please download and install Visual Studio 2017 with C++ components.");
 			}
 
-			// Finally, default to VS2017 anyway
-			return WindowsCompiler.VisualStudio2017;
+			// Finally, default to VS2019 anyway
+			return WindowsCompiler.VisualStudio2019;
 		}
 
 		/// <summary>
