Prefer 2019 when generating project files.

diff --git a/Engine/Source/Programs/UnrealBuildTool/ProjectFiles/VisualStudio/VCProjectFileGenerator.cs b/Engine/Source/Programs/UnrealBuildTool/ProjectFiles/VisualStudio/VCProjectFileGenerator.cs
--- a/Engine/Source/Programs/UnrealBuildTool/ProjectFiles/VisualStudio/VCProjectFileGenerator.cs
+++ b/Engine/Source/Programs/UnrealBuildTool/ProjectFiles/VisualStudio/VCProjectFileGenerator.cs
@@ -292,7 +292,11 @@ namespace UnrealBuildTool
 			if (ProjectFileFormat == VCProjectFileFormat.Default)
 			{
 				// Pick the best platform installed by default
-				if (WindowsPlatform.HasCompiler(WindowsCompiler.VisualStudio2017) && WindowsPlatform.HasIDE(WindowsCompiler.VisualStudio2017))
+				if (WindowsPlatform.HasCompiler(WindowsCompiler.VisualStudio2019) && WindowsPlatform.HasIDE(WindowsCompiler.VisualStudio2019))
+				{
+					ProjectFileFormat = VCProjectFileFormat.VisualStudio2019;
+				}
+				else if (WindowsPlatform.HasCompiler(WindowsCompiler.VisualStudio2017) && WindowsPlatform.HasIDE(WindowsCompiler.VisualStudio2017))
 				{
 					ProjectFileFormat = VCProjectFileFormat.VisualStudio2017;
 				}
@@ -300,10 +304,6 @@ namespace UnrealBuildTool
 				{
 					ProjectFileFormat = VCProjectFileFormat.VisualStudio2015;
 				}
-				else if (WindowsPlatform.HasCompiler(WindowsCompiler.VisualStudio2019) && WindowsPlatform.HasIDE(WindowsCompiler.VisualStudio2019))
-				{
-					ProjectFileFormat = VCProjectFileFormat.VisualStudio2019;
-				}
 
 				// Allow the SDKs to override
 				foreach (UnrealTargetPlatform SupportedPlatform in SupportedPlatforms)
