Fix parsing of dependencies file on Mac. Revert to older code for non ISPC output.

#rnx

diff --git a/Engine/Source/Programs/UnrealBuildTool/System/CppDependencyCache.cs b/Engine/Source/Programs/UnrealBuildTool/System/CppDependencyCache.cs
--- a/Engine/Source/Programs/UnrealBuildTool/System/CppDependencyCache.cs
+++ b/Engine/Source/Programs/UnrealBuildTool/System/CppDependencyCache.cs
@@ -93,7 +93,7 @@ namespace UnrealBuildTool
 			this.BaseDir = BaseDir;
 			this.Parent = Parent;
 
-			if(FileReference.Exists(Location))
+			if (FileReference.Exists(Location))
 			{
 				Read();
 			}
@@ -107,7 +107,7 @@ namespace UnrealBuildTool
 		/// <returns>True if the input file exists and the dependencies were read</returns>
 		public bool TryGetDependencies(FileItem InputFile, out List<FileItem> OutDependencyItems)
 		{
-			if(!InputFile.Exists)
+			if (!InputFile.Exists)
 			{
 				OutDependencyItems = null;
 				return false;
@@ -117,7 +117,7 @@ namespace UnrealBuildTool
 			{
 				return TryGetDependenciesInternal(InputFile, out OutDependencyItems);
 			}
-			catch(Exception Ex)
+			catch (Exception Ex)
 			{
 				Log.TraceLog("Unable to read {0}:\n{1}", InputFile, ExceptionUtils.FormatExceptionDetails(Ex));
 				OutDependencyItems = null;
@@ -133,14 +133,14 @@ namespace UnrealBuildTool
 		/// <returns>True if the input file exists and the dependencies were read</returns>
 		private bool TryGetDependenciesInternal(FileItem InputFile, out List<FileItem> OutDependencyItems)
 		{
-			if(Parent != null && !InputFile.Location.IsUnderDirectory(BaseDir))
+			if (Parent != null && !InputFile.Location.IsUnderDirectory(BaseDir))
 			{
 				return Parent.TryGetDependencies(InputFile, out OutDependencyItems);
 			}
 			else
 			{
 				DependencyInfo Info;
-				if(DependencyFileToInfo.TryGetValue(InputFile, out Info) && InputFile.LastWriteTimeUtc.Ticks <= Info.LastWriteTimeUtc)
+				if (DependencyFileToInfo.TryGetValue(InputFile, out Info) && InputFile.LastWriteTimeUtc.Ticks <= Info.LastWriteTimeUtc)
 				{
 					OutDependencyItems = Info.Files;
 					return true;
@@ -169,10 +169,10 @@ namespace UnrealBuildTool
 		{
 			CppDependencyCache Cache = null;
 
-			if(ProjectFile == null || !UnrealBuildTool.IsEngineInstalled())
+			if (ProjectFile == null || !UnrealBuildTool.IsEngineInstalled())
 			{
 				string AppName;
-				if(TargetType == TargetType.Program)
+				if (TargetType == TargetType.Program)
 				{
 					AppName = TargetName;
 				}
@@ -185,7 +185,7 @@ namespace UnrealBuildTool
 				Cache = FindOrAddCache(EngineCacheLocation, UnrealBuildTool.EngineDirectory, Cache);
 			}
 
-			if(ProjectFile != null)
+			if (ProjectFile != null)
 			{
 				FileReference ProjectCacheLocation = FileReference.Combine(ProjectFile.Directory, UEBuildTarget.GetPlatformIntermediateFolder(Platform, Architecture), TargetName, Configuration.ToString(), "DependencyCache.bin");
 				Cache = FindOrAddCache(ProjectCacheLocation, ProjectFile.Directory, Cache);
@@ -203,10 +203,10 @@ namespace UnrealBuildTool
 		/// <returns>Reference to a dependency cache with the given settings</returns>
 		static CppDependencyCache FindOrAddCache(FileReference Location, DirectoryReference BaseDir, CppDependencyCache Parent)
 		{
-			lock(Caches)
+			lock (Caches)
 			{
 				CppDependencyCache Cache;
-				if(Caches.TryGetValue(Location, out Cache))
+				if (Caches.TryGetValue(Location, out Cache))
 				{
 					Debug.Assert(Cache.BaseDir == BaseDir);
 					Debug.Assert(Cache.Parent == Parent);
@@ -225,7 +225,7 @@ namespace UnrealBuildTool
 		/// </summary>
 		public static void SaveAll()
 		{
-			Parallel.ForEach(Caches.Values, Cache => { if(Cache.bModified){ Cache.Write(); } });
+			Parallel.ForEach(Caches.Values, Cache => { if (Cache.bModified) { Cache.Write(); } });
 		}
 
 		/// <summary>
@@ -235,24 +235,24 @@ namespace UnrealBuildTool
 		{
 			try
 			{
-				using(BinaryArchiveReader Reader = new BinaryArchiveReader(Location))
+				using (BinaryArchiveReader Reader = new BinaryArchiveReader(Location))
 				{
 					int Version = Reader.ReadInt();
-					if(Version != CurrentVersion)
+					if (Version != CurrentVersion)
 					{
 						Log.TraceLog("Unable to read dependency cache from {0}; version {1} vs current {2}", Location, Version, CurrentVersion);
 						return;
 					}
 
 					int Count = Reader.ReadInt();
-					for(int Idx = 0; Idx < Count; Idx++)
+					for (int Idx = 0; Idx < Count; Idx++)
 					{
 						FileItem File = Reader.ReadFileItem();
 						DependencyFileToInfo[File] = DependencyInfo.Read(Reader);
 					}
 				}
 			}
-			catch(Exception Ex)
+			catch (Exception Ex)
 			{
 				Log.TraceWarning("Unable to read {0}. See log for additional information.", Location);
 				Log.TraceLog("{0}", ExceptionUtils.FormatExceptionDetails(Ex));
@@ -265,14 +265,14 @@ namespace UnrealBuildTool
 		private void Write()
 		{
 			DirectoryReference.CreateDirectory(Location.Directory);
-			using(FileStream Stream = File.Open(Location.FullName, FileMode.Create, FileAccess.Write, FileShare.Read))
+			using (FileStream Stream = File.Open(Location.FullName, FileMode.Create, FileAccess.Write, FileShare.Read))
 			{
-				using(BinaryArchiveWriter Writer = new BinaryArchiveWriter(Stream))
+				using (BinaryArchiveWriter Writer = new BinaryArchiveWriter(Stream))
 				{
 					Writer.WriteInt(CurrentVersion);
 
 					Writer.WriteInt(DependencyFileToInfo.Count);
-					foreach(KeyValuePair<FileItem, DependencyInfo> Pair in DependencyFileToInfo)
+					foreach (KeyValuePair<FileItem, DependencyInfo> Pair in DependencyFileToInfo)
 					{
 						Writer.WriteFileItem(Pair.Key);
 						Pair.Value.Write(Writer);
@@ -289,22 +289,154 @@ namespace UnrealBuildTool
 		/// <returns>List of included dependencies</returns>
 		static List<FileItem> ReadDependenciesFile(FileReference InputFile)
 		{
-			string[] Lines = FileReference.ReadAllLines(InputFile);
+			if (InputFile.HasExtension(".d"))
+			{
+				string Text = FileReference.ReadAllText(InputFile);
+
+				List<string> Tokens = new List<string>();
+
+				StringBuilder Token = new StringBuilder();
+				for (int Idx = 0; TryReadMakefileToken(Text, ref Idx, Token);)
+				{
+					Tokens.Add(Token.ToString());
+				}
+
+				int TokenIdx = 0;
+				while (TokenIdx < Tokens.Count && Tokens[TokenIdx] == "\n")
+				{
+					TokenIdx++;
+				}
+
+				if (TokenIdx + 1 >= Tokens.Count || Tokens[TokenIdx + 1] != ":")
+				{
+					throw new BuildException("Unable to parse dependency file");
+				}
+
+				TokenIdx += 2;
+
+				List<FileItem> NewDependencyFiles = new List<FileItem>();
+				for (; TokenIdx < Tokens.Count && Tokens[TokenIdx] != "\n"; TokenIdx++)
+				{
+					NewDependencyFiles.Add(FileItem.GetItemByPath(Tokens[TokenIdx]));
+				}
+
+				while (TokenIdx < Tokens.Count && Tokens[TokenIdx] == "\n")
+				{
+					TokenIdx++;
+				}
 
-			HashSet<FileItem> DependencyItems = new HashSet<FileItem>();
-			foreach(string Line in Lines)
+				if (TokenIdx != Tokens.Count)
+				{
+					throw new BuildException("Unable to parse dependency file");
+				}
+
+				return NewDependencyFiles;
+			}
+			else if (InputFile.HasExtension(".txt"))
 			{
-				if(Line.Length > 0)
+				string[] Lines = FileReference.ReadAllLines(InputFile);
+
+				HashSet<FileItem> DependencyItems = new HashSet<FileItem>();
+				foreach (string Line in Lines)
 				{
-					// Ignore *.tlh and *.tli files generated by the compiler from COM DLLs
-					if(!Line.EndsWith(".tlh", StringComparison.OrdinalIgnoreCase) && !Line.EndsWith(".tli", StringComparison.OrdinalIgnoreCase))
+					if (Line.Length > 0)
 					{
-						string FixedLine = Line.Replace("\\\\", "\\"); // ISPC outputs files with escaped slashes
-						DependencyItems.Add(FileItem.GetItemByPath(FixedLine));
+						// Ignore *.tlh and *.tli files generated by the compiler from COM DLLs
+						if (!Line.EndsWith(".tlh", StringComparison.OrdinalIgnoreCase) && !Line.EndsWith(".tli", StringComparison.OrdinalIgnoreCase))
+						{
+							string FixedLine = Line.Replace("\\\\", "\\"); // ISPC outputs files with escaped slashes
+							DependencyItems.Add(FileItem.GetItemByPath(FixedLine));
+						}
 					}
 				}
+				return DependencyItems.ToList();
+			}
+			else
+			{
+				throw new BuildException("Unknown dependency list file type: {0}", InputFile);
+			}
+		}
+
+		/// <summary>
+		/// Attempts to read a single token from a makefile
+		/// </summary>
+		/// <param name="Text">Text to read from</param>
+		/// <param name="RefIdx">Current position within the file</param>
+		/// <param name="Token">Receives the token characters</param>
+		/// <returns>True if a token was read, false if the end of the buffer was reached</returns>
+		static bool TryReadMakefileToken(string Text, ref int RefIdx, StringBuilder Token)
+		{
+			Token.Clear();
+
+			int Idx = RefIdx;
+			for (; ; )
+			{
+				if (Idx == Text.Length)
+				{
+					return false;
+				}
+
+				// Skip whitespace
+				while (Text[Idx] == ' ' || Text[Idx] == '\t')
+				{
+					if (++Idx == Text.Length)
+					{
+						return false;
+					}
+				}
+
+				// Colon token
+				if (Text[Idx] == ':')
+				{
+					Token.Append(':');
+					RefIdx = Idx + 1;
+					return true;
+				}
+
+				// Check for a newline
+				if (Text[Idx] == '\r' || Text[Idx] == '\n')
+				{
+					Token.Append('\n');
+					RefIdx = Idx + 1;
+					return true;
+				}
+
+				// Check for an escaped newline
+				if (Text[Idx] == '\\' && Idx + 1 < Text.Length)
+				{
+					if (Text[Idx + 1] == '\n')
+					{
+						Idx += 2;
+						continue;
+					}
+					if (Text[Idx + 1] == '\r' && Idx + 2 < Text.Length && Text[Idx + 2] == '\n')
+					{
+						Idx += 3;
+						continue;
+					}
+				}
+
+				// Read a token. Special handling for drive letters on Windows!
+				for (; Idx < Text.Length; Idx++)
+				{
+					if (Text[Idx] == ' ' || Text[Idx] == '\t' || Text[Idx] == '\r' || Text[Idx] == '\n')
+					{
+						break;
+					}
+					if (Text[Idx] == ':' && Token.Length > 1)
+					{
+						break;
+					}
+					if (Text[Idx] == '\\' && Idx + 1 < Text.Length && Text[Idx + 1] == ' ')
+					{
+						Idx++;
+					}
+					Token.Append(Text[Idx]);
+				}
+
+				RefIdx = Idx;
+				return true;
 			}
-			return DependencyItems.ToList();
 		}
 	}
 }
