Dont over optimize when using MSan the same way we do with ASan

diff --git a/Engine/Source/Programs/UnrealBuildTool/Platform/Linux/LinuxToolChain.cs b/Engine/Source/Programs/UnrealBuildTool/Platform/Linux/LinuxToolChain.cs
--- a/Engine/Source/Programs/UnrealBuildTool/Platform/Linux/LinuxToolChain.cs
+++ b/Engine/Source/Programs/UnrealBuildTool/Platform/Linux/LinuxToolChain.cs
@@ -737,8 +737,8 @@ namespace UnrealBuildTool
 			}
 			else
 			{
-				// Don't over optimise if using AddressSanitizer or you'll get false positive errors due to erroneous optimisation of necessary AddressSanitizer instrumentation.
-				if (Options.HasFlag(LinuxToolChainOptions.EnableAddressSanitizer))
+				// Don't over optimise if using Address/MemorySanitizer or you'll get false positive errors due to erroneous optimisation of necessary Address/MemorySanitizer instrumentation.
+				if (Options.HasFlag(LinuxToolChainOptions.EnableAddressSanitizer) || Options.HasFlag(LinuxToolChainOptions.EnableMemorySanitizer))
 				{
 					Result += " -O1 -g -fno-optimize-sibling-calls -fno-omit-frame-pointer";
 				}
