Append a suffix to the binary name before the extension when using llvm sanitizers to avoid overwriting binary/*.so names

diff --git a/Engine/Source/Programs/UnrealBuildTool/Platform/Linux/UEBuildLinux.cs b/Engine/Source/Programs/UnrealBuildTool/Platform/Linux/UEBuildLinux.cs
--- a/Engine/Source/Programs/UnrealBuildTool/Platform/Linux/UEBuildLinux.cs
+++ b/Engine/Source/Programs/UnrealBuildTool/Platform/Linux/UEBuildLinux.cs
@@ -537,6 +537,41 @@ namespace UnrealBuildTool
 			};
 		}
 
+		public override List<FileReference> FinalizeBinaryPaths(FileReference BinaryName, FileReference ProjectFile, ReadOnlyTargetRules Target)
+		{
+			List<FileReference> FinalBinaryPath = new List<FileReference>();
+
+			string SanitizerSuffix = null;
+			if(Target.LinuxPlatform.bEnableAddressSanitizer)
+			{
+				SanitizerSuffix = "ASan";
+			}
+			else if(Target.LinuxPlatform.bEnableThreadSanitizer)
+			{
+				SanitizerSuffix = "TSan";
+			}
+			else if(Target.LinuxPlatform.bEnableUndefinedBehaviorSanitizer)
+			{
+				SanitizerSuffix = "UBSan";
+			}
+			else if(Target.LinuxPlatform.bEnableMemorySanitizer)
+			{
+				SanitizerSuffix = "MSan";
+			}
+
+			if (String.IsNullOrEmpty(SanitizerSuffix))
+			{
+				FinalBinaryPath.Add(BinaryName);
+			}
+			else
+			{
+				// Append the sanitizer suffix to the binary name but before the extension type
+				FinalBinaryPath.Add(new FileReference(Path.Combine(BinaryName.Directory.FullName, BinaryName.GetFileNameWithoutExtension() + "-" + SanitizerSuffix + BinaryName.GetExtension())));
+			}
+
+			return FinalBinaryPath;
+		}
+
 		/// <summary>
 		/// Creates a toolchain instance for the given platform.
 		/// </summary>
