UBT: Fix UHT always running for installed engine builds. Was dependent on timestamps of binaries downloaded by launcher.

diff --git a/Engine/Source/Programs/UnrealBuildTool/System/ExternalExecution.cs b/Engine/Source/Programs/UnrealBuildTool/System/ExternalExecution.cs
--- a/Engine/Source/Programs/UnrealBuildTool/System/ExternalExecution.cs
+++ b/Engine/Source/Programs/UnrealBuildTool/System/ExternalExecution.cs
@@ -708,22 +708,26 @@ namespace UnrealBuildTool
 					return false;
 				}
 
-				TargetReceipt Receipt;
-				if (!TargetReceipt.TryRead(ReceiptPath, out Receipt))
+				// Don't check timestamps for individual binaries if we're using the installed version of UHT. It will always be up to date.
+				if (!UnrealBuildTool.IsFileInstalled(ReceiptFile.Location))
 				{
-					Timestamp = DateTime.MaxValue;
-					return false;
-				}
-
-				// Make sure all the build products exist, and that the receipt is newer
-				foreach (BuildProduct BuildProduct in Receipt.BuildProducts)
-				{
-					FileItem BuildProductItem = FileItem.GetItemByFileReference(BuildProduct.Path);
-					if(!BuildProductItem.Exists || BuildProductItem.LastWriteTimeUtc > ReceiptFile.LastWriteTimeUtc)
+					TargetReceipt Receipt;
+					if (!TargetReceipt.TryRead(ReceiptPath, out Receipt))
 					{
 						Timestamp = DateTime.MaxValue;
 						return false;
 					}
+
+					// Make sure all the build products exist, and that the receipt is newer
+					foreach (BuildProduct BuildProduct in Receipt.BuildProducts)
+					{
+						FileItem BuildProductItem = FileItem.GetItemByFileReference(BuildProduct.Path);
+						if (!BuildProductItem.Exists || BuildProductItem.LastWriteTimeUtc > ReceiptFile.LastWriteTimeUtc)
+						{
+							Timestamp = DateTime.MaxValue;
+							return false;
+						}
+					}
 				}
 
 				// Return the timestamp for all the binaries
