Native MSBuild implementation of UnrealBuildTool's discovery of platform extensions.

(ushell-p4-cherrypick of 7973988 by [REDACTED_PERSON])

(ushell-p4-cherrypick of 8124246 by [REDACTED_PERSON])

diff --git a/Engine/Build/BatchFiles/FindPlatformExtensionSources.bat b/Engine/Build/BatchFiles/FindPlatformExtensionSources.bat
deleted file mode 100644
diff --git a/Engine/Build/BatchFiles/FindPlatformExtensionSources.sh b/Engine/Build/BatchFiles/FindPlatformExtensionSources.sh
deleted file mode 100755
diff --git a/Engine/Build/BatchFiles/GenerateProjectFiles.bat b/Engine/Build/BatchFiles/GenerateProjectFiles.bat
--- a/Engine/Build/BatchFiles/GenerateProjectFiles.bat
+++ b/Engine/Build/BatchFiles/GenerateProjectFiles.bat
@@ -57,9 +57,6 @@ if not [%MSBUILD_EXE%] == [%MSBUILD_15_EXE_WITH_NUGET%] goto Error_RequireNugetP
 
 :NoVsWhere
 
-rem Find platform extension source files
-call "%~dp0FindPlatformExtensionSources.bat"
-
 rem Check to see if the files in the UBT directory have changed. We conditionally include platform files from the .csproj file, but MSBuild doesn't recognize the dependency when new files are added. 
 md ..\Intermediate\Build >nul 2>nul
 dir /s /b Programs\UnrealBuildTool\*.cs >..\Intermediate\Build\UnrealBuildToolFiles.txt
@@ -130,4 +127,4 @@ goto Exit
 :Exit
 rem ## Restore original CWD in case we change it
 popd
-exit /B 1
\ No newline at end of file
+exit /B 1
diff --git a/Engine/Build/BatchFiles/Linux/GenerateProjectFiles.sh b/Engine/Build/BatchFiles/Linux/GenerateProjectFiles.sh
--- a/Engine/Build/BatchFiles/Linux/GenerateProjectFiles.sh
+++ b/Engine/Build/BatchFiles/Linux/GenerateProjectFiles.sh
@@ -22,9 +22,6 @@ fi
 
 source "$BASE_PATH/SetupMono.sh" "$BASE_PATH"
 
-# make sure the UBT project has references to auto-discovered platform extension source files
-"${BASE_PATH}/../FindPlatformExtensionSources.sh"
-
 if [ -f "$BASE_PATH/../../../Source/Programs/UnrealBuildTool/UnrealBuildTool.csproj" ]; then
 	xbuild "$BASE_PATH/../../../Source/Programs/UnrealBuildTool/UnrealBuildTool.csproj" /property:Configuration="Development" /verbosity:quiet /nologo /p:NoWarn=1591
 fi
diff --git a/Engine/Build/BatchFiles/Mac/GenerateProjectFiles.sh b/Engine/Build/BatchFiles/Mac/GenerateProjectFiles.sh
--- a/Engine/Build/BatchFiles/Mac/GenerateProjectFiles.sh
+++ b/Engine/Build/BatchFiles/Mac/GenerateProjectFiles.sh
@@ -21,9 +21,6 @@ fi
 
 . "$BASE_PATH/SetupMono.sh" "$BASE_PATH"
 
-# make sure the UBT project has references to auto-discovered platform extension source files
-"${BASE_PATH}/../FindPlatformExtensionSources.sh"
-
 if [ -f "$BASE_PATH/../../../Source/Programs/UnrealBuildTool/UnrealBuildTool.csproj" ]; then
 	xbuild "$BASE_PATH/../../../Source/Programs/UnrealBuildTool/UnrealBuildTool.csproj" /property:Configuration="Development" /verbosity:quiet /nologo /p:NoWarn=1591 |grep -i error
 fi
diff --git a/Engine/Build/BatchFiles/RunUAT.bat b/Engine/Build/BatchFiles/RunUAT.bat
--- a/Engine/Build/BatchFiles/RunUAT.bat
+++ b/Engine/Build/BatchFiles/RunUAT.bat
@@ -30,9 +30,6 @@ rem ## check if the UAT projects are present. if not, we'll just use the precomp
 if not exist Source\Programs\AutomationTool\AutomationTool.csproj goto RunPrecompiled
 if not exist Source\Programs\AutomationToolLauncher\AutomationToolLauncher.csproj goto RunPrecompiled
 
-rem ## find platform extension source code that UBT will need when compiling platform extension automation projects
-call "%~dp0FindPlatformExtensionSources.bat"
-
 rem ## Get the path to MSBuild
 call "%~dp0GetMSBuildPath.bat"
 if errorlevel 1 goto RunPrecompiled
diff --git a/Engine/Build/BatchFiles/RunUAT.sh b/Engine/Build/BatchFiles/RunUAT.sh
--- a/Engine/Build/BatchFiles/RunUAT.sh
+++ b/Engine/Build/BatchFiles/RunUAT.sh
@@ -73,9 +73,6 @@ if [ "$UATCompileArg" = "-compile" ]; then
 		UATCompileArg=
 	else
 
-		# make sure the UBT project has references to auto-discovered platform extension source files
-		"${SCRIPT_DIR}/FindPlatformExtensionSources.sh"
-
 		# mono 5.0 and up include msbuild
 		if [ "$IS_MS_BUILD_AVAILABLE" == "1" ]; then
 			BUILD_TOOL=msbuild
diff --git a/Engine/Source/Programs/UnrealBuildTool/UnrealBuildTool.csproj b/Engine/Source/Programs/UnrealBuildTool/UnrealBuildTool.csproj
--- a/Engine/Source/Programs/UnrealBuildTool/UnrealBuildTool.csproj
+++ b/Engine/Source/Programs/UnrealBuildTool/UnrealBuildTool.csproj
@@ -441,20 +441,18 @@
     <None Include="app.config" />
     <None Include="app.manifest" />
   </ItemGroup>
+
   <Import Project="$(MSBuildBinPath)\Microsoft.CSharp.targets" />
-  <!-- The next two sections allow for GenerateProjectFiles script to hunt down "platform extension" C# files
-       without them needing to be added to this project. If they were compiled in a separate assembly like
-	   rules files, we would need to make a huge amount of classes, etc public, which we definitely do not 
-	   want to do. However, this does require that GenerateProjectFiles is run to create the References file
-	   below. 
-       The PlatformExtensionCompile compile item compiles the files found in the References file. This indirection
-	   is used so that Visual Studio shows these files. (If the References file directly compiled files, the 
-	   files would not be visible in the Solution Explorer).
-  -->
-  <Import Project="..\..\..\Intermediate\ProjectFiles\UnrealBuildTool.csproj.References" Condition="Exists('..\..\..\Intermediate\ProjectFiles\UnrealBuildTool.csproj.References')" />
+
+  <!-- The following will glob all platform extensions' UnrealBuildTool source
+  files and include them in the build. They will also appear in VisualStudio -->
   <ItemGroup>
-    <Compile Include="@(PlatformExtensionCompile)" />
+    <PlatformExtensions Include="../../../Platforms/*/Source/Programs/UnrealBuildTool/**/*.cs" />
+    <Compile Include="@(PlatformExtensions)">
+      <Link>Platform/$([System.Text.RegularExpressions.Regex]::Split(`%(Identity)`, `[\\/]`)[4])/_</Link>
+    </Compile>
   </ItemGroup>
+
   <!-- To modify your build process, add your task inside one of the targets below and uncomment it. 
        Other similar extension points exist, see Microsoft.Common.targets.
   <Target Name="BeforeBuild">
