Undid changelist 8124316. Turns out Main and its children only have a partial implementation of platform extensions and therefore aren't ready for this change yet.

diff --git a/Engine/Build/BatchFiles/FindPlatformExtensionSources.bat b/Engine/Build/BatchFiles/FindPlatformExtensionSources.bat
new file mode 100644
--- /dev/null
+++ b/Engine/Build/BatchFiles/FindPlatformExtensionSources.bat
@@ -0,0 +1,36 @@
+@echo off
+setlocal
+
+rem ## Unreal Engine 4 AutomationTool setup script
+rem ## Copyright 1998-2019 Epic Games, Inc. All Rights Reserved.
+rem ##
+rem ## This script is expecting to exist in the UE4/Engine/Build/BatchFiles directory.  It will not work correctly
+rem ## if you copy it to a different location and run it.
+
+pushd "%~dp0..\..\Source"
+
+mkdir ..\Intermediate\ProjectFiles >NUL 2>NUL
+
+rem Look for any platform extension .cs files, and add them to a file that will be refernced by UBT project, so it can bring them in automatically
+set REFERENCE_FILE="%~dp0..\..\Intermediate\ProjectFiles\UnrealBuildTool.csproj.References"
+echo ^<?xml version="1.0" encoding="utf-8"?^> > %REFERENCE_FILE%
+echo ^<Project ToolsVersion="15.0" DefaultTargets="Build" xmlns="http://schemas.microsoft.com/developer/msbuild/2003"^> >> %REFERENCE_FILE%
+echo ^<ItemGroup^> >> %REFERENCE_FILE%
+
+if exist "..\..\Platforms" (
+	pushd "..\..\Platforms"
+	for /d %%i in ("*") do (
+		if exist "%%i\Engine\Source\Programs\UnrealBuildTool" (
+			pushd "%%i\Engine\Source\Programs\UnrealBuildTool"
+			for %%j in ("*.cs") do (
+				echo   ^<PlatformExtensionCompile Include="..\..\..\..\Platforms\%%i\Engine\Source\Programs\UnrealBuildTool\%%j"^> >> %REFERENCE_FILE%
+				echo     ^<Link^>Platform\%%i\%%j^</Link^> >> %REFERENCE_FILE%
+				echo   ^</PlatformExtensionCompile^> >> %REFERENCE_FILE%
+			)
+			popd
+		)
+	)
+	popd
+)
+echo ^</ItemGroup^> >> %REFERENCE_FILE%
+echo ^</Project^> >> %REFERENCE_FILE%
\ No newline at end of file
diff --git a/Engine/Build/BatchFiles/FindPlatformExtensionSources.sh b/Engine/Build/BatchFiles/FindPlatformExtensionSources.sh
new file mode 100755
--- /dev/null
+++ b/Engine/Build/BatchFiles/FindPlatformExtensionSources.sh
@@ -0,0 +1,29 @@
+#!/bin/bash
+## Copyright 1998-2019 Epic Games, Inc. All Rights Reserved.
+##
+## Unreal Engine 4 AutomationTool setup script
+##
+## This script is expecting to exist in the UE4/Engine/Build/BatchFiles directory.  It will not work
+## correctly if you copy it to a different location and run it.
+
+
+SCRIPT_DIR=$(cd "`dirname "$0"`" && pwd)
+cd "$SCRIPT_DIR/../../Source"
+
+mkdir -p ../Intermediate/ProjectFiles
+
+## Look for any platform extension .cs files, and add them to a file that will be refernced by UBT project, so it can bring them in automatically
+REFERENCE_FILE=../Intermediate/ProjectFiles/UnrealBuildTool.csproj.References
+echo "<?xml version=\"1.0\" encoding=\"utf-8\"?>" > $REFERENCE_FILE
+echo "<Project ToolsVersion=\"15.0\" DefaultTargets=\"Build\" xmlns=\"http://schemas.microsoft.com/developer/msbuild/2003\">" >> $REFERENCE_FILE
+echo "<ItemGroup>" >> $REFERENCE_FILE
+
+for PLATFORM in ../../Platforms/*; do
+	for SOURCE_FILE in ${PLATFORM}/Engine/Source/Programs/UnrealBuildTool/*; do
+		echo "<PlatformExtensionCompile Include=\"..\\..\\${SOURCE_FILE//\//\\}\" />" >> $REFERENCE_FILE
+	done
+done
+
+echo "</ItemGroup>" >> $REFERENCE_FILE
+echo "</Project>" >> $REFERENCE_FILE
+
diff --git a/Engine/Build/BatchFiles/GenerateProjectFiles.bat b/Engine/Build/BatchFiles/GenerateProjectFiles.bat
--- a/Engine/Build/BatchFiles/GenerateProjectFiles.bat
+++ b/Engine/Build/BatchFiles/GenerateProjectFiles.bat
@@ -57,6 +57,9 @@ if not [%MSBUILD_EXE%] == [%MSBUILD_15_EXE_WITH_NUGET%] goto Error_RequireNugetP
 
 :NoVsWhere
 
+rem Find platform extension source files
+call "%~dp0FindPlatformExtensionSources.bat"
+
 rem Check to see if the files in the UBT directory have changed. We conditionally include platform files from the .csproj file, but MSBuild doesn't recognize the dependency when new files are added. 
 md ..\Intermediate\Build >nul 2>nul
 dir /s /b Programs\UnrealBuildTool\*.cs >..\Intermediate\Build\UnrealBuildToolFiles.txt
@@ -127,4 +130,4 @@ goto Exit
 :Exit
 rem ## Restore original CWD in case we change it
 popd
-exit /B 1
+exit /B 1
\ No newline at end of file
diff --git a/Engine/Build/BatchFiles/Linux/GenerateProjectFiles.sh b/Engine/Build/BatchFiles/Linux/GenerateProjectFiles.sh
--- a/Engine/Build/BatchFiles/Linux/GenerateProjectFiles.sh
+++ b/Engine/Build/BatchFiles/Linux/GenerateProjectFiles.sh
@@ -22,6 +22,9 @@ fi
 
 source "$BASE_PATH/SetupMono.sh" "$BASE_PATH"
 
+# make sure the UBT project has references to auto-discovered platform extension source files
+"${BASE_PATH}/../FindPlatformExtensionSources.sh"
+
 if [ -f "$BASE_PATH/../../../Source/Programs/UnrealBuildTool/UnrealBuildTool.csproj" ]; then
 	xbuild "$BASE_PATH/../../../Source/Programs/UnrealBuildTool/UnrealBuildTool.csproj" /property:Configuration="Development" /verbosity:quiet /nologo /p:NoWarn=1591
 fi
diff --git a/Engine/Build/BatchFiles/Mac/GenerateProjectFiles.sh b/Engine/Build/BatchFiles/Mac/GenerateProjectFiles.sh
--- a/Engine/Build/BatchFiles/Mac/GenerateProjectFiles.sh
+++ b/Engine/Build/BatchFiles/Mac/GenerateProjectFiles.sh
@@ -21,6 +21,9 @@ fi
 
 . "$BASE_PATH/SetupMono.sh" "$BASE_PATH"
 
+# make sure the UBT project has references to auto-discovered platform extension source files
+"${BASE_PATH}/../FindPlatformExtensionSources.sh"
+
 if [ -f "$BASE_PATH/../../../Source/Programs/UnrealBuildTool/UnrealBuildTool.csproj" ]; then
 	xbuild "$BASE_PATH/../../../Source/Programs/UnrealBuildTool/UnrealBuildTool.csproj" /property:Configuration="Development" /verbosity:quiet /nologo /p:NoWarn=1591 |grep -i error
 fi
diff --git a/Engine/Build/BatchFiles/RunUAT.bat b/Engine/Build/BatchFiles/RunUAT.bat
--- a/Engine/Build/BatchFiles/RunUAT.bat
+++ b/Engine/Build/BatchFiles/RunUAT.bat
@@ -30,6 +30,9 @@ rem ## check if the UAT projects are present. if not, we'll just use the precomp
 if not exist Source\Programs\AutomationTool\AutomationTool.csproj goto RunPrecompiled
 if not exist Source\Programs\AutomationToolLauncher\AutomationToolLauncher.csproj goto RunPrecompiled
 
+rem ## find platform extension source code that UBT will need when compiling platform extension automation projects
+call "%~dp0FindPlatformExtensionSources.bat"
+
 rem ## Get the path to MSBuild
 call "%~dp0GetMSBuildPath.bat"
 if errorlevel 1 goto RunPrecompiled
diff --git a/Engine/Build/BatchFiles/RunUAT.sh b/Engine/Build/BatchFiles/RunUAT.sh
--- a/Engine/Build/BatchFiles/RunUAT.sh
+++ b/Engine/Build/BatchFiles/RunUAT.sh
@@ -73,6 +73,9 @@ if [ "$UATCompileArg" = "-compile" ]; then
 		UATCompileArg=
 	else
 
+		# make sure the UBT project has references to auto-discovered platform extension source files
+		"${SCRIPT_DIR}/FindPlatformExtensionSources.sh"
+
 		# mono 5.0 and up include msbuild
 		if [ "$IS_MS_BUILD_AVAILABLE" == "1" ]; then
 			BUILD_TOOL=msbuild
diff --git a/Engine/Source/Programs/UnrealBuildTool/UnrealBuildTool.csproj b/Engine/Source/Programs/UnrealBuildTool/UnrealBuildTool.csproj
--- a/Engine/Source/Programs/UnrealBuildTool/UnrealBuildTool.csproj
+++ b/Engine/Source/Programs/UnrealBuildTool/UnrealBuildTool.csproj
@@ -441,18 +441,20 @@
     <None Include="app.config" />
     <None Include="app.manifest" />
   </ItemGroup>
-
   <Import Project="$(MSBuildBinPath)\Microsoft.CSharp.targets" />
-
-  <!-- The following will glob all platform extensions' UnrealBuildTool source
-  files and include them in the build. They will also appear in VisualStudio -->
+  <!-- The next two sections allow for GenerateProjectFiles script to hunt down "platform extension" C# files
+       without them needing to be added to this project. If they were compiled in a separate assembly like
+	   rules files, we would need to make a huge amount of classes, etc public, which we definitely do not 
+	   want to do. However, this does require that GenerateProjectFiles is run to create the References file
+	   below. 
+       The PlatformExtensionCompile compile item compiles the files found in the References file. This indirection
+	   is used so that Visual Studio shows these files. (If the References file directly compiled files, the 
+	   files would not be visible in the Solution Explorer).
+  -->
+  <Import Project="..\..\..\Intermediate\ProjectFiles\UnrealBuildTool.csproj.References" Condition="Exists('..\..\..\Intermediate\ProjectFiles\UnrealBuildTool.csproj.References')" />
   <ItemGroup>
-    <PlatformExtensions Include="../../../Platforms/*/Source/Programs/UnrealBuildTool/**/*.cs" />
-    <Compile Include="@(PlatformExtensions)">
-      <Link>Platform/$([System.Text.RegularExpressions.Regex]::Split(`%(Identity)`, `[\\/]`)[4])/_</Link>
-    </Compile>
+    <Compile Include="@(PlatformExtensionCompile)" />
   </ItemGroup>
-
   <!-- To modify your build process, add your task inside one of the targets below and uncomment it. 
        Other similar extension points exist, see Microsoft.Common.targets.
   <Target Name="BeforeBuild">
