Fix broken UBT -FastPDB command line option by using underlying type for nullable types when parsing command lines

diff --git a/Engine/Source/Programs/UnrealBuildTool/System/CommandLine.cs b/Engine/Source/Programs/UnrealBuildTool/System/CommandLine.cs
--- a/Engine/Source/Programs/UnrealBuildTool/System/CommandLine.cs
+++ b/Engine/Source/Programs/UnrealBuildTool/System/CommandLine.cs
@@ -46,6 +46,16 @@ namespace UnrealBuildTool
 			ParseAndRemoveArguments(Arguments.ToList(), TargetObject);
 		}
 
+		private static Type NonNullableType(Type type)
+		{
+			Type UnderlyingType = Nullable.GetUnderlyingType(type);
+			if (UnderlyingType != null)
+			{
+				return UnderlyingType;
+			}
+			return type;
+		}
+
 		/// <summary>
 		/// Parse the given list of arguments, and remove any that are parsed successfully
 		/// </summary>
@@ -65,7 +75,7 @@ namespace UnrealBuildTool
 						string Prefix = Attribute.Prefix;
 						if(Prefix == null)
 						{
-							if(FieldInfo.FieldType == typeof(bool))
+							if(NonNullableType(FieldInfo.FieldType) == typeof(bool))
 							{
 								Prefix = String.Format("-{0}", FieldInfo.Name);
 							}
@@ -76,7 +86,7 @@ namespace UnrealBuildTool
 						}
 						else
 						{
-							if(FieldInfo.FieldType != typeof(bool) && Attribute.Value == null && !Prefix.EndsWith("=") && !Prefix.EndsWith(":"))
+							if(NonNullableType(FieldInfo.FieldType) != typeof(bool) && Attribute.Value == null && !Prefix.EndsWith("=") && !Prefix.EndsWith(":"))
 							{
 								Prefix = Prefix + "=";
 							}
@@ -119,7 +129,7 @@ namespace UnrealBuildTool
 						{
 							AssignValue(Parameter, Argument.Substring(EqualsIdx + 1), TargetObject, AssignedFieldToParameter);
 						}
-						else if(Parameter.FieldInfo.FieldType == typeof(bool))
+						else if(NonNullableType(Parameter.FieldInfo.FieldType) == typeof(bool))
 						{
 							AssignValue(Parameter, "true", TargetObject, AssignedFieldToParameter);
 						}
@@ -268,6 +278,7 @@ namespace UnrealBuildTool
 		/// <returns>True if the text could be parsed, false otherwise</returns>
 		static bool TryParseValue(Type FieldType, string Text, out object Value)
 		{
+			FieldType = NonNullableType(FieldType);
 			if(FieldType.IsEnum)
 			{
 				// Special handling for enums; parse the value ignoring case.
