Backed out local change submitted in error.

diff --git a/Engine/Source/Programs/UnrealBuildTool/System/Unity.cs b/Engine/Source/Programs/UnrealBuildTool/System/Unity.cs
--- a/Engine/Source/Programs/UnrealBuildTool/System/Unity.cs
+++ b/Engine/Source/Programs/UnrealBuildTool/System/Unity.cs
@@ -231,7 +231,7 @@ namespace UnrealBuildTool
 					{
 						// Every single file in the module appears in the working set, so don't bother using adaptive unity for this
 						// module.  Otherwise it would make full builds really slow.
-						//bUseAdaptiveUnityBuild = false;
+						bUseAdaptiveUnityBuild = false;
 					}
 				}
 
