Deterministic generation of .sln and .vcxproj files

Visual Studio solution files and project files are now generated deterministically by UnrealBuildTool to reduce solution and project reloads.

Three sources of non-determinism are addressed by this change:

1. VisualStudioSolutionFolder now uses a deterministic guid based on its folder path in the solution. Note that this relies on MasterProjectFolder.AddSubFolder avoiding sibling folders with duplicate names.
2. The Visualizers folder now uses a hard-coded Guid.
3. RulesCompiler.FindAllRulesFiles now sorts cached paths.

diff --git a/Engine/Source/Programs/UnrealBuildTool/ProjectFiles/VisualStudio/VCProjectFileGenerator.cs b/Engine/Source/Programs/UnrealBuildTool/ProjectFiles/VisualStudio/VCProjectFileGenerator.cs
--- a/Engine/Source/Programs/UnrealBuildTool/ProjectFiles/VisualStudio/VCProjectFileGenerator.cs
+++ b/Engine/Source/Programs/UnrealBuildTool/ProjectFiles/VisualStudio/VCProjectFileGenerator.cs
@@ -2,9 +2,10 @@
 
 using System;
 using System.Collections.Generic;
-using System.Text;
 using System.IO;
 using System.Linq;
+using System.Security.Cryptography;
+using System.Text;
 using System.Xml.Linq;
 using Tools.DotNETCommon;
 
@@ -21,18 +22,6 @@ namespace UnrealBuildTool
 		public VisualStudioSolutionFolder(ProjectFileGenerator InitOwnerProjectFileGenerator, string InitFolderName)
 			: base(InitOwnerProjectFileGenerator, InitFolderName)
 		{
-			// Generate a unique GUID for this folder
-			// NOTE: When saving generated project files, we ignore differences in GUIDs if every other part of the file
-			//       matches identically with the pre-existing file
-			FolderGUID = Guid.NewGuid();
-		}
-
-
-		/// GUID for this folder
-		public Guid FolderGUID
-		{
-			get;
-			private set;
 		}
 	}
 
@@ -378,6 +367,38 @@ namespace UnrealBuildTool
 			return SolutionConfigName;
 		}
 
+		static IDictionary<MasterProjectFolder, Guid> GenerateProjectFolderGuids(MasterProjectFolder RootFolder)
+		{
+			IDictionary<MasterProjectFolder, Guid> Guids = new Dictionary<MasterProjectFolder, Guid>();
+			foreach (MasterProjectFolder Folder in RootFolder.SubFolders)
+			{
+				GenerateProjectFolderGuids("UE4", Folder, Guids);
+			}
+			return Guids;
+		}
+
+		static void GenerateProjectFolderGuids(string ParentPath, MasterProjectFolder Folder, IDictionary<MasterProjectFolder, Guid> Guids)
+		{
+			string Path = String.Format("{0}/{1}", ParentPath, Folder.FolderName);
+			Guids[Folder] = MakeMd5Guid(Encoding.UTF8.GetBytes(Path));
+
+			foreach (MasterProjectFolder SubFolder in Folder.SubFolders)
+			{
+				GenerateProjectFolderGuids(Path, SubFolder, Guids);
+			}
+		}
+
+		static Guid MakeMd5Guid(byte[] Input)
+		{
+			byte[] Hash = MD5.Create().ComputeHash(Input);
+			Hash[6] = (byte)(0x30 | (Hash[6] & 0x0f)); // 0b0011'xxxx Version 3 UUID (MD5)
+			Hash[8] = (byte)(0x80 | (Hash[8] & 0x3f)); // 0b10xx'xxxx RFC 4122 UUID
+			Array.Reverse(Hash, 0, 4);
+			Array.Reverse(Hash, 4, 2);
+			Array.Reverse(Hash, 6, 2);
+			return new Guid(Hash);
+		}
+
 		/// <summary>
 		/// Writes the project files to disk
 		/// </summary>
@@ -467,6 +488,8 @@ namespace UnrealBuildTool
 				throw new BuildException("Unexpected ProjectFileFormat");
 			}
 
+			IDictionary<MasterProjectFolder, Guid> ProjectFolderGuids = GenerateProjectFolderGuids(RootFolder);
+
 			// Solution folders, files and project entries
 			{
 				// This the GUID that Visual Studio uses to identify a solution folder
@@ -474,23 +497,10 @@ namespace UnrealBuildTool
 
 				// Solution folders
 				{
-					List<MasterProjectFolder> AllSolutionFolders = new List<MasterProjectFolder>();
-					System.Action<List<MasterProjectFolder> /* Folders */ > GatherFoldersFunction = null;
-					GatherFoldersFunction = FolderList =>
-						{
-							AllSolutionFolders.AddRange(FolderList);
-							foreach (MasterProjectFolder CurSubFolder in FolderList)
-							{
-								GatherFoldersFunction(CurSubFolder.SubFolders);
-							}
-						};
-					GatherFoldersFunction(RootFolder.SubFolders);
-
-					AllSolutionFolders.Sort((Lhs, Rhs) => Lhs.FolderName.CompareTo(Rhs.FolderName));
-
-					foreach (VisualStudioSolutionFolder CurFolder in AllSolutionFolders)
+					IEnumerable<MasterProjectFolder> AllSolutionFolders = ProjectFolderGuids.Keys.OrderBy(Folder => Folder.FolderName).ThenBy(Folder => ProjectFolderGuids[Folder]);
+					foreach (MasterProjectFolder CurFolder in AllSolutionFolders)
 					{
-						string FolderGUIDString = CurFolder.FolderGUID.ToString("B").ToUpperInvariant();
+						string FolderGUIDString = ProjectFolderGuids[CurFolder].ToString("B").ToUpperInvariant();
 						VCSolutionFileContent.AppendLine("Project(\"" + SolutionFolderEntryGUID + "\") = \"" + CurFolder.FolderName + "\", \"" + CurFolder.FolderName + "\", \"" + FolderGUIDString + "\"");
 
 						// Add any files that are inlined right inside the solution folder
@@ -560,7 +570,7 @@ namespace UnrealBuildTool
 				FileReference VisualizersFile = FileReference.Combine(UnrealBuildTool.EngineDirectory, "Extras", "VisualStudioDebugging", "UE4.natvis");
 
 				// Add the visualizers at the solution level. Doesn't seem to be picked up from a makefile project in VS2017 15.8.5.
-				VCSolutionFileContent.AppendLine(String.Format("Project(\"{{2150E333-8FDC-42A3-9474-1A3956D46DE8}}\") = \"Visualizers\", \"Visualizers\", \"{0}\"", Guid.NewGuid().ToString("B").ToUpperInvariant()));
+				VCSolutionFileContent.AppendLine(String.Format("Project(\"{0}\") = \"Visualizers\", \"Visualizers\", \"{{1CCEC849-CC72-4C59-8C36-2F7C38706D4C}}\"", SolutionFolderEntryGUID));
 				VCSolutionFileContent.AppendLine("\tProjectSection(SolutionItems) = preProject");
 				VCSolutionFileContent.AppendLine("\t\t{0} = {0}", VisualizersFile.MakeRelativeTo(MasterProjectPath));
 				VCSolutionFileContent.AppendLine("\tEndProjectSection");
@@ -722,9 +732,9 @@ namespace UnrealBuildTool
 						System.Action<StringBuilder /* VCSolutionFileContent */, List<MasterProjectFolder> /* Folders */ > FolderProcessorFunction = null;
 						FolderProcessorFunction = (LocalVCSolutionFileContent, LocalMasterProjectFolders) =>
 							{
-								foreach (VisualStudioSolutionFolder CurFolder in LocalMasterProjectFolders)
+								foreach (MasterProjectFolder CurFolder in LocalMasterProjectFolders)
 								{
-									string CurFolderGUIDString = CurFolder.FolderGUID.ToString("B").ToUpperInvariant();
+									string CurFolderGUIDString = ProjectFolderGuids[CurFolder].ToString("B").ToUpperInvariant();
 
 									foreach (MSBuildProjectFile ChildProject in CurFolder.ChildProjects)
 									{
@@ -732,10 +742,10 @@ namespace UnrealBuildTool
 										LocalVCSolutionFileContent.AppendLine("		" + ChildProject.ProjectGUID.ToString("B").ToUpperInvariant() + " = " + CurFolderGUIDString);
 									}
 
-									foreach (VisualStudioSolutionFolder SubFolder in CurFolder.SubFolders)
+									foreach (MasterProjectFolder SubFolder in CurFolder.SubFolders)
 									{
 										//	e.g. "{BF6FB09F-A2A6-468F-BE6F-DEBE07EAD3EA} = {C43B6BB5-3EF0-4784-B896-4099753BCDA9}"
-										LocalVCSolutionFileContent.AppendLine("		" + SubFolder.FolderGUID.ToString("B").ToUpperInvariant() + " = " + CurFolderGUIDString);
+										LocalVCSolutionFileContent.AppendLine("		" + ProjectFolderGuids[SubFolder].ToString("B").ToUpperInvariant() + " = " + CurFolderGUIDString);
 									}
 
 									// Recurse into subfolders
diff --git a/Engine/Source/Programs/UnrealBuildTool/System/RulesCompiler.cs b/Engine/Source/Programs/UnrealBuildTool/System/RulesCompiler.cs
--- a/Engine/Source/Programs/UnrealBuildTool/System/RulesCompiler.cs
+++ b/Engine/Source/Programs/UnrealBuildTool/System/RulesCompiler.cs
@@ -256,6 +256,9 @@ namespace UnrealBuildTool
 					DirectoryItem BaseDirectory = DirectoryItem.GetItemByDirectoryReference(Directory);
 					Queue.Enqueue(() => FindAllRulesFilesRecursively(BaseDirectory, Cache, Queue));
 				}
+				Cache.ModuleRules.Sort((A, B) => A.FullName.CompareTo(B.FullName));
+				Cache.TargetRules.Sort((A, B) => A.FullName.CompareTo(B.FullName));
+				Cache.AutomationModules.Sort((A, B) => A.FullName.CompareTo(B.FullName));
 				RootFolderToRulesFileCache[Directory] = Cache;
 			}
 
