UAT: Change the default DeployFolder name to be ShortProjectName-RootFolderName-UserName truncated to 50 characters.
This will deploy the same project from different streams or users to to different workspaces.
Use the DeployFolder instead of the ShortProjectName for all platforms.
The existing -deploy=DeployFolder argument can be used to specify any custom workspace name.

diff --git a/Engine/Source/Programs/AutomationTool/AutomationUtils/ProjectParams.cs b/Engine/Source/Programs/AutomationTool/AutomationUtils/ProjectParams.cs
--- a/Engine/Source/Programs/AutomationTool/AutomationUtils/ProjectParams.cs
+++ b/Engine/Source/Programs/AutomationTool/AutomationUtils/ProjectParams.cs
@@ -765,7 +765,7 @@ namespace AutomationTool
 			// if the user specified -deploy but no folder, set the default
 			if (this.Deploy && string.IsNullOrEmpty(this.DeployFolder))
 			{
-				this.DeployFolder = this.ShortProjectName;
+				this.DeployFolder = UnrealBuildTool.DeployExports.GetDefaultDeployFolder(this.ShortProjectName);
 			}
 			else if (string.IsNullOrEmpty(this.DeployFolder) == false)
 			{
diff --git a/Engine/Source/Programs/UnrealBuildTool/System/DeployExports.cs b/Engine/Source/Programs/UnrealBuildTool/System/DeployExports.cs
new file mode 100644
--- /dev/null
+++ b/Engine/Source/Programs/UnrealBuildTool/System/DeployExports.cs
@@ -0,0 +1,48 @@
+// Copyright Epic Games, Inc. All Rights Reserved.
+
+using System;
+using System.Collections.Generic;
+using System.Linq;
+using System.Text;
+using System.Threading.Tasks;
+using Tools.DotNETCommon;
+
+namespace UnrealBuildTool
+{
+	/// <summary>
+	/// General deploy support functions, exported for UAT
+	/// </summary>
+	public class DeployExports
+	{
+		/// <summary>
+		/// The max length of the deploy folder name.
+		/// </summary>
+		public const int DeployFolderMaxLength = 50; 
+
+		/// <summary>
+		/// Gets the default deploy folder name on the format ShortProjectName-RootDirectoryName-UserName,
+		/// truncated to DeployFolderMaxLength characters, with dots '.' and whitespace ' ' stripped.
+		/// </summary>
+		/// <param name="ShortProjectName">Project to deploy</param>
+		/// <returns>The default deploy folder name.</returns>
+		public static string GetDefaultDeployFolder(string ShortProjectName)
+		{
+			string DeployFolder = string.Format("{0}-{1}-{2}",
+				ShortProjectName,
+				UnrealBuildTool.RootDirectory.GetDirectoryName(),
+				Environment.UserName);
+
+			DeployFolder = DeployFolder.Replace(" ", "");
+			DeployFolder = DeployFolder.Replace(".", "");
+
+			if (DeployFolder.Length > DeployFolderMaxLength)
+			{
+				DeployFolder = DeployFolder.Substring(0, DeployFolderMaxLength);
+				char[] CharsToTrim = {'-'};
+				DeployFolder = DeployFolder.Trim(CharsToTrim);
+			}
+
+			return DeployFolder;
+		}
+	}
+}
diff --git a/Engine/Source/Programs/UnrealBuildTool/UnrealBuildTool.csproj b/Engine/Source/Programs/UnrealBuildTool/UnrealBuildTool.csproj
--- a/Engine/Source/Programs/UnrealBuildTool/UnrealBuildTool.csproj
+++ b/Engine/Source/Programs/UnrealBuildTool/UnrealBuildTool.csproj
@@ -229,6 +229,7 @@
     <Compile Include="System\ConfigFileAttribute.cs" />
     <Compile Include="System\ConfigHierarchy.cs" />
     <Compile Include="System\CustomBuildSteps.cs" />
+    <Compile Include="System\DeployExports.cs" />
     <Compile Include="System\InstalledPlatformInfo.cs" />
     <Compile Include="System\LocalizationDescriptor.cs" />
     <Compile Include="System\ObsoleteOverrideAttribute.cs" />
