Prevent IOS provisioning parse failure when remote building on a Mac without internet access

diff --git a/Engine/Source/Programs/AutomationTool/IOS/ExportIPAFromArchive.cs b/Engine/Source/Programs/AutomationTool/IOS/ExportIPAFromArchive.cs
--- a/Engine/Source/Programs/AutomationTool/IOS/ExportIPAFromArchive.cs
+++ b/Engine/Source/Programs/AutomationTool/IOS/ExportIPAFromArchive.cs
@@ -44,6 +44,7 @@ public class ExportIPAFromArchive : BuildCommand
 	protected void CreateExportPlist(string OutputPath)
 	{
 		XmlDocument Doc = new XmlDocument();
+		Doc.XmlResolver = null;
 		Doc.Load(OptionsFile);
 
 		// select all entries via xpath
diff --git a/Engine/Source/Programs/UnrealBuildTool/Platform/IOS/MobileProvisionContents.cs b/Engine/Source/Programs/UnrealBuildTool/Platform/IOS/MobileProvisionContents.cs
--- a/Engine/Source/Programs/UnrealBuildTool/Platform/IOS/MobileProvisionContents.cs
+++ b/Engine/Source/Programs/UnrealBuildTool/Platform/IOS/MobileProvisionContents.cs
@@ -157,6 +157,7 @@ namespace UnrealBuildTool
 									byte[] Data = Reader.ReadBytes(NextField.Length);
 
 									XmlDocument Document = new XmlDocument();
+									Document.XmlResolver = null;
 									Document.Load(new MemoryStream(Data));
 									return Document;
 								}
