- Fixed the platform extension config location checks in C# to match C++

diff --git a/Engine/Source/Programs/UnrealBuildTool/System/ConfigHierarchy.cs b/Engine/Source/Programs/UnrealBuildTool/System/ConfigHierarchy.cs
--- a/Engine/Source/Programs/UnrealBuildTool/System/ConfigHierarchy.cs
+++ b/Engine/Source/Programs/UnrealBuildTool/System/ConfigHierarchy.cs
@@ -764,17 +764,17 @@ namespace UnrealBuildTool
 		static ConfigLayer[] ConfigLayers = 
 		{
 			// Engine/Base.ini
-			new ConfigLayer { Path = "{ENGINE}/Config/Base.ini" }, //, Flag = EConfigFlag.Required },
+			new ConfigLayer { Path = "{ENGINE}/Base.ini" }, //, Flag = EConfigFlag.Required },
 			// Engine/Base*.ini
- 			new ConfigLayer { Path = "{ENGINE}/Config/{ED}{EF}Base{TYPE}.ini" },
+ 			new ConfigLayer { Path = "{ENGINE}/{ED}{EF}Base{TYPE}.ini" },
 			// Engine/Platform/BasePlatform*.ini
-			new ConfigLayer { Path = "{ENGINE}/Config/{ED}{PLATFORM}/{EF}Base{PLATFORM}{TYPE}.ini", ExtEnginePath = "{EXTENGINE}/Config/{ED}{EF}Base{PLATFORM}{TYPE}.ini" },
+			new ConfigLayer { Path = "{ENGINE}/{ED}{PLATFORM}/{EF}Base{PLATFORM}{TYPE}.ini", ExtEnginePath = "{EXTENGINE}/{ED}{EF}Base{PLATFORM}{TYPE}.ini" },
 			// Project/Default*.ini
-			new ConfigLayer { Path = "{PROJECT}/Config/{ED}{EF}Default{TYPE}.ini" }, //, Flag = EConfigFlag.AllowCommandLineOverride },
+			new ConfigLayer { Path = "{PROJECT}/{ED}{EF}Default{TYPE}.ini" }, //, Flag = EConfigFlag.AllowCommandLineOverride },
 			// Engine/Platform/Platform*.ini
-			new ConfigLayer { Path = "{ENGINE}/Config/{ED}{PLATFORM}/{EF}{PLATFORM}{TYPE}.ini", ExtEnginePath = "{EXTENGINE}/Config/{ED}{EF}{PLATFORM}{TYPE}.ini" },
+			new ConfigLayer { Path = "{ENGINE}/{ED}{PLATFORM}/{EF}{PLATFORM}{TYPE}.ini", ExtEnginePath = "{EXTENGINE}/{ED}{EF}{PLATFORM}{TYPE}.ini" },
 			// Project/Platform/Platform*.ini
-			new ConfigLayer { Path = "{PROJECT}/Config/{ED}{PLATFORM}/{EF}{PLATFORM}{TYPE}.ini", ExtProjectPath = "{EXTPROJECT}/Config/{ED}{EF}{PLATFORM}{TYPE}.ini" },
+			new ConfigLayer { Path = "{PROJECT}/{ED}{PLATFORM}/{EF}{PLATFORM}{TYPE}.ini", ExtProjectPath = "{EXTPROJECT}/{ED}{EF}{PLATFORM}{TYPE}.ini" },
 
 			// UserSettings/.../User*.ini
 			new ConfigLayer { Path = "{USERSETTINGS}/Unreal Engine/Engine/Config/User{TYPE}.ini" },
@@ -826,8 +826,8 @@ namespace UnrealBuildTool
 			out bool bHasPlatformTag, out bool bHasProjectTag, out bool bHasExpansionTag)
 		{
 			// cache some platform extension information that can be used inside the loops
-			string PlatformExtensionEngineConfigDir = DirectoryReference.Combine(UnrealBuildTool.EnginePlatformExtensionsDirectory, PlatformExtensionName).FullName;
-			string PlatformExtensionProjectConfigDir = ProjectDir != null ? DirectoryReference.Combine(UnrealBuildTool.ProjectPlatformExtensionsDirectory(ProjectDir), PlatformExtensionName).FullName : null;
+			string PlatformExtensionEngineConfigDir = DirectoryReference.Combine(UnrealBuildTool.EnginePlatformExtensionsDirectory, PlatformExtensionName, "Config").FullName;
+			string PlatformExtensionProjectConfigDir = ProjectDir != null ? DirectoryReference.Combine(UnrealBuildTool.ProjectPlatformExtensionsDirectory(ProjectDir), PlatformExtensionName, "Config").FullName : null;
 			bool bHasPlatformExtensionEngineConfigDir = Directory.Exists(PlatformExtensionEngineConfigDir);
 			bool bHasPlatformExtensionProjectConfigDir = PlatformExtensionProjectConfigDir != null && Directory.Exists(PlatformExtensionProjectConfigDir);
 
@@ -855,7 +855,7 @@ namespace UnrealBuildTool
 				}
 				else
 				{
-					LayerPath = Layer.Path.Replace("{PROJECT}", ProjectDir.FullName);
+					LayerPath = Layer.Path.Replace("{PROJECT}", Path.Combine(ProjectDir.FullName, "Config"));
 				}
 			}
 			else
@@ -866,7 +866,7 @@ namespace UnrealBuildTool
 				}
 				else
 				{
-					LayerPath = Layer.Path.Replace("{ENGINE}", UnrealBuildTool.EngineDirectory.FullName);
+					LayerPath = Layer.Path.Replace("{ENGINE}", Path.Combine(UnrealBuildTool.EngineDirectory.FullName, "Config"));
 				}
 			}
 			LayerPath = LayerPath.Replace("{TYPE}", BaseIniName);
