Use VS2017 by default when running static analysis with PVS studio.

#rnx

diff --git a/Engine/Source/Programs/UnrealBuildTool/Platform/Windows/UEBuildWindows.cs b/Engine/Source/Programs/UnrealBuildTool/Platform/Windows/UEBuildWindows.cs
--- a/Engine/Source/Programs/UnrealBuildTool/Platform/Windows/UEBuildWindows.cs
+++ b/Engine/Source/Programs/UnrealBuildTool/Platform/Windows/UEBuildWindows.cs
@@ -683,7 +683,14 @@ namespace UnrealBuildTool
 			// Set the compiler version if necessary
 			if (Target.WindowsPlatform.Compiler == WindowsCompiler.Default)
 			{
-				Target.WindowsPlatform.Compiler = GetDefaultCompiler(Target.ProjectFile);
+				if (Target.WindowsPlatform.StaticAnalyzer == WindowsStaticAnalyzer.PVSStudio && HasCompiler(WindowsCompiler.VisualStudio2017))
+				{
+					Target.WindowsPlatform.Compiler = WindowsCompiler.VisualStudio2017;
+				}
+				else
+				{
+					Target.WindowsPlatform.Compiler = GetDefaultCompiler(Target.ProjectFile);
+				}
 			}
 
 			// Disable linking if we're using a static analyzer
