UBT: Output toolchain info when doing XGE export.

#rnx

diff --git a/Engine/Source/Programs/UnrealBuildTool/Modes/BuildMode.cs b/Engine/Source/Programs/UnrealBuildTool/Modes/BuildMode.cs
--- a/Engine/Source/Programs/UnrealBuildTool/Modes/BuildMode.cs
+++ b/Engine/Source/Programs/UnrealBuildTool/Modes/BuildMode.cs
@@ -280,16 +280,20 @@ namespace UnrealBuildTool
 				// Execute the actions
 				if ((Options & BuildOptions.XGEExport) != 0)
 				{
+					OutputToolchainInfo(TargetDescriptors, Makefiles);
+
 					// Just export to an XML file
-					using(Timeline.ScopeEvent("XGE.ExportActions()"))
+					using (Timeline.ScopeEvent("XGE.ExportActions()"))
 					{
 						XGE.ExportActions(MergedActionsToExecute);
 					}
 				}
 				else if(WriteOutdatedActionsFile != null)
 				{
+					OutputToolchainInfo(TargetDescriptors, Makefiles);
+
 					// Write actions to an output file
-					using(Timeline.ScopeEvent("ActionGraph.WriteActions"))
+					using (Timeline.ScopeEvent("ActionGraph.WriteActions"))
 					{
 						ActionGraph.ExportJson(MergedActionsToExecute, WriteOutdatedActionsFile);
 					}
