For non dll, assume our ../../.. structure for rpaths. Keep the new way to DSOs which will play a role when it comes to plugins etc.

diff --git a/Engine/Source/Programs/UnrealBuildTool/Platform/Linux/LinuxToolChain.cs b/Engine/Source/Programs/UnrealBuildTool/Platform/Linux/LinuxToolChain.cs
--- a/Engine/Source/Programs/UnrealBuildTool/Platform/Linux/LinuxToolChain.cs
+++ b/Engine/Source/Programs/UnrealBuildTool/Platform/Linux/LinuxToolChain.cs
@@ -1667,14 +1667,19 @@ namespace UnrealBuildTool
 
 				if ((AdditionalLibrary.Contains("Plugins") || AdditionalLibrary.Contains("Binaries/ThirdParty") || AdditionalLibrary.Contains("Binaries\\ThirdParty")) && Path.GetDirectoryName(AdditionalLibrary) != Path.GetDirectoryName(OutputFile.AbsolutePath))
 				{
-					// Remove the root UnrealBuildTool.RootDirectory from the RuntimeLibaryPath
-					string AdditionalLibraryRootPath = new FileReference(AdditionalLibrary).Directory.MakeRelativeTo(UnrealBuildTool.RootDirectory);
+					string RelativePath = new FileReference(AdditionalLibrary).Directory.MakeRelativeTo(OutputFile.Location.Directory);
 
-					// Figure out how many dirs we need to go back
-					string RelativeRootPath = UnrealBuildTool.RootDirectory.MakeRelativeTo(OutputFile.Location.Directory);
+					if (LinkEnvironment.bIsBuildingDLL)
+					{
+						// Remove the root UnrealBuildTool.RootDirectory from the RuntimeLibaryPath
+						string AdditionalLibraryRootPath = new FileReference(AdditionalLibrary).Directory.MakeRelativeTo(UnrealBuildTool.RootDirectory);
+
+						// Figure out how many dirs we need to go back
+						string RelativeRootPath = UnrealBuildTool.RootDirectory.MakeRelativeTo(OutputFile.Location.Directory);
 
-					// Combine the two together ie. number of ../ + the path after the root
-					string RelativePath = Path.Combine(RelativeRootPath, AdditionalLibraryRootPath);
+						// Combine the two together ie. number of ../ + the path after the root
+						RelativePath = Path.Combine(RelativeRootPath, AdditionalLibraryRootPath);
+					}
 
 					// On Windows, MakeRelativeTo can silently fail if the engine and the project are located on different drives
 					if (CrossCompiling() && RelativePath.StartsWith(UnrealBuildTool.RootDirectory.FullName))
@@ -1699,14 +1704,24 @@ namespace UnrealBuildTool
 
 				if(!RelativePath.StartsWith("$"))
 				{
-					// Remove the root UnrealBuildTool.RootDirectory from the RuntimeLibaryPath
-					string RuntimeLibraryRootPath = new DirectoryReference(RuntimeLibaryPath).MakeRelativeTo(UnrealBuildTool.RootDirectory);
+					if (LinkEnvironment.bIsBuildingDLL)
+					{
+						// Remove the root UnrealBuildTool.RootDirectory from the RuntimeLibaryPath
+						string RuntimeLibraryRootPath = new DirectoryReference(RuntimeLibaryPath).MakeRelativeTo(UnrealBuildTool.RootDirectory);
 
-					// Figure out how many dirs we need to go back
-					string RelativeRootPath = UnrealBuildTool.RootDirectory.MakeRelativeTo(OutputFile.Location.Directory);
+						// Figure out how many dirs we need to go back
+						string RelativeRootPath = UnrealBuildTool.RootDirectory.MakeRelativeTo(OutputFile.Location.Directory);
 
-					// Combine the two together ie. number of ../ + the path after the root
-					RelativePath = Path.Combine(RelativeRootPath, RuntimeLibraryRootPath);
+						// Combine the two together ie. number of ../ + the path after the root
+						RelativePath = Path.Combine(RelativeRootPath, RuntimeLibraryRootPath);
+					}
+					else
+					{
+						string RelativeRootPath = new DirectoryReference(RuntimeLibaryPath).MakeRelativeTo(UnrealBuildTool.RootDirectory);
+
+						// We're assuming that the binary will be placed according to our ProjectName/Binaries/Platform scheme
+						RelativePath = Path.Combine("..", "..", "..", RelativeRootPath);
+					}
 				}
 
 				// On Windows, MakeRelativeTo can silently fail if the engine and the project are located on different drives
