HTML5 - emscripten (1.38.31) delete EmscriptenCache folder upon new toolchain detection

#rn

diff --git a/Engine/Source/Programs/UnrealBuildTool/Platform/HTML5/HTML5SDKInfo.cs b/Engine/Source/Programs/UnrealBuildTool/Platform/HTML5/HTML5SDKInfo.cs
--- a/Engine/Source/Programs/UnrealBuildTool/Platform/HTML5/HTML5SDKInfo.cs
+++ b/Engine/Source/Programs/UnrealBuildTool/Platform/HTML5/HTML5SDKInfo.cs
@@ -263,6 +263,58 @@ namespace UnrealBuildTool
 				Log.TraceErrorOnce(" Recreation of Emscripten Temp folder failed because of " + Ex.ToString());
 			}
 
+			// ----------------------------------------
+			// double check clang version, if different -- NUKE cached asmjs folder
+			string CachePath = Path.Combine(HTML5Intermediatory, "EmscriptenCache");
+			try
+			{
+				string isVanillaPath = Path.Combine(CachePath, "is_vanilla.txt");
+				bool deleteAsmjsCache = true;
+				if (File.Exists(isVanillaPath))
+				{
+					string cachedClang = File.ReadAllText(isVanillaPath);
+					if ( cachedClang.Equals("0:"+LLVM_ROOT, StringComparison.Ordinal) )
+					{
+						deleteAsmjsCache = false;
+					}
+				}
+				// ----------------------------------------
+				// sometimes, CIS machines need EmscriptenCache folder to be deleted manually
+				// automate this on every toolchain upgrade
+				string ue4emsdkcachePath = Path.Combine(HTML5Intermediatory, "ue4emsdk.txt");
+				if (File.Exists(ue4emsdkcachePath))
+				{
+					string cachedClang = File.ReadAllText(ue4emsdkcachePath);
+					if ( cachedClang.Equals(SDKVersion, StringComparison.Ordinal) )
+					{
+						deleteAsmjsCache = false;
+					}
+					else
+					{
+						File.Delete(ue4emsdkcachePath); // to be recreated
+					}
+				}
+				else
+				{
+					deleteAsmjsCache = true; // this overrides everything
+				}
+				// ----------------------------------------
+				if ( deleteAsmjsCache )
+				{
+					string asmjsPath = Path.Combine(CachePath, "asmjs");
+					if (Directory.Exists(asmjsPath))
+					{
+						Directory.Delete(asmjsPath, true);
+					}
+					File.WriteAllText(ue4emsdkcachePath, SDKVersion);
+				}
+			}
+			catch (Exception Ex)
+			{
+				Log.TraceErrorOnce(" Recreation of Emscripten Temp folder failed because of " + Ex.ToString());
+			}
+
+			// done
 			return TempPath;
 		}
 
