HTML5 - emscripten (1.38.31) delete EmscriptenCache folder upon new toolchain detection

(trying to get CIS green again)

#rnx

diff --git a/Engine/Source/Programs/UnrealBuildTool/Platform/HTML5/HTML5SDKInfo.cs b/Engine/Source/Programs/UnrealBuildTool/Platform/HTML5/HTML5SDKInfo.cs
--- a/Engine/Source/Programs/UnrealBuildTool/Platform/HTML5/HTML5SDKInfo.cs
+++ b/Engine/Source/Programs/UnrealBuildTool/Platform/HTML5/HTML5SDKInfo.cs
@@ -268,8 +268,8 @@ namespace UnrealBuildTool
 			string CachePath = Path.Combine(HTML5Intermediatory, "EmscriptenCache");
 			try
 			{
-				string isVanillaPath = Path.Combine(CachePath, "is_vanilla.txt");
 				bool deleteAsmjsCache = true;
+				string isVanillaPath = Path.Combine(CachePath, "is_vanilla.txt");
 				if (File.Exists(isVanillaPath))
 				{
 					string cachedClang = File.ReadAllText(isVanillaPath);
@@ -301,10 +301,9 @@ namespace UnrealBuildTool
 				// ----------------------------------------
 				if ( deleteAsmjsCache )
 				{
-					string asmjsPath = Path.Combine(CachePath, "asmjs");
-					if (Directory.Exists(asmjsPath))
+					if (Directory.Exists(CachePath))
 					{
-						Directory.Delete(asmjsPath, true);
+						Directory.Delete(CachePath, true);
 					}
 					File.WriteAllText(ue4emsdkcachePath, SDKVersion);
 				}
