HTML5 - mac python: include dynamic lib path to build

#rnx

diff --git a/Engine/Source/Programs/UnrealBuildTool/Platform/HTML5/HTML5SDKInfo.cs b/Engine/Source/Programs/UnrealBuildTool/Platform/HTML5/HTML5SDKInfo.cs
--- a/Engine/Source/Programs/UnrealBuildTool/Platform/HTML5/HTML5SDKInfo.cs
+++ b/Engine/Source/Programs/UnrealBuildTool/Platform/HTML5/HTML5SDKInfo.cs
@@ -488,6 +488,16 @@ namespace UnrealBuildTool
 			return new FileReference(PYTHON);
 		}
 
+		/// <summary>
+		///
+		/// </summary>
+		/// <returns></returns>
+		public static string MacPythonLib() // UE-75402
+		{
+			string UE4PythonPath = FileReference.Combine(UnrealBuildTool.EngineDirectory, "Binaries", "ThirdParty", "Python").FullName;
+			return Path.Combine(UE4PythonPath, "Mac", "lib", "python2.7", "lib-dynload");
+		}
+
 		/// <summary>
 		///
 		/// </summary>
diff --git a/Engine/Source/Programs/UnrealBuildTool/Platform/HTML5/HTML5ToolChain.cs b/Engine/Source/Programs/UnrealBuildTool/Platform/HTML5/HTML5ToolChain.cs
--- a/Engine/Source/Programs/UnrealBuildTool/Platform/HTML5/HTML5ToolChain.cs
+++ b/Engine/Source/Programs/UnrealBuildTool/Platform/HTML5/HTML5ToolChain.cs
@@ -221,6 +221,10 @@ namespace UnrealBuildTool
 				// Packaging on Window needs this - zap any existing HOME environment variables to prevent any accidental pick ups
 				Environment.SetEnvironmentVariable("HOME", "");
 			}
+			if (BuildHostPlatform.Current.Platform == UnrealTargetPlatform.Mac)
+			{
+				Environment.SetEnvironmentVariable("LD_LIBRARY_PATH", HTML5SDKInfo.MacPythonLib()); // UE-75402
+			}
 			return Result;
 		}
 
