Fix for framework submission. I had used a new C# feature that we don't support yet.

diff --git a/Engine/Source/Programs/UnrealBuildTool/ProjectFiles/Xcode/XcodeFrameworkWrapperProject.cs b/Engine/Source/Programs/UnrealBuildTool/ProjectFiles/Xcode/XcodeFrameworkWrapperProject.cs
--- a/Engine/Source/Programs/UnrealBuildTool/ProjectFiles/Xcode/XcodeFrameworkWrapperProject.cs
+++ b/Engine/Source/Programs/UnrealBuildTool/ProjectFiles/Xcode/XcodeFrameworkWrapperProject.cs
@@ -271,14 +271,16 @@ namespace UnrealBuildTool.ProjectFiles.Xcode
 		public static bool GetBuildAsFramework(FileReference ProjectFile)
 		{
 			ConfigHierarchy Ini = ConfigCache.ReadHierarchy(ConfigHierarchyType.Engine, DirectoryReference.FromFile(ProjectFile), UnrealTargetPlatform.IOS);
-			Ini.GetBool("/Script/IOSRuntimeSettings.IOSRuntimeSettings", "bBuildAsFramework", out bool bBuildAsFramework);
+			bool bBuildAsFramework;
+			Ini.GetBool("/Script/IOSRuntimeSettings.IOSRuntimeSettings", "bBuildAsFramework", out bBuildAsFramework);
 			return bBuildAsFramework;
 		}
 		
 		public static bool GetGenerateFrameworkWrapperProject(FileReference ProjectFile)
 		{
 		ConfigHierarchy Ini = ConfigCache.ReadHierarchy(ConfigHierarchyType.Engine, DirectoryReference.FromFile(ProjectFile), UnrealTargetPlatform.IOS);
-			Ini.GetBool("/Script/IOSRuntimeSettings.IOSRuntimeSettings", "bGenerateFrameworkWrapperProject", out bool bGenerateFrameworkWrapperProject);
+			bool bGenerateFrameworkWrapperProject;
+			Ini.GetBool("/Script/IOSRuntimeSettings.IOSRuntimeSettings", "bGenerateFrameworkWrapperProject", out bGenerateFrameworkWrapperProject);
 			return bGenerateFrameworkWrapperProject;
 		}
 	}
