Fixed blueprint only projects building as frameworks and ironed out some other issues building as a framework.

*IOSToolChain - Added quotes around commandline args for IOSPostBuildSync. Paths with spaces were broken.
*IOSPlatform.Automation - Fixed "Failed to find game binary" if UE4Game built as framework.
*XCodeFrameworkWrapperProject - Ensure files are writable before modifying them (fixes issue with UE4Game being in Perforce locked dir)
*UEBuildIOS - Get whether to compile as framework from ini instead of looking at Target so that UE4Game projects will use project setting and not just default BaseEngine.ini value.
*BaseEngine.ini - Set bBuildAsFramework to default False so legacy projects don't break.
*UnrealView.mm - Don't add ue4commandline.txt to CommandLine because this happens somewhere else as well.
*Config.cs - Removed my changes for frameworks
*CookTime.cs - Removed my changes for frameworks
*iPhonePackager.cs - Removed my changes for frameworks
*ue4commandline.txt - Added template for commandline instead of being blank
*IOSRuntimeSettings - Removed GenerateFrameworkWrapperIfNonexistent setting. We should just always generate it if it is nonexistent.

diff --git a/Engine/Build/IOS/Resources/FrameworkWrapper/FrameworkWrapper/Views/UnrealView.mm b/Engine/Build/IOS/Resources/FrameworkWrapper/FrameworkWrapper/Views/UnrealView.mm
--- a/Engine/Build/IOS/Resources/FrameworkWrapper/FrameworkWrapper/Views/UnrealView.mm
+++ b/Engine/Build/IOS/Resources/FrameworkWrapper/FrameworkWrapper/Views/UnrealView.mm
@@ -89,21 +89,6 @@ -(void) createUnrealView
     if (sharedUnrealView == nil)
 	{
 #if CAN_USE_UE4
-		extern FString GSavedCommandLine;
-
-		NSString* Path = [[NSBundle mainBundle] pathForResource:@"ue4commandline" ofType:@"txt"];
-		NSError* Error = nil;
-		NSString* CommandLineFileContents = [[NSString stringWithContentsOfFile:Path encoding:NSUTF8StringEncoding error:&Error]
-											 stringByTrimmingCharactersInSet: [NSCharacterSet whitespaceAndNewlineCharacterSet]];
-		if (Error == nil && [CommandLineFileContents length] > 0)
-		{
-			GSavedCommandLine = CommandLineFileContents;
-		}
-		else
-		{
-			GSavedCommandLine = self.CommandLine;
-		}
-
 		
 //       [self SetupEmbeddedToNativeCallback];
 
diff --git a/Engine/Build/IOS/Resources/FrameworkWrapper/FrameworkWrapper/ue4commandline.txt b/Engine/Build/IOS/Resources/FrameworkWrapper/FrameworkWrapper/ue4commandline.txt
--- a/Engine/Build/IOS/Resources/FrameworkWrapper/FrameworkWrapper/ue4commandline.txt
+++ b/Engine/Build/IOS/Resources/FrameworkWrapper/FrameworkWrapper/ue4commandline.txt
@@ -1 +1 @@
-
+../../../PROJECT_NAME/PROJECT_NAME.uproject
diff --git a/Engine/Config/BaseEngine.ini b/Engine/Config/BaseEngine.ini
--- a/Engine/Config/BaseEngine.ini
+++ b/Engine/Config/BaseEngine.ini
@@ -1901,7 +1901,7 @@ bAllowRemoteRotation=True
 bUseRemoteAsVirtualJoystick=True
 bUseAbsoluteDpadValues=False
 bAllowControllers=True
-bBuildAsFramework=True
+bBuildAsFramework=False
 bGenerateFrameworkWrapperProject=True
 bGeneratedSYMFile=False
 bDisableHTTPS=false
diff --git a/Engine/Source/Programs/AutomationTool/IOS/IOSPlatform.Automation.cs b/Engine/Source/Programs/AutomationTool/IOS/IOSPlatform.Automation.cs
--- a/Engine/Source/Programs/AutomationTool/IOS/IOSPlatform.Automation.cs
+++ b/Engine/Source/Programs/AutomationTool/IOS/IOSPlatform.Automation.cs
@@ -292,11 +292,13 @@ public class IOSPlatform : Platform
 	{
 		LogInformation("Package {0}", Params.RawProjectPath);
 
+		bool bIsBuiltAsFramework = IsBuiltAsFramework(Params, SC);
+
 		// ensure the ue4game binary exists, if applicable
 #if !PLATFORM_MAC
 		string ProjectGameExeFilename = Params.GetProjectExeForPlatform(Platform).ToString();
 		string FullExePath = CombinePaths(Path.GetDirectoryName(ProjectGameExeFilename), SC.StageExecutables[0] + (UnrealBuildTool.BuildHostPlatform.Current.Platform != UnrealTargetPlatform.Mac ? ".stub" : ""));
-		if (!SC.IsCodeBasedProject && !FileExists_NoExceptions(FullExePath))
+		if (!SC.IsCodeBasedProject && !FileExists_NoExceptions(FullExePath) && !bIsBuiltAsFramework)
 		{
 			LogError("Failed to find game binary " + FullExePath);
 			throw new AutomationException(ExitCode.Error_MissingExecutable, "Stage Failed. Could not find binary {0}. You may need to build the UE4 project with your target configuration and platform.", FullExePath);
@@ -324,14 +326,14 @@ public class IOSPlatform : Platform
 		//   - Code here probably needs to be updated to write 0 byte files as 1 byte (difference with IPP, was required at one point when using Ionic.Zip to prevent issues on device, maybe not needed anymore?)
 		if (UnrealBuildTool.BuildHostPlatform.Current.Platform == UnrealTargetPlatform.Mac)
 		{
-            Console.WriteLine("RYRY - IsBuiltAsFramework: " + IsBuiltAsFramework(Params, SC));
 			// If we're building as a framework, then we already have everything we need in the .app
 			// so simply package it up as an ipa
-			if (IsBuiltAsFramework(Params, SC))
-			{
-				PackageIPA(Params, ProjectGameExeFilename, SC);
-				return;
-			}
+            if (bIsBuiltAsFramework)
+            {
+                PackageIPA(Params, ProjectGameExeFilename, SC);
+                return;
+            }
+
 
 			// copy in all of the artwork and plist
 			PrepForUATPackageOrDeploy(TargetConfiguration, Params.RawProjectPath,
diff --git a/Engine/Source/Programs/IOS/iPhonePackager/Config.cs b/Engine/Source/Programs/IOS/iPhonePackager/Config.cs
--- a/Engine/Source/Programs/IOS/iPhonePackager/Config.cs
+++ b/Engine/Source/Programs/IOS/iPhonePackager/Config.cs
@@ -1,4 +1,4 @@
-/*
+﻿/*
  * Copyright 1998-2019 Epic Games, Inc. All Rights Reserved.
  */
 
@@ -335,11 +335,6 @@ namespace iPhonePackager
 		/// </summary>
 		public static bool bCookOnTheFly = false;
 
-		/// <summary>
-		/// Whether to build the project as an IOS framework.
-		/// </summary>
-		public static bool bBuildAsFramework = false;
-
 		/// <summary>
 		/// Whether the install should be performed on a provision
 		/// </summary>
diff --git a/Engine/Source/Programs/IOS/iPhonePackager/CookTime.cs b/Engine/Source/Programs/IOS/iPhonePackager/CookTime.cs
--- a/Engine/Source/Programs/IOS/iPhonePackager/CookTime.cs
+++ b/Engine/Source/Programs/IOS/iPhonePackager/CookTime.cs
@@ -1,4 +1,4 @@
-/**
+﻿/**
  * Copyright 1998-2019 Epic Games, Inc. All Rights Reserved.
  */
 
@@ -26,7 +26,7 @@ namespace iPhonePackager
 		/// List of files being inserted or updated in the Zip
 		/// </summary>
 		static private HashSet<string> FilesBeingModifiedToPrintOut = new HashSet<string>();
-		
+
 		/**
 		 * Create and open a work IPA file
 		 */
@@ -36,7 +36,7 @@ namespace iPhonePackager
 			string WorkIPA = Config.GetIPAPath(".ipa");
 			return CreateWorkingIPA(ReferenceZipPath, WorkIPA);
 		}
-		
+
 		/// <summary>
 		/// Creates a copy of a source IPA to a working path and opens it up as a Zip for further modifications
 		/// </summary>
@@ -76,9 +76,9 @@ namespace iPhonePackager
 				Program.ReturnCode = (int)ErrorCodes.Error_StubNotSignedCorrectly;
 			}
 
-            // Set encoding to support unicode filenames
-            Stub.AlternateEncodingUsage = ZipOption.Always;
-            Stub.AlternateEncoding = Encoding.UTF8;
+			// Set encoding to support unicode filenames
+			Stub.AlternateEncodingUsage = ZipOption.Always;
+			Stub.AlternateEncoding = Encoding.UTF8;
 
 			return Stub;
 		}
@@ -133,7 +133,7 @@ namespace iPhonePackager
 					FileOperations.CopyRequiredFile(GameFilename, DestFilename);
 				}
 			}
-			
+
 			// look for matching engine files that weren't in game
 			foreach (string EngineFilename in EngineFiles)
 			{
@@ -172,7 +172,7 @@ namespace iPhonePackager
 		/**
 		 * Callback for setting progress when saving zip file
 		 */
-		static private void UpdateSaveProgress( object Sender, SaveProgressEventArgs Event )
+		static private void UpdateSaveProgress(object Sender, SaveProgressEventArgs Event)
 		{
 			if (Event.EventType == ZipProgressEventType.Saving_BeforeWriteEntry)
 			{
@@ -306,10 +306,6 @@ namespace iPhonePackager
 			{
 				string SourceDir = Path.GetFullPath(ZipSourceDir);
 				string[] PayloadFiles = Directory.GetFiles(SourceDir, "*.*", Config.bIterate ? SearchOption.TopDirectoryOnly : SearchOption.AllDirectories);
-				HashSet<string> ExistingFiles = new HashSet<string>(FileSystem.GetAllPayloadFiles());
-				// We don't want to replace ue4commandline if we're building as a framework, since the wrapper project has its own ue4commandline file.
-				HashSet<string> FilesToNotReplace = new HashSet<string>()
-					{ "ue4commandline.txt" };
 
 				foreach (string Filename in PayloadFiles)
 				{
@@ -317,29 +313,26 @@ namespace iPhonePackager
 					// deeper than the base dir, since they were generated from a search)
 					string AbsoluteFilename = Path.GetFullPath(Filename);
 					string RelativeFilename = AbsoluteFilename.Substring(SourceDir.Length + 1).Replace('\\', '/');
-					
-					if (!Config.bBuildAsFramework || (!(FilesToNotReplace.Contains(RelativeFilename) && ExistingFiles.Contains(RelativeFilename))))
+
+					string ZipAbsolutePath = String.Format("Payload/{0}{1}.app/{2}",
+						Program.GameName + (Program.IsClient ? "Client" : ""),
+						Program.Architecture,
+						RelativeFilename);
+
+					byte[] FileContents = File.ReadAllBytes(AbsoluteFilename);
+					if (FileContents.Length == 0)
+					{
+						// Zero-length files added by Ionic cause installation/upgrade to fail on device with error 0xE8000050
+						// We store a single byte in the files as a workaround for now
+						FileContents = new byte[1];
+						FileContents[0] = 0;
+					}
+
+					FileSystem.WriteAllBytes(RelativeFilename, FileContents);
+
+					if ((FileContents.Length >= 1024 * 1024) || (Config.bVerbose))
 					{
-						string ZipAbsolutePath = String.Format("Payload/{0}{1}.app/{2}",
-							Program.GameName + (Program.IsClient ? "Client" : ""),
-							Program.Architecture,
-							RelativeFilename);
-
-						byte[] FileContents = File.ReadAllBytes(AbsoluteFilename);
-						if (FileContents.Length == 0)
-						{
-							// Zero-length files added by Ionic cause installation/upgrade to fail on device with error 0xE8000050
-							// We store a single byte in the files as a workaround for now
-							FileContents = new byte[1];
-							FileContents[0] = 0;
-						}
-
-						FileSystem.WriteAllBytes(RelativeFilename, FileContents);
-
-						if ((FileContents.Length >= 1024 * 1024) || (Config.bVerbose))
-						{
-							FilesBeingModifiedToPrintOut.Add(ZipAbsolutePath);
-						}
+						FilesBeingModifiedToPrintOut.Add(ZipAbsolutePath);
 					}
 				}
 			}
diff --git a/Engine/Source/Programs/IOS/iPhonePackager/iPhonePackager.cs b/Engine/Source/Programs/IOS/iPhonePackager/iPhonePackager.cs
--- a/Engine/Source/Programs/IOS/iPhonePackager/iPhonePackager.cs
+++ b/Engine/Source/Programs/IOS/iPhonePackager/iPhonePackager.cs
@@ -1,4 +1,4 @@
-/**
+﻿/**
  * Copyright 1998-2019 Epic Games, Inc. All Rights Reserved.
  */
 
@@ -257,9 +257,6 @@ namespace iPhonePackager
 							case "-autosigning":
 								Config.bAutomaticSigning = true;
 								break;
-							case "-buildasframework":
-								Config.bBuildAsFramework = true;
-								break;
 						}
 
 						// get the stage dir path
diff --git a/Engine/Source/Programs/UnrealBuildTool/Platform/IOS/IOSToolChain.cs b/Engine/Source/Programs/UnrealBuildTool/Platform/IOS/IOSToolChain.cs
--- a/Engine/Source/Programs/UnrealBuildTool/Platform/IOS/IOSToolChain.cs
+++ b/Engine/Source/Programs/UnrealBuildTool/Platform/IOS/IOSToolChain.cs
@@ -1,4 +1,4 @@
-// Copyright 1998-2019 Epic Games, Inc. All Rights Reserved.
+﻿// Copyright 1998-2019 Epic Games, Inc. All Rights Reserved.
 
 
 using System;
@@ -131,11 +131,16 @@ namespace UnrealBuildTool
 
 		public override void ModifyBuildProducts(ReadOnlyTargetRules Target, UEBuildBinary Binary, List<string> Libraries, List<UEBuildBundleResource> BundleResources, Dictionary<FileReference, BuildProductType> BuildProducts)
 		{
-			if (Target.IOSPlatform.bBuildAsFramework)
+			bool bBuildAsFramework = IOSToolChain.GetBuildAsFramework(Target.ProjectFile);
+			if (bBuildAsFramework)
 			{
 				int Index = Binary.OutputFilePath.GetFileNameWithoutExtension().IndexOf("-");
 				string OutputFile = Binary.OutputFilePath.GetFileNameWithoutExtension().Substring(0, Index > 0 ? Index : Binary.OutputFilePath.GetFileNameWithoutExtension().Length);
-				FileReference StubFile = FileReference.Combine(Binary.OutputDir.ParentDirectory.ParentDirectory, Binary.OutputFilePath.GetFileNameWithoutExtension() + ".stub");
+
+				DirectoryReference OutputDir = Binary.OutputDir;
+				if (OutputDir.ToString().Contains(".framework"))
+					OutputDir = OutputDir.ParentDirectory.ParentDirectory;
+				FileReference StubFile = FileReference.Combine(OutputDir, Binary.OutputFilePath.GetFileNameWithoutExtension() + ".stub");
 				BuildProducts.Add(StubFile, BuildProductType.Package);
 			}
 			else if (Target.IOSPlatform.bCreateStubIPA && Binary.Type != UEBuildBinaryType.StaticLibrary)
@@ -1501,7 +1506,7 @@ namespace UnrealBuildTool
 				FileReference PostBuildSyncFile = FileReference.Combine(BinaryLinkEnvironment.IntermediateDirectory, "PostBuildSync.dat");
 				BinaryFormatterUtils.Save(PostBuildSyncFile, PostBuildSyncTarget);
 
-				string PostBuildSyncArguments = String.Format("-Input=\"{0}\" -XmlConfigCache=\"{1}\" -remoteini={2}", PostBuildSyncFile, XmlConfig.CacheFile, UnrealBuildTool.GetRemoteIniPath());
+				string PostBuildSyncArguments = String.Format("-Input=\"{0}\" -XmlConfigCache=\"{1}\" -remoteini=\"{2}\"", PostBuildSyncFile, XmlConfig.CacheFile, UnrealBuildTool.GetRemoteIniPath());
 
 				Action PostBuildSyncAction = Action.CreateRecursiveAction<IOSPostBuildSyncMode>(ActionType.CreateAppBundle, PostBuildSyncArguments);
 				PostBuildSyncAction.WorkingDirectory = UnrealBuildTool.EngineSourceDirectory;
@@ -1671,7 +1676,7 @@ namespace UnrealBuildTool
 			IOSExports.WriteEntitlements(Target.Platform, Ini, AppName, MobileProvisionFile, Target.bForDistribution, IntermediateDir);
 		}
 
-		private static void WriteMissingEmbeddedPlist(FileReference MissingPlistPath, FileReference ProjectFile)
+		private static void WriteMissingEmbeddedPlist(FileReference MissingPlistPath, DirectoryReference ProjectDirectory, FileReference ProjectFile)
 		{
 			FileReference PlistTemplatePath = FileReference.Combine(UnrealBuildTool.EngineDirectory, "Build", "IOS", "Resources", "FrameworkPlist", "Embedded.plist");
 			if (!DirectoryReference.Exists(MissingPlistPath.Directory))
@@ -1679,8 +1684,8 @@ namespace UnrealBuildTool
 				DirectoryReference.CreateDirectory(MissingPlistPath.Directory);
 			}
 
-			string BundleID = ProjectFiles.Xcode.XcodeFrameworkWrapperUtils.GetBundleID(ProjectFile);
-			string BundleName = ProjectFiles.Xcode.XcodeFrameworkWrapperUtils.GetBundleName(ProjectFile);
+			string BundleID = ProjectFiles.Xcode.XcodeFrameworkWrapperUtils.GetBundleID(ProjectDirectory, ProjectFile);
+			string BundleName = ProjectFiles.Xcode.XcodeFrameworkWrapperUtils.GetBundleName(ProjectDirectory, ProjectFile);
 
 			BundleID += ".embedded"; // This is necessary because the wrapper that contains this framework cannot have the same BundleID.
 
@@ -1691,43 +1696,67 @@ namespace UnrealBuildTool
 			FileReference.WriteAllText(MissingPlistPath, Plist);
 		}
 
+        /// <summary>
+        /// If the project is a UE4Game project, Target.ProjectDirectory refers to the engine dir, not the actual dir of the project. So this method gets the 
+        /// actual directory of the project whether it is a UE4Game project or not.
+        /// </summary>
+        /// <returns>The actual project directory.</returns>
+        /// <param name="ProjectFile">The path to the project file</param>
+        private static DirectoryReference GetActualProjectDirectory(FileReference ProjectFile)
+		{
+			DirectoryReference ProjectDirectory = (ProjectFile == null ? DirectoryReference.FromString(UnrealBuildTool.GetRemoteIniPath()) : DirectoryReference.FromFile(ProjectFile)); 
+			return ProjectDirectory;
+		}
+
+        /// <summary>
+        /// Gets whether to build the project as a framework. Since it uses GetActualProjectDirectory internally, it ensures that the correct value is retrieved whether we're building a code-based project or blueprint-only project.
+        /// </summary>
+        /// <returns><c>true</c>, if the project should be built as a framework, <c>false</c> otherwise.</returns>
+        /// <param name="ProjectFile">A reference to the project file of the project being built.</param>
+        public static bool GetBuildAsFramework(FileReference ProjectFile)
+		{
+
+			DirectoryReference ProjectDirectory = GetActualProjectDirectory(ProjectFile);
+			ConfigHierarchy Ini = ConfigCache.ReadHierarchy(ConfigHierarchyType.Engine, ProjectDirectory, UnrealTargetPlatform.IOS);
+			bool bBuildAsFramework;
+			Ini.GetBool("/Script/IOSRuntimeSettings.IOSRuntimeSettings", "bBuildAsFramework", out bBuildAsFramework);
+			return bBuildAsFramework;
+		}
+
 		private static void GenerateFrameworkWrapperIfNonexistent(IOSPostBuildSyncTarget Target)
 		{
-			if (Target?.ProjectFile != null)
-			{
-				DirectoryReference FrameworkProjectDirectory = DirectoryReference.FromFile(Target.ProjectFile);
-				DirectoryReference WrapperDirectory = DirectoryReference.Combine(FrameworkProjectDirectory, "Wrapper");
-				string FrameworkName = Target.TargetName;
-				string BundleId = ProjectFiles.Xcode.XcodeFrameworkWrapperUtils.GetBundleID(Target.ProjectFile);
-				string EnginePath = UnrealBuildTool.EngineDirectory.ToString();
-				string SrcFrameworkPath = DirectoryReference.Combine(FrameworkProjectDirectory, "Binaries", "IOS", Target.Configuration.ToString()).ToString();
-				string CookedDataPath = DirectoryReference.Combine(FrameworkProjectDirectory, "Saved", "StagedBuilds", "IOS", "cookeddata").ToString();
-				bool GenerateFrameworkWrapperProject = ProjectFiles.Xcode.XcodeFrameworkWrapperUtils.GetGenerateFrameworkWrapperProject(Target.ProjectFile);
-
-				string ProvisionName = string.Empty;
-				string TeamUUID = string.Empty;
-				IOSProjectSettings ProjectSettings = ((IOSPlatform)UEBuildPlatform.GetBuildPlatform(UnrealTargetPlatform.IOS)).ReadProjectSettings(Target.ProjectFile);
-				if (ProjectSettings != null)
-				{
-					IOSProvisioningData Data = ((IOSPlatform)UEBuildPlatform.GetBuildPlatform(UnrealTargetPlatform.IOS)).ReadProvisioningData(ProjectSettings, Target.bForDistribution);
-					if (Data != null)
-					{
-						ProvisionName = Data.MobileProvisionName;
-						TeamUUID = Data.TeamUUID;
-					}
-				}
+			DirectoryReference ProjectDirectory = GetActualProjectDirectory(Target.ProjectFile);
+			DirectoryReference WrapperDirectory = DirectoryReference.Combine(ProjectDirectory, "Wrapper");
 
-				// Create the framework wrapper project, only if it doesn't exist and only if the user wants to generate the wrapper.
-				if (!DirectoryReference.Exists(DirectoryReference.Combine(WrapperDirectory, FrameworkName)) && GenerateFrameworkWrapperProject)
-				{
-					ProjectFiles.Xcode.XcodeFrameworkWrapperProject.GenerateXcodeFrameworkWrapper(WrapperDirectory.ToString(), FrameworkName, BundleId, SrcFrameworkPath, EnginePath, CookedDataPath, ProvisionName, TeamUUID);
-				}
+			string ProjectName = ProjectDirectory.GetDirectoryName();
+			string FrameworkName = Target.TargetName;
+			string BundleId = ProjectFiles.Xcode.XcodeFrameworkWrapperUtils.GetBundleID(Target.ProjectDirectory, Target.ProjectFile);
+			string EnginePath = UnrealBuildTool.EngineDirectory.ToString();
+			string SrcFrameworkPath = DirectoryReference.Combine(Target.ProjectDirectory , "Binaries", "IOS", Target.Configuration.ToString()).ToString(); // We use Target.ProjectDirectory because if it is a UE4Game we want the engine dir and not the actual project dir.
+			string CookedDataPath = DirectoryReference.Combine(ProjectDirectory, "Saved", "StagedBuilds", "IOS", "cookeddata").ToString();
+			bool GenerateFrameworkWrapperProject = ProjectFiles.Xcode.XcodeFrameworkWrapperUtils.GetGenerateFrameworkWrapperProject(Target.ProjectDirectory);
+
+			string ProvisionName = string.Empty;
+			string TeamUUID = string.Empty;
+			IOSProjectSettings ProjectSettings = ((IOSPlatform)UEBuildPlatform.GetBuildPlatform(UnrealTargetPlatform.IOS)).ReadProjectSettings(Target.ProjectFile);
+
+			IOSProvisioningData Data = ((IOSPlatform)UEBuildPlatform.GetBuildPlatform(UnrealTargetPlatform.IOS)).ReadProvisioningData(ProjectSettings, Target.bForDistribution);
+			if (Data != null)
+			{
+				ProvisionName = Data.MobileProvisionName;
+				TeamUUID = Data.TeamUUID;
+			}
+
+			// Create the framework wrapper project, only if it doesn't exist and only if the user wants to generate the wrapper.
+			if (!DirectoryReference.Exists(DirectoryReference.Combine(WrapperDirectory, FrameworkName)) && GenerateFrameworkWrapperProject)
+			{
+				ProjectFiles.Xcode.XcodeFrameworkWrapperProject.GenerateXcodeFrameworkWrapper(WrapperDirectory.ToString(), ProjectName, FrameworkName, BundleId, SrcFrameworkPath, EnginePath, CookedDataPath, ProvisionName, TeamUUID);
 			}
 		}
 
 		public static void PostBuildSync(IOSPostBuildSyncTarget Target)
 		{
-			ConfigHierarchy Ini = ConfigCache.ReadHierarchy(ConfigHierarchyType.Engine, Target.ProjectFile == null ? DirectoryReference.FromString(UnrealBuildTool.GetRemoteIniPath()) : DirectoryReference.FromFile(Target.ProjectFile), UnrealTargetPlatform.IOS);
+			ConfigHierarchy Ini = ConfigCache.ReadHierarchy(ConfigHierarchyType.Engine, GetActualProjectDirectory(Target.ProjectFile), UnrealTargetPlatform.IOS);
 			string BundleID;
 			Ini.GetString("/Script/IOSRuntimeSettings.IOSRuntimeSettings", "BundleIdentifier", out BundleID);
 
@@ -1737,14 +1766,14 @@ namespace UnrealBuildTool
 			string PathToDsymZip = Target.OutputPath.FullName + ".dSYM.zip";
 			if (BuildHostPlatform.Current.Platform == UnrealTargetPlatform.Mac)
 			{
-				if (IsCompiledAsFramework(Target.OutputPath.FullName))
+				if (GetBuildAsFramework(Target.ProjectFile))
 				{
 					// make sure the framework has a plist
-					FileReference PlistSrcLocation = FileReference.Combine(Target.ProjectFile.Directory, "Build", "IOS", "Embedded.plist");
+					FileReference PlistSrcLocation = FileReference.Combine(Target.ProjectDirectory, "Build", "IOS", "Embedded.plist");
 					if (!FileReference.Exists(PlistSrcLocation))
 					{
 
-						WriteMissingEmbeddedPlist(PlistSrcLocation, Target.ProjectFile);
+						WriteMissingEmbeddedPlist(PlistSrcLocation, Target.ProjectDirectory, Target.ProjectFile);
 						if (!FileReference.Exists(PlistSrcLocation))
 							throw new BuildException("Unable to find plist for output framework ({0})", PlistSrcLocation);
 					}
@@ -1784,8 +1813,8 @@ namespace UnrealBuildTool
 
 				// generate the dummy project so signing works
 				DirectoryReference XcodeWorkspaceDir = null;
-				if (!bBuildAsFramework)
-				{
+                if (!bBuildAsFramework)
+                {
 					if (AppName == "UE4Game" || AppName == "UE4Client" || Target.ProjectFile == null || Target.ProjectFile.IsUnderDirectory(UnrealBuildTool.EngineDirectory))
 					{
 						XcodeWorkspaceDir = DirectoryReference.Combine(UnrealBuildTool.RootDirectory, String.Format("UE4_{0}.xcworkspace", (Target.Platform == UnrealTargetPlatform.IOS ? "IOS" : "TVOS")));
@@ -1794,7 +1823,7 @@ namespace UnrealBuildTool
 					{
 						XcodeWorkspaceDir = DirectoryReference.Combine(Target.ProjectDirectory, String.Format("{0}_{1}.xcworkspace", Target.ProjectFile.GetFileNameWithoutExtension(), (Target.Platform == UnrealTargetPlatform.IOS ? "IOS" : "TVOS")));
 					}
-				}
+                }
 
 				// Path to the temporary keychain. When -ImportCertificate is specified, we will temporarily add this to the list of keychains to search, and remove it later.
 				FileReference TempKeychain = FileReference.Combine(Target.ProjectIntermediateDirectory, "TempKeychain.keychain");
@@ -1892,7 +1921,14 @@ namespace UnrealBuildTool
 					string SchemeName;
 					if (AppName == "UE4Game" || AppName == "UE4Client")
 					{
-						SchemeName = "UE4";
+						if (bBuildAsFramework)
+						{
+							SchemeName = "UE4Game";
+						}
+						else
+						{
+							SchemeName = "UE4";
+						}
 					}
 					else
 					{
@@ -1904,9 +1940,9 @@ namespace UnrealBuildTool
 					string CmdLine;
 					if (bBuildAsFramework)
 					{
-						DirectoryReference FrameworkProjectDirectory = DirectoryReference.FromFile(Target.ProjectFile);
+						DirectoryReference FrameworkProjectDirectory = GetActualProjectDirectory(Target.ProjectFile);
 						DirectoryReference WrapperDirectory = DirectoryReference.Combine(FrameworkProjectDirectory, "Wrapper");
-                        string FrameworkName = Target.TargetName;
+						string FrameworkName = Target.TargetName;
 						DirectoryReference WrapperProject = DirectoryReference.Combine(WrapperDirectory, FrameworkName, FrameworkName + ".xcodeproj");
 
 						// Delete everything in payload directory before building
@@ -1936,7 +1972,7 @@ namespace UnrealBuildTool
 					}
 
 					CmdLine += String.Format(" CODE_SIGN_IDENTITY='{0}'", SigningCertificate);
-					
+
 					if (!ProjectSettings.bAutomaticSigning)
 					{
 						CmdLine += (!string.IsNullOrEmpty(MobileProvisionUUID) ? (" PROVISIONING_PROFILE_SPECIFIER=" + MobileProvisionUUID) : "");
@@ -2016,7 +2052,8 @@ namespace UnrealBuildTool
 				{
 					CleanIntermediateDirectory(FrameworkDerivedDataDir);
 				}
-				else
+
+				if (XcodeWorkspaceDir != null)
 				{
 					// delete the temp project
 					DirectoryReference.Delete(XcodeWorkspaceDir, true);
@@ -2036,7 +2073,6 @@ namespace UnrealBuildTool
 				{
 					PackageStub(RemoteShadowDirectoryMac, AppName, Target.OutputPath.GetFileNameWithoutExtension());
 				}
-
 			}
 			else
 			{
diff --git a/Engine/Source/Programs/UnrealBuildTool/Platform/IOS/UEBuildIOS.cs b/Engine/Source/Programs/UnrealBuildTool/Platform/IOS/UEBuildIOS.cs
--- a/Engine/Source/Programs/UnrealBuildTool/Platform/IOS/UEBuildIOS.cs
+++ b/Engine/Source/Programs/UnrealBuildTool/Platform/IOS/UEBuildIOS.cs
@@ -1,4 +1,4 @@
-// Copyright 1998-2019 Epic Games, Inc. All Rights Reserved.
+﻿// Copyright 1998-2019 Epic Games, Inc. All Rights Reserved.
 
 using System;
 using System.Collections.Generic;
@@ -787,7 +787,10 @@ namespace UnrealBuildTool
 				Target.GlobalDefinitions.Add("HAS_METAL=0");
 			}
 
-			if (Target.IOSPlatform.bBuildAsFramework)
+
+			bool bBuildAsFramework = IOSToolChain.GetBuildAsFramework(Target.ProjectFile);
+
+			if (bBuildAsFramework)
 			{
 				Target.bShouldCompileAsDLL = true;
 				int PreviousDefinition = Target.GlobalDefinitions.FindIndex(s => s.Contains("BUILD_EMBEDDED_APP"));
diff --git a/Engine/Source/Programs/UnrealBuildTool/Platform/IOS/UEDeployIOS.cs b/Engine/Source/Programs/UnrealBuildTool/Platform/IOS/UEDeployIOS.cs
--- a/Engine/Source/Programs/UnrealBuildTool/Platform/IOS/UEDeployIOS.cs
+++ b/Engine/Source/Programs/UnrealBuildTool/Platform/IOS/UEDeployIOS.cs
@@ -1,4 +1,4 @@
-// Copyright 1998-2019 Epic Games, Inc. All Rights Reserved.
+﻿// Copyright 1998-2019 Epic Games, Inc. All Rights Reserved.
 
 using System;
 using System.Collections;
@@ -770,7 +770,7 @@ namespace UnrealBuildTool
 
 			File.WriteAllText(PListFile, Text.ToString());
 
-            if (BuildHostPlatform.Current.Platform == UnrealTargetPlatform.Mac && !ProjectFiles.Xcode.XcodeFrameworkWrapperUtils.GetBuildAsFramework(ProjectFile))
+            if (BuildHostPlatform.Current.Platform == UnrealTargetPlatform.Mac && !ProjectFiles.Xcode.XcodeFrameworkWrapperUtils.GetBuildAsFramework(DirectoryReference.FromFile(ProjectFile)))
 			{
 				if (!Directory.Exists(AppDirectory))
 				{
@@ -956,7 +956,7 @@ namespace UnrealBuildTool
 			}
 
 			// If we are building as a framework, we don't need to do all of this.
-			if (ProjectFiles.Xcode.XcodeFrameworkWrapperUtils.GetBuildAsFramework(ProjectFile))
+			if (IOSToolChain.GetBuildAsFramework(ProjectFile))
 				return false;
 
 
diff --git a/Engine/Source/Programs/UnrealBuildTool/ProjectFiles/Xcode/XcodeFrameworkWrapperProject.cs b/Engine/Source/Programs/UnrealBuildTool/ProjectFiles/Xcode/XcodeFrameworkWrapperProject.cs
--- a/Engine/Source/Programs/UnrealBuildTool/ProjectFiles/Xcode/XcodeFrameworkWrapperProject.cs
+++ b/Engine/Source/Programs/UnrealBuildTool/ProjectFiles/Xcode/XcodeFrameworkWrapperProject.cs
@@ -17,6 +17,8 @@ namespace UnrealBuildTool.ProjectFiles.Xcode
 		private readonly static string PROJECT_FILE_SEARCH_EXPRESSION = "*.pbxproj";
 		private readonly static string TEMPLATE_NAME = "FrameworkWrapper";
 		private readonly static string FRAMEWORK_WRAPPER_TEMPLATE_DIRECTORY = Path.Combine(UnrealBuildTool.EngineDirectory.ToNormalizedPath(), "Build", "IOS", "Resources", TEMPLATE_NAME);
+		private readonly static string TEMPLATE_PROJECT_NAME = "PROJECT_NAME";
+		private readonly static string COMMANDLINE_FILENAME = "ue4commandline.txt";
 
 		/// <summary>
 		/// Recursively copies all of the files and directories that are inside <paramref name="SourceDirectory"/> into <paramref name="DestinationDirectory"/>.
@@ -117,11 +119,12 @@ namespace UnrealBuildTool.ProjectFiles.Xcode
 		/// with <paramref name="NewValue"/>.
 		/// </summary>
 		/// <param name="RootDirectory">The directory in which all files should be subject to replacements.</param>
+        /// <param name="SearchPattern">Only replace text in files that match this pattern. Default is all files.</param>
 		/// <param name="OldValue">The value that should be replaced in all files.</param>
 		/// <param name="NewValue">The replacement value.</param>
-		private static void ReplaceTextInFiles(string RootDirectory, string OldValue, string NewValue)
+		private static void ReplaceTextInFiles(string RootDirectory, string OldValue, string NewValue, string SearchPattern = "*")
 		{
-			IEnumerable<string> Files = Directory.EnumerateFiles(RootDirectory, "*", SearchOption.AllDirectories);
+			IEnumerable<string> Files = Directory.EnumerateFiles(RootDirectory, SearchPattern, SearchOption.AllDirectories);
 			foreach (string SrcFile in Files)
 			{
 				string FileContents = File.ReadAllText(SrcFile);
@@ -169,10 +172,23 @@ namespace UnrealBuildTool.ProjectFiles.Xcode
 					ProjectContents = ChangeProjectSetting(ProjectContents, Setting.Key, Setting.Value);
 				}
 
-				File.WriteAllText(ProjectFiles[0], ProjectContents);
+                File.WriteAllText(ProjectFiles[0], ProjectContents);
 			}
 		}
-		
+
+        /// <summary>
+        /// Removes the readonly attribute from all files in a directory file while retaining all other attributes, thus making them writeable.
+        /// </summary>
+        /// <param name="RootDirectory">The path to the directory that will be make writeable.</param>
+        private static void MakeAllFilesWriteable(string RootDirectory)
+		{
+			IEnumerable<string> FileNames = Directory.EnumerateFiles(RootDirectory, "*", SearchOption.AllDirectories);
+			foreach (string FileName in FileNames)
+			{
+				File.SetAttributes(FileName, File.GetAttributes(FileName) & ~FileAttributes.ReadOnly);
+			}
+		}
+
 		/// <summary>
 		/// Changes the value of a setting in a project file. 
 		/// </summary>
@@ -229,7 +245,7 @@ namespace UnrealBuildTool.ProjectFiles.Xcode
 		/// <param name="CookedDataPath">The path to the 'cookeddata' folder that accompanies the framework.</param>
 		/// <param name="ProvisionName"></param>
 		/// <param name="TeamUUID"></param>
-		public static void GenerateXcodeFrameworkWrapper(string OutputDirectory, string FrameworkName, string BundleId, string SrcFrameworkPath, string EnginePath, string CookedDataPath,string ProvisionName, string TeamUUID)
+		public static void GenerateXcodeFrameworkWrapper(string OutputDirectory, string ProjectName, string FrameworkName, string BundleId, string SrcFrameworkPath, string EnginePath, string CookedDataPath,string ProvisionName, string TeamUUID)
 		{
 			string OutputDir = Path.Combine(OutputDirectory, FrameworkName);
 
@@ -237,20 +253,31 @@ namespace UnrealBuildTool.ProjectFiles.Xcode
 
 			DeleteUnwantedDirectories(OutputDir);
 
+			MakeAllFilesWriteable(OutputDir);
+
 			RenameFilesAndDirectories(OutputDir, TEMPLATE_NAME, FrameworkName);
 
 			SetProjectFileSettings(OutputDir, FrameworkName, BundleId, SrcFrameworkPath, EnginePath, CookedDataPath, ProvisionName, TeamUUID);
 
 			ReplaceTextInFiles(OutputDir, TEMPLATE_NAME, FrameworkName);
+
+			ReplaceTextInFiles(OutputDir, TEMPLATE_PROJECT_NAME, ProjectName, COMMANDLINE_FILENAME);
+
 		}
 
 	}
 	
 	class XcodeFrameworkWrapperUtils
 	{
-		public static string GetBundleID(FileReference ProjectFile)
+        private static ConfigHierarchy GetIni(DirectoryReference ProjectDirectory)
+		{
+			return ConfigCache.ReadHierarchy(ConfigHierarchyType.Engine, ProjectDirectory, UnrealTargetPlatform.IOS);
+			//return ConfigCache.ReadHierarchy(ConfigHierarchyType.Engine, DirectoryReference.FromFile(ProjectFile), UnrealTargetPlatform.IOS);
+		}
+
+		public static string GetBundleID(DirectoryReference ProjectDirectory, FileReference ProjectFile)
 		{
-			ConfigHierarchy Ini = ConfigCache.ReadHierarchy(ConfigHierarchyType.Engine, DirectoryReference.FromFile(ProjectFile), UnrealTargetPlatform.IOS);
+			ConfigHierarchy Ini = GetIni(ProjectDirectory);
 			string BundleId;
 			Ini.GetString("/Script/IOSRuntimeSettings.IOSRuntimeSettings", "BundleIdentifier", out BundleId);
 
@@ -258,9 +285,9 @@ namespace UnrealBuildTool.ProjectFiles.Xcode
 			return BundleId;
 		}
 
-		public static string GetBundleName(FileReference ProjectFile)
+		public static string GetBundleName(DirectoryReference ProjectDirectory, FileReference ProjectFile)
 		{
-			ConfigHierarchy Ini = ConfigCache.ReadHierarchy(ConfigHierarchyType.Engine, DirectoryReference.FromFile(ProjectFile), UnrealTargetPlatform.IOS);
+			ConfigHierarchy Ini = GetIni(ProjectDirectory);
 			string BundleName;
 			Ini.GetString("/Script/IOSRuntimeSettings.IOSRuntimeSettings", "BundleName", out BundleName);
 
@@ -268,17 +295,17 @@ namespace UnrealBuildTool.ProjectFiles.Xcode
 			return BundleName;
 		}
 
-		public static bool GetBuildAsFramework(FileReference ProjectFile)
+		public static bool GetBuildAsFramework(DirectoryReference ProjectDirectory)
 		{
-			ConfigHierarchy Ini = ConfigCache.ReadHierarchy(ConfigHierarchyType.Engine, DirectoryReference.FromFile(ProjectFile), UnrealTargetPlatform.IOS);
+			ConfigHierarchy Ini = GetIni(ProjectDirectory);
 			bool bBuildAsFramework;
 			Ini.GetBool("/Script/IOSRuntimeSettings.IOSRuntimeSettings", "bBuildAsFramework", out bBuildAsFramework);
 			return bBuildAsFramework;
 		}
 		
-		public static bool GetGenerateFrameworkWrapperProject(FileReference ProjectFile)
+		public static bool GetGenerateFrameworkWrapperProject(DirectoryReference ProjectDirectory)
 		{
-		ConfigHierarchy Ini = ConfigCache.ReadHierarchy(ConfigHierarchyType.Engine, DirectoryReference.FromFile(ProjectFile), UnrealTargetPlatform.IOS);
+			ConfigHierarchy Ini = GetIni(ProjectDirectory);
 			bool bGenerateFrameworkWrapperProject;
 			Ini.GetBool("/Script/IOSRuntimeSettings.IOSRuntimeSettings", "bGenerateFrameworkWrapperProject", out bGenerateFrameworkWrapperProject);
 			return bGenerateFrameworkWrapperProject;
diff --git a/Engine/Source/Runtime/IOS/IOSRuntimeSettings/Classes/IOSRuntimeSettings.h b/Engine/Source/Runtime/IOS/IOSRuntimeSettings/Classes/IOSRuntimeSettings.h
--- a/Engine/Source/Runtime/IOS/IOSRuntimeSettings/Classes/IOSRuntimeSettings.h
+++ b/Engine/Source/Runtime/IOS/IOSRuntimeSettings/Classes/IOSRuntimeSettings.h
@@ -248,10 +248,6 @@ public:
     UPROPERTY(GlobalConfig, EditAnywhere, Category = Build, meta = (DisplayName = "Build project as a framework"))
     bool bBuildAsFramework;
 
-    // Whether to generate a native Xcode project as a wrapper for the framework.
-    UPROPERTY(GlobalConfig, EditAnywhere, Category = Build, meta = (DisplayName = "Generate a native Xcode wrapper project for the framework"))
-    bool bGenerateFrameworkWrapperProject;
-
 	// Remotely compile shaders offline
 	UPROPERTY(GlobalConfig, EditAnywhere, Category = Build)
 	bool EnableRemoteShaderCompile;
diff --git a/Engine/Source/Runtime/IOS/IOSRuntimeSettings/Private/IOSRuntimeSettings.cpp b/Engine/Source/Runtime/IOS/IOSRuntimeSettings/Private/IOSRuntimeSettings.cpp
--- a/Engine/Source/Runtime/IOS/IOSRuntimeSettings/Private/IOSRuntimeSettings.cpp
+++ b/Engine/Source/Runtime/IOS/IOSRuntimeSettings/Private/IOSRuntimeSettings.cpp
@@ -23,7 +23,6 @@ UIOSRuntimeSettings::UIOSRuntimeSettings(const FObjectInitializer& ObjectInitial
 	bSupportsIPhone = true;
 	MinimumiOSVersion = EIOSVersion::IOS_10;
     bBuildAsFramework = true;
-    bGenerateFrameworkWrapperProject = true;
 	EnableRemoteShaderCompile = false;
 	bGeneratedSYMFile = false;
 	bGeneratedSYMBundle = false;
