Fix warning for missing documentation of method parameter in XcodeFrameworkWrapperProject.cs

diff --git a/Engine/Source/Programs/UnrealBuildTool/ProjectFiles/Xcode/XcodeFrameworkWrapperProject.cs b/Engine/Source/Programs/UnrealBuildTool/ProjectFiles/Xcode/XcodeFrameworkWrapperProject.cs
--- a/Engine/Source/Programs/UnrealBuildTool/ProjectFiles/Xcode/XcodeFrameworkWrapperProject.cs
+++ b/Engine/Source/Programs/UnrealBuildTool/ProjectFiles/Xcode/XcodeFrameworkWrapperProject.cs
@@ -238,6 +238,7 @@ namespace UnrealBuildTool.ProjectFiles.Xcode
 		/// settings in the project to accommodate the framework, and replacing text in all the files to match the framework.
 		/// </summary>
 		/// <param name="OutputDirectory">The directory in which to place the framework. The framework will be placed in 'outputDirectory/frameworkName/'.</param>
+        /// <param name="ProjectName">The name of the project. If blueprint-only, use the actual name of the project, not just UE4Game.</param>
 		/// <param name="FrameworkName">The name of the framework that this project is wrapping.</param>
 		/// <param name="BundleId">The Bundle ID to give to the wrapper project.</param>
 		/// <param name="SrcFrameworkPath">The path to the directory containing the framework to be wrapped.</param>
