Disable testOnly injection for Android Studio builds

#ue4
#android

diff --git a/Engine/Source/Programs/UnrealBuildTool/Platform/Android/UEDeployAndroid.cs b/Engine/Source/Programs/UnrealBuildTool/Platform/Android/UEDeployAndroid.cs
--- a/Engine/Source/Programs/UnrealBuildTool/Platform/Android/UEDeployAndroid.cs
+++ b/Engine/Source/Programs/UnrealBuildTool/Platform/Android/UEDeployAndroid.cs
@@ -3158,6 +3158,7 @@ namespace UnrealBuildTool
 
 			GradleProperties.AppendLine("org.gradle.daemon=false");
 			GradleProperties.AppendLine("org.gradle.jvmargs=-XX:MaxHeapSize=4096m -Xmx9216m");
+			GradleProperties.AppendLine("android.injected.testOnly=false");
 			GradleProperties.AppendLine(string.Format("COMPILE_SDK_VERSION={0}", CompileSDKVersion));
 			GradleProperties.AppendLine(string.Format("BUILD_TOOLS_VERSION={0}", BuildToolsVersion));
 			GradleProperties.AppendLine(string.Format("PACKAGE_NAME={0}", PackageName));
