Remove Android support for Ant build system

#ue4
#android

diff --git a/Engine/Build/Android/Legacy/src/com/epicgames/ue4/LocalNotificationReceiver.java b/Engine/Build/Android/Legacy/src/com/epicgames/ue4/LocalNotificationReceiver.java
deleted file mode 100644
diff --git a/Engine/Plugins/Online/OnlineSubsystemFacebook/Source/OnlineSubsystemFacebook_UPL.xml b/Engine/Plugins/Online/OnlineSubsystemFacebook/Source/OnlineSubsystemFacebook_UPL.xml
--- a/Engine/Plugins/Online/OnlineSubsystemFacebook/Source/OnlineSubsystemFacebook_UPL.xml
+++ b/Engine/Plugins/Online/OnlineSubsystemFacebook/Source/OnlineSubsystemFacebook_UPL.xml
@@ -4,14 +4,7 @@
 	<!-- init section is always evaluated once per architecture -->
 	<init>
 		<log text="Facebook SDK Android init"/>
-
-		<!-- check if using Gradle -->
-		<setBoolFromProperty result="bGradle" ini="Engine" section="/Script/AndroidRuntimeSettings.AndroidRuntimeSettings" property="bEnableGradle" default="false"/>
-		<if condition="bGradle">
-			<true>
-				<log text="Facebook requires minimum SDK of 15"/>
-			</true>
-		</if>
+		<log text="Facebook requires minimum SDK of 15"/>
 		
 		<!-- get packaging for Facebook from ini and reset it if architecture not supported -->
 		<setBoolFromProperty result="bEnabled" ini="Engine" section="OnlineSubsystemFacebook" property="bEnabled" default="false"/>
@@ -42,22 +35,6 @@
 		<insertNewline/>
 	</minimumSDKAPI>
 	
-	<!-- AAR dependencies -->
-	<AARImports>
-		<if condition="bEnabled">
-			<true>
-				<if condition="bGradle">
-					<false>
-						<insertValue value="repository $S(PluginDir)/ThirdParty/Android/NotForLicensees"/>
-						<insertNewline/>
-						<insertValue value="FacebookSDK,facebook-android-sdk,4.19.0" />
-						<insertNewline/>
-					</false>
-				</if>
-			</true>
-		</if>
-	</AARImports>
-
 	<buildGradleAdditions>
 		<if condition="bEnabled">
 			<true>
diff --git a/Engine/Plugins/Online/OnlineSubsystemGoogle/Source/OnlineSubsystemGoogle_UPL.xml b/Engine/Plugins/Online/OnlineSubsystemGoogle/Source/OnlineSubsystemGoogle_UPL.xml
--- a/Engine/Plugins/Online/OnlineSubsystemGoogle/Source/OnlineSubsystemGoogle_UPL.xml
+++ b/Engine/Plugins/Online/OnlineSubsystemGoogle/Source/OnlineSubsystemGoogle_UPL.xml
@@ -4,14 +4,7 @@
 	<!-- init section is always evaluated once per architecture -->
 	<init>
 		<log text="Google SDK Android init"/>
-
-        <!-- check if using Gradle -->
-        <setBoolFromProperty result="bGradle" ini="Engine" section="/Script/AndroidRuntimeSettings.AndroidRuntimeSettings" property="bEnableGradle" default="false"/>
-        <if condition="bGradle">
-            <true>
-                <log text="Google requires minimum SDK of 15"/>
-            </true>
-        </if>
+		<log text="Google requires minimum SDK of 15"/>
         
 		<!-- <trace enable="true"/> -->
 		<!-- get packaging for Google from ini and reset it if architecture not supported -->
@@ -36,20 +29,6 @@
 		</if>
 	</prebuildCopies>
 	
-	<!-- AAR dependencies -->
-	<AARImports>
-		<if condition="bEnabled">
-			<true>
-                <if condition="bGradle">
-                    <false>
-                        <insertValue value="com.google.android.gms,play-services-auth,11.0.4" />
-                        <insertNewline/>
-                    </false>
-                </if>
-			</true>
-		</if>
-	</AARImports>
-
     <buildGradleAdditions>
         <insert>
             dependencies {
diff --git a/Engine/Plugins/Runtime/AR/Google/GoogleARCore/Source/GoogleARCoreBase/GoogleARCoreBase_APL.xml b/Engine/Plugins/Runtime/AR/Google/GoogleARCore/Source/GoogleARCoreBase/GoogleARCoreBase_APL.xml
--- a/Engine/Plugins/Runtime/AR/Google/GoogleARCore/Source/GoogleARCoreBase/GoogleARCoreBase_APL.xml
+++ b/Engine/Plugins/Runtime/AR/Google/GoogleARCore/Source/GoogleARCoreBase/GoogleARCoreBase_APL.xml
@@ -4,10 +4,9 @@
   <!-- init section is always evaluated once per architecture -->
   <init>
     <log text="GoogleARCorePlugin APL init"/>
-    <setBoolFromProperty result="bGradle" ini="Engine" section="/Script/AndroidRuntimeSettings.AndroidRuntimeSettings" property="bEnableGradle" default="false"/>
-		<!-- Get the bARCoreRequiredApp setting from the Engine ini. -->
-		<setBoolFromProperty result="bARCoreRequiredApp" ini="Engine" section="/Script/GoogleARCoreBase.GoogleARCoreEditorSettings" property="bARCoreRequiredApp" default="true" />
-	</init>
+	<!-- Get the bARCoreRequiredApp setting from the Engine ini. -->
+	<setBoolFromProperty result="bARCoreRequiredApp" ini="Engine" section="/Script/GoogleARCoreBase.GoogleARCoreEditorSettings" property="bARCoreRequiredApp" default="true" />
+  </init>
 
   <!-- optional files or directories to copy to Intermediate/Android/APK -->
   <resourceCopies>
@@ -118,15 +117,4 @@
       </insert>
   </buildGradleAdditions>
 
-  <AARImports>
-    <insert>repository $(ENGINEDIR)/Source/ThirdParty/GoogleARCore/lib/android-aar</insert>
-    <insertNewline/>
-    <if condition="bGradle">
-      <false>
-        <insert>com.google.ar, core, 1.9.0</insert>
-        <insertNewline/>
-      </false>
-    </if>
-  </AARImports>
-
 </root>
diff --git a/Engine/Plugins/Runtime/AR/Google/GoogleARCoreServices/Source/GoogleARCoreServices/GoogleARCoreServices_APL.xml b/Engine/Plugins/Runtime/AR/Google/GoogleARCoreServices/Source/GoogleARCoreServices/GoogleARCoreServices_APL.xml
--- a/Engine/Plugins/Runtime/AR/Google/GoogleARCoreServices/Source/GoogleARCoreServices/GoogleARCoreServices_APL.xml
+++ b/Engine/Plugins/Runtime/AR/Google/GoogleARCoreServices/Source/GoogleARCoreServices/GoogleARCoreServices_APL.xml
@@ -4,7 +4,6 @@
   <!-- init section is always evaluated once per architecture -->
   <init>
     <log text="GoogleARServices APL init"/>
-    <setBoolFromProperty result="bGradle" ini="Engine" section="/Script/AndroidRuntimeSettings.AndroidRuntimeSettings" property="bEnableGradle" default="false"/>
     <setStringFromProperty result="APIKey" ini="Engine" section="/Script/GoogleARCoreServices.GoogleARCoreServicesEditorSettings" property="AndroidAPIKey" default=""/>
 	</init>
 
diff --git a/Engine/Plugins/Runtime/Analytics/Adjust/Source/AndroidAdjust/Adjust_UPL.xml b/Engine/Plugins/Runtime/Analytics/Adjust/Source/AndroidAdjust/Adjust_UPL.xml
--- a/Engine/Plugins/Runtime/Analytics/Adjust/Source/AndroidAdjust/Adjust_UPL.xml
+++ b/Engine/Plugins/Runtime/Analytics/Adjust/Source/AndroidAdjust/Adjust_UPL.xml
@@ -32,9 +32,6 @@
 		<setBoolIsEqual result="bNoDefaultTracker" arg1="$S(DefaultTracker)" arg2=""/>
 		<setBoolIsEqual result="bNoProcessName" arg1="$S(ProcessName)" arg2=""/>
 		<setBoolIsEqual result="bNoDelayStart" arg1="$S(DelayStart)" arg2="0.0"/>
-
-		<!-- check if using Gradle -->
-		<setBoolFromProperty result="bGradle" ini="Engine" section="/Script/AndroidRuntimeSettings.AndroidRuntimeSettings" property="bEnableGradle" default="false"/>
 	</init>
 	
 	<!-- optional files or directories to copy to Intermediate/Android/APK -->
@@ -48,16 +45,6 @@
 					dst = "$S(BuildDir)/JavaLibs/adjust_library" />
 	</prebuildCopies>
 
-	<!-- AAR dependencies -->
-	<AARImports>
-		<if condition="bGradle">
-			<false>
-				<insertValue value="com.google.android.gms,play-services-analytics,11.0.4" />
-				<insertNewline/>
-			</false>
-		</if>
-	</AARImports>
-
 	<buildGradleAdditions>
 		<insert>
 dependencies {
diff --git a/Engine/Plugins/Runtime/AndroidPermission/Source/AndroidPermission/AndroidPermission_APL.xml b/Engine/Plugins/Runtime/AndroidPermission/Source/AndroidPermission/AndroidPermission_APL.xml
--- a/Engine/Plugins/Runtime/AndroidPermission/Source/AndroidPermission/AndroidPermission_APL.xml
+++ b/Engine/Plugins/Runtime/AndroidPermission/Source/AndroidPermission/AndroidPermission_APL.xml
@@ -2,8 +2,6 @@
 <root xmlns:android="http://schemas.android.com/apk/res/android">
 	<init>
 		<log text="Android Permission Plugin Init"/>
-		<!-- check if using Gradle -->
-		<setBoolFromProperty result="bGradle" ini="Engine" section="/Script/AndroidRuntimeSettings.AndroidRuntimeSettings" property="bEnableGradle" default="false"/>
 	</init>
 	
 	<prebuildCopies>
@@ -20,15 +18,6 @@
 		</insert>
 	</proguardAdditions>
 
-	<AARImports>
-		<if condition="bGradle">
-			<false>
-				<insert>com.android.support,support-v13,25.2.0</insert>
-				<insertNewline/>
-			</false>
-		</if>
-	</AARImports>
-
 	<buildGradleAdditions>
 		<insert>
 dependencies {
diff --git a/Engine/Plugins/Runtime/GoogleCloudMessaging/Source/GoogleCloudMessaging/GoogleCloudMessaging_UPL.xml b/Engine/Plugins/Runtime/GoogleCloudMessaging/Source/GoogleCloudMessaging/GoogleCloudMessaging_UPL.xml
--- a/Engine/Plugins/Runtime/GoogleCloudMessaging/Source/GoogleCloudMessaging/GoogleCloudMessaging_UPL.xml
+++ b/Engine/Plugins/Runtime/GoogleCloudMessaging/Source/GoogleCloudMessaging/GoogleCloudMessaging_UPL.xml
@@ -9,30 +9,18 @@
     <setBoolIsEqual result="GCMClientSenderInvalid" arg1="$S(GCMClientSenderID)" arg2=""/>
     <setBoolNot result="GCMEnabled" source="$B(GCMClientSenderInvalid)"/>
     <log text="GCMClientSenderID set: $B(GCMEnabled)"/>
-		<!-- check if using Gradle -->
-		<setBoolFromProperty result="bGradle" ini="Engine" section="/Script/AndroidRuntimeSettings.AndroidRuntimeSettings" property="bEnableGradle" default="false"/>
 	</init>
 
-	<AARImports>
-		<!-- only modify AAR imports if enabled -->
+	<buildGradleAdditions>
 		<if condition="GCMEnabled">
 			<true>
-				<if condition="bGradle">
-					<false>
-						<insert>com.google.android.gms,play-services-gcm,11.0.4</insert>
-						<insertNewline/>
-					</false>
-				</if>
-			</true>
-		</if>
-	</AARImports>
-
-	<buildGradleAdditions>
-		<insert>
+				<insert>
 dependencies {
 	implementation('com.google.android.gms:play-services-gcm:11.8.0')
 }
-		</insert>
+				</insert>
+			</true>
+		</if>
 	</buildGradleAdditions>
 
 	<!-- optional updates applied to AndroidManifest.xml -->
diff --git a/Engine/Source/Developer/Android/AndroidPlatformEditor/Private/AndroidLicenseDialog.cpp b/Engine/Source/Developer/Android/AndroidPlatformEditor/Private/AndroidLicenseDialog.cpp
--- a/Engine/Source/Developer/Android/AndroidPlatformEditor/Private/AndroidLicenseDialog.cpp
+++ b/Engine/Source/Developer/Android/AndroidPlatformEditor/Private/AndroidLicenseDialog.cpp
@@ -270,7 +270,7 @@ FReply SAndroidLicenseDialog::OnAgree()
 FReply SAndroidLicenseDialog::OnCancel()
 {
 	// turn off Gradle checkbox
-	GetMutableDefault<UAndroidRuntimeSettings>()->bEnableGradle = false;
+	//GetMutableDefault<UAndroidRuntimeSettings>()->bEnableGradle = false;
 
 	TSharedRef<SWindow> ParentWindow = FSlateApplication::Get().FindWidgetWindow(AsShared()).ToSharedRef();
 	FSlateApplication::Get().RequestDestroyWindow(ParentWindow);
diff --git a/Engine/Source/Developer/Android/AndroidPlatformEditor/Private/AndroidSDKSettings.cpp b/Engine/Source/Developer/Android/AndroidPlatformEditor/Private/AndroidSDKSettings.cpp
--- a/Engine/Source/Developer/Android/AndroidPlatformEditor/Private/AndroidSDKSettings.cpp
+++ b/Engine/Source/Developer/Android/AndroidPlatformEditor/Private/AndroidSDKSettings.cpp
@@ -51,13 +51,6 @@ void UAndroidSDKSettings::UpdateTargetModulePaths()
 		Values.Add(NDKPath.Path);
 	}
 	
-	if (!ANTPath.Path.IsEmpty())
-	{
-		FPaths::NormalizeFilename(ANTPath.Path);
-		Keys.Add(TEXT("ANT_HOME"));
-		Values.Add(ANTPath.Path);
-	}
-
 	if (!JavaPath.Path.IsEmpty())
 	{
 		FPaths::NormalizeFilename(JavaPath.Path);
diff --git a/Engine/Source/Developer/Android/AndroidPlatformEditor/Private/AndroidSDKSettings.h b/Engine/Source/Developer/Android/AndroidPlatformEditor/Private/AndroidSDKSettings.h
--- a/Engine/Source/Developer/Android/AndroidPlatformEditor/Private/AndroidSDKSettings.h
+++ b/Engine/Source/Developer/Android/AndroidPlatformEditor/Private/AndroidSDKSettings.h
@@ -28,10 +28,6 @@ public:
 	UPROPERTY(GlobalConfig, EditAnywhere, Category = SDKConfig, Meta = (DisplayName = "Location of Android NDK (the directory usually contains 'android-ndk-')"))
 	FDirectoryPath NDKPath;
 
-	// Location on disk of the ANT tool (falls back to ANT_HOME environment variable if this is left blank)
-	UPROPERTY(GlobalConfig, EditAnywhere, Category = SDKConfig, Meta = (DisplayName = "Location of ANT (the directory usually contains 'apache-ant-')"))
-	FDirectoryPath ANTPath;
-
 	// Location on disk of Java (falls back to JAVA_HOME environment variable if this is left blank)
 	UPROPERTY(GlobalConfig, EditAnywhere, Category = SDKConfig, Meta = (DisplayName = "Location of JAVA (the directory usually contains 'jdk')"))
 	FDirectoryPath JavaPath;
diff --git a/Engine/Source/Developer/Android/AndroidPlatformEditor/Private/AndroidTargetSettingsCustomization.cpp b/Engine/Source/Developer/Android/AndroidPlatformEditor/Private/AndroidTargetSettingsCustomization.cpp
--- a/Engine/Source/Developer/Android/AndroidPlatformEditor/Private/AndroidTargetSettingsCustomization.cpp
+++ b/Engine/Source/Developer/Android/AndroidPlatformEditor/Private/AndroidTargetSettingsCustomization.cpp
@@ -307,8 +307,9 @@ void FAndroidTargetSettingsCustomization::BuildAppManifestSection(IDetailLayoutB
 
 	// check for Gradle change
 	TSharedRef<IPropertyHandle> EnableGradleProperty = DetailLayout.GetProperty(GET_MEMBER_NAME_CHECKED(UAndroidRuntimeSettings, bEnableGradle));
-	FSimpleDelegate EnableGradleChange = FSimpleDelegate::CreateSP(this, &FAndroidTargetSettingsCustomization::OnEnableGradleChange);
-	EnableGradleProperty->SetOnPropertyValueChanged(EnableGradleChange);
+	EnableGradleProperty->MarkHiddenByCustomization();
+	//FSimpleDelegate EnableGradleChange = FSimpleDelegate::CreateSP(this, &FAndroidTargetSettingsCustomization::OnEnableGradleChange);
+	//EnableGradleProperty->SetOnPropertyValueChanged(EnableGradleChange);
 
 	// check for GoogleVR change
 	TSharedRef<IPropertyHandle> GoogleVRCapsProperty = DetailLayout.GetProperty(GET_MEMBER_NAME_CHECKED(UAndroidRuntimeSettings, GoogleVRCaps));
diff --git a/Engine/Source/Programs/UnrealBuildTool/Platform/Android/AndroidProjectGenerator.cs b/Engine/Source/Programs/UnrealBuildTool/Platform/Android/AndroidProjectGenerator.cs
--- a/Engine/Source/Programs/UnrealBuildTool/Platform/Android/AndroidProjectGenerator.cs
+++ b/Engine/Source/Programs/UnrealBuildTool/Platform/Android/AndroidProjectGenerator.cs
@@ -422,7 +422,6 @@ namespace UnrealBuildTool
 								"	<PropertyGroup Label=\"UserMacros\">" + ProjectFileGenerator.NewLine +
 								"		<ANDROID_HOME>" + vsAndroidSDKPath + "</ANDROID_HOME>" + ProjectFileGenerator.NewLine +
 								"		<JAVA_HOME>$(VS_JavaHome)</JAVA_HOME>" + ProjectFileGenerator.NewLine +
-								"		<ANT_HOME>$(VS_AntHome)</ANT_HOME>" + ProjectFileGenerator.NewLine +
 								"		<NDKROOT>$(VS_NdkRoot)</NDKROOT>" + ProjectFileGenerator.NewLine +
 								"	</PropertyGroup>" + ProjectFileGenerator.NewLine +
 								"	<PropertyGroup />" + ProjectFileGenerator.NewLine +
@@ -436,10 +435,6 @@ namespace UnrealBuildTool
 								"			<Value>$(JAVA_HOME)</Value>" + ProjectFileGenerator.NewLine +
 								"			<EnvironmentVariable>true</EnvironmentVariable>" + ProjectFileGenerator.NewLine +
 								"		</BuildMacro>" + ProjectFileGenerator.NewLine +
-								"		<BuildMacro Include=\"ANT_HOME\">" + ProjectFileGenerator.NewLine +
-								"			<Value>$(ANT_HOME)</Value>" + ProjectFileGenerator.NewLine +
-								"			<EnvironmentVariable>true</EnvironmentVariable>" + ProjectFileGenerator.NewLine +
-								"		</BuildMacro>" + ProjectFileGenerator.NewLine +
 								"		<BuildMacro Include=\"NDKROOT\">" + ProjectFileGenerator.NewLine +
 								"			<Value>$(NDKROOT)</Value>" + ProjectFileGenerator.NewLine +
 								"			<EnvironmentVariable>true</EnvironmentVariable>" + ProjectFileGenerator.NewLine +
diff --git a/Engine/Source/Programs/UnrealBuildTool/Platform/Android/UEBuildAndroid.cs b/Engine/Source/Programs/UnrealBuildTool/Platform/Android/UEBuildAndroid.cs
--- a/Engine/Source/Programs/UnrealBuildTool/Platform/Android/UEBuildAndroid.cs
+++ b/Engine/Source/Programs/UnrealBuildTool/Platform/Android/UEBuildAndroid.cs
@@ -537,7 +537,6 @@ namespace UnrealBuildTool
 				Dictionary<string, string> EnvVarNames = new Dictionary<string, string> { 
                                                          {"ANDROID_HOME", "SDKPath"}, 
                                                          {"NDKROOT", "NDKPath"}, 
-                                                         {"ANT_HOME", "ANTPath"},
                                                          {"JAVA_HOME", "JavaPath"}
                                                          };
 
diff --git a/Engine/Source/Programs/UnrealBuildTool/Platform/Android/UEDeployAndroid.cs b/Engine/Source/Programs/UnrealBuildTool/Platform/Android/UEDeployAndroid.cs
--- a/Engine/Source/Programs/UnrealBuildTool/Platform/Android/UEDeployAndroid.cs
+++ b/Engine/Source/Programs/UnrealBuildTool/Platform/Android/UEDeployAndroid.cs
@@ -55,16 +55,12 @@ namespace UnrealBuildTool
 			bPackageDataInsideApk = InForcePackageData || IniValue == true;
 		}
 
-		// Enable Gradle instead of Ant (project-level setting for now)
-		private bool bGradleEnabled = false;
-
 		private UnrealPluginLanguage UPL = null;
 		private string UPLHashCode = null;
 		private bool ARCorePluginEnabled = false;
 		private bool FacebookPluginEnabled = false;
 		private bool OculusMobilePluginEnabled = false;
 		private bool GoogleVRPluginEnabled = false;
-		private bool CrashlyticsPluginEnabled = false;
 
 		public void SetAndroidPluginData(List<string> Architectures, List<string> inPluginExtraData)
 		{
@@ -79,7 +75,6 @@ namespace UnrealBuildTool
 			FacebookPluginEnabled = false;
 			GoogleVRPluginEnabled = false;
 			OculusMobilePluginEnabled = false;
-			CrashlyticsPluginEnabled = false;
 			foreach (string Plugin in inPluginExtraData)
 			{
 				// check if the Facebook plugin was enabled
@@ -109,14 +104,6 @@ namespace UnrealBuildTool
 					GoogleVRPluginEnabled = true;
 					continue;
 				}
-
-				// check if Crashlytics plugin was enabled
-				// NOTE: There is a thirdparty plugin using Crashlytics_UPL_Android.xml which shouldn't use this code so check full name
-				if (Plugin.Contains("Crashlytics_UPL.xml"))
-				{
-					CrashlyticsPluginEnabled = true;
-					continue;
-				}
 			}
 
 			UPL = new UnrealPluginLanguage(ProjectFile, inPluginExtraData, NDKArches, "http://schemas.android.com/apk/res/android", "xmlns:android=\"http://schemas.android.com/apk/res/android\"", UnrealTargetPlatform.Android);
@@ -207,7 +194,7 @@ namespace UnrealBuildTool
 		}
 
 		private string CachedSDKLevel = null;
-		private string GetSdkApiLevel(AndroidToolChain ToolChain, bool bGradleEnabled = false)
+		private string GetSdkApiLevel(AndroidToolChain ToolChain)
 		{
 			if (CachedSDKLevel == null)
 			{
@@ -249,30 +236,16 @@ namespace UnrealBuildTool
 					SDKLevelInt = GetApiLevelInt(SDKLevel);
 					if (SDKLevelInt < MinimumSDKLevel)
 					{
-						if (bGradleEnabled)
-						{
-							SDKLevelInt = MinimumSDKLevel;
-							SDKLevel = "android-" + MinimumSDKLevel.ToString();
-							Log.TraceInformation("Gradle will attempt to download SDK API level {0}", SDKLevelInt);
-						}
-						else
-						{
-							throw new BuildException("Can't make an APK without SDK API 'android-" + MinimumSDKLevel.ToString() + "' minimum installed");
-						}
+						SDKLevelInt = MinimumSDKLevel;
+						SDKLevel = "android-" + MinimumSDKLevel.ToString();
+						Log.TraceInformation("Gradle will attempt to download SDK API level {0}", SDKLevelInt);
 					}
 				}
 
 				// validate the platform SDK is installed
 				if (!ValidateSDK(PlatformsDir, SDKLevel))
 				{
-					if (bGradleEnabled)
-					{
-						Log.TraceWarning("The SDK API requested '{0}' not installed in {1}; Gradle will attempt to download it.", SDKLevel, PlatformsDir);
-					}
-					else
-					{
-						throw new BuildException("The SDK API requested '{0}' not installed in {1}", SDKLevel, PlatformsDir);
-					}
+					Log.TraceWarning("The SDK API requested '{0}' not installed in {1}; Gradle will attempt to download it.", SDKLevel, PlatformsDir);
 				}
 
 				Log.TraceInformation("Building Java with SDK API level '{0}'", SDKLevel);
@@ -305,7 +278,7 @@ namespace UnrealBuildTool
 			return Value;
 		}
 
-		private string GetBuildToolsVersion(bool bGradle)
+		private string GetBuildToolsVersion()
 		{
 			// return cached path if ANDROID_HOME has not changed
 			string HomePath = Environment.ExpandEnvironmentVariables("%ANDROID_HOME%");
@@ -362,7 +335,7 @@ namespace UnrealBuildTool
 			}
 
 			// with Gradle enabled use at least 24.0.2 (will be installed by Gradle if missing)
-			if (bGradle && (BestVersion < ((24 << 24) | (0 << 16) | (2 << 8))))
+			if (BestVersion < ((24 << 24) | (0 << 16) | (2 << 8)))
 			{
 				BestVersionString = "24.0.2";
 			}
@@ -507,43 +480,6 @@ namespace UnrealBuildTool
 			return bDisableVerifyOBBOnStartUp;
 		}
 
-		private static string GetAntPath()
-		{
-			// look up an ANT_HOME environment variable
-			string AntHome = Environment.GetEnvironmentVariable("ANT_HOME");
-			if (!string.IsNullOrEmpty(AntHome) && Directory.Exists(AntHome))
-			{
-				string AntPath = AntHome + "/bin/ant" + (Utils.IsRunningOnMono ? "" : ".bat");
-				// use it if found
-				if (File.Exists(AntPath))
-				{
-					return AntPath;
-				}
-			}
-
-			// otherwise, look in the eclipse install for the ant plugin (matches the unzipped Android ADT from Google)
-			string PluginDir = Environment.ExpandEnvironmentVariables("%ANDROID_HOME%/../eclipse/plugins");
-			if (Directory.Exists(PluginDir))
-			{
-				string[] Plugins = Directory.GetDirectories(PluginDir, "org.apache.ant*");
-				// use the first one with ant.bat
-				if (Plugins.Length > 0)
-				{
-					foreach (string PluginName in Plugins)
-					{
-						string AntPath = PluginName + "/bin/ant" + (Utils.IsRunningOnMono ? "" : ".bat");
-						// use it if found
-						if (File.Exists(AntPath))
-						{
-							return AntPath;
-						}
-					}
-				}
-			}
-
-			throw new BuildException("Unable to find ant.bat (via %ANT_HOME% or %ANDROID_HOME%/../eclipse/plugins/org.apache.ant*");
-		}
-
 		private static bool SafeDeleteFile(string Filename, bool bCheckExists = true)
 		{
 			if (!bCheckExists || File.Exists(Filename))
@@ -1252,7 +1188,7 @@ namespace UnrealBuildTool
 			Utils.RunLocalProcessAndLogOutput(StartInfo);
 		}
 
-		private static void CopySTL(AndroidToolChain ToolChain, string UE4BuildPath, string UE4Arch, string NDKArch, bool bForDistribution, bool bGradleEnabled)
+		private static void CopySTL(AndroidToolChain ToolChain, string UE4BuildPath, string UE4Arch, string NDKArch, bool bForDistribution)
 		{
 			string GccVersion = "4.6";
 			if (Directory.Exists(Environment.ExpandEnvironmentVariables("%NDKROOT%/sources/cxx-stl/gnu-libstdc++/4.9")))
@@ -1264,53 +1200,18 @@ namespace UnrealBuildTool
 				GccVersion = "4.8";
 			}
 
-			if (bGradleEnabled)
-			{
-				// copy it in!
-				string SourceSTLSOName = Environment.ExpandEnvironmentVariables("%NDKROOT%/sources/cxx-stl/gnu-libstdc++/") + GccVersion + "/libs/" + NDKArch + "/libgnustl_shared.so";
-				string FinalSTLSOName = UE4BuildPath + "/jni/" + NDKArch + "/libgnustl_shared.so";
-
-				// remove from location for Ant if previously copied
-				string WrongSTLSOName = UE4BuildPath + "/libs/" + NDKArch + "/libgnustl_shared.so";
-				SafeDeleteFile(WrongSTLSOName);
+			// copy it in!
+			string SourceSTLSOName = Environment.ExpandEnvironmentVariables("%NDKROOT%/sources/cxx-stl/gnu-libstdc++/") + GccVersion + "/libs/" + NDKArch + "/libgnustl_shared.so";
+			string FinalSTLSOName = UE4BuildPath + "/jni/" + NDKArch + "/libgnustl_shared.so";
 
-				// check to see if libgnustl_shared.so is newer than last time we copied
-				bool bFileExists = File.Exists(FinalSTLSOName);
-				TimeSpan Diff = File.GetLastWriteTimeUtc(FinalSTLSOName) - File.GetLastWriteTimeUtc(SourceSTLSOName);
-				if (!bFileExists || Diff.TotalSeconds < -1 || Diff.TotalSeconds > 1)
-				{
-					SafeDeleteFile(FinalSTLSOName);
-					Directory.CreateDirectory(Path.GetDirectoryName(FinalSTLSOName));
-					File.Copy(SourceSTLSOName, FinalSTLSOName, true);
-				}
-			}
-			else
+			// check to see if libgnustl_shared.so is newer than last time we copied
+			bool bFileExists = File.Exists(FinalSTLSOName);
+			TimeSpan Diff = File.GetLastWriteTimeUtc(FinalSTLSOName) - File.GetLastWriteTimeUtc(SourceSTLSOName);
+			if (!bFileExists || Diff.TotalSeconds < -1 || Diff.TotalSeconds > 1)
 			{
-				// copy it in!
-				string SourceSTLSOName = Environment.ExpandEnvironmentVariables("%NDKROOT%/sources/cxx-stl/gnu-libstdc++/") + GccVersion + "/libs/" + NDKArch + "/libgnustl_shared.so";
-				string FinalSTLSOName = UE4BuildPath + "/libs/" + NDKArch + "/libgnustl_shared.so";
-
-				// remove from location for Gradle if previously copied
-				string WrongSTLSOName = UE4BuildPath + "/jni/" + NDKArch + "/libgnustl_shared.so";
-				SafeDeleteFile(WrongSTLSOName);
-
-				// check to see if libgnustl_shared.so is newer than last time we copied (or needs stripping for distribution)
-				bool bFileExists = File.Exists(FinalSTLSOName);
-				TimeSpan Diff = File.GetLastWriteTimeUtc(FinalSTLSOName) - File.GetLastWriteTimeUtc(SourceSTLSOName);
-				if (bForDistribution || !bFileExists || Diff.TotalSeconds < -1 || Diff.TotalSeconds > 1)
-				{
-					SafeDeleteFile(FinalSTLSOName);
-					Directory.CreateDirectory(Path.GetDirectoryName(FinalSTLSOName));
-					if (bForDistribution)
-					{
-						// Strip debug symbols for distribution builds
-						StripDebugSymbols(SourceSTLSOName, FinalSTLSOName, UE4Arch, true);
-					}
-					else
-					{
-						File.Copy(SourceSTLSOName, FinalSTLSOName, true);
-					}
-				}
+				SafeDeleteFile(FinalSTLSOName);
+				Directory.CreateDirectory(Path.GetDirectoryName(FinalSTLSOName));
+				File.Copy(SourceSTLSOName, FinalSTLSOName, true);
 			}
 		}
 
@@ -1577,104 +1478,15 @@ namespace UnrealBuildTool
 			return true;
 		}
 
-		private void UpdateProjectProperties(AndroidToolChain ToolChain, string UE4BuildPath, string ProjectName)
-		{
-			Log.TraceInformation("\n===={0}====UPDATING BUILD CONFIGURATION FILES====================================================", DateTime.Now.ToString());
-
-			// get all of the libs (from engine + project)
-			string JavaLibsDir = Path.Combine(UE4BuildPath, "JavaLibs");
-			string[] LibDirs = Directory.GetDirectories(JavaLibsDir);
-
-			// get existing project.properties lines (if any)
-			string ProjectPropertiesFile = Path.Combine(UE4BuildPath, "project.properties");
-			string[] PropertiesLines = new string[] { };
-			if (File.Exists(ProjectPropertiesFile))
-			{
-				PropertiesLines = File.ReadAllLines(ProjectPropertiesFile);
-			}
-
-			// figure out how many libraries were already listed (if there were more than this listed, then we need to start the file over, because we need to unreference a library)
-			int NumOutstandingAlreadyReferencedLibs = 0;
-			foreach (string Line in PropertiesLines)
-			{
-				if (Line.StartsWith("android.library.reference."))
-				{
-					NumOutstandingAlreadyReferencedLibs++;
-				}
-			}
-
-			// now go through each one and verify they are listed in project properties, and if not, add them
-			List<string> LibsToBeAdded = new List<string>();
-			foreach (string LibDir in LibDirs)
-			{
-				// put it in terms of the subdirectory that would be in the project.properties
-				string RelativePath = "JavaLibs/" + Path.GetFileName(LibDir);
-
-				// now look for this in the existing file
-				bool bWasReferencedAlready = false;
-				foreach (string Line in PropertiesLines)
-				{
-					if (Line.StartsWith("android.library.reference.") && Line.EndsWith(RelativePath))
-					{
-						// this lib was already referenced, don't need to readd
-						bWasReferencedAlready = true;
-						break;
-					}
-				}
-
-				if (bWasReferencedAlready)
-				{
-					// if it was, no further action needed, and count it off
-					NumOutstandingAlreadyReferencedLibs--;
-				}
-				else
-				{
-					// otherwise, we need to add it to the project properties
-					LibsToBeAdded.Add(RelativePath);
-				}
-			}
-
-			// now at this point, if there are any outstanding already referenced libs, we have too many, so we have to start over
-			if (NumOutstandingAlreadyReferencedLibs > 0)
-			{
-				// @todo android: If a user had a project.properties in the game, NEVER do this
-				Log.TraceInformation("There were too many libs already referenced in project.properties, tossing it");
-				SafeDeleteFile(ProjectPropertiesFile);
-
-				LibsToBeAdded.Clear();
-				foreach (string LibDir in LibDirs)
-				{
-					// put it in terms of the subdirectory that would be in the project.properties
-					LibsToBeAdded.Add("JavaLibs/" + Path.GetFileName(LibDir));
-				}
-			}
-
-			// now update the project for each library
-			string AndroidCommandPath = Environment.ExpandEnvironmentVariables("%ANDROID_HOME%/tools/android" + (Utils.IsRunningOnMono ? "" : ".bat"));
-			string UpdateCommandLine = "--silent update project --subprojects --name " + ProjectName + " --path . --target " + GetSdkApiLevel(ToolChain);
-			foreach (string Lib in LibsToBeAdded)
-			{
-				string LocalUpdateCommandLine = UpdateCommandLine + " --library " + Lib;
-
-				// make sure each library has a build.xml - --subprojects doesn't create build.xml files, but it will create project.properties
-				// and later code needs each lib to have a build.xml
-				RunCommandLineProgramWithException(UE4BuildPath, AndroidCommandPath, "--silent update lib-project --path " + Lib + " --target " + GetSdkApiLevel(ToolChain), "");
-				RunCommandLineProgramWithException(UE4BuildPath, AndroidCommandPath, LocalUpdateCommandLine, "Updating project.properties, local.properties, and build.xml for " + Path.GetFileName(Lib) + "...");
-			}
-
-		}
-
-
-		private string GetAllBuildSettings(AndroidToolChain ToolChain, string BuildPath, bool bForDistribution, bool bMakeSeparateApks, bool bPackageDataInsideApk, bool bDisableVerifyOBBOnStartUp, bool bUseExternalFilesDir, bool bGradleEnabled, string TemplatesHashCode)
+		private string GetAllBuildSettings(AndroidToolChain ToolChain, string BuildPath, bool bForDistribution, bool bMakeSeparateApks, bool bPackageDataInsideApk, bool bDisableVerifyOBBOnStartUp, bool bUseExternalFilesDir, string TemplatesHashCode)
 		{
 			// make the settings string - this will be char by char compared against last time
 			StringBuilder CurrentSettings = new StringBuilder();
 			CurrentSettings.AppendLine(string.Format("NDKROOT={0}", Environment.GetEnvironmentVariable("NDKROOT")));
 			CurrentSettings.AppendLine(string.Format("ANDROID_HOME={0}", Environment.GetEnvironmentVariable("ANDROID_HOME")));
-			CurrentSettings.AppendLine(string.Format("ANT_HOME={0}", Environment.GetEnvironmentVariable("ANT_HOME")));
 			CurrentSettings.AppendLine(string.Format("JAVA_HOME={0}", Environment.GetEnvironmentVariable("JAVA_HOME")));
 			CurrentSettings.AppendLine(string.Format("NDKVersion={0}", ToolChain.GetNdkApiLevel()));
-			CurrentSettings.AppendLine(string.Format("SDKVersion={0}", GetSdkApiLevel(ToolChain, bGradleEnabled)));
+			CurrentSettings.AppendLine(string.Format("SDKVersion={0}", GetSdkApiLevel(ToolChain)));
 			CurrentSettings.AppendLine(string.Format("bForDistribution={0}", bForDistribution));
 			CurrentSettings.AppendLine(string.Format("bMakeSeparateApks={0}", bMakeSeparateApks));
 			CurrentSettings.AppendLine(string.Format("bPackageDataInsideApk={0}", bPackageDataInsideApk));
@@ -2422,10 +2234,7 @@ namespace UnrealBuildTool
 				}
 			}
 			Text.AppendLine("\t             android:hardwareAccelerated=\"true\"");
-			if (bGradleEnabled)
-			{
-				Text.AppendLine("\t				android:name=\"com.epicgames.ue4.GameApplication\"");
-			}
+			Text.AppendLine("\t				android:name=\"com.epicgames.ue4.GameApplication\"");
 			Text.AppendLine("\t             android:hasCode=\"true\">");
 			if (bShowLaunchImage)
 			{
@@ -2620,11 +2429,6 @@ namespace UnrealBuildTool
 			}
 			else
 			{
-				if (!bGradleEnabled)
-				{
-					// need just the number part of the sdk
-					Text.AppendLine(string.Format("\t<uses-sdk android:minSdkVersion=\"{0}\" android:targetSdkVersion=\"{1}\"/>", MinSDKVersion, TargetSDKVersion));
-				}
 				Text.AppendLine("\t<uses-feature android:glEsVersion=\"" + AndroidToolChain.GetGLESVersionFromGPUArch(GPUArch, bOnlyETC2Enabled, bBuildForES2, bBuildForES31) + "\" android:required=\"true\" />");
 				Text.AppendLine("\t<uses-permission android:name=\"android.permission.INTERNET\"/>");
 				Text.AppendLine("\t<uses-permission android:name=\"android.permission.WRITE_EXTERNAL_STORAGE\"/>");
@@ -2915,84 +2719,6 @@ namespace UnrealBuildTool
 			}
 		}
 
-		private void PatchAntBatIfNeeded()
-		{
-			// only need to do this for Windows (other platforms are Mono so use that for check)
-			if (Utils.IsRunningOnMono)
-			{
-				return;
-			}
-
-			string AntBinPath = Environment.ExpandEnvironmentVariables("%ANT_HOME%/bin");
-			string AntBatFilename = Path.Combine(AntBinPath, "ant.bat");
-			string AntOrigBatFilename = Path.Combine(AntBinPath, "ant.orig.bat");
-
-			// check for an unused drive letter
-			string UnusedDriveLetter = "";
-			bool bFound = true;
-			DriveInfo[] AllDrives = DriveInfo.GetDrives();
-			for (char DriveLetter = 'Z'; DriveLetter >= 'A'; DriveLetter--)
-			{
-				UnusedDriveLetter = Char.ToString(DriveLetter) + ":";
-				bFound = false;
-				for (int DriveIndex = AllDrives.Length-1; DriveIndex >= 0; DriveIndex--)
-				{
-					if (AllDrives[DriveIndex].Name.ToUpper().StartsWith(UnusedDriveLetter))
-					{
-						bFound = true;
-						break;
-					}
-				}
-				if (!bFound)
-				{
-					break;
-				}
-			}
-
-			if (bFound)
-			{
-				Log.TraceInformation("\nUnable to apply fixed ant.bat (all drive letters in use!)");
-				return;
-			}
-
-			Log.TraceInformation("\nPatching ant.bat to work around commandline length limit (using unused drive letter {0})", UnusedDriveLetter);
-
-			if (!File.Exists(AntOrigBatFilename))
-			{
-				// copy the existing ant.bat to ant.orig.bat
-				File.Copy(AntBatFilename, AntOrigBatFilename, true);
-			}
-
-			// make sure ant.bat isn't read-only
-			FileAttributes Attribs = File.GetAttributes(AntBatFilename);
-			if (Attribs.HasFlag(FileAttributes.ReadOnly))
-			{
-				File.SetAttributes(AntBatFilename, Attribs & ~FileAttributes.ReadOnly);
-			}
-
-			// generate new ant.bat with an unused drive letter for subst
-			string AntBatText =
-					"@echo off\n" +
-					"setlocal\n" +
-					"set ANTPATH=%~dp0\n" +
-					"set ANT_CMD_LINE_ARGS=\n" +
-					":setupArgs\n" +
-					"if \"\"%1\"\"==\"\"\"\" goto doneStart\n" +
-					"set ANT_CMD_LINE_ARGS=%ANT_CMD_LINE_ARGS% %1\n" +
-					"shift\n" +
-					"goto setupArgs\n\n" +
-					":doneStart\n" +
-					"subst " + UnusedDriveLetter + " \"%CD%\"\n" +
-					"pushd " + UnusedDriveLetter + "\n" +
-					"call \"%ANTPATH%\\ant.orig.bat\" %ANT_CMD_LINE_ARGS%\n" +
-					"set ANTERROR=%ERRORLEVEL%\n" +
-					"popd\n" +
-					"subst " + UnusedDriveLetter + " /d\n" +
-					"exit /b %ANTERROR%\n";
-
-			File.WriteAllText(AntBatFilename, AntBatText);
-		}
-
 		private bool CreateRunGradle(string GradlePath)
 		{
 			string RunGradleBatFilename = Path.Combine(GradlePath, "rungradle.bat");
@@ -3037,7 +2763,7 @@ namespace UnrealBuildTool
 				}
 			}
 
-			// generate new ant.bat with an unused drive letter for subst
+			// generate new rungradle.bat with an unused drive letter for subst
 			string RunGradleBatText =
 					"@echo off\n" +
 					"setlocal\n" +
@@ -3113,7 +2839,7 @@ namespace UnrealBuildTool
 
 			// Make sure minSdkVersion is at least 13 (need this for appcompat-v13 used by AndroidPermissions)
 			// this may be changed by active plugins (Google Play Services 11.0.4 needs 14 for example)
-			if (bGradleEnabled && MinSDKVersion < MinimumSDKLevelForGradle)
+			if (MinSDKVersion < MinimumSDKLevelForGradle)
 			{
 				MinSDKVersion = MinimumSDKLevelForGradle;
 				Log.TraceInformation("Fixing minSdkVersion; requires minSdkVersion of {0} with Gradle based on active plugins", MinimumSDKLevelForGradle);
@@ -3323,6 +3049,11 @@ namespace UnrealBuildTool
 		{
 			Log.TraceInformation("\n===={0}====PREPARING TO MAKE APK=================================================================", DateTime.Now.ToString());
 
+			if (!GradleEnabled())
+			{
+				throw new BuildException("Support for building APK without Gradle is depreciated; please update your project Engine.ini.");
+			}
+
 			if (UPL.GetLastError() != null)
 			{
 				throw new BuildException("Cannot make APK with UPL errors");
@@ -3333,15 +3064,10 @@ namespace UnrealBuildTool
 
 			SetMinimumSDKLevelForGradle();
 
-			// check for Gradle enabled for this project
-			bGradleEnabled = GradleEnabled();
-
-			if (bGradleEnabled)
+			// Verify license agreement since we require Gradle
+			if (!IsLicenseAgreementValid())
 			{
-				if (!IsLicenseAgreementValid())
-				{
-					throw new BuildException("Android SDK license file not found.  Please agree to license in Android project settings in the editor.");
-				}
+				throw new BuildException("Android SDK license file not found.  Please agree to license in Android project settings in the editor.");
 			}
 
 			LogBuildSetup();
@@ -3349,14 +3075,9 @@ namespace UnrealBuildTool
 			bool bIsBuildMachine = Environment.GetEnvironmentVariable("IsBuildMachine") == "1";
 
 			// do this here so we'll stop early if there is a problem with the SDK API level (cached so later calls will return the same)
-			string SDKAPILevel = GetSdkApiLevel(ToolChain, bGradleEnabled);
+			string SDKAPILevel = GetSdkApiLevel(ToolChain);
 			int SDKLevelInt = GetApiLevelInt(SDKAPILevel);
-			string BuildToolsVersion = GetBuildToolsVersion(bGradleEnabled);
-
-			if (!bGradleEnabled)
-			{
-				PatchAntBatIfNeeded();
-			}
+			string BuildToolsVersion = GetBuildToolsVersion();
 
 			// cache some tools paths
 			string NDKBuildPath = Environment.ExpandEnvironmentVariables("%NDKROOT%/ndk-build" + (Utils.IsRunningOnMono ? "" : ".cmd"));
@@ -3531,7 +3252,7 @@ namespace UnrealBuildTool
 
 			// check to see if any "meta information" is newer than last time we build
 			string TemplatesHashCode = GenerateTemplatesHashCode(EngineDirectory);
-			string CurrentBuildSettings = GetAllBuildSettings(ToolChain, UE4BuildPath, bForDistribution, bMakeSeparateApks, bPackageDataInsideApk, bDisableVerifyOBBOnStartUp, bUseExternalFilesDir, bGradleEnabled, TemplatesHashCode);
+			string CurrentBuildSettings = GetAllBuildSettings(ToolChain, UE4BuildPath, bForDistribution, bMakeSeparateApks, bPackageDataInsideApk, bDisableVerifyOBBOnStartUp, bUseExternalFilesDir, TemplatesHashCode);
 			string BuildSettingsCacheFile = Path.Combine(UE4BuildPath, "UEBuildSettings.txt");
 
 			// do we match previous build settings?
@@ -3624,14 +3345,6 @@ namespace UnrealBuildTool
 			Dictionary<string, string> Replacements = new Dictionary<string, string>();
 			Replacements.Add("${EXECUTABLE_NAME}", ApplicationDisplayName);
 
-			if (!bIncrementalPackage)
-			{
-				// Wipe the Intermediate/Build/APK directory first, except for dexedLibs, because Google Services takes FOREVER to predex, and it almost never changes
-				// so allow the ANT checking to win here - if this grows a bit with extra libs, it's fine, it _should_ only pull in dexedLibs it needs
-				Log.TraceInformation("Performing complete package - wiping {0}, except for predexedLibs", UE4BuildPath);
-				DeleteDirectory(UE4BuildPath, "dexedLibs");
-			}
-
 			// If we are packaging for Amazon then we need to copy the  file to the correct location
 			Log.TraceInformation("bPackageDataInsideApk = {0}", bPackageDataInsideApk);
 			if (bPackageDataInsideApk)
@@ -3684,17 +3397,6 @@ namespace UnrealBuildTool
 				SafeDeleteFile(CommandLineDestFileName);
 			}
 
-			string AARExtractListFilename = Path.Combine(UE4BuildPath, "JavaLibs", "AARExtractList.txt");
-			if (bGradleEnabled)
-			{
-				//Need to clear out JavaLibs if last run was with Ant
-				if (File.Exists(AARExtractListFilename))
-				{
-					Log.TraceInformation("Cleanup up JavaLibs from previous Ant packaging");
-					DeleteDirectory(Path.Combine(UE4BuildPath, "JavaLibs"));
-				}
-			}
-
 			//Copy build files to the intermediate folder in this order (later overrides earlier):
 			//	- Shared Engine
 			//  - Shared Engine NoRedist (for Epic secret files)
@@ -3703,24 +3405,12 @@ namespace UnrealBuildTool
 			CopyFileDirectory(UE4BuildFilesPath, UE4BuildPath, Replacements);
 			CopyFileDirectory(UE4BuildFilesPath + "/NotForLicensees", UE4BuildPath, Replacements);
 			CopyFileDirectory(UE4BuildFilesPath + "/NoRedist", UE4BuildPath, Replacements);
-			if (!bGradleEnabled)
-			{
-				CopyFileDirectory(Path.Combine(EngineDirectory, "Build", "Android", "Legacy"), UE4BuildPath, Replacements);
-			}
 			CopyFileDirectory(GameBuildFilesPath, UE4BuildPath, Replacements);
 			CopyFileDirectory(GameBuildFilesPath + "/NotForLicensees", UE4BuildPath, Replacements);
 			CopyFileDirectory(GameBuildFilesPath + "/NoRedist", UE4BuildPath, Replacements);
 
-			if (!bGradleEnabled)
-			{
-				//Extract AAR and Jar files with dependencies if not using Gradle
-				ExtractAARAndJARFiles(EngineDirectory, UE4BuildPath, NDKArches, PackageName, AARExtractListFilename);
-			}
-			else
-			{
-				//Generate Gradle AAR dependencies
-				GenerateGradleAARImports(EngineDirectory, UE4BuildPath, NDKArches);
-			}
+			//Generate Gradle AAR dependencies
+			GenerateGradleAARImports(EngineDirectory, UE4BuildPath, NDKArches);
 
 			//Now validate GooglePlay app_id if enabled
 			ValidateGooglePlay(UE4BuildPath);
@@ -3753,11 +3443,6 @@ namespace UnrealBuildTool
 			Log.TraceInformation("\n===={0}====PREPARING NATIVE CODE=================================================================", DateTime.Now.ToString());
 			bool HasNDKPath = File.Exists(NDKBuildPath);
 
-			// get Ant verbosity level
-			ConfigHierarchy Ini = GetConfigCacheIni(ConfigHierarchyType.Engine);
-			string AntVerbosity;
-			Ini.GetString("/Script/AndroidRuntimeSettings.AndroidRuntimeSettings", "AntVerbosity", out AntVerbosity);
-
 			// use Gradle for compile/package
 			string UE4BuildGradlePath = Path.Combine(UE4BuildPath, "gradle");
 			string UE4BuildGradleAppPath = Path.Combine(UE4BuildGradlePath, "app");
@@ -3783,28 +3468,12 @@ namespace UnrealBuildTool
 				XDocument AdditionalBuildPathFilesDoc = new XDocument(new XElement("files"));
 				UPL.ProcessPluginNode(NDKArch, "additionalBuildPathFiles", "", ref AdditionalBuildPathFilesDoc);
 
-				if (!bGradleEnabled)
-				{
-					// update metadata files (like project.properties, build.xml) if we are missing a build.xml or if we just overwrote project.properties with a bad version in it (from game/engine dir)
-					UpdateProjectProperties(ToolChain, UE4BuildPath, ProjectName);
-
-					// modify the generated build.xml before the final include
-					UpdateBuildXML(Arch, NDKArch, EngineDirectory, UE4BuildPath);
-
-					// Write Crashlytics data if enabled (not needed for Gradle)
-					if (CrashlyticsPluginEnabled)
-					{
-						Trace.TraceInformation("Writing Crashlytics resources");
-						WriteCrashlyticsResources(Path.Combine(ProjectDirectory, "Build", "Android"), PackageName, ApplicationDisplayName, bSkipGradleBuild, Arch);
-					}
-				}
-
 				// Generate the OBBData.java file again in case architecture has different store version
 				WriteJavaOBBDataFile(UE4OBBDataFileName, PackageName, RequiredOBBFiles, CookFlavor, bPackageDataInsideApk, Arch);
 
 				// update GameActivity.java and GameApplication.java if out of date
 				UpdateGameActivity(Arch, NDKArch, EngineDirectory, UE4BuildPath);
-				UpdateGameApplication(Arch, NDKArch, EngineDirectory, UE4BuildPath, bGradleEnabled);
+				UpdateGameApplication(Arch, NDKArch, EngineDirectory, UE4BuildPath);
 
 				// we don't actually need the SO for the bSkipGradleBuild case
 				string FinalSOName = null;
@@ -3837,84 +3506,22 @@ namespace UnrealBuildTool
 						throw new BuildException("Can't make an APK without the jni directory [{0}/jni]", UE4BuildFilesPath);
 					}
 
-					if (bGradleEnabled)
-					{
-						string JniDir = UE4BuildPath + "/jni/" + NDKArch;
-						FinalSOName = JniDir + "/libUE4.so";
+					string JniDir = UE4BuildPath + "/jni/" + NDKArch;
+					FinalSOName = JniDir + "/libUE4.so";
 
-						// clear out libs directory like ndk-build would have
-						string LibsDir = Path.Combine(UE4BuildPath, "libs");
-						DeleteDirectory(LibsDir);
-						MakeDirectoryIfRequired(LibsDir);
+					// clear out libs directory like ndk-build would have
+					string LibsDir = Path.Combine(UE4BuildPath, "libs");
+					DeleteDirectory(LibsDir);
+					MakeDirectoryIfRequired(LibsDir);
 
-						// check to see if libUE4.so needs to be copied
-						if (BuildListComboTotal > 1 || FilesAreDifferent(SourceSOName, FinalSOName))
-						{
-							Log.TraceInformation("\nCopying new .so {0} file to jni folder...", SourceSOName);
-							Directory.CreateDirectory(JniDir);
-							// copy the binary to the standard .so location
-							File.Copy(SourceSOName, FinalSOName, true);
-							File.SetLastWriteTimeUtc(FinalSOName, File.GetLastWriteTimeUtc(SourceSOName));
-						}
-					}
-					else
+					// check to see if libUE4.so needs to be copied
+					if (BuildListComboTotal > 1 || FilesAreDifferent(SourceSOName, FinalSOName))
 					{
-						if (HasNDKPath)
-						{
-							string LibDir = UE4BuildPath + "/jni/" + NDKArch;
-							FinalSOName = LibDir + "/libUE4.so";
-
-							// check to see if libUE4.so needs to be copied
-							if (BuildListComboTotal > 1 || FilesAreDifferent(SourceSOName, FinalSOName))
-							{
-								Log.TraceInformation("\nCopying new .so {0} file to jni folder...", SourceSOName);
-								Directory.CreateDirectory(LibDir);
-								// copy the binary to the standard .so location
-								File.Copy(SourceSOName, FinalSOName, true);
-							}
-
-							// remove any read only flags
-							FileInfo DestFileInfo2 = new FileInfo(FinalSOName);
-							DestFileInfo2.Attributes = DestFileInfo2.Attributes & ~FileAttributes.ReadOnly;
-							File.SetLastWriteTimeUtc(FinalSOName, File.GetLastWriteTimeUtc(SourceSOName));
-
-							// run ndk-build for Ant (will stage libUE4.so into libs)
-							string LibSOName = UE4BuildPath + "/libs/" + NDKArch + "/libUE4.so";
-
-							// always delete libs up to this point so fat binaries and incremental builds work together (otherwise we might end up with multiple
-							// so files in an apk that doesn't want them)
-							// note that we don't want to delete all libs, just the ones we copied
-							TimeSpan Diff = File.GetLastWriteTimeUtc(LibSOName) - File.GetLastWriteTimeUtc(FinalSOName);
-							if (!File.Exists(LibSOName) || Diff.TotalSeconds < -1 || Diff.TotalSeconds > 1)
-							{
-								foreach (string Lib in Directory.EnumerateFiles(UE4BuildPath + "/libs", "libUE4*.so", SearchOption.AllDirectories))
-								{
-									File.Delete(Lib);
-								}
-
-								string CommandLine = "APP_ABI=\"" + NDKArch + " " + "\"";
-								if (!bForDistribution)
-								{
-									CommandLine += " NDK_DEBUG=1";
-								}
-								RunCommandLineProgramWithException(UE4BuildPath, NDKBuildPath, CommandLine, "Preparing native code for debugging...", true);
-
-								File.SetLastWriteTimeUtc(LibSOName, File.GetLastWriteTimeUtc(FinalSOName));
-							}
-						}
-						else
-						{
-							// if no NDK, we don't need any of the debugger stuff, so we just copy the .so to where it will end up
-							FinalSOName = UE4BuildPath + "/libs/" + NDKArch + "/libUE4.so";
-
-							// check to see if libUE4.so needs to be copied
-							if (BuildListComboTotal > 1 || FilesAreDifferent(SourceSOName, FinalSOName))
-							{
-								Log.TraceInformation("\nCopying .so {0} file to jni folder...", SourceSOName);
-								Directory.CreateDirectory(Path.GetDirectoryName(FinalSOName));
-								File.Copy(SourceSOName, FinalSOName, true);
-							}
-						}
+						Log.TraceInformation("\nCopying new .so {0} file to jni folder...", SourceSOName);
+						Directory.CreateDirectory(JniDir);
+						// copy the binary to the standard .so location
+						File.Copy(SourceSOName, FinalSOName, true);
+						File.SetLastWriteTimeUtc(FinalSOName, File.GetLastWriteTimeUtc(SourceSOName));
 					}
 
 					// remove any read only flags
@@ -3925,7 +3532,7 @@ namespace UnrealBuildTool
 
 				// after ndk-build is called, we can now copy in the stl .so (ndk-build deletes old files)
 				// copy libgnustl_shared.so to library (use 4.8 if possible, otherwise 4.6)
-				CopySTL(ToolChain, UE4BuildPath, Arch, NDKArch, bForDistribution, bGradleEnabled);
+				CopySTL(ToolChain, UE4BuildPath, Arch, NDKArch, bForDistribution);
 				CopyGfxDebugger(UE4BuildPath, Arch, NDKArch);
 				CopyVulkanValidationLayers(UE4BuildPath, Arch, NDKArch, Configuration.ToString());
 
@@ -3936,260 +3543,177 @@ namespace UnrealBuildTool
 
 				Log.TraceInformation("\n===={0}====PERFORMING FINAL APK PACKAGE OPERATION================================================", DateTime.Now.ToString());
 
-
-				if (!bGradleEnabled)
+				// check if any plugins want to increase the required compile SDK version
+				string CompileSDKMin = UPL.ProcessPluginNode(NDKArch, "minimumSDKAPI", "");
+				if (CompileSDKMin != "")
 				{
-					// use Ant for compile/package
-					string AntBuildType = "debug";
-					string AntOutputSuffix = "-debug";
-					if (bForDistribution)
+					int CompileSDKVersionInt;
+					if (!int.TryParse(CompileSDKVersion, out CompileSDKVersionInt))
 					{
-						// Generate the Proguard file contents and write it
-						string ProguardContents = GenerateProguard(NDKArch, UE4BuildFilesPath, GameBuildFilesPath);
-						string ProguardFilename = Path.Combine(UE4BuildPath, "proguard-project.txt");
-						SafeDeleteFile(ProguardFilename);
-						File.WriteAllText(ProguardFilename, ProguardContents);
-
-						// this will write out ant.properties with info needed to sign a distribution build
-						PrepareToSignApk(UE4BuildPath);
-						AntBuildType = "release";
-						AntOutputSuffix = "-release";
+						CompileSDKVersionInt = 23;
 					}
 
-					// Use ant to build the .apk file
-					string AntOptions = AntBuildType + " -Djava.source=1.7 -Djava.target=1.7";
-					string ShellExecutable = BuildHostPlatform.Current.Shell.FullName;
-					string ShellParametersBegin = (BuildHostPlatform.Current.ShellType == ShellType.Sh) ? "-c '" : "/c ";
-					string ShellParametersEnd = (BuildHostPlatform.Current.ShellType == ShellType.Sh) ? "'" : "";
-					switch (AntVerbosity.ToLower())
+					bool bUpdatedCompileSDK = false;
+					string[] CompileSDKLines = CompileSDKMin.Split(new[] { "\r\n", "\r", "\n" }, StringSplitOptions.None);
+					foreach (string CompileLine in CompileSDKLines)
 					{
-						default:
-						case "quiet":
-							if (RunCommandLineProgramAndReturnResult(UE4BuildPath, ShellExecutable, ShellParametersBegin + "\"" + GetAntPath() + "\" -quiet " + AntOptions + ShellParametersEnd, "Making .apk with Ant... (note: it's safe to ignore javac obsolete warnings)") != 0)
+						string VersionString = CompileLine.Replace("android-", "");
+						int VersionInt;
+						if (int.TryParse(CompileLine, out VersionInt))
+						{
+							if (VersionInt > CompileSDKVersionInt)
 							{
-								RunCommandLineProgramWithException(UE4BuildPath, ShellExecutable, ShellParametersBegin + "\"" + GetAntPath() + "\" " + AntOptions + ShellParametersEnd, "Making .apk with Ant again to show errors");
+								CompileSDKVersionInt = VersionInt;
+								bUpdatedCompileSDK = true;
 							}
-							break;
-
-						case "normal":
-							RunCommandLineProgramWithException(UE4BuildPath, ShellExecutable, ShellParametersBegin + "\"" + GetAntPath() + "\" " + AntOptions + ShellParametersEnd, "Making .apk with Ant again to show errors");
-							break;
-
-						case "verbose":
-							RunCommandLineProgramWithException(UE4BuildPath, ShellExecutable, ShellParametersBegin + "\"" + GetAntPath() + "\" -verbose " + AntOptions + ShellParametersEnd, "Making .apk with Ant again to show errors");
-							break;
+						}
 					}
 
-					// upload Crashlytics symbols if plugin enabled and using build machine
-					if (CrashlyticsPluginEnabled && bIsBuildMachine)
+					if (bUpdatedCompileSDK)
 					{
-						AntOptions = "crashlytics-upload-symbols";
-						RunCommandLineProgramWithException(UE4BuildPath, ShellExecutable, ShellParametersBegin + "\"" + GetAntPath() + "\" " + AntOptions + ShellParametersEnd, "Uploading Crashlytics symbols");
+						CompileSDKVersion = CompileSDKVersionInt.ToString();
+						Log.TraceInformation("Building Java with SDK API Level 'android-{0}' due to enabled plugin requirements", CompileSDKVersion);
 					}
-
-					// make sure destination exists
-					Directory.CreateDirectory(Path.GetDirectoryName(DestApkName));
-
-					// now copy to the final location
-					File.Copy(UE4BuildPath + "/bin/" + ProjectName + AntOutputSuffix + ".apk", DestApkName, true);
 				}
-				else
-				{
-					// check if any plugins want to increase the required compile SDK version
-					string CompileSDKMin = UPL.ProcessPluginNode(NDKArch, "minimumSDKAPI", "");
-					if (CompileSDKMin != "")
-					{
-						int CompileSDKVersionInt;
-						if (!int.TryParse(CompileSDKVersion, out CompileSDKVersionInt))
-						{
-							CompileSDKVersionInt = 23;
-						}
-
-						bool bUpdatedCompileSDK = false;
-						string[] CompileSDKLines = CompileSDKMin.Split(new[] { "\r\n", "\r", "\n" }, StringSplitOptions.None);
-						foreach (string CompileLine in CompileSDKLines)
-						{
-							string VersionString = CompileLine.Replace("android-", "");
-							int VersionInt;
-							if (int.TryParse(CompileLine, out VersionInt))
-							{
-								if (VersionInt > CompileSDKVersionInt)
-								{
-									CompileSDKVersionInt = VersionInt;
-									bUpdatedCompileSDK = true;
-								}
-							}
-						}
-
-						if (bUpdatedCompileSDK)
-						{
-							CompileSDKVersion = CompileSDKVersionInt.ToString();
-							Log.TraceInformation("Building Java with SDK API Level 'android-{0}' due to enabled plugin requirements", CompileSDKVersion);
-						}
-					}
 
-					// stage files into gradle app directory
-					string GradleManifest = Path.Combine(UE4BuildGradleMainPath, "AndroidManifest.xml");
-					MakeDirectoryIfRequired(GradleManifest);
-					File.Copy(Path.Combine(UE4BuildPath, "AndroidManifest.xml"), GradleManifest, true);
+				// stage files into gradle app directory
+				string GradleManifest = Path.Combine(UE4BuildGradleMainPath, "AndroidManifest.xml");
+				MakeDirectoryIfRequired(GradleManifest);
+				File.Copy(Path.Combine(UE4BuildPath, "AndroidManifest.xml"), GradleManifest, true);
 
-					string[] Excludes;
-					switch (NDKArch)
-					{
-						default:
-						case "armeabi-v7a":
-							Excludes = new string[] { "arm64-v8a", "x86", "x86-64" };
-							break;
+				string[] Excludes;
+				switch (NDKArch)
+				{
+					default:
+					case "armeabi-v7a":
+						Excludes = new string[] { "arm64-v8a", "x86", "x86-64" };
+						break;
 
-						case "arm64-v8a":
-							Excludes = new string[] { "armeabi-v7a", "x86", "x86-64" };
-							break;
+					case "arm64-v8a":
+						Excludes = new string[] { "armeabi-v7a", "x86", "x86-64" };
+						break;
 
-						case "x86":
-							Excludes = new string[] { "armeabi-v7a", "arm64-v8a", "x86-64" };
-							break;
+					case "x86":
+						Excludes = new string[] { "armeabi-v7a", "arm64-v8a", "x86-64" };
+						break;
 
-						case "x86_64":
-							Excludes = new string[] { "armeabi-v7a", "arm64-v8a", "x86" };
-							break;
-					}
+					case "x86_64":
+						Excludes = new string[] { "armeabi-v7a", "arm64-v8a", "x86" };
+						break;
+				}
 
-					CleanCopyDirectory(Path.Combine(UE4BuildPath, "jni"), Path.Combine(UE4BuildGradleMainPath, "jniLibs"), Excludes);  // has debug symbols
-					CleanCopyDirectory(Path.Combine(UE4BuildPath, "libs"), Path.Combine(UE4BuildGradleMainPath, "libs"), Excludes);
+				CleanCopyDirectory(Path.Combine(UE4BuildPath, "jni"), Path.Combine(UE4BuildGradleMainPath, "jniLibs"), Excludes);  // has debug symbols
+				CleanCopyDirectory(Path.Combine(UE4BuildPath, "libs"), Path.Combine(UE4BuildGradleMainPath, "libs"), Excludes);
 
-					CleanCopyDirectory(Path.Combine(UE4BuildPath, "assets"), Path.Combine(UE4BuildGradleMainPath, "assets"));
-					CleanCopyDirectory(Path.Combine(UE4BuildPath, "res"), Path.Combine(UE4BuildGradleMainPath, "res"));
-					CleanCopyDirectory(Path.Combine(UE4BuildPath, "src"), Path.Combine(UE4BuildGradleMainPath, "java"));
+				CleanCopyDirectory(Path.Combine(UE4BuildPath, "assets"), Path.Combine(UE4BuildGradleMainPath, "assets"));
+				CleanCopyDirectory(Path.Combine(UE4BuildPath, "res"), Path.Combine(UE4BuildGradleMainPath, "res"));
+				CleanCopyDirectory(Path.Combine(UE4BuildPath, "src"), Path.Combine(UE4BuildGradleMainPath, "java"));
 
-					// do any plugin requested copies
-					UPL.ProcessPluginNode(NDKArch, "gradleCopies", "");
+				// do any plugin requested copies
+				UPL.ProcessPluginNode(NDKArch, "gradleCopies", "");
 
-					// get min and target SDK versions
-					int MinSDKVersion = 0;
-					int TargetSDKVersion = 0;
-					int NDKLevelInt = 0;
-					GetMinTargetSDKVersions(ToolChain, Arch, UPL, NDKArch, out MinSDKVersion, out TargetSDKVersion, out NDKLevelInt);
+				// get min and target SDK versions
+				int MinSDKVersion = 0;
+				int TargetSDKVersion = 0;
+				int NDKLevelInt = 0;
+				GetMinTargetSDKVersions(ToolChain, Arch, UPL, NDKArch, out MinSDKVersion, out TargetSDKVersion, out NDKLevelInt);
 					
-					// move JavaLibs into subprojects
-					string JavaLibsDir = Path.Combine(UE4BuildPath, "JavaLibs");
-					PrepareJavaLibsForGradle(JavaLibsDir, UE4BuildGradlePath, MinSDKVersion.ToString(), TargetSDKVersion.ToString(), CompileSDKVersion, BuildToolsVersion);
+				// move JavaLibs into subprojects
+				string JavaLibsDir = Path.Combine(UE4BuildPath, "JavaLibs");
+				PrepareJavaLibsForGradle(JavaLibsDir, UE4BuildGradlePath, MinSDKVersion.ToString(), TargetSDKVersion.ToString(), CompileSDKVersion, BuildToolsVersion);
 
-					// Create local.properties
-					String LocalPropertiesFilename = Path.Combine(UE4BuildGradlePath, "local.properties");
-					StringBuilder LocalProperties = new StringBuilder();
-					LocalProperties.AppendLine(string.Format("ndk.dir={0}", Environment.GetEnvironmentVariable("NDKROOT").Replace("\\", "/")));
-					LocalProperties.AppendLine(string.Format("sdk.dir={0}", Environment.GetEnvironmentVariable("ANDROID_HOME").Replace("\\", "/")));
-					File.WriteAllText(LocalPropertiesFilename, LocalProperties.ToString());
+				// Create local.properties
+				String LocalPropertiesFilename = Path.Combine(UE4BuildGradlePath, "local.properties");
+				StringBuilder LocalProperties = new StringBuilder();
+				LocalProperties.AppendLine(string.Format("ndk.dir={0}", Environment.GetEnvironmentVariable("NDKROOT").Replace("\\", "/")));
+				LocalProperties.AppendLine(string.Format("sdk.dir={0}", Environment.GetEnvironmentVariable("ANDROID_HOME").Replace("\\", "/")));
+				File.WriteAllText(LocalPropertiesFilename, LocalProperties.ToString());
 
-					CreateGradlePropertiesFiles(Arch, MinSDKVersion, TargetSDKVersion, CompileSDKVersion, BuildToolsVersion, PackageName, DestApkName, NDKArch,
-						UE4BuildFilesPath, GameBuildFilesPath, UE4BuildGradleAppPath, UE4BuildPath, UE4BuildGradlePath, bForDistribution, bSkipGradleBuild);
+				CreateGradlePropertiesFiles(Arch, MinSDKVersion, TargetSDKVersion, CompileSDKVersion, BuildToolsVersion, PackageName, DestApkName, NDKArch,
+					UE4BuildFilesPath, GameBuildFilesPath, UE4BuildGradleAppPath, UE4BuildPath, UE4BuildGradlePath, bForDistribution, bSkipGradleBuild);
 
-					if (!bSkipGradleBuild)
+				if (!bSkipGradleBuild)
+				{
+					string GradleScriptPath = Path.Combine(UE4BuildGradlePath, "gradlew");
+					if (Utils.IsRunningOnMono)
+					{
+						// fix permissions for Mac/Linux
+						RunCommandLineProgramWithException(UE4BuildGradlePath, "/bin/sh", string.Format("-c 'chmod 0755 \"{0}\"'", GradleScriptPath.Replace("'", "'\"'\"'")), "Fix gradlew permissions");
+					}
+					else
 					{
-						string GradleScriptPath = Path.Combine(UE4BuildGradlePath, "gradlew");
-						if (Utils.IsRunningOnMono)
+						if (CreateRunGradle(UE4BuildGradlePath))
 						{
-							// fix permissions for Mac/Linux
-							RunCommandLineProgramWithException(UE4BuildGradlePath, "/bin/sh", string.Format("-c 'chmod 0755 \"{0}\"'", GradleScriptPath.Replace("'", "'\"'\"'")), "Fix gradlew permissions");
+							GradleScriptPath = Path.Combine(UE4BuildGradlePath, "rungradle.bat");
 						}
 						else
 						{
-							if (CreateRunGradle(UE4BuildGradlePath))
-							{
-								GradleScriptPath = Path.Combine(UE4BuildGradlePath, "rungradle.bat");
-							}
-							else
-							{
-								GradleScriptPath = Path.Combine(UE4BuildGradlePath, "gradlew.bat");
-							}
+							GradleScriptPath = Path.Combine(UE4BuildGradlePath, "gradlew.bat");
 						}
+					}
 
-						string GradleBuildType = bForDistribution ? ":app:assembleRelease" : ":app:assembleDebug";
-
-						// collect optional additional Gradle parameters from plugins
-						string GradleOptions = UPL.ProcessPluginNode(NDKArch, "gradleParameters", GradleBuildType); //  "--stacktrace --debug " + GradleBuildType);
-						string GradleSecondCallOptions = UPL.ProcessPluginNode(NDKArch, "gradleSecondCallParameters", "");
+					string GradleBuildType = bForDistribution ? ":app:assembleRelease" : ":app:assembleDebug";
 
-						// check for Android Studio project, call Gradle if doesn't exist (assume user will build with Android Studio)
-						string GradleAppImlFilename = Path.Combine(UE4BuildGradlePath, "app.iml");
-						if (!File.Exists(GradleAppImlFilename))
-						{
-							// make sure destination exists
-							Directory.CreateDirectory(Path.GetDirectoryName(DestApkName));
+					// collect optional additional Gradle parameters from plugins
+					string GradleOptions = UPL.ProcessPluginNode(NDKArch, "gradleParameters", GradleBuildType); //  "--stacktrace --debug " + GradleBuildType);
+					string GradleSecondCallOptions = UPL.ProcessPluginNode(NDKArch, "gradleSecondCallParameters", "");
 
-							// Use gradle to build the .apk file
-							string ShellExecutable = Utils.IsRunningOnMono ? "/bin/sh" : "cmd.exe";
-							string ShellParametersBegin = Utils.IsRunningOnMono ? "-c '" : "/c ";
-							string ShellParametersEnd = Utils.IsRunningOnMono ? "'" : "";
-							RunCommandLineProgramWithExceptionAndFiltering(UE4BuildGradlePath, ShellExecutable, ShellParametersBegin + "\"" + GradleScriptPath + "\" " + GradleOptions + ShellParametersEnd, "Making .apk with Gradle...");
+					// check for Android Studio project, call Gradle if doesn't exist (assume user will build with Android Studio)
+					string GradleAppImlFilename = Path.Combine(UE4BuildGradlePath, "app.iml");
+					if (!File.Exists(GradleAppImlFilename))
+					{
+						// make sure destination exists
+						Directory.CreateDirectory(Path.GetDirectoryName(DestApkName));
 
-							if (GradleSecondCallOptions != "")
-							{
-								RunCommandLineProgramWithExceptionAndFiltering(UE4BuildGradlePath, ShellExecutable, ShellParametersBegin + "\"" + GradleScriptPath + "\" " + GradleSecondCallOptions + ShellParametersEnd, "Additional Gradle steps...");
-							}
-							
-							// For build machine run a clean afterward to clean up intermediate files (does not remove final APK)
-							if (bIsBuildMachine)
-							{
-								//GradleOptions = "tasks --all";
-								//RunCommandLineProgramWithException(UE4BuildGradlePath, ShellExecutable, ShellParametersBegin + "\"" + GradleScriptPath + "\" " + GradleOptions + ShellParametersEnd, "Listing all tasks...");
+						// Use gradle to build the .apk file
+						string ShellExecutable = Utils.IsRunningOnMono ? "/bin/sh" : "cmd.exe";
+						string ShellParametersBegin = Utils.IsRunningOnMono ? "-c '" : "/c ";
+						string ShellParametersEnd = Utils.IsRunningOnMono ? "'" : "";
+						RunCommandLineProgramWithExceptionAndFiltering(UE4BuildGradlePath, ShellExecutable, ShellParametersBegin + "\"" + GradleScriptPath + "\" " + GradleOptions + ShellParametersEnd, "Making .apk with Gradle...");
 
-								GradleOptions = "clean";
-								RunCommandLineProgramWithExceptionAndFiltering(UE4BuildGradlePath, ShellExecutable, ShellParametersBegin + "\"" + GradleScriptPath + "\" " + GradleOptions + ShellParametersEnd, "Cleaning Gradle intermediates...");
-							}
+						if (GradleSecondCallOptions != "")
+						{
+							RunCommandLineProgramWithExceptionAndFiltering(UE4BuildGradlePath, ShellExecutable, ShellParametersBegin + "\"" + GradleScriptPath + "\" " + GradleSecondCallOptions + ShellParametersEnd, "Additional Gradle steps...");
 						}
-						else
+							
+						// For build machine run a clean afterward to clean up intermediate files (does not remove final APK)
+						if (bIsBuildMachine)
 						{
-							Log.TraceInformation("=============================================================================================");
-							Log.TraceInformation("Android Studio project found, skipping Gradle; complete creation of APK in Android Studio!!!!");
-							Log.TraceInformation("Delete '{0} if you want to have UnrealBuildTool run Gradle for future runs.", GradleAppImlFilename);
-							Log.TraceInformation("=============================================================================================");
+							//GradleOptions = "tasks --all";
+							//RunCommandLineProgramWithException(UE4BuildGradlePath, ShellExecutable, ShellParametersBegin + "\"" + GradleScriptPath + "\" " + GradleOptions + ShellParametersEnd, "Listing all tasks...");
+
+							GradleOptions = "clean";
+							RunCommandLineProgramWithExceptionAndFiltering(UE4BuildGradlePath, ShellExecutable, ShellParametersBegin + "\"" + GradleScriptPath + "\" " + GradleOptions + ShellParametersEnd, "Cleaning Gradle intermediates...");
 						}
 					}
-
-					bool bBuildWithHiddenSymbolVisibility = false;
-					bool bSaveSymbols = false;
-					Ini.GetBool("/Script/AndroidRuntimeSettings.AndroidRuntimeSettings", "bBuildWithHiddenSymbolVisibility", out bBuildWithHiddenSymbolVisibility);
-					Ini.GetBool("/Script/AndroidRuntimeSettings.AndroidRuntimeSettings", "bSaveSymbols", out bSaveSymbols);
-					bSaveSymbols = true;
-					if (bSaveSymbols || (Configuration == UnrealTargetConfiguration.Shipping && bBuildWithHiddenSymbolVisibility))
+					else
 					{
-						// Copy .so with symbols to 
-						int StoreVersion = GetStoreVersion(Arch);
-						string SymbolSODirectory = Path.Combine(DestApkDirectory, ProjectName + "_Symbols_v" + StoreVersion + "/" + ProjectName + Arch + GPUArchitecture);
-						string SymbolifiedSOPath = Path.Combine(SymbolSODirectory, Path.GetFileName(FinalSOName));
-						MakeDirectoryIfRequired(SymbolifiedSOPath);
-						Log.TraceInformation("Writing symbols to {0}", SymbolifiedSOPath);
-
-						File.Copy(FinalSOName, SymbolifiedSOPath, true);
+						Log.TraceInformation("=============================================================================================");
+						Log.TraceInformation("Android Studio project found, skipping Gradle; complete creation of APK in Android Studio!!!!");
+						Log.TraceInformation("Delete '{0} if you want to have UnrealBuildTool run Gradle for future runs.", GradleAppImlFilename);
+						Log.TraceInformation("=============================================================================================");
 					}
 				}
-			}
-		}
-
-		private void PrepareToSignApk(string BuildPath)
-		{
-			// ini file to get settings from
-			ConfigHierarchy Ini = GetConfigCacheIni(ConfigHierarchyType.Engine);
-			string KeyAlias, KeyStore, KeyStorePassword, KeyPassword;
-			Ini.GetString("/Script/AndroidRuntimeSettings.AndroidRuntimeSettings", "KeyAlias", out KeyAlias);
-			Ini.GetString("/Script/AndroidRuntimeSettings.AndroidRuntimeSettings", "KeyStore", out KeyStore);
-			Ini.GetString("/Script/AndroidRuntimeSettings.AndroidRuntimeSettings", "KeyStorePassword", out KeyStorePassword);
-			Ini.GetString("/Script/AndroidRuntimeSettings.AndroidRuntimeSettings", "KeyPassword", out KeyPassword);
 
-			if (string.IsNullOrEmpty(KeyAlias) || string.IsNullOrEmpty(KeyStore) || string.IsNullOrEmpty(KeyStorePassword))
-			{
-				throw new BuildException("DistributionSigning settings are not all set. Check the DistributionSettings section in the Android tab of Project Settings");
-			}
+				bool bBuildWithHiddenSymbolVisibility = false;
+				bool bSaveSymbols = false;
 
-			string[] AntPropertiesLines = new string[4];
-			AntPropertiesLines[0] = "key.store=" + KeyStore;
-			AntPropertiesLines[1] = "key.alias=" + KeyAlias;
-			AntPropertiesLines[2] = "key.store.password=" + KeyStorePassword;
-			AntPropertiesLines[3] = "key.alias.password=" + ((string.IsNullOrEmpty(KeyPassword) || KeyPassword == "_sameaskeystore_") ? KeyStorePassword : KeyPassword);
+				ConfigHierarchy Ini = GetConfigCacheIni(ConfigHierarchyType.Engine);
+				Ini.GetBool("/Script/AndroidRuntimeSettings.AndroidRuntimeSettings", "bBuildWithHiddenSymbolVisibility", out bBuildWithHiddenSymbolVisibility);
+				Ini.GetBool("/Script/AndroidRuntimeSettings.AndroidRuntimeSettings", "bSaveSymbols", out bSaveSymbols);
+				bSaveSymbols = true;
+				if (bSaveSymbols || (Configuration == UnrealTargetConfiguration.Shipping && bBuildWithHiddenSymbolVisibility))
+				{
+					// Copy .so with symbols to 
+					int StoreVersion = GetStoreVersion(Arch);
+					string SymbolSODirectory = Path.Combine(DestApkDirectory, ProjectName + "_Symbols_v" + StoreVersion + "/" + ProjectName + Arch + GPUArchitecture);
+					string SymbolifiedSOPath = Path.Combine(SymbolSODirectory, Path.GetFileName(FinalSOName));
+					MakeDirectoryIfRequired(SymbolifiedSOPath);
+					Log.TraceInformation("Writing symbols to {0}", SymbolifiedSOPath);
 
-			// now write out the properties
-			File.WriteAllLines(Path.Combine(BuildPath, "ant.properties"), AntPropertiesLines);
+					File.Copy(FinalSOName, SymbolifiedSOPath, true);
+				}
+			}
 		}
 
 		private List<string> CollectPluginDataPaths(TargetReceipt Receipt)
@@ -4331,46 +3855,6 @@ namespace UnrealBuildTool
 			}
 		}
 
-		private void UpdateBuildXML(string UE4Arch, string NDKArch, string EngineDir, string UE4BuildPath)
-		{
-			string SourceFilename = Path.Combine(UE4BuildPath, "build.xml");
-			string DestFilename = SourceFilename;
-
-			Dictionary<string, string> Replacements = new Dictionary<string, string>{
-			  { "<import file=\"${sdk.dir}/tools/ant/build.xml\" />", UPL.ProcessPluginNode(NDKArch, "buildXmlPropertyAdditions", "")}
-			};
-
-			string[] TemplateSrc = File.ReadAllLines(SourceFilename);
-			string[] TemplateDest = File.Exists(DestFilename) ? File.ReadAllLines(DestFilename) : null;
-
-			for (int LineIndex = 0; LineIndex < TemplateSrc.Length; ++LineIndex)
-			{
-				string SrcLine = TemplateSrc[LineIndex];
-				bool Changed = false;
-				foreach (KeyValuePair<string, string> KVP in Replacements)
-				{
-					if (SrcLine.Contains(KVP.Key))
-					{
-						// insert replacement before the <import>
-						SrcLine = SrcLine.Replace(KVP.Key, KVP.Value);
-						// then add the <import>
-						SrcLine += KVP.Key;
-						Changed = true;
-					}
-				}
-				if (Changed)
-				{
-					TemplateSrc[LineIndex] = SrcLine;
-				}
-			}
-
-			if (TemplateDest == null || TemplateSrc.Length != TemplateDest.Length || !TemplateSrc.SequenceEqual(TemplateDest))
-			{
-				Log.TraceInformation("\n==== Writing new build.xml file to {0} ====", DestFilename);
-				File.WriteAllLines(DestFilename, TemplateSrc);
-			}
-		}
-
 		private string GenerateTemplatesHashCode(string EngineDir)
 		{
 			string SourceDirectory = Path.Combine(EngineDir, "Build", "Android", "Java");
@@ -4522,18 +4006,11 @@ namespace UnrealBuildTool
 			}
 		}
 
-		private void UpdateGameApplication(string UE4Arch, string NDKArch, string EngineDir, string UE4BuildPath, bool bGradleEnabled)
+		private void UpdateGameApplication(string UE4Arch, string NDKArch, string EngineDir, string UE4BuildPath)
 		{
 			string SourceFilename = Path.Combine(EngineDir, "Build", "Android", "Java", "src", "com", "epicgames", "ue4", "GameApplication.java.template");
 			string DestFilename = Path.Combine(UE4BuildPath, "src", "com", "epicgames", "ue4", "GameApplication.java");
 
-			if (!bGradleEnabled)
-			{
-				// do not use GameApplication for Ant
-				SafeDeleteFile(DestFilename);
-				return;
-			}
-			
 			Dictionary<string, string> Replacements = new Dictionary<string, string>{
 				{ "//$${gameApplicationImportAdditions}$$", UPL.ProcessPluginNode(NDKArch, "gameApplicationImportAdditions", "")},
 				{ "//$${gameApplicationOnCreateAdditions}$$", UPL.ProcessPluginNode(NDKArch, "gameApplicationOnCreateAdditions", "")},
@@ -4613,7 +4090,7 @@ namespace UnrealBuildTool
 			}
 		}
 
-		private AndroidAARHandler CreateAARHandler(string EngineDir, string UE4BuildPath, List<string> NDKArches, bool bGradleEnabled, bool HandleDependencies=true)
+		private AndroidAARHandler CreateAARHandler(string EngineDir, string UE4BuildPath, List<string> NDKArches, bool HandleDependencies=true)
 		{
 			AndroidAARHandler AARHandler = new AndroidAARHandler();
 			string ImportList = "";
@@ -4624,18 +4101,10 @@ namespace UnrealBuildTool
 
 			// Add the AARs from the default aar-imports.txt
 			// format: Package,Name,Version
-			string ImportsAntFile = Path.Combine(UE4BuildPath, "aar-imports-ant.txt");
-			if (!bGradleEnabled && File.Exists(ImportsAntFile))
-			{
-				ImportList = File.ReadAllText(ImportsAntFile);
-			}
-			else
+			string ImportsFile = Path.Combine(UE4BuildPath, "aar-imports.txt");
+			if (File.Exists(ImportsFile))
 			{
-				string ImportsFile = Path.Combine(UE4BuildPath, "aar-imports.txt");
-				if (File.Exists(ImportsFile))
-				{
-					ImportList = File.ReadAllText(ImportsFile);
-				}
+				ImportList = File.ReadAllText(ImportsFile);
 			}
 
 			// Run the UPL imports section for each architecture and add any new imports (duplicates will be removed)
@@ -4847,7 +4316,7 @@ namespace UnrealBuildTool
 
 		private void GenerateGradleAARImports(string EngineDir, string UE4BuildPath, List<string> NDKArches)
 		{
-			AndroidAARHandler AARHandler = CreateAARHandler(EngineDir, UE4BuildPath, NDKArches, true, false);
+			AndroidAARHandler AARHandler = CreateAARHandler(EngineDir, UE4BuildPath, NDKArches, false);
 			StringBuilder AARImportsContent = new StringBuilder();
 
 			// Add repositories
@@ -4870,22 +4339,5 @@ namespace UnrealBuildTool
 			string AARImportsFilename = Path.Combine(UE4BuildPath, "gradle", "app", "aar-imports.gradle");
 			File.WriteAllText(AARImportsFilename, AARImportsContent.ToString());
 		}
-
-		private void ExtractAARAndJARFiles(string EngineDir, string UE4BuildPath, List<string> NDKArches, string AppPackageName, string AARExtractListFilename)
-		{
-			AndroidAARHandler AARHandler = CreateAARHandler(EngineDir, UE4BuildPath, NDKArches, false, true);
-
-			// Finally, extract the AARs and copy the JARs
-			AARHandler.ExtractAARs(UE4BuildPath, AppPackageName);
-			AARHandler.CopyJARs(UE4BuildPath);
-
-			// Write list of AAR files extracted
-			StringBuilder AARListContents = new StringBuilder();
-			foreach (AndroidAARHandler.AndroidAAREntry Dependency in AARHandler.AARList)
-			{
-				AARListContents.AppendLine(Dependency.Filename);
-			}
-			File.WriteAllText(AARExtractListFilename, AARListContents.ToString());
-		}
 	}
 }
diff --git a/Engine/Source/Runtime/Advertising/Android/AndroidAdvertising/AndroidAdvertising_APL.xml b/Engine/Source/Runtime/Advertising/Android/AndroidAdvertising/AndroidAdvertising_APL.xml
--- a/Engine/Source/Runtime/Advertising/Android/AndroidAdvertising/AndroidAdvertising_APL.xml
+++ b/Engine/Source/Runtime/Advertising/Android/AndroidAdvertising/AndroidAdvertising_APL.xml
@@ -17,31 +17,18 @@
 				<setBool result="bRequestGAID" value="false"/>
 			</false>
 		</if>
-
-		<!-- check if using Gradle -->
-		<setBoolFromProperty result="bGradle" ini="Engine" section="/Script/AndroidRuntimeSettings.AndroidRuntimeSettings" property="bEnableGradle" default="false"/>
 	</init>
 
-	<!-- AAR dependencies -->
-	<AARImports>
+	<buildGradleAdditions>
 		<if condition="bSupportAdMob">
 			<true>
-				<if condition="bGradle">
-					<false>
-						<insertValue value="com.google.android.gms,play-services-ads,11.0.4"/>
-						<insertNewline/>
-					</false>
-				</if>
-			</true>
-		</if>
-	</AARImports>
-
-	<buildGradleAdditions>
-		<insert>
+				<insert>
 dependencies {
 	implementation('com.google.android.gms:play-services-ads:11.8.0')
 }
-		</insert>
+				</insert>
+			</true>
+		</if>
 	</buildGradleAdditions>
 
 	<!-- optional updates applied to AndroidManifest.xml -->
diff --git a/Engine/Source/Runtime/Android/AndroidRuntimeSettings/Classes/AndroidRuntimeSettings.h b/Engine/Source/Runtime/Android/AndroidRuntimeSettings/Classes/AndroidRuntimeSettings.h
--- a/Engine/Source/Runtime/Android/AndroidRuntimeSettings/Classes/AndroidRuntimeSettings.h
+++ b/Engine/Source/Runtime/Android/AndroidRuntimeSettings/Classes/AndroidRuntimeSettings.h
@@ -12,22 +12,6 @@
 
 DECLARE_LOG_CATEGORY_EXTERN(LogAndroidRuntimeSettings, Log, All);
 
-UENUM()
-namespace EAndroidAntVerbosity
-{
-	enum Type
-	{
-		/** Extra quiet logging (-quiet), errors will be logged by second run at normal verbosity. */
-		Quiet,
-
-		/** Normal logging (no options) */
-		Normal,
-
-		/** Extra verbose logging (-verbose) */
-		Verbose,
-	};
-}
-
 UENUM()
 namespace EAndroidScreenOrientation
 {
@@ -287,10 +271,6 @@ public:
 	UPROPERTY(GlobalConfig, EditAnywhere, Category = "APK Packaging", Meta = (DisplayName = "Restore scheduled notifications on reboot"))
 	bool bRestoreNotificationsOnReboot;
 
-	// Level of verbosity to use during packaging with Ant
-	UPROPERTY(GlobalConfig, EditAnywhere, Category = "APK Packaging")
-	TEnumAsByte<EAndroidAntVerbosity::Type> AntVerbosity;
-
 	// Should the software navigation buttons be hidden or not
 	UPROPERTY(GlobalConfig, EditAnywhere, Category = "APK Packaging", Meta = (DisplayName = "Enable FullScreen Immersive on KitKat and above devices."))
 	bool bFullScreen;
