Fix UBT including developer modules from platform extensions in the wrong builds. Modules in the Developer folder must be explicitly whitelisted, but since the IsValidForTarget() function was not filtering out platform extension paths, all Developer modules from platform extensions were being included regardless, breaking several platform builds.

diff --git a/Engine/Source/Programs/UnrealBuildTool/Configuration/ModuleRules.cs b/Engine/Source/Programs/UnrealBuildTool/Configuration/ModuleRules.cs
--- a/Engine/Source/Programs/UnrealBuildTool/Configuration/ModuleRules.cs
+++ b/Engine/Source/Programs/UnrealBuildTool/Configuration/ModuleRules.cs
@@ -1176,7 +1176,7 @@ namespace UnrealBuildTool
 					case ModuleRules.PrecompileTargetsType.None:
 						return false;
 					case ModuleRules.PrecompileTargetsType.Default:
-						return (Target.Type == TargetType.Editor || !RulesFile.IsUnderDirectory(UnrealBuildTool.EngineSourceDeveloperDirectory) || Plugin != null);
+						return (Target.Type == TargetType.Editor || !UnrealBuildTool.GetAllEngineDirectories("Source/Developer").Any(Dir => RulesFile.IsUnderDirectory(Dir)) || Plugin != null);
 					case ModuleRules.PrecompileTargetsType.Game:
 						return (Target.Type == TargetType.Client || Target.Type == TargetType.Server || Target.Type == TargetType.Game);
 					case ModuleRules.PrecompileTargetsType.Editor:
