Fix bug in UBT config reader. Automation was looking for "Base" config files in game projects rather than regular ones.

diff --git a/Engine/Source/Programs/UnrealBuildTool/System/ConfigHierarchy.cs b/Engine/Source/Programs/UnrealBuildTool/System/ConfigHierarchy.cs
--- a/Engine/Source/Programs/UnrealBuildTool/System/ConfigHierarchy.cs
+++ b/Engine/Source/Programs/UnrealBuildTool/System/ConfigHierarchy.cs
@@ -769,7 +769,7 @@ namespace UnrealBuildTool
 			// Engine/Platform/Platform*.ini
 			new ConfigLayer { Path = "{ENGINE}/Config/{ED}{PLATFORM}/{EF}{PLATFORM}{TYPE}.ini", ExtEnginePath = "{EXTENGINE}/Config/{ED}{EF}{PLATFORM}{TYPE}.ini" },
 			// Project/Platform/Platform*.ini
-			new ConfigLayer { Path = "{PROJECT}/Config/{ED}{PLATFORM}/{EF}{PLATFORM}{TYPE}.ini", ExtProjectPath = "{EXTPROJECT}/Config/{ED}{EF}Base{PLATFORM}{TYPE}.ini" },
+			new ConfigLayer { Path = "{PROJECT}/Config/{ED}{PLATFORM}/{EF}{PLATFORM}{TYPE}.ini", ExtProjectPath = "{EXTPROJECT}/Config/{ED}{EF}{PLATFORM}{TYPE}.ini" },
 
 			// UserSettings/.../User*.ini
 			new ConfigLayer { Path = "{USERSETTINGS}/Unreal Engine/Engine/Config/User{TYPE}.ini" },
