DRPS - Bump stack size as debug editor builds were out of mem
#rnx

diff --git a/Engine/Source/Programs/UnrealBuildTool/Platform/Windows/UEBuildWindows.cs b/Engine/Source/Programs/UnrealBuildTool/Platform/Windows/UEBuildWindows.cs
--- a/Engine/Source/Programs/UnrealBuildTool/Platform/Windows/UEBuildWindows.cs
+++ b/Engine/Source/Programs/UnrealBuildTool/Platform/Windows/UEBuildWindows.cs
@@ -208,14 +208,14 @@ namespace UnrealBuildTool
 		}
 
 		/// <summary>
-		/// The stack size when linking a non-editor target
+		/// The stack size when linking
 		/// </summary>
 		[RequiresUniqueBuildEnvironment]
 		[ConfigFile(ConfigHierarchyType.Engine, "/Script/WindowsTargetPlatform.WindowsTargetSettings")]
-		public int DefaultStackSize = 5000000;
+		public int DefaultStackSize = 12000000;
 
 		/// <summary>
-		/// The stack size to commit when linking a non-editor target
+		/// The stack size to commit when linking
 		/// </summary>
 		[RequiresUniqueBuildEnvironment]
 		[ConfigFile(ConfigHierarchyType.Engine, "/Script/WindowsTargetPlatform.WindowsTargetSettings")]
