Hololens  ERROR: The element 'InProcessServer'...
-AppxManifest_arm64.xml was not matching its schema because the SceneUnderstanding entry did not have any activable class defined.  This is because a change to the code that finds that attribute was comparing the AssemblyQualifiedName instead of the FullName to detect StaticAttribute.

diff --git a/Engine/Source/Programs/UnrealBuildTool/Platform/HoloLens/WinMDRegistrationInfo.cs b/Engine/Source/Programs/UnrealBuildTool/Platform/HoloLens/WinMDRegistrationInfo.cs
--- a/Engine/Source/Programs/UnrealBuildTool/Platform/HoloLens/WinMDRegistrationInfo.cs
+++ b/Engine/Source/Programs/UnrealBuildTool/Platform/HoloLens/WinMDRegistrationInfo.cs
@@ -62,7 +62,7 @@ namespace UnrealBuildTool
 					string ThreadingModelName = "both";
 					foreach (CustomAttributeData Attr in WinMDType.CustomAttributes)
 					{
-						if (Attr.AttributeType.FullName == "Windows.Foundation.Metadata.ActivatableAttribute" || Attr.AttributeType.AssemblyQualifiedName == "Windows.Foundation.Metadata.StaticAttribute")
+						if (Attr.AttributeType.FullName == "Windows.Foundation.Metadata.ActivatableAttribute" || Attr.AttributeType.FullName == "Windows.Foundation.Metadata.StaticAttribute")
 						{
 							IsActivatable = true;
 						}
