Add COMPILED_PLATFORM_HEADER_WITH_PREFIX to SN-DBS computed include rules
#rnx

diff --git a/Engine/Source/Programs/UnrealBuildTool/Executors/SNDBS.cs b/Engine/Source/Programs/UnrealBuildTool/Executors/SNDBS.cs
--- a/Engine/Source/Programs/UnrealBuildTool/Executors/SNDBS.cs
+++ b/Engine/Source/Programs/UnrealBuildTool/Executors/SNDBS.cs
@@ -506,15 +506,6 @@ namespace UnrealBuildTool
 
 		private void GenerateSNDBSIncludeRewriteRules()
 		{
-			if (!Directory.Exists(IncludeRewriteRulesFile.Directory.FullName))
-			{
-				Directory.CreateDirectory(IncludeRewriteRulesFile.Directory.FullName);
-			}
-
-			List<string> IncludeRewriteRulesText = new List<string>();
-			IncludeRewriteRulesText.Add("[computed-include-rules]");
-			IncludeRewriteRulesText.Add(@"pattern1=^COMPILED_PLATFORM_HEADER\(\s*([^ ,]+)\)");
-
 			// Get all registered platforms. Most will just use the name as is, but some may have an override so
 			// add all distinct entries to the list.
 			IEnumerable<UnrealTargetPlatform> Platforms = UEBuildPlatform.GetRegisteredPlatforms();
@@ -528,8 +519,23 @@ namespace UnrealBuildTool
 				}
 			}
 
-			IEnumerable<string> PlatformExpansions = PlatformNames.Select(p => String.Format("{0}/{0}$1|{0}$1", p));
-			IncludeRewriteRulesText.Add(String.Format("expansions1={0}", String.Join("|", PlatformExpansions)));
+			if (!Directory.Exists(IncludeRewriteRulesFile.Directory.FullName))
+			{
+				Directory.CreateDirectory(IncludeRewriteRulesFile.Directory.FullName);
+			}
+
+			List<string> IncludeRewriteRulesText = new List<string>();
+			IncludeRewriteRulesText.Add("[computed-include-rules]");
+			{
+				IncludeRewriteRulesText.Add(@"pattern1=^COMPILED_PLATFORM_HEADER\(\s*([^ ,]+)\)");
+				IEnumerable<string> PlatformExpansions = PlatformNames.Select(p => String.Format("{0}/{0}$1|{0}$1", p));
+				IncludeRewriteRulesText.Add(String.Format("expansions1={0}", String.Join("|", PlatformExpansions)));
+			}
+			{
+				IncludeRewriteRulesText.Add(@"pattern2=^COMPILED_PLATFORM_HEADER_WITH_PREFIX\(\s*([^ ,]+)\s*,\s*([^ ,]+)\)");
+				IEnumerable<string> PlatformExpansions = PlatformNames.Select(p => String.Format("$1/{0}/{0}$2|$1/{0}$2", p));
+				IncludeRewriteRulesText.Add(String.Format("expansions2={0}", String.Join("|", PlatformExpansions)));
+			}
 			File.WriteAllText(IncludeRewriteRulesFile.FullName, String.Join(Environment.NewLine, IncludeRewriteRulesText));
 		}
 	}
