Copying //UE4/Dev-Console@10648120 to Dev-RenderPlat-Staging (//UE4/Dev-RenderPlat-Staging)
#rnx

diff --git a/Engine/Build/SNDBSTemplates/clang++.exe.sn-dbs-tool.ini b/Engine/Build/SNDBSTemplates/clang++.exe.sn-dbs-tool.ini
new file mode 100644
--- /dev/null
+++ b/Engine/Build/SNDBSTemplates/clang++.exe.sn-dbs-tool.ini
@@ -0,0 +1,28 @@
+[tool]
+family=clang-cl
+include_path01=..\include
+include_path02=..\include\c++\v1
+include_path03=..\lib\clang\*\include
+
+[files]
+main=clang++.exe
+file01=clang-shared.dll
+file02=libclang.dll
+file03=NXMangledNamePrinter.dll
+file04=..\lib\*
+
+[output-file-patterns]
+outputfile01=\s*"([^ ",]+\.cpp\.txt)\"
+
+[output-file-rules]
+rule01=*.log|discard=true
+rule02=*.dat|discard=true
+rule03=*.tmp|discard=true
+
+[system-file-filters]
+filter01=msvcr*.dll
+filter02=msvcp*.dll
+filter03=vcruntime140*.dll
+filter04=appcrt140*.dll
+filter05=desktopcrt140*.dll
+filter06=concrt140*.dll
\ No newline at end of file
diff --git a/Engine/Plugins/Media/PixelStreaming/Source/PixelStreaming/Private/InputDevice.cpp b/Engine/Plugins/Media/PixelStreaming/Source/PixelStreaming/Private/InputDevice.cpp
--- a/Engine/Plugins/Media/PixelStreaming/Source/PixelStreaming/Private/InputDevice.cpp
+++ b/Engine/Plugins/Media/PixelStreaming/Source/PixelStreaming/Private/InputDevice.cpp
@@ -36,8 +36,6 @@ class FCursor : public ICursor
 public:
 	FCursor() {}
 	virtual ~FCursor() = default;
-	virtual void* CreateCursorFromFile(const FString& InPathToCursorWithoutExtension, FVector2D HotSpot) override { return nullptr; }
-	virtual void* CreateCursorFromRGBABuffer(const FColor* Pixels, int32 Width, int32 Height, FVector2D InHotSpot) override { return nullptr; }
 	virtual FVector2D GetPosition() const override { return Position; }
 	virtual void SetPosition(const int32 X, const int32 Y) override { Position = FVector2D(X, Y); };
 	virtual void SetType(const EMouseCursor::Type InNewCursor) override {};
diff --git a/Engine/Source/Developer/AutomationDriver/Private/AutomatedApplication.cpp b/Engine/Source/Developer/AutomationDriver/Private/AutomatedApplication.cpp
--- a/Engine/Source/Developer/AutomationDriver/Private/AutomatedApplication.cpp
+++ b/Engine/Source/Developer/AutomationDriver/Private/AutomatedApplication.cpp
@@ -13,16 +13,6 @@ public:
 	virtual ~FAutomatedCursor()
 	{ }
 
-	virtual void* CreateCursorFromFile(const FString& InPathToCursorWithoutExtension, FVector2D HotSpot) override
-	{
-		return nullptr;
-	}
-
-	virtual void* CreateCursorFromRGBABuffer(const FColor* Pixels, int32 Width, int32 Height, FVector2D InHotSpot) override
-	{
-		return nullptr;
-	}
-
 	virtual FVector2D GetPosition() const override
 	{
 		if (bAllowMessageHandling)
diff --git a/Engine/Source/Programs/AutomationTool/Scripts/CopyBuildToStagingDirectory.Automation.cs b/Engine/Source/Programs/AutomationTool/Scripts/CopyBuildToStagingDirectory.Automation.cs
--- a/Engine/Source/Programs/AutomationTool/Scripts/CopyBuildToStagingDirectory.Automation.cs
+++ b/Engine/Source/Programs/AutomationTool/Scripts/CopyBuildToStagingDirectory.Automation.cs
@@ -3481,6 +3481,8 @@ public partial class Project : CommandUtils
 				// clean the staging directories first
 				foreach (var SC in DeployContextList)
 				{
+					SC.StageTargetPlatform.PreStage(Params, SC);
+
 					// write out the commandline file now so it can go into the manifest
 					WriteStageCommandline(Params, SC);
 					CreateStagingManifest(Params, SC);
diff --git a/Engine/Source/Programs/UnrealBuildTool/Executors/SNDBS.cs b/Engine/Source/Programs/UnrealBuildTool/Executors/SNDBS.cs
--- a/Engine/Source/Programs/UnrealBuildTool/Executors/SNDBS.cs
+++ b/Engine/Source/Programs/UnrealBuildTool/Executors/SNDBS.cs
@@ -237,13 +237,15 @@ namespace UnrealBuildTool
 						{
 							// Create a dummy force-included file which references PCH files, so that SN-DBS knows they are dependencies.
 							string AdditionalStubIncludes = "";
-							if (Action.CommandPath.GetFileName().Equals("cl.exe", StringComparison.OrdinalIgnoreCase))
+							if (Action.CommandPath.GetFileName().Equals("cl.exe", StringComparison.OrdinalIgnoreCase) || Action.CommandPath.GetFileName().Equals("cl-filter.exe", StringComparison.OrdinalIgnoreCase))
 							{
-								string ResponseFile = Action.CommandArguments.Replace("\"", "").Replace("@", "").Trim();
+								string DummyPCHIncludeFile = Action.DependencyListFile.AbsolutePath.Replace("\"", "").Replace("@", "").Trim();
+								DummyPCHIncludeFile = Path.ChangeExtension(DummyPCHIncludeFile, null);
+
 								StringBuilder WrapperContents = new StringBuilder();
 								using (StringWriter Writer = new StringWriter(WrapperContents))
 								{
-									Writer.WriteLine("// PCH dependencies for {0}", ResponseFile);
+									Writer.WriteLine("// PCH dependencies for {0}", DummyPCHIncludeFile);
 									Writer.WriteLine("#if 0");
 									foreach (FileItem Preqrequisite in Action.PrerequisiteItems)
 									{
@@ -255,8 +257,8 @@ namespace UnrealBuildTool
 									Writer.WriteLine("#endif");
 								}
 
-								FileItem DummyResponseFileDependency = FileItem.CreateIntermediateTextFile(new FileReference(ResponseFile + ".dummy.h"), WrapperContents.ToString());
-								AdditionalStubIncludes = string.Format("/FI\"{0}\"", DummyResponseFileDependency);
+								FileItem DummyPCHIncludeFileDependency = FileItem.CreateIntermediateTextFile(new FileReference(DummyPCHIncludeFile + ".dummy.h"), WrapperContents.ToString());
+								AdditionalStubIncludes = string.Format("/FI\"{0}\"", DummyPCHIncludeFileDependency);
 							}
 
 							// Add to script for execution by SN-DBS
@@ -536,6 +538,11 @@ namespace UnrealBuildTool
 				IEnumerable<string> PlatformExpansions = PlatformNames.Select(p => String.Format("$1/{0}/{0}$2|$1/{0}$2", p));
 				IncludeRewriteRulesText.Add(String.Format("expansions2={0}", String.Join("|", PlatformExpansions)));
 			}
+			{
+				IncludeRewriteRulesText.Add(@"pattern3=ULANG_STRINGIFY\(\s*(\S*)ULANG_PLATFORM/ULANG_PLATFORM(\S*)\s*\)");
+				IEnumerable<string> PlatformExpansions = PlatformNames.Select(p => String.Format("$1{0}/{0}$2", p));
+				IncludeRewriteRulesText.Add(String.Format("expansions3={0}", String.Join("|", PlatformExpansions)));
+			}
 			File.WriteAllText(IncludeRewriteRulesFile.FullName, String.Join(Environment.NewLine, IncludeRewriteRulesText));
 		}
 	}
diff --git a/Engine/Source/Programs/UnrealBuildTool/Platform/Windows/VCToolChain.cs b/Engine/Source/Programs/UnrealBuildTool/Platform/Windows/VCToolChain.cs
--- a/Engine/Source/Programs/UnrealBuildTool/Platform/Windows/VCToolChain.cs
+++ b/Engine/Source/Programs/UnrealBuildTool/Platform/Windows/VCToolChain.cs
@@ -1292,7 +1292,8 @@ namespace UnrealBuildTool
 
 				// Resource tool can run remotely if possible
 				CompileAction.bCanExecuteRemotely = true;
-
+				CompileAction.bCanExecuteRemotelyWithSNDBS = false;	// no tool template for SN-DBS results in warnings
+			
 				List<string> Arguments = new List<string>();
 
 				// Suppress header spew
@@ -1333,7 +1334,7 @@ namespace UnrealBuildTool
 					}
 				}
 
-				// Figure the icon to use. We can only use a custom icon when compiling to a project-specific intemediate directory (and not for the shared editor executable, for example).
+				// Figure the icon to use. We can only use a custom icon when compiling to a project-specific intermediate directory (and not for the shared editor executable, for example).
 				FileReference IconFile;
 				if(Target.ProjectFile != null && !CompileEnvironment.bUseSharedBuildEnvironment)
 				{
diff --git a/Engine/Source/Programs/UnrealBuildTool/System/Action.cs b/Engine/Source/Programs/UnrealBuildTool/System/Action.cs
--- a/Engine/Source/Programs/UnrealBuildTool/System/Action.cs
+++ b/Engine/Source/Programs/UnrealBuildTool/System/Action.cs
@@ -168,6 +168,12 @@ namespace UnrealBuildTool
 		public Action(ActionType InActionType)
 		{
 			ActionType = InActionType;
+
+			// link actions are going to run locally on SN-DBS so don't try to distribute them as that generates warnings for missing tool templates
+			if ( ActionType == ActionType.Link )
+			{
+				bCanExecuteRemotelyWithSNDBS = false;
+			}
 		}
 
 		public Action(BinaryArchiveReader Reader)
diff --git a/Engine/Source/Runtime/ApplicationCore/Private/Linux/LinuxCursor.cpp b/Engine/Source/Runtime/ApplicationCore/Private/Linux/LinuxCursor.cpp
--- a/Engine/Source/Runtime/ApplicationCore/Private/Linux/LinuxCursor.cpp
+++ b/Engine/Source/Runtime/ApplicationCore/Private/Linux/LinuxCursor.cpp
@@ -172,11 +172,6 @@ FLinuxCursor::~FLinuxCursor()
 	}
 }
 
-void* FLinuxCursor::CreateCursorFromFile(const FString& InPathToCursorWithoutExtension, FVector2D HotSpot)
-{
-	return nullptr;
-}
-
 void* FLinuxCursor::CreateCursorFromRGBABuffer(const FColor* Pixels, int32 Width, int32 Height, FVector2D InHotSpot)
 {
 	uint32 Rmask = 0x000000ff;
diff --git a/Engine/Source/Runtime/ApplicationCore/Public/Android/AndroidCursor.h b/Engine/Source/Runtime/ApplicationCore/Public/Android/AndroidCursor.h
--- a/Engine/Source/Runtime/ApplicationCore/Public/Android/AndroidCursor.h
+++ b/Engine/Source/Runtime/ApplicationCore/Public/Android/AndroidCursor.h
@@ -15,16 +15,6 @@ public:
 	{
 	}
 
-	virtual void* CreateCursorFromFile(const FString& InPathToCursorWithoutExtension, FVector2D HotSpot) override
-	{
-		return nullptr;
-	}
-
-	virtual void* CreateCursorFromRGBABuffer(const FColor* Pixels, int32 Width, int32 Height, FVector2D InHotSpot) override
-	{
-		return nullptr;
-	}
-
 	virtual FVector2D GetPosition() const override
 	{
 		return CurrentPosition;
diff --git a/Engine/Source/Runtime/ApplicationCore/Public/GenericPlatform/GenericPlatformSoftwareCursor.h b/Engine/Source/Runtime/ApplicationCore/Public/GenericPlatform/GenericPlatformSoftwareCursor.h
--- a/Engine/Source/Runtime/ApplicationCore/Public/GenericPlatform/GenericPlatformSoftwareCursor.h
+++ b/Engine/Source/Runtime/ApplicationCore/Public/GenericPlatform/GenericPlatformSoftwareCursor.h
@@ -18,16 +18,6 @@ public:
 	{
 	}
 
-	virtual void* CreateCursorFromFile(const FString& InPathToCursorWithoutExtension, FVector2D HotSpot) override
-	{
-		return nullptr;
-	}
-
-	virtual void* CreateCursorFromRGBABuffer(const FColor* Pixels, int32 Width, int32 Height, FVector2D InHotSpot) override
-	{
-		return nullptr;
-	}
-
 	virtual FVector2D GetPosition() const override
 	{
 		return CurrentPosition;
diff --git a/Engine/Source/Runtime/ApplicationCore/Public/GenericPlatform/ICursor.h b/Engine/Source/Runtime/ApplicationCore/Public/GenericPlatform/ICursor.h
--- a/Engine/Source/Runtime/ApplicationCore/Public/GenericPlatform/ICursor.h
+++ b/Engine/Source/Runtime/ApplicationCore/Public/GenericPlatform/ICursor.h
@@ -69,11 +69,14 @@ namespace EMouseCursor
 class ICursor
 {
 public:
-	/** Creates a hardware cursor from file. Can return nullptr when not availbale. */
-	virtual void* CreateCursorFromFile(const FString& InPathToCursorWithoutExtension, FVector2D HotSpot) = 0;
+	/** Creates a hardware cursor from file. Can return nullptr when not available. */
+	virtual void* CreateCursorFromFile(const FString& InPathToCursorWithoutExtension, FVector2D HotSpot) { return nullptr; }
+
+	/** Is CreateCursorFromRGBABuffer() supported by this cursor? */
+	virtual bool IsCreateCursorFromRGBABufferSupported() const { return false; }
 
 	/** Creates a hardware cursor from bitmap data. Can return nullptr when not available. */
-	virtual void* CreateCursorFromRGBABuffer(const FColor* Pixels, int32 Width, int32 Height, FVector2D InHotSpot) = 0;
+	virtual void* CreateCursorFromRGBABuffer(const FColor* Pixels, int32 Width, int32 Height, FVector2D InHotSpot) { return nullptr; }
 
 	/** The position of the cursor */
 	virtual FVector2D GetPosition() const = 0;
diff --git a/Engine/Source/Runtime/ApplicationCore/Public/HoloLens/HoloLensCursor.h b/Engine/Source/Runtime/ApplicationCore/Public/HoloLens/HoloLensCursor.h
--- a/Engine/Source/Runtime/ApplicationCore/Public/HoloLens/HoloLensCursor.h
+++ b/Engine/Source/Runtime/ApplicationCore/Public/HoloLens/HoloLensCursor.h
@@ -58,17 +58,6 @@ public:
 	*/
 	void SetTypeShape(EMouseCursor::Type InCursorType, void* CursorHandle);
 
-	virtual void* CreateCursorFromFile(const FString& InPathToCursorWithoutExtension, FVector2D HotSpot) override
-	{
-		return nullptr;
-	}
-
-	/** Creates a hardware cursor from bitmap data. Can return nullptr when not available. */
-	virtual void* CreateCursorFromRGBABuffer(const FColor* Pixels, int32 Width, int32 Height, FVector2D InHotSpot) override
-	{
-		return nullptr;
-	}
-
 private:
 
 	void SetUseRawMouse(bool bUse);
diff --git a/Engine/Source/Runtime/ApplicationCore/Public/IOS/IOSCursor.h b/Engine/Source/Runtime/ApplicationCore/Public/IOS/IOSCursor.h
--- a/Engine/Source/Runtime/ApplicationCore/Public/IOS/IOSCursor.h
+++ b/Engine/Source/Runtime/ApplicationCore/Public/IOS/IOSCursor.h
@@ -15,16 +15,6 @@ public:
 	{
 	}
 
-	virtual void* CreateCursorFromFile(const FString& InPathToCursorWithoutExtension, FVector2D HotSpot) override
-	{
-		return nullptr;
-	}
-
-	virtual void* CreateCursorFromRGBABuffer(const FColor* Pixels, int32 Width, int32 Height, FVector2D InHotSpot) override
-	{
-		return nullptr;
-	}
-
 	virtual FVector2D GetPosition() const override
 	{
 		return CurrentPosition;
diff --git a/Engine/Source/Runtime/ApplicationCore/Public/Linux/LinuxCursor.h b/Engine/Source/Runtime/ApplicationCore/Public/Linux/LinuxCursor.h
--- a/Engine/Source/Runtime/ApplicationCore/Public/Linux/LinuxCursor.h
+++ b/Engine/Source/Runtime/ApplicationCore/Public/Linux/LinuxCursor.h
@@ -16,7 +16,10 @@ public:
 
 	virtual ~FLinuxCursor();
 
-	virtual void* CreateCursorFromFile(const FString& InPathToCursorWithoutExtension, FVector2D HotSpot) override;
+	virtual bool IsCreateCursorFromRGBABufferSupported() const override
+	{
+		return true;
+	}
 
 	virtual void* CreateCursorFromRGBABuffer(const FColor* Pixels, int32 Width, int32 Height, FVector2D InHotSpot) override;
 
diff --git a/Engine/Source/Runtime/ApplicationCore/Public/Mac/MacCursor.h b/Engine/Source/Runtime/ApplicationCore/Public/Mac/MacCursor.h
--- a/Engine/Source/Runtime/ApplicationCore/Public/Mac/MacCursor.h
+++ b/Engine/Source/Runtime/ApplicationCore/Public/Mac/MacCursor.h
@@ -17,6 +17,8 @@ public:
 
 	virtual void* CreateCursorFromFile(const FString& InPathToCursorWithoutExtension, FVector2D HotSpot) override;
 
+	virtual bool IsCreateCursorFromRGBABufferSupported() const override { return true; }
+
 	virtual void* CreateCursorFromRGBABuffer(const FColor* Pixels, int32 Width, int32 Height, FVector2D InHotSpot) override;
 
 	virtual FVector2D GetPosition() const override;
diff --git a/Engine/Source/Runtime/ApplicationCore/Public/Windows/WindowsCursor.h b/Engine/Source/Runtime/ApplicationCore/Public/Windows/WindowsCursor.h
--- a/Engine/Source/Runtime/ApplicationCore/Public/Windows/WindowsCursor.h
+++ b/Engine/Source/Runtime/ApplicationCore/Public/Windows/WindowsCursor.h
@@ -16,6 +16,11 @@ public:
 
 	virtual void* CreateCursorFromFile(const FString& InPathToCursorWithoutExtension, FVector2D HotSpot) override;
 
+	virtual bool IsCreateCursorFromRGBABufferSupported() const override
+	{
+		return true;
+	}
+
 	virtual void* CreateCursorFromRGBABuffer(const FColor* Pixels, int32 Width, int32 Height, FVector2D InHotSpot) override;
 
 	virtual FVector2D GetPosition() const override;
diff --git a/Engine/Source/Runtime/ClothingSystemRuntimeNv/Private/ClothingSimulationNv.cpp b/Engine/Source/Runtime/ClothingSystemRuntimeNv/Private/ClothingSimulationNv.cpp
--- a/Engine/Source/Runtime/ClothingSystemRuntimeNv/Private/ClothingSimulationNv.cpp
+++ b/Engine/Source/Runtime/ClothingSystemRuntimeNv/Private/ClothingSimulationNv.cpp
@@ -1155,9 +1155,23 @@ void FClothingSimulationNv::UpdateLod(int32 InPredictedLod, const FTransform& Co
 			bool bOldLodMapped = LodMap.IsValidIndex(CurrentMeshLodIndex) && LodMap[CurrentMeshLodIndex] != INDEX_NONE;
 
 			// Get the clothing LOD mapped from the mesh predicted LOD
-			const int32 PredictedClothingLod = LodMap[InPredictedLod];
+
+
 			const int32 OldClothingLod = bOldLodMapped ? LodMap[CurrentMeshLodIndex] : INDEX_NONE;
 
+			// If potentialLod doesn't map to a valid LOD, we try higher LOD levels for a valid LOD.
+			// Asset might only have lod on LOD 1 and not 0, however if mesh doesn't force LOD to 1, 
+			// asset will not be assigned valid LOD index and will not generate sim data, breaking things.
+			int32 PredictedClothingLod = INDEX_NONE;
+			for (int32 PotentialLod = InPredictedLod; PotentialLod < LodMap.Num(); ++PotentialLod)
+			{
+				if (LodMap[PotentialLod] != INDEX_NONE)
+				{
+					PredictedClothingLod = LodMap[PotentialLod];
+					break;
+				}
+			}
+
 			if(PredictedClothingLod == Actor.CurrentLodIndex)
 			{
 				// We must have forced a LOD update because we added a new actor - this actor is good to go though.
diff --git a/Engine/Source/Runtime/Engine/Private/GameViewportClient.cpp b/Engine/Source/Runtime/Engine/Private/GameViewportClient.cpp
--- a/Engine/Source/Runtime/Engine/Private/GameViewportClient.cpp
+++ b/Engine/Source/Runtime/Engine/Private/GameViewportClient.cpp
@@ -3817,6 +3817,11 @@ bool UGameViewportClient::GetUseMouseForTouch() const
 
 void* UGameViewportClient::LoadCursorFromPngs(ICursor& PlatformCursor, const FString& InPathToCursorWithoutExtension, FVector2D InHotSpot)
 {
+	if (!PlatformCursor.IsCreateCursorFromRGBABufferSupported())
+	{
+		return nullptr;
+	}
+
 	TArray<TSharedPtr<FPngFileData>> CursorPngFiles;
 	if (!LoadAvailableCursorPngs(CursorPngFiles, InPathToCursorWithoutExtension))
 	{
diff --git a/Engine/Source/Runtime/Engine/Private/PrecomputedVolumetricLightmap.cpp b/Engine/Source/Runtime/Engine/Private/PrecomputedVolumetricLightmap.cpp
--- a/Engine/Source/Runtime/Engine/Private/PrecomputedVolumetricLightmap.cpp
+++ b/Engine/Source/Runtime/Engine/Private/PrecomputedVolumetricLightmap.cpp
@@ -538,6 +538,8 @@ ENGINE_API void FPrecomputedVolumetricLightmapData::AddToSceneData(FPrecomputedV
 				// Steal the indirection texture. When the sublevels are unloaded the values will be restored.
 				IndirectionTexture = SceneData->IndirectionTexture;
 			}
+
+			RHICmdList.TransitionResource(EResourceTransitionAccess::EReadable, IndirectionTexture.Texture);
 		}
 		else
 		{
@@ -1057,6 +1059,20 @@ void FVolumetricLightmapBrickAtlas::Insert(int32 Index, FPrecomputedVolumetricLi
 	// Replace with new allcations
 	Allocations = NewAllocations;
 	TextureSet = NewTextureSet; // <-- Old texture references are released here
+
+	FRHITexture* Textures[3 + UE_ARRAY_COUNT(TextureSet.SHCoefficients)] = 
+	{
+		TextureSet.AmbientVector.Texture,
+		TextureSet.SkyBentNormal.Texture,
+		TextureSet.DirectionalLightShadowing.Texture
+	};
+
+	for (int32 TextureIndex = 0; TextureIndex < UE_ARRAY_COUNT(TextureSet.SHCoefficients); ++TextureIndex)
+	{
+		Textures[TextureIndex + 3] = TextureSet.SHCoefficients[TextureIndex].Texture;
+	}
+
+	RHICmdList.TransitionResources(EResourceTransitionAccess::EReadable, Textures, UE_ARRAY_COUNT(Textures));
 }
 
 void FVolumetricLightmapBrickAtlas::Remove(FPrecomputedVolumetricLightmapData* Data)
diff --git a/Engine/Source/Runtime/Experimental/Chaos/Public/Chaos/PBDRigidParticles.h b/Engine/Source/Runtime/Experimental/Chaos/Public/Chaos/PBDRigidParticles.h
--- a/Engine/Source/Runtime/Experimental/Chaos/Public/Chaos/PBDRigidParticles.h
+++ b/Engine/Source/Runtime/Experimental/Chaos/Public/Chaos/PBDRigidParticles.h
@@ -5,7 +5,7 @@
 #include "Chaos/RigidParticles.h"
 #include "Chaos/Rotation.h"
 
-#if !PLATFORM_PS4
+#if defined(_MSC_VER)
 #pragma warning(push)
 #pragma warning(disable:4946)
 #endif
@@ -187,6 +187,6 @@ FChaosArchive& operator<<(FChaosArchive& Ar, TPBDRigidParticles<T, d>& Particles
 }
 }
 
-#if !PLATFORM_PS4
+#if defined(MSC_VER)
 #pragma warning(pop)
 #endif
diff --git a/Engine/Source/Runtime/RenderCore/Private/ClearReplacementShaders.cpp b/Engine/Source/Runtime/RenderCore/Private/ClearReplacementShaders.cpp
--- a/Engine/Source/Runtime/RenderCore/Private/ClearReplacementShaders.cpp
+++ b/Engine/Source/Runtime/RenderCore/Private/ClearReplacementShaders.cpp
@@ -23,3 +23,4 @@ IMPLEMENT_SHADER_TYPE2_WITH_TEMPLATE_PREFIX(template<>, FClearReplacementCS_Text
 IMPLEMENT_SHADER_TYPE2_WITH_TEMPLATE_PREFIX(template<>, FClearReplacementCS_Texture3D_Uint4         , SF_Compute);
 IMPLEMENT_SHADER_TYPE2_WITH_TEMPLATE_PREFIX(template<>, FClearReplacementCS_Texture2D_Uint4         , SF_Compute);
 IMPLEMENT_SHADER_TYPE2_WITH_TEMPLATE_PREFIX(template<>, FClearReplacementCS_Texture2DArray_Uint4    , SF_Compute);
+IMPLEMENT_SHADER_TYPE2_WITH_TEMPLATE_PREFIX(template<>, FClearReplacementCS_Buffer_Uint4            , SF_Compute);
diff --git a/Engine/Source/Runtime/RenderCore/Public/ClearReplacementShaders.h b/Engine/Source/Runtime/RenderCore/Public/ClearReplacementShaders.h
--- a/Engine/Source/Runtime/RenderCore/Public/ClearReplacementShaders.h
+++ b/Engine/Source/Runtime/RenderCore/Public/ClearReplacementShaders.h
@@ -277,3 +277,4 @@ typedef TClearReplacementCS<EClearReplacementResourceType::Texture2DArray, FClea
 typedef TClearReplacementCS<EClearReplacementResourceType::Texture3D,      FClearReplacementBase_Uint4>         FClearReplacementCS_Texture3D_Uint4;
 typedef TClearReplacementCS<EClearReplacementResourceType::Texture2D,      FClearReplacementBase_Uint4>         FClearReplacementCS_Texture2D_Uint4;
 typedef TClearReplacementCS<EClearReplacementResourceType::Texture2DArray, FClearReplacementBase_Uint4>         FClearReplacementCS_Texture2DArray_Uint4;
+typedef TClearReplacementCS<EClearReplacementResourceType::Buffer,         FClearReplacementBase_Uint4>         FClearReplacementCS_Buffer_Uint4;
diff --git a/Engine/Source/Runtime/Renderer/Private/PostProcess/DiaphragmDOF.h b/Engine/Source/Runtime/Renderer/Private/PostProcess/DiaphragmDOF.h
--- a/Engine/Source/Runtime/Renderer/Private/PostProcess/DiaphragmDOF.h
+++ b/Engine/Source/Runtime/Renderer/Private/PostProcess/DiaphragmDOF.h
@@ -10,11 +10,6 @@
 #include "RendererInterface.h"
 #include "PostProcess/PostProcessing.h"
 
-
-// Whitelist diaphragm DOF for platforms that actually have been tested.
-#define WITH_DIAPHRAGM_DOF (PLATFORM_WINDOWS || PLATFORM_XBOXONE || PLATFORM_PS4 || PLATFORM_MAC || PLATFORM_LINUX || PLATFORM_IOS || PLATFORM_SWITCH)
-
-
 class FViewInfo;
 class FSceneTextureParameters;
 struct FTemporalAAHistory;
@@ -122,11 +117,6 @@ struct FBokehModel
 /** Returns whether DOF is supported. */
 inline bool IsSupported(EShaderPlatform ShaderPlatform)
 {
-	// Since this is still prototype, only allow it on D3D.
-	#if !WITH_DIAPHRAGM_DOF
-		return false;
-	#endif
-
 	// Only compile diaphragm DOF on platform it has been tested to ensure this is not blocking anyone else.
 	return 
 		ShaderPlatform == SP_PCD3D_SM5 ||
diff --git a/Engine/Source/Runtime/Slate/Private/Framework/Application/SlateApplication.cpp b/Engine/Source/Runtime/Slate/Private/Framework/Application/SlateApplication.cpp
--- a/Engine/Source/Runtime/Slate/Private/Framework/Application/SlateApplication.cpp
+++ b/Engine/Source/Runtime/Slate/Private/Framework/Application/SlateApplication.cpp
@@ -64,8 +64,6 @@ class FFauxSlateCursor : public ICursor
 public:
 	FFauxSlateCursor() {}
 	virtual ~FFauxSlateCursor() {}
-	virtual void* CreateCursorFromFile(const FString& InPathToCursorWithoutExtension, FVector2D HotSpot) { return nullptr; }
-	virtual void* CreateCursorFromRGBABuffer(const FColor* Pixels, int32 Width, int32 Height, FVector2D InHotSpot) { return nullptr; }
 	virtual void SetTypeShape(EMouseCursor::Type InCursorType, void* CursorHandle) override {}
 	
 	virtual FVector2D GetPosition() const override { return CurrentPosition; }
@@ -2706,6 +2704,17 @@ void FSlateApplication::SetAllUserFocus(const FWidgetPath& InFocusPath, const EF
 	ForEachUser([&] (FSlateUser& User) {
 		SetUserFocus(User, InFocusPath, InCause);
 	});
+
+	// cache the focus path so it can be applied to any new users
+	if (InFocusPath.IsValid())
+	{
+		LastAllUsersFocusWidget = InFocusPath.GetLastWidget();
+	}
+	else
+	{
+		LastAllUsersFocusWidget.Reset();
+	}
+	LastAllUsersFocusCause = InCause;
 }
 
 void FSlateApplication::SetAllUserFocusAllowingDescendantFocus(const FWidgetPath& InFocusPath, const EFocusCause InCause)
@@ -2718,6 +2727,10 @@ void FSlateApplication::SetAllUserFocusAllowingDescendantFocus(const FWidgetPath
 			SetUserFocus(User, InFocusPath, InCause);
 		}
 	});
+
+	// cache the focus path so it can be applied to any new users
+	LastAllUsersFocusWidget = FocusWidget;
+	LastAllUsersFocusCause = InCause;
 }
 
 FModifierKeysState FSlateApplication::GetModifierKeys() const
@@ -3783,6 +3796,12 @@ TSharedRef<FSlateUser> FSlateApplication::RegisterNewUser(int32 UserIndex, bool
 
 	Users[UserIndex] = NewUser;
 
+	// Apply the last known "all users" focus widget path to this new user if they do not have a specific one
+	if (!NewUser->HasValidFocusPath() && LastAllUsersFocusWidget.IsValid())
+	{
+		SetUserFocus(NewUser->GetUserIndex(), LastAllUsersFocusWidget.Pin(), LastAllUsersFocusCause);
+	}
+
 	UE_LOG(LogSlate, Log, TEXT("Slate User Registered.  User Index %d, Is Virtual User: %d"), UserIndex, bIsVirtual);
 	UserRegisteredEvent.Broadcast(UserIndex);
 
diff --git a/Engine/Source/Runtime/Slate/Public/Framework/Application/SlateApplication.h b/Engine/Source/Runtime/Slate/Public/Framework/Application/SlateApplication.h
--- a/Engine/Source/Runtime/Slate/Public/Framework/Application/SlateApplication.h
+++ b/Engine/Source/Runtime/Slate/Public/Framework/Application/SlateApplication.h
@@ -1612,6 +1612,12 @@ private:
 	/** Weak pointers to the allocated virtual users. */
 	TArray<TWeakPtr<FSlateVirtualUserHandle>> VirtualUsers;
 
+	/**
+	  * Last widget that was set for 'all users' focus and the cause.
+	  */
+	TWeakPtr<SWidget> LastAllUsersFocusWidget;
+	EFocusCause LastAllUsersFocusCause;
+
 	/**
 	 * Application throttling
 	 */
