IsAllowedRemoteNotifications returns false if GCM or Firebase plugin not included for Android

#ue4
#android

diff --git a/Engine/Build/Android/Java/src/com/epicgames/ue4/GameActivity.java.template b/Engine/Build/Android/Java/src/com/epicgames/ue4/GameActivity.java.template
--- a/Engine/Build/Android/Java/src/com/epicgames/ue4/GameActivity.java.template
+++ b/Engine/Build/Android/Java/src/com/epicgames/ue4/GameActivity.java.template
@@ -5132,7 +5132,9 @@ public class GameActivity extends $${gameActivitySuperClass}$$ implements Surfac
 
 	public boolean AndroidThunkJava_IsAllowedRemoteNotifications()
 	{
-		return NotificationManagerCompat.from(getApplicationContext()).areNotificationsEnabled();
+		bool bPluginEnabled = false;
+//$${gameActivityAllowedRemoteNotificationsAdditions}$$
+		return bPluginEnabled && NotificationManagerCompat.from(getApplicationContext()).areNotificationsEnabled();
 	}
 
 	public void AndroidThunkJava_SetDesiredViewSize(int width, int height)
diff --git a/Engine/Plugins/Runtime/Firebase/Source/Firebase.upl.xml b/Engine/Plugins/Runtime/Firebase/Source/Firebase.upl.xml
--- a/Engine/Plugins/Runtime/Firebase/Source/Firebase.upl.xml
+++ b/Engine/Plugins/Runtime/Firebase/Source/Firebase.upl.xml
@@ -84,4 +84,13 @@ apply plugin: 'com.google.gms.google-services'
 		</if>
 	</gradleCopies>
 
+	<gameActivityAllowedRemoteNotificationsAdditions>
+		<if condition="bEnabled">
+			<true>
+				<insert>		bPluginEnabled = true;</insert>
+				<insertNewline/>
+			</true>
+		</if>
+	</gameActivityAllowedRemoteNotificationsAdditions>
+
 </root>
diff --git a/Engine/Plugins/Runtime/GoogleCloudMessaging/Source/GoogleCloudMessaging/GoogleCloudMessaging_UPL.xml b/Engine/Plugins/Runtime/GoogleCloudMessaging/Source/GoogleCloudMessaging/GoogleCloudMessaging_UPL.xml
--- a/Engine/Plugins/Runtime/GoogleCloudMessaging/Source/GoogleCloudMessaging/GoogleCloudMessaging_UPL.xml
+++ b/Engine/Plugins/Runtime/GoogleCloudMessaging/Source/GoogleCloudMessaging/GoogleCloudMessaging_UPL.xml
@@ -144,6 +144,15 @@ dependencies {
     </if>
   </gameActivityClassAdditions>
 
+	<gameActivityAllowedRemoteNotificationsAdditions>
+		<if condition="GCMEnabled">
+			<true>
+				<insert>		bPluginEnabled = true;</insert>
+				<insertNewline/>
+			</true>
+		</if>
+	</gameActivityAllowedRemoteNotificationsAdditions>
+
   <!-- optional additions to GameActivity onCreate metadata reading in GameActivity.java
 	<gameActivityReadMetadataAdditions>
 	</gameActivityReadMetadataAdditions>
diff --git a/Engine/Source/Programs/UnrealBuildTool/Platform/Android/UEDeployAndroid.cs b/Engine/Source/Programs/UnrealBuildTool/Platform/Android/UEDeployAndroid.cs
--- a/Engine/Source/Programs/UnrealBuildTool/Platform/Android/UEDeployAndroid.cs
+++ b/Engine/Source/Programs/UnrealBuildTool/Platform/Android/UEDeployAndroid.cs
@@ -3975,6 +3975,7 @@ namespace UnrealBuildTool
 				{ "//$${gameActivityGetCommandLineAdditions}$$", UPL.ProcessPluginNode(NDKArch, "gameActivityGetCommandLineAdditions", "")},
 				{ "//$${gameActivityGetLoginIdAdditions}$$", UPL.ProcessPluginNode(NDKArch, "gameActivityGetLoginIdAdditions", "")},
 				{ "//$${gameActivityGetFunnelIdAdditions}$$", UPL.ProcessPluginNode(NDKArch, "gameActivityGetFunnelIdAdditions", "")},
+				{ "//$${gameActivityAllowedRemoteNotificationsAdditions}$$", UPL.ProcessPluginNode(NDKArch, "gameActivityAllowedRemoteNotificationsAdditions", "")},
 				{ "//$${soLoadLibrary}$$", UPL.ProcessPluginNode(NDKArch, "soLoadLibrary", LoadLibraryDefaults)},
 				{ "$${gameActivitySuperClass}$$", SuperClassDefault},
 			};
