When staging, set the CFBundleExecutable Info.plist entry to include Client when packaging IOS and TVOS Client C++ targets from PC. Previously it was only set for the stub and then it got overwritten by the stage version when merging plists.

diff --git a/Engine/Source/Programs/AutomationTool/IOS/IOSPlatform.Automation.cs b/Engine/Source/Programs/AutomationTool/IOS/IOSPlatform.Automation.cs
--- a/Engine/Source/Programs/AutomationTool/IOS/IOSPlatform.Automation.cs
+++ b/Engine/Source/Programs/AutomationTool/IOS/IOSPlatform.Automation.cs
@@ -211,7 +211,7 @@ public class IOSPlatform : Platform
 		IOSExports.GetProvisioningData(InProject, bDistribution, out MobileProvision, out SigningCertificate, out TeamUUID, out bAutomaticSigning);
 	}
 
-	public virtual bool DeployGeneratePList(FileReference ProjectFile, UnrealTargetConfiguration Config, DirectoryReference ProjectDirectory, bool bIsUE4Game, string GameName, string ProjectName, DirectoryReference InEngineDir, DirectoryReference AppDirectory, string InExecutablePath, out bool bSupportsPortrait, out bool bSupportsLandscape, out bool bSkipIcons)
+	public virtual bool DeployGeneratePList(FileReference ProjectFile, UnrealTargetConfiguration Config, DirectoryReference ProjectDirectory, bool bIsUE4Game, string GameName, bool bIsClient, string ProjectName, DirectoryReference InEngineDir, DirectoryReference AppDirectory, string InExecutablePath, out bool bSupportsPortrait, out bool bSupportsLandscape, out bool bSkipIcons)
 	{
 		string TargetName = Path.GetFileNameWithoutExtension(InExecutablePath).Split("-".ToCharArray())[0];
 		FileReference TargetReceiptFileName;
@@ -223,7 +223,7 @@ public class IOSPlatform : Platform
 		{
 			TargetReceiptFileName = TargetReceipt.GetDefaultPath(ProjectDirectory, TargetName, UnrealTargetPlatform.IOS, Config, "");
 		}
-		return IOSExports.GeneratePList(ProjectFile, Config, ProjectDirectory, bIsUE4Game, GameName, ProjectName, InEngineDir, AppDirectory, TargetReceiptFileName, out bSupportsPortrait, out bSupportsLandscape, out bSkipIcons);
+		return IOSExports.GeneratePList(ProjectFile, Config, ProjectDirectory, bIsUE4Game, GameName, bIsClient, ProjectName, InEngineDir, AppDirectory, TargetReceiptFileName, out bSupportsPortrait, out bSupportsLandscape, out bSkipIcons);
 	}
 
 	protected string MakeIPAFileName(UnrealTargetConfiguration TargetConfiguration, ProjectParams Params, DeploymentContext SC, bool bAllowDistroPrefix)
@@ -867,8 +867,9 @@ public class IOSPlatform : Platform
 							(SC.IsCodeBasedProject ? SC.ProjectRoot : DirectoryReference.Combine(SC.LocalRoot, "Engine")),
 							!SC.IsCodeBasedProject,
 							(SC.IsCodeBasedProject ? SC.StageExecutables[0] : "UE4Game"),
+							Params.Client,
 							SC.ShortProjectName, DirectoryReference.Combine(SC.LocalRoot, "Engine"),
-														DirectoryReference.Combine((SC.IsCodeBasedProject ? SC.ProjectRoot : DirectoryReference.Combine(SC.LocalRoot, "Engine")), "Binaries", PlatformName, "Payload", (SC.IsCodeBasedProject ? SC.ShortProjectName : "UE4Game") + ".app"),
+							DirectoryReference.Combine((SC.IsCodeBasedProject ? SC.ProjectRoot : DirectoryReference.Combine(SC.LocalRoot, "Engine")), "Binaries", PlatformName, "Payload", (SC.IsCodeBasedProject ? SC.ShortProjectName : "UE4Game") + ".app"),
 							SC.StageExecutables[0],
 							out bSupportsPortrait,
 							out bSupportsLandscape,
diff --git a/Engine/Source/Programs/AutomationTool/TVOS/TVOSPlatform.Automation.cs b/Engine/Source/Programs/AutomationTool/TVOS/TVOSPlatform.Automation.cs
--- a/Engine/Source/Programs/AutomationTool/TVOS/TVOSPlatform.Automation.cs
+++ b/Engine/Source/Programs/AutomationTool/TVOS/TVOSPlatform.Automation.cs
@@ -37,7 +37,7 @@ public class TVOSPlatform : IOSPlatform
 		TVOSExports.GetProvisioningData(InProject, bDistribution, out MobileProvision, out SigningCertificate, out Team, out bAutomaticSigning);
     }
 
-	public override bool DeployGeneratePList(FileReference ProjectFile, UnrealTargetConfiguration Config, DirectoryReference ProjectDirectory, bool bIsUE4Game, string GameName, string ProjectName, DirectoryReference InEngineDir, DirectoryReference AppDirectory, string InExecutablePath, out bool bSupportsPortrait, out bool bSupportsLandscape, out bool bSkipIcons)
+	public override bool DeployGeneratePList(FileReference ProjectFile, UnrealTargetConfiguration Config, DirectoryReference ProjectDirectory, bool bIsUE4Game, string GameName, bool bIsClient, string ProjectName, DirectoryReference InEngineDir, DirectoryReference AppDirectory, string InExecutablePath, out bool bSupportsPortrait, out bool bSupportsLandscape, out bool bSkipIcons)
 	{
 		string TargetName = Path.GetFileNameWithoutExtension(InExecutablePath).Split("-".ToCharArray())[0];
 		FileReference TargetReceiptFileName;
@@ -49,7 +49,7 @@ public class TVOSPlatform : IOSPlatform
 		{
 			TargetReceiptFileName = TargetReceipt.GetDefaultPath(ProjectDirectory, TargetName, UnrealTargetPlatform.TVOS, Config, "");
 		}
-		return TVOSExports.GeneratePList(ProjectFile, Config, ProjectDirectory, bIsUE4Game, GameName, ProjectName, InEngineDir, AppDirectory, TargetReceiptFileName, out bSupportsPortrait, out bSupportsLandscape, out bSkipIcons);
+		return TVOSExports.GeneratePList(ProjectFile, Config, ProjectDirectory, bIsUE4Game, GameName, bIsClient, ProjectName, InEngineDir, AppDirectory, TargetReceiptFileName, out bSupportsPortrait, out bSupportsLandscape, out bSkipIcons);
 	}
 
     public override string GetCookPlatform(bool bDedicatedServer, bool bIsClientOnly)
@@ -112,7 +112,7 @@ public class TVOSPlatform : IOSPlatform
                     bool bSupportsPortrait = false;
                     bool bSupportsLandscape = false;
                     bool bSkipIcons = false;
-                    DeployGeneratePList(SC.RawProjectPath, TargetConfiguration, (SC.IsCodeBasedProject ? SC.ProjectRoot : SC.EngineRoot), !SC.IsCodeBasedProject, (SC.IsCodeBasedProject ? SC.ShortProjectName : "UE4Game"), SC.ShortProjectName, SC.EngineRoot, DirectoryReference.Combine((SC.IsCodeBasedProject ? SC.ProjectRoot : SC.EngineRoot), "Binaries", "TVOS", "Payload", (SC.IsCodeBasedProject ? SC.ShortProjectName : "UE4Game") + ".app"), SC.StageExecutables[0], out bSupportsPortrait, out bSupportsLandscape, out bSkipIcons);
+                    DeployGeneratePList(SC.RawProjectPath, TargetConfiguration, (SC.IsCodeBasedProject ? SC.ProjectRoot : SC.EngineRoot), !SC.IsCodeBasedProject, (SC.IsCodeBasedProject ? SC.ShortProjectName : "UE4Game"), Params.Client, SC.ShortProjectName, SC.EngineRoot, DirectoryReference.Combine((SC.IsCodeBasedProject ? SC.ProjectRoot : SC.EngineRoot), "Binaries", "TVOS", "Payload", (SC.IsCodeBasedProject ? SC.ShortProjectName : "UE4Game") + ".app"), SC.StageExecutables[0], out bSupportsPortrait, out bSupportsLandscape, out bSkipIcons);
                 }
 
                 SC.StageFile(StagedFileType.SystemNonUFS, TargetPListFile, new StagedFileReference("Info.plist"));
diff --git a/Engine/Source/Programs/UnrealBuildTool/Platform/IOS/IOSExports.cs b/Engine/Source/Programs/UnrealBuildTool/Platform/IOS/IOSExports.cs
--- a/Engine/Source/Programs/UnrealBuildTool/Platform/IOS/IOSExports.cs
+++ b/Engine/Source/Programs/UnrealBuildTool/Platform/IOS/IOSExports.cs
@@ -93,6 +93,7 @@ namespace UnrealBuildTool
 		/// <param name="ProjectDirectory"></param>
 		/// <param name="bIsUE4Game"></param>
 		/// <param name="GameName"></param>
+		/// <param name="bIsClient"></param>
 		/// <param name="ProjectName"></param>
 		/// <param name="InEngineDir"></param>
 		/// <param name="AppDirectory"></param>
@@ -101,10 +102,10 @@ namespace UnrealBuildTool
 		/// <param name="bSupportsLandscape"></param>
 		/// <param name="bSkipIcons"></param>
 		/// <returns></returns>
-		public static bool GeneratePList(FileReference ProjectFile, UnrealTargetConfiguration Config, DirectoryReference ProjectDirectory, bool bIsUE4Game, string GameName, string ProjectName, DirectoryReference InEngineDir, DirectoryReference AppDirectory, FileReference BuildReceiptFileName, out bool bSupportsPortrait, out bool bSupportsLandscape, out bool bSkipIcons)
+		public static bool GeneratePList(FileReference ProjectFile, UnrealTargetConfiguration Config, DirectoryReference ProjectDirectory, bool bIsUE4Game, string GameName, bool bIsClient, string ProjectName, DirectoryReference InEngineDir, DirectoryReference AppDirectory, FileReference BuildReceiptFileName, out bool bSupportsPortrait, out bool bSupportsLandscape, out bool bSkipIcons)
 		{
 			TargetReceipt Receipt = TargetReceipt.Read(BuildReceiptFileName);
-			return new UEDeployIOS().GeneratePList(ProjectFile, Config, ProjectDirectory.FullName, bIsUE4Game, GameName, ProjectName, InEngineDir.FullName, AppDirectory.FullName, Receipt, out bSupportsPortrait, out bSupportsLandscape, out bSkipIcons);
+			return new UEDeployIOS().GeneratePList(ProjectFile, Config, ProjectDirectory.FullName, bIsUE4Game, GameName, bIsClient, ProjectName, InEngineDir.FullName, AppDirectory.FullName, Receipt, out bSupportsPortrait, out bSupportsLandscape, out bSkipIcons);
 		}
 
 		/// <summary>
diff --git a/Engine/Source/Programs/UnrealBuildTool/Platform/IOS/UEDeployIOS.cs b/Engine/Source/Programs/UnrealBuildTool/Platform/IOS/UEDeployIOS.cs
--- a/Engine/Source/Programs/UnrealBuildTool/Platform/IOS/UEDeployIOS.cs
+++ b/Engine/Source/Programs/UnrealBuildTool/Platform/IOS/UEDeployIOS.cs
@@ -195,7 +195,7 @@ namespace UnrealBuildTool
 			return result;
 		}
 
-		public static bool GenerateIOSPList(FileReference ProjectFile, UnrealTargetConfiguration Config, string ProjectDirectory, bool bIsUE4Game, string GameName, string ProjectName, string InEngineDir, string AppDirectory, VersionNumber SdkVersion, UnrealPluginLanguage UPL, string BundleID, bool bBuildAsFramework, out bool bSupportsPortrait, out bool bSupportsLandscape, out bool bSkipIcons)
+		public static bool GenerateIOSPList(FileReference ProjectFile, UnrealTargetConfiguration Config, string ProjectDirectory, bool bIsUE4Game, string GameName, bool bIsClient, string ProjectName, string InEngineDir, string AppDirectory, VersionNumber SdkVersion, UnrealPluginLanguage UPL, string BundleID, bool bBuildAsFramework, out bool bSupportsPortrait, out bool bSupportsLandscape, out bool bSkipIcons)
 		{
 			// generate the Info.plist for future use
 			string BuildDirectory = ProjectDirectory + "/Build/IOS";
@@ -430,7 +430,12 @@ namespace UnrealBuildTool
 			Text.AppendLine("\t<key>CFBundleDisplayName</key>");
 			Text.AppendLine(string.Format("\t<string>{0}</string>", EncodeBundleName(BundleDisplayName, ProjectName)));
 			Text.AppendLine("\t<key>CFBundleExecutable</key>");
-			Text.AppendLine(string.Format("\t<string>{0}</string>", bIsUE4Game ? "UE4Game" : GameName));
+			string BundleExecutable = bIsUE4Game ? "UE4Game" : GameName;
+			if (bIsClient)
+			{
+				BundleExecutable += "Client";
+			}
+			Text.AppendLine(string.Format("\t<string>{0}</string>", BundleExecutable));
 			Text.AppendLine("\t<key>CFBundleIdentifier</key>");
 			Text.AppendLine(string.Format("\t<string>{0}</string>", BundleIdentifier.Replace("[PROJECT_NAME]", ProjectName).Replace("_", "")));
 			Text.AppendLine("\t<key>CFBundleInfoDictionaryVersion</key>");
@@ -814,15 +819,15 @@ namespace UnrealBuildTool
 			return false;
 		}
 
-		public bool GeneratePList(FileReference ProjectFile, UnrealTargetConfiguration Config, string ProjectDirectory, bool bIsUE4Game, string GameName, string ProjectName, string InEngineDir, string AppDirectory, TargetReceipt Receipt, out bool bSupportsPortrait, out bool bSupportsLandscape, out bool bSkipIcons)
+ 		public bool GeneratePList(FileReference ProjectFile, UnrealTargetConfiguration Config, string ProjectDirectory, bool bIsUE4Game, string GameName, bool bIsClient, string ProjectName, string InEngineDir, string AppDirectory, TargetReceipt Receipt, out bool bSupportsPortrait, out bool bSupportsLandscape, out bool bSkipIcons)
 		{
 			List<string> UPLScripts = CollectPluginDataPaths(Receipt.AdditionalProperties);
 			VersionNumber SdkVersion = GetSdkVersion(Receipt);
 			bool bBuildAsFramework = GetCompileAsDll(Receipt);
-			return GeneratePList(ProjectFile, Config, ProjectDirectory, bIsUE4Game, GameName, ProjectName, InEngineDir, AppDirectory, UPLScripts, SdkVersion, "", bBuildAsFramework, out bSupportsPortrait, out bSupportsLandscape, out bSkipIcons);
+			return GeneratePList(ProjectFile, Config, ProjectDirectory, bIsUE4Game, GameName, bIsClient, ProjectName, InEngineDir, AppDirectory, UPLScripts, SdkVersion, "", bBuildAsFramework, out bSupportsPortrait, out bSupportsLandscape, out bSkipIcons);
 		}
 
-		public virtual bool GeneratePList(FileReference ProjectFile, UnrealTargetConfiguration Config, string ProjectDirectory, bool bIsUE4Game, string GameName, string ProjectName, string InEngineDir, string AppDirectory, List<string> UPLScripts, VersionNumber SdkVersion, string BundleID, bool bBuildAsFramework, out bool bSupportsPortrait, out bool bSupportsLandscape, out bool bSkipIcons)
+		public virtual bool GeneratePList(FileReference ProjectFile, UnrealTargetConfiguration Config, string ProjectDirectory, bool bIsUE4Game, string GameName, bool bIsClient, string ProjectName, string InEngineDir, string AppDirectory, List<string> UPLScripts, VersionNumber SdkVersion, string BundleID, bool bBuildAsFramework, out bool bSupportsPortrait, out bool bSupportsLandscape, out bool bSkipIcons)
 		{
 			// remember name with -IOS-Shipping, etc
 			string ExeName = GameName;
@@ -854,7 +859,7 @@ namespace UnrealBuildTool
 			// Passing in true for distribution is not ideal here but given the way that ios packaging happens and this call chain it seems unavoidable for now, maybe there is a way to correctly pass it in that I can't find?
 			UPL.Init(ProjectArches, true, RelativeEnginePath, BundlePath, ProjectDirectory, Config.ToString(), false);
 
-			return GenerateIOSPList(ProjectFile, Config, ProjectDirectory, bIsUE4Game, GameName, ProjectName, InEngineDir, AppDirectory, SdkVersion, UPL, BundleID, bBuildAsFramework, out bSupportsPortrait, out bSupportsLandscape, out bSkipIcons);
+			return GenerateIOSPList(ProjectFile, Config, ProjectDirectory, bIsUE4Game, GameName, bIsClient, ProjectName, InEngineDir, AppDirectory, SdkVersion, UPL, BundleID, bBuildAsFramework, out bSupportsPortrait, out bSupportsLandscape, out bSkipIcons);
 		}
 
 		protected virtual void CopyCloudResources(string InEngineDir, string AppDirectory)
@@ -1125,7 +1130,7 @@ namespace UnrealBuildTool
 			bool bSupportsPortrait = true;
 			bool bSupportsLandscape = false;
 			bool bSkipIcons = false;
-			bool bSkipDefaultPNGs = GeneratePList(ProjectFile, Config, InProjectDirectory, bIsUE4Game, GameExeName, InProjectName, InEngineDir, AppDirectory, UPLScripts, SdkVersion, BundleID, bBuildAsFramework, out bSupportsPortrait, out bSupportsLandscape, out bSkipIcons);
+			bool bSkipDefaultPNGs = GeneratePList(ProjectFile, Config, InProjectDirectory, bIsUE4Game, GameExeName, false, InProjectName, InEngineDir, AppDirectory, UPLScripts, SdkVersion, BundleID, bBuildAsFramework, out bSupportsPortrait, out bSupportsLandscape, out bSkipIcons);
 
 			// ensure the destination is writable
 			if (File.Exists(AppDirectory + "/" + GameName))
@@ -1202,7 +1207,7 @@ namespace UnrealBuildTool
 			{
 				// @todo tvos merge: This used to copy the bundle back - where did that code go? It needs to be fixed up for TVOS directories
 				bool bSupportPortrait, bSupportLandscape, bSkipIcons;
-				GeneratePList(ProjectFile, Configuration, ProjectDirectory, bIsUE4Game, GameName, (ProjectFile == null) ? "" : Path.GetFileNameWithoutExtension(ProjectFile.FullName), "../../Engine", "", UPLScripts, SdkVersion, BundleID, bBuildAsFramework, out bSupportPortrait, out bSupportLandscape, out bSkipIcons);
+				GeneratePList(ProjectFile, Configuration, ProjectDirectory, bIsUE4Game, GameName, false, (ProjectFile == null) ? "" : Path.GetFileNameWithoutExtension(ProjectFile.FullName), "../../Engine", "", UPLScripts, SdkVersion, BundleID, bBuildAsFramework, out bSupportPortrait, out bSupportLandscape, out bSkipIcons);
 			}
 			return true;
 		}
diff --git a/Engine/Source/Programs/UnrealBuildTool/Platform/TVOS/TVOSExports.cs b/Engine/Source/Programs/UnrealBuildTool/Platform/TVOS/TVOSExports.cs
--- a/Engine/Source/Programs/UnrealBuildTool/Platform/TVOS/TVOSExports.cs
+++ b/Engine/Source/Programs/UnrealBuildTool/Platform/TVOS/TVOSExports.cs
@@ -90,6 +90,7 @@ namespace UnrealBuildTool
 		/// <param name="ProjectDirectory"></param>
 		/// <param name="bIsUE4Game"></param>
 		/// <param name="GameName"></param>
+		/// <param name="bIsClient"></param>
 		/// <param name="ProjectName"></param>
 		/// <param name="InEngineDir"></param>
 		/// <param name="AppDirectory"></param>
@@ -98,10 +99,10 @@ namespace UnrealBuildTool
 		/// <param name="bSupportsLandscape"></param>
 		/// <param name="bSkipIcons"></param>
 		/// <returns></returns>
-		public static bool GeneratePList(FileReference ProjectFile, UnrealTargetConfiguration Config, DirectoryReference ProjectDirectory, bool bIsUE4Game, string GameName, string ProjectName, DirectoryReference InEngineDir, DirectoryReference AppDirectory, FileReference BuildReceiptFileName, out bool bSupportsPortrait, out bool bSupportsLandscape, out bool bSkipIcons)
+		public static bool GeneratePList(FileReference ProjectFile, UnrealTargetConfiguration Config, DirectoryReference ProjectDirectory, bool bIsUE4Game, string GameName, bool bIsClient, string ProjectName, DirectoryReference InEngineDir, DirectoryReference AppDirectory, FileReference BuildReceiptFileName, out bool bSupportsPortrait, out bool bSupportsLandscape, out bool bSkipIcons)
 		{
 			TargetReceipt Receipt = TargetReceipt.Read(BuildReceiptFileName);
-			return new UEDeployTVOS().GeneratePList(ProjectFile, Config, ProjectDirectory.FullName, bIsUE4Game, GameName, ProjectName, InEngineDir.FullName, AppDirectory.FullName, Receipt, out bSupportsPortrait, out bSupportsLandscape, out bSkipIcons);
+			return new UEDeployTVOS().GeneratePList(ProjectFile, Config, ProjectDirectory.FullName, bIsUE4Game, GameName, bIsClient, ProjectName, InEngineDir.FullName, AppDirectory.FullName, Receipt, out bSupportsPortrait, out bSupportsLandscape, out bSkipIcons);
 		}
 	}
 }
diff --git a/Engine/Source/Programs/UnrealBuildTool/Platform/TVOS/UEDeployTVOS.cs b/Engine/Source/Programs/UnrealBuildTool/Platform/TVOS/UEDeployTVOS.cs
--- a/Engine/Source/Programs/UnrealBuildTool/Platform/TVOS/UEDeployTVOS.cs
+++ b/Engine/Source/Programs/UnrealBuildTool/Platform/TVOS/UEDeployTVOS.cs
@@ -24,7 +24,7 @@ namespace UnrealBuildTool
 			return "TVOS";
 		}
 
-		public static bool GenerateTVOSPList(string ProjectDirectory, bool bIsUE4Game, string GameName, string ProjectName, string InEngineDir, string AppDirectory, UnrealPluginLanguage UPL, string BundleID)
+		public static bool GenerateTVOSPList(string ProjectDirectory, bool bIsUE4Game, string GameName, bool bIsClient, string ProjectName, string InEngineDir, string AppDirectory, UnrealPluginLanguage UPL, string BundleID)
 		{
 			// @todo tvos: THIS!
 
@@ -116,7 +116,12 @@ namespace UnrealBuildTool
 			Text.AppendLine("\t<key>CFBundleDisplayName</key>");
 			Text.AppendLine(string.Format("\t<string>{0}</string>", EncodeBundleName(BundleDisplayName, ProjectName)));
 			Text.AppendLine("\t<key>CFBundleExecutable</key>");
-			Text.AppendLine(string.Format("\t<string>{0}</string>", bIsUE4Game ? "UE4Game" : GameName));
+			string BundleExecutable = bIsUE4Game ? "UE4Game" : GameName;
+			if (bIsClient)
+			{
+				BundleExecutable += "Client";
+			}
+			Text.AppendLine(string.Format("\t<string>{0}</string>", BundleExecutable));
 			Text.AppendLine("\t<key>CFBundleIdentifier</key>");
 			Text.AppendLine(string.Format("\t<string>{0}</string>", BundleIdentifier.Replace("[PROJECT_NAME]", ProjectName).Replace("_","")));
 			Text.AppendLine("\t<key>CFBundleInfoDictionaryVersion</key>");
@@ -349,11 +354,11 @@ namespace UnrealBuildTool
 			return bSkipDefaultPNGs;
 		}
 
-		public override bool GeneratePList(FileReference ProjectFile, UnrealTargetConfiguration Config, string ProjectDirectory, bool bIsUE4Game, string GameName, string ProjectName, string InEngineDir, string AppDirectory, List<string> UPLScripts, VersionNumber SdkVersion, string BundleID, bool bBuildAsFramework, out bool bSupportsPortrait, out bool bSupportsLandscape, out bool bSkipIcons)
+		public override bool GeneratePList(FileReference ProjectFile, UnrealTargetConfiguration Config, string ProjectDirectory, bool bIsUE4Game, string GameName, bool bIsClient, string ProjectName, string InEngineDir, string AppDirectory, List<string> UPLScripts, VersionNumber SdkVersion, string BundleID, bool bBuildAsFramework, out bool bSupportsPortrait, out bool bSupportsLandscape, out bool bSkipIcons)
 		{
 			bSupportsLandscape = bSupportsPortrait = true;
             bSkipIcons = true;
-			return GenerateTVOSPList(ProjectDirectory, bIsUE4Game, GameName, ProjectName, InEngineDir, AppDirectory, null, BundleID);
+			return GenerateTVOSPList(ProjectDirectory, bIsUE4Game, GameName, bIsClient, ProjectName, InEngineDir, AppDirectory, null, BundleID);
 		}
 
         protected override void CopyGraphicsResources(bool bSkipDefaultPNGs, bool bSkipIcons, string InEngineDir, string AppDirectory, string BuildDirectory, string IntermediateDir, bool bSupportsPortrait, bool bSupportsLandscape)
diff --git a/Engine/Source/Programs/UnrealBuildTool/ProjectFiles/Xcode/XcodeProject.cs b/Engine/Source/Programs/UnrealBuildTool/ProjectFiles/Xcode/XcodeProject.cs
--- a/Engine/Source/Programs/UnrealBuildTool/ProjectFiles/Xcode/XcodeProject.cs
+++ b/Engine/Source/Programs/UnrealBuildTool/ProjectFiles/Xcode/XcodeProject.cs
@@ -1265,12 +1265,12 @@ namespace UnrealBuildTool
 							TargetReceipt.TryRead(ReceiptFilename, out Receipt);
 							VersionNumber SdkVersion = UEDeployIOS.GetSdkVersion(Receipt);
 							bool bBuildAsFramework = UEDeployIOS.GetCompileAsDll(Receipt);
-							UEDeployIOS.GenerateIOSPList(ProjectFile, Config.BuildConfig, ProjectPath.FullName, bIsUE4Game, GameName, Config.BuildTarget, EngineDir.FullName, ProjectPath + "/Binaries/IOS/Payload", SdkVersion, null, BundleIdentifier, bBuildAsFramework, out bSupportPortrait, out bSupportLandscape, out bSkipIcons);
+							UEDeployIOS.GenerateIOSPList(ProjectFile, Config.BuildConfig, ProjectPath.FullName, bIsUE4Game, GameName, false, Config.BuildTarget, EngineDir.FullName, ProjectPath + "/Binaries/IOS/Payload", SdkVersion, null, BundleIdentifier, bBuildAsFramework, out bSupportPortrait, out bSupportLandscape, out bSkipIcons);
 						}
 						if (bCreateTVOSInfoPlist)
 						{
 							Directory.CreateDirectory(Path.GetDirectoryName(TVOSInfoPlistPath));
-							UEDeployTVOS.GenerateTVOSPList(ProjectPath.FullName, bIsUE4Game, GameName, Config.BuildTarget, EngineDir.FullName, ProjectPath + "/Binaries/TVOS/Payload", null, BundleIdentifier);
+							UEDeployTVOS.GenerateTVOSPList(ProjectPath.FullName, bIsUE4Game, GameName, false, Config.BuildTarget, EngineDir.FullName, ProjectPath + "/Binaries/TVOS/Payload", null, BundleIdentifier);
 						}
 					}
 				}
