Remove Gnmx from hard-coded list of modules that aren't IWYU compatible.

#rnx

diff --git a/Engine/Source/Programs/UnrealBuildTool/Configuration/UEBuildModuleCPP.cs b/Engine/Source/Programs/UnrealBuildTool/Configuration/UEBuildModuleCPP.cs
--- a/Engine/Source/Programs/UnrealBuildTool/Configuration/UEBuildModuleCPP.cs
+++ b/Engine/Source/Programs/UnrealBuildTool/Configuration/UEBuildModuleCPP.cs
@@ -1157,7 +1157,6 @@ namespace UnrealBuildTool
 				case "VulkanRHI":
 				case "OpenGLDrv":
 				case "MetalRHI":
-				case "Gnmx":
 					return true;
 			}
 			return false;
