Move to new llvm clang 9.0.1 v16 toolchain

diff --git a/Engine/Build/BatchFiles/Linux/SetupToolchain.sh b/Engine/Build/BatchFiles/Linux/SetupToolchain.sh
--- a/Engine/Build/BatchFiles/Linux/SetupToolchain.sh
+++ b/Engine/Build/BatchFiles/Linux/SetupToolchain.sh
@@ -4,7 +4,7 @@
 # put ourselves into Engine directory (two up from location of this script)
 pushd "`dirname "$0"`/../../.." > /dev/null
 
-TOOLCHAIN_VERSION=v15_clang-8.0.1-centos7
+TOOLCHAIN_VERSION=v16_clang-9.0.1-centos7
 TOOLCHAIN_ARCHIVE=$TOOLCHAIN_VERSION.tar.gz
 
 TOOLCHAIN_URL=http://cdn.unrealengine.com/Toolchain_Linux/native-linux-$TOOLCHAIN_ARCHIVE
diff --git a/Engine/Source/Programs/UnrealBuildTool/Platform/Linux/LinuxToolChain.cs b/Engine/Source/Programs/UnrealBuildTool/Platform/Linux/LinuxToolChain.cs
--- a/Engine/Source/Programs/UnrealBuildTool/Platform/Linux/LinuxToolChain.cs
+++ b/Engine/Source/Programs/UnrealBuildTool/Platform/Linux/LinuxToolChain.cs
@@ -173,10 +173,10 @@ namespace UnrealBuildTool
 				throw new BuildException("Unable to build: no compatible clang version found. Please run Setup.sh");
 			}
 			// prevent unknown clangs since the build is likely to fail on too old or too new compilers
-			else if ((CompilerVersionMajor * 10 + CompilerVersionMinor) > 80 || (CompilerVersionMajor * 10 + CompilerVersionMinor) < 60)
+			else if ((CompilerVersionMajor * 10 + CompilerVersionMinor) > 90 || (CompilerVersionMajor * 10 + CompilerVersionMinor) < 60)
 			{
 				throw new BuildException(
-					string.Format("This version of the Unreal Engine can only be compiled with clang 8.0, 7.0 and 6.0. clang {0} may not build it - please use a different version.",
+					string.Format("This version of the Unreal Engine can only be compiled with clang 9.0, 8.0, 7.0 and 6.0. clang {0} may not build it - please use a different version.",
 						CompilerVersionString)
 					);
 			}
diff --git a/Engine/Source/Programs/UnrealBuildTool/Platform/Linux/UEBuildLinux.cs b/Engine/Source/Programs/UnrealBuildTool/Platform/Linux/UEBuildLinux.cs
--- a/Engine/Source/Programs/UnrealBuildTool/Platform/Linux/UEBuildLinux.cs
+++ b/Engine/Source/Programs/UnrealBuildTool/Platform/Linux/UEBuildLinux.cs
@@ -639,7 +639,7 @@ namespace UnrealBuildTool
 		/// <summary>
 		/// This is the SDK version we support
 		/// </summary>
-		static string ExpectedSDKVersion = "v15_clang-8.0.1-centos7";	// now unified for all the architectures
+		static string ExpectedSDKVersion = "v16_clang-9.0.1-centos7";	// now unified for all the architectures
 
 		/// <summary>
 		/// Platform name (embeds architecture for now)
