Fixing a bug in UBT when generating timing data.

#rnx

diff --git a/Engine/Source/Programs/UnrealBuildTool/Platform/Windows/VCToolChain.cs b/Engine/Source/Programs/UnrealBuildTool/Platform/Windows/VCToolChain.cs
--- a/Engine/Source/Programs/UnrealBuildTool/Platform/Windows/VCToolChain.cs
+++ b/Engine/Source/Programs/UnrealBuildTool/Platform/Windows/VCToolChain.cs
@@ -1275,6 +1275,10 @@ namespace UnrealBuildTool
 				{
 					// Generate the file manifest for the aggregator.
 					FileReference ManifestFile = FileReference.Combine(Makefile.ProjectIntermediateDirectory, $"{Target.Name}TimingManifest.txt");
+					if (!DirectoryReference.Exists(ManifestFile.Directory))
+					{
+						DirectoryReference.CreateDirectory(ManifestFile.Directory);
+					}
 					File.WriteAllLines(ManifestFile.FullName, TimingJsonFiles.Select(f => f.AbsolutePath));
 
 					FileReference ExpectedCompileTimeFile = FileReference.FromString(Path.Combine(Makefile.ProjectIntermediateDirectory.FullName, String.Format("{0}.json", Target.Name)));
