Fix null reference when generating project files in binary release.

#rnx

diff --git a/Engine/Source/Programs/UnrealBuildTool/ProjectFiles/VisualStudio/VCProject.cs b/Engine/Source/Programs/UnrealBuildTool/ProjectFiles/VisualStudio/VCProject.cs
--- a/Engine/Source/Programs/UnrealBuildTool/ProjectFiles/VisualStudio/VCProject.cs
+++ b/Engine/Source/Programs/UnrealBuildTool/ProjectFiles/VisualStudio/VCProject.cs
@@ -558,6 +558,7 @@ namespace UnrealBuildTool
 		/// <summary>
 		/// Gets highest C++ version which is used in this project
 		/// </summary>
+		/// <returns>C++ standard version</returns>
 		public CppStandardVersion GetIntelliSenseCppVersion()
 		{
 			if (IntelliSenseCppVersion != CppStandardVersion.Default)
@@ -565,10 +566,10 @@ namespace UnrealBuildTool
 				return IntelliSenseCppVersion;
 			}
 
-			var Version = CppStandardVersion.Default;
-			foreach (var Combination in ProjectConfigAndTargetCombinations)
+			CppStandardVersion Version = CppStandardVersion.Default;
+			foreach (ProjectConfigAndTargetCombination Combination in ProjectConfigAndTargetCombinations)
 			{
-				if (Combination.ProjectTarget.TargetRules.CppStandard > Version)
+				if (Combination.ProjectTarget.TargetRules != null && Combination.ProjectTarget.TargetRules.CppStandard > Version)
 				{
 					Version = Combination.ProjectTarget.TargetRules.CppStandard;
 				}
