Update BundleTool to 0.13.0

#ue4
#android

diff --git a/Engine/Source/Programs/UnrealBuildTool/Platform/Android/UEDeployAndroid.cs b/Engine/Source/Programs/UnrealBuildTool/Platform/Android/UEDeployAndroid.cs
--- a/Engine/Source/Programs/UnrealBuildTool/Platform/Android/UEDeployAndroid.cs
+++ b/Engine/Source/Programs/UnrealBuildTool/Platform/Android/UEDeployAndroid.cs
@@ -18,6 +18,9 @@ namespace UnrealBuildTool
 	{
 		private const string XML_HEADER = "<?xml version=\"1.0\" encoding=\"utf-8\"?>";
 
+		// filename of current BundleTool
+		private const string BUNDLETOOL_JAR = "bundletool-all-0.13.0.jar";
+
 		// Minimum Android SDK that must be used for Java compiling
 		readonly int MinimumSDKLevel = 28;
 
@@ -2989,7 +2992,7 @@ namespace UnrealBuildTool
 				GradleProperties.AppendLine(string.Format("OBB_FILE{0}={1}", OBBFileIndex++, OBBFile.Replace("\\", "/")));
 			}
 
-			GradleProperties.AppendLine("BUNDLETOOL_JAR=" + Path.GetFullPath(Path.Combine(UE4BuildFilesPath, "..", "Prebuilt", "bundletool", "bundletool-all-0.10.3.jar")).Replace("\\", "/"));
+			GradleProperties.AppendLine("BUNDLETOOL_JAR=" + Path.GetFullPath(Path.Combine(UE4BuildFilesPath, "..", "Prebuilt", "bundletool", BUNDLETOOL_JAR)).Replace("\\", "/"));
 			GradleProperties.AppendLine("GENUNIVERSALAPK_JAR=" + Path.GetFullPath(Path.Combine(UE4BuildFilesPath, "..", "Prebuilt", "GenUniversalAPK", "bin", "GenUniversalAPK.jar")).Replace("\\", "/"));
 
 			// add any Gradle properties from UPL
