Fixed iOS toolchain crash with missing provision
#rnx

diff --git a/Engine/Source/Programs/UnrealBuildTool/Platform/IOS/IOSToolChain.cs b/Engine/Source/Programs/UnrealBuildTool/Platform/IOS/IOSToolChain.cs
--- a/Engine/Source/Programs/UnrealBuildTool/Platform/IOS/IOSToolChain.cs
+++ b/Engine/Source/Programs/UnrealBuildTool/Platform/IOS/IOSToolChain.cs
@@ -1884,7 +1884,7 @@ namespace UnrealBuildTool
 						MobileProvisionFile = ProvisioningData.MobileProvisionFile;
 						MobileProvisionUUID = ProvisioningData.MobileProvisionUUID;
 						TeamUUID = ProvisioningData.TeamUUID;
-						if (!ProvisioningData.BundleIdentifier.Contains("*"))
+						if (ProvisioningData.BundleIdentifier != null && !ProvisioningData.BundleIdentifier.Contains("*"))
 						{
 							// If the BundleIndentifer contains a wild card it will not be valid to use in the plist.
 							BundleID = ProvisioningData.BundleIdentifier;
@@ -1905,7 +1905,7 @@ namespace UnrealBuildTool
 						}
 					}
 
-					if(MobileProvisionFile == null)
+					if (MobileProvisionFile == null)
 					{
 						throw new BuildException("Unable to find valid certificate/mobile provision pair.");
 					}
