UnrealBuildTool: Turned off parsing of platform data that often fails for plugins that are disabled during project generation.

diff --git a/Engine/Source/Programs/UnrealBuildTool/System/PluginReferenceDescriptor.cs b/Engine/Source/Programs/UnrealBuildTool/System/PluginReferenceDescriptor.cs
--- a/Engine/Source/Programs/UnrealBuildTool/System/PluginReferenceDescriptor.cs
+++ b/Engine/Source/Programs/UnrealBuildTool/System/PluginReferenceDescriptor.cs
@@ -167,21 +167,26 @@ namespace UnrealBuildTool
 		/// <returns>New PluginReferenceDescriptor object</returns>
 		public static PluginReferenceDescriptor FromJsonObject(JsonObject RawObject)
 		{
-			string[] WhitelistPlatformNames;
-			string[] BlacklistPlatformNames;
-			string[] SupportedTargetPlatformNames;
+			string[] WhitelistPlatformNames = null;
+			string[] BlacklistPlatformNames = null;
+			string[] SupportedTargetPlatformNames = null;
 
 			PluginReferenceDescriptor Descriptor = new PluginReferenceDescriptor(RawObject.GetStringField("Name"), null, RawObject.GetBoolField("Enabled"));
 			RawObject.TryGetBoolField("Optional", out Descriptor.bOptional);
 			RawObject.TryGetStringField("Description", out Descriptor.Description);
 			RawObject.TryGetStringField("MarketplaceURL", out Descriptor.MarketplaceURL);
-			RawObject.TryGetStringArrayField("WhitelistPlatforms", out WhitelistPlatformNames);
-			RawObject.TryGetStringArrayField("BlacklistPlatforms", out BlacklistPlatformNames);
-			RawObject.TryGetEnumArrayField<UnrealTargetConfiguration>("WhitelistTargetConfigurations", out Descriptor.WhitelistTargetConfigurations);
-			RawObject.TryGetEnumArrayField<UnrealTargetConfiguration>("BlacklistTargetConfigurations", out Descriptor.BlacklistTargetConfigurations);
-			RawObject.TryGetEnumArrayField<TargetType>("WhitelistTargets", out Descriptor.WhitelistTargets);
-			RawObject.TryGetEnumArrayField<TargetType>("BlacklistTargets", out Descriptor.BlacklistTargets);
-			RawObject.TryGetStringArrayField("SupportedTargetPlatforms", out SupportedTargetPlatformNames);
+
+			// Only parse platform information if enabled
+			if (Descriptor.bEnabled)
+			{
+				RawObject.TryGetStringArrayField("WhitelistPlatforms", out WhitelistPlatformNames);
+				RawObject.TryGetStringArrayField("BlacklistPlatforms", out BlacklistPlatformNames);
+				RawObject.TryGetEnumArrayField<UnrealTargetConfiguration>("WhitelistTargetConfigurations", out Descriptor.WhitelistTargetConfigurations);
+				RawObject.TryGetEnumArrayField<UnrealTargetConfiguration>("BlacklistTargetConfigurations", out Descriptor.BlacklistTargetConfigurations);
+				RawObject.TryGetEnumArrayField<TargetType>("WhitelistTargets", out Descriptor.WhitelistTargets);
+				RawObject.TryGetEnumArrayField<TargetType>("BlacklistTargets", out Descriptor.BlacklistTargets);
+				RawObject.TryGetStringArrayField("SupportedTargetPlatforms", out SupportedTargetPlatformNames);
+			}
 
 			try
 			{
