Fixed UnrealBuildTool crashing on launch if initialization of UnrealTargetPlatform static members variables happens out of order

## Summary

`UnrealTargetPlatform` and `UnrealPlatformGroup` are partial structs which each have a static member variable declared like this:

private static UniqueStringRegistry StringRegistry = new UniqueStringRegistry();

These partial structs have additional implementations in /Engine/Platforms/XXX/Source/Programs/UnrealBuildTool/UEBuildXXX.cs, which also have static member variables declared like this:

public static UnrealTargetPlatform XXX = FindOrAddByName("XXX");

`FindOrAddByName()` is a static method on `UnrealTargetPlatform` and `UnrealPlatformGroup` which accesses `StringRegistry` and requires it to be initialized i.e. non-null.

It appears that there's no guarantee in .NET runtime that the static member variables in the "original" implementation of the partial structs will be initialized _before_ the ones from the "additional" implementations. This means `XXX` above can be initialized in the additonal implementation before `StringRegistry` has been initialized in the original implementation, which means it's still null when `FindOrAddByName()` is called so the tool crashes.

In practice, this is 100% reproducible in our public GitHub mirror when opening `UnrealBuildTool.csproj` with VisualStudio Mac. This IDE rewrites the project file which affects when / how UEBuildXXX.cs is compiled, and static member variables for these partial structs end up being initialized out or order.

The only solution I found that worked was to replace all accesses to the `StringRegistry` static member variables by a wrapper function `GetUniqueStringRegistry()` which takes care of initializing it - only once. Since all this happens while the tool launches, this code should not need to be multithread aware / re-entrant.

NOTE: There is another partial struct `RestrictedFolder` in this codebase which also use `UniqueStringRegistry` in a similar way but the 2 implementations are in the same source file, so the bug may not happen and I left the implementation untouched.

## Test Plan

- Opened `UnrealBuildTool.csproj` with VisualStudio Mac
- Built UnrealBuildTool and verified it crashed on launch due to `StringRegistry` being NULL when used by `FindOrAddByName()`
- Confirmed the crash doesn't happen anymore after applying this change and that the `StringRegistry` static member variable was only set once by `GetUniqueStringRegistry()`

diff --git a/Engine/Source/Programs/UnrealBuildTool/Configuration/UEBuildTarget.cs b/Engine/Source/Programs/UnrealBuildTool/Configuration/UEBuildTarget.cs
--- a/Engine/Source/Programs/UnrealBuildTool/Configuration/UEBuildTarget.cs
+++ b/Engine/Source/Programs/UnrealBuildTool/Configuration/UEBuildTarget.cs
@@ -26,11 +26,22 @@ namespace UnrealBuildTool
 		private int Id;
 
 		// shared string instance registry - pass in a delegate to create a new one with a name that wasn't made yet
-		private static UniqueStringRegistry StringRegistry = new UniqueStringRegistry();
+		private static UniqueStringRegistry StringRegistry;
+
+		// #jira UE-88908 if parts of a partial struct have each static member variables, their initialization order does not appear guaranteed
+		// here this means initializing "StringRegistry" directly to "new UniqueStringRegistry()" may not be executed before FindOrAddByName() has been called as part of initializing a static member variable of another part of the partial struct
+		private static UniqueStringRegistry GetUniqueStringRegistry()
+		{
+			if (StringRegistry == null)
+			{
+				StringRegistry = new UniqueStringRegistry();
+			}
+			return StringRegistry;
+		}
 
 		private UnrealTargetPlatform(string Name)
 		{
-			Id = StringRegistry.FindOrAddByName(Name);
+			Id = GetUniqueStringRegistry().FindOrAddByName(Name);
 		}
 
 		private UnrealTargetPlatform(int InId)
@@ -55,7 +66,7 @@ namespace UnrealBuildTool
 		/// <param name="Context"></param>
 		public UnrealTargetPlatform(SerializationInfo Info, StreamingContext Context)
 		{
-			Id = StringRegistry.FindOrAddByName((string)Info.GetValue("Name", typeof(string)));
+			Id = GetUniqueStringRegistry().FindOrAddByName((string)Info.GetValue("Name", typeof(string)));
 		}
 
 		/// <summary>
@@ -65,7 +76,7 @@ namespace UnrealBuildTool
 		/// <returns></returns>
 		static private UnrealTargetPlatform FindOrAddByName(string Name)
 		{
-			return new UnrealTargetPlatform(StringRegistry.FindOrAddByName(Name));
+			return new UnrealTargetPlatform(GetUniqueStringRegistry().FindOrAddByName(Name));
 		}
 
 		/// <summary>
@@ -114,7 +125,7 @@ namespace UnrealBuildTool
 		/// <returns></returns>
 		public override string ToString()
 		{
-			return StringRegistry.GetStringForId(Id);
+			return GetUniqueStringRegistry().GetStringForId(Id);
 		}
 
 		/// <summary>
@@ -125,9 +136,9 @@ namespace UnrealBuildTool
 		/// <returns></returns>
 		static public bool TryParse(string Name, out UnrealTargetPlatform Platform)
 		{
-			if (StringRegistry.HasString(Name))
+			if (GetUniqueStringRegistry().HasString(Name))
 			{
-				Platform.Id = StringRegistry.FindOrAddByName(Name);
+				Platform.Id = GetUniqueStringRegistry().FindOrAddByName(Name);
 				return true;
 			}
 
@@ -142,13 +153,13 @@ namespace UnrealBuildTool
 		/// <returns></returns>
 		static public UnrealTargetPlatform Parse(string Name)
 		{
-			if (StringRegistry.HasString(Name))
+			if (GetUniqueStringRegistry().HasString(Name))
 			{
 				return new UnrealTargetPlatform(Name);
 			}
 
 			throw new BuildException(string.Format("The platform name {0} is not a valid platform name. Valid names are ({1})", Name, 
-				string.Join(",", StringRegistry.GetStringNames())));
+				string.Join(",", GetUniqueStringRegistry().GetStringNames())));
 		}
 
 
@@ -158,7 +169,7 @@ namespace UnrealBuildTool
 		/// <returns></returns>
 		public static UnrealTargetPlatform[] GetValidPlatforms()
 		{
-			return Array.ConvertAll(StringRegistry.GetStringIds(), x => new UnrealTargetPlatform(x));
+			return Array.ConvertAll(GetUniqueStringRegistry().GetStringIds(), x => new UnrealTargetPlatform(x));
 		}
 
 		/// <summary>
@@ -167,7 +178,7 @@ namespace UnrealBuildTool
 		/// <returns></returns>
 		public static string[] GetValidPlatformNames()
 		{
-			return StringRegistry.GetStringNames();
+			return GetUniqueStringRegistry().GetStringNames();
 		}
 
 		/// <summary>
@@ -177,7 +188,7 @@ namespace UnrealBuildTool
 		/// <returns></returns>
 		public static bool IsValidName(string Name)
 		{
-			return StringRegistry.HasString(Name);
+			return GetUniqueStringRegistry().HasString(Name);
 		}
 
 		/// <summary>
@@ -301,11 +312,21 @@ namespace UnrealBuildTool
 		private int Id;
 
 		// shared string instance registry - pass in a delegate to create a new one with a name that wasn't made yet
-		private static UniqueStringRegistry StringRegistry = new UniqueStringRegistry();
+		private static UniqueStringRegistry StringRegistry;
+
+		// #jira UE-88908 (see above)
+		private static UniqueStringRegistry GetUniqueStringRegistry()
+		{
+			if (StringRegistry == null)
+			{
+				StringRegistry = new UniqueStringRegistry();
+			}
+			return StringRegistry;
+		}
 
 		private UnrealPlatformGroup(string Name)
 		{
-			Id = StringRegistry.FindOrAddByName(Name);
+			Id = GetUniqueStringRegistry().FindOrAddByName(Name);
 		}
 
 		private UnrealPlatformGroup(int InId)
@@ -330,7 +351,7 @@ namespace UnrealBuildTool
 		/// <param name="Context"></param>
 		public UnrealPlatformGroup(SerializationInfo Info, StreamingContext Context)
 		{
-			Id = StringRegistry.FindOrAddByName((string)Info.GetValue("Name", typeof(string)));
+			Id = GetUniqueStringRegistry().FindOrAddByName((string)Info.GetValue("Name", typeof(string)));
 		}
 
 		/// <summary>
@@ -340,7 +361,7 @@ namespace UnrealBuildTool
 		/// <returns></returns>
 		static private UnrealPlatformGroup FindOrAddByName(string Name)
 		{
-			return new UnrealPlatformGroup(StringRegistry.FindOrAddByName(Name));
+			return new UnrealPlatformGroup(GetUniqueStringRegistry().FindOrAddByName(Name));
 		}
 
 		/// <summary>
@@ -390,7 +411,7 @@ namespace UnrealBuildTool
 		/// <returns></returns>
 		public override string ToString()
 		{
-			return StringRegistry.GetStringForId(Id);
+			return GetUniqueStringRegistry().GetStringForId(Id);
 		}
 		/// <summary>
 		/// 
@@ -398,7 +419,7 @@ namespace UnrealBuildTool
 		/// <returns></returns>
 		public static UnrealPlatformGroup[] GetValidGroups()
 		{
-			return Array.ConvertAll(StringRegistry.GetStringIds(), x => new UnrealPlatformGroup(x));
+			return Array.ConvertAll(GetUniqueStringRegistry().GetStringIds(), x => new UnrealPlatformGroup(x));
 		}
 
 		/// <summary>
@@ -407,7 +428,7 @@ namespace UnrealBuildTool
 		/// <returns></returns>
 		public static string[] GetValidGroupNames()
 		{
-			return StringRegistry.GetStringNames();
+			return GetUniqueStringRegistry().GetStringNames();
 		}
 
 		/// <summary>
@@ -417,7 +438,7 @@ namespace UnrealBuildTool
 		/// <returns></returns>
 		public static bool IsValidName(string Name)
 		{
-			return StringRegistry.HasString(Name);
+			return GetUniqueStringRegistry().HasString(Name);
 		}
 
 
