Make sure libgnustl_shared is writable in intermediates folder.

#tests Android builds from AutoSDK can now be cleaned.

diff --git a/Engine/Source/Programs/UnrealBuildTool/Platform/Android/UEDeployAndroid.cs b/Engine/Source/Programs/UnrealBuildTool/Platform/Android/UEDeployAndroid.cs
--- a/Engine/Source/Programs/UnrealBuildTool/Platform/Android/UEDeployAndroid.cs
+++ b/Engine/Source/Programs/UnrealBuildTool/Platform/Android/UEDeployAndroid.cs
@@ -1221,6 +1221,9 @@ namespace UnrealBuildTool
 				SafeDeleteFile(FinalSTLSOName);
 				Directory.CreateDirectory(Path.GetDirectoryName(FinalSTLSOName));
 				File.Copy(SourceSTLSOName, FinalSTLSOName, true);
+
+				// make sure it's writable if the source was readonly (e.g. autosdks)
+				new FileInfo(FinalSTLSOName).IsReadOnly = false;
 			}
 		}
 
