Optimized FASTBuild's child process handling on Mac. New implementation of the Process class using NSTask and Grand Dispatch Central replaces the slow poll + sleep loop previously used.

Also, worked around a problem with compiling SQLiteCore plugin with FASTBuild by forcing its C file to compile locally.

diff --git a/Engine/Source/Programs/UnrealBuildTool/Platform/Mac/MacToolChain.cs b/Engine/Source/Programs/UnrealBuildTool/Platform/Mac/MacToolChain.cs
--- a/Engine/Source/Programs/UnrealBuildTool/Platform/Mac/MacToolChain.cs
+++ b/Engine/Source/Programs/UnrealBuildTool/Platform/Mac/MacToolChain.cs
@@ -556,7 +556,7 @@ namespace UnrealBuildTool
 				CompileAction.StatusDescription = Path.GetFileName(SourceFile.AbsolutePath);
 				CompileAction.bIsGCCCompiler = true;
 				// We're already distributing the command by execution on Mac.
-				CompileAction.bCanExecuteRemotely = true;
+				CompileAction.bCanExecuteRemotely = Extension != ".C";
 				CompileAction.bShouldOutputStatusDescription = true;
 				Actions.Add(CompileAction);
 			}
