Add linker flag for folding reducing binary output size

diff --git a/Engine/Source/Programs/UnrealBuildTool/Platform/Linux/LinuxToolChain.cs b/Engine/Source/Programs/UnrealBuildTool/Platform/Linux/LinuxToolChain.cs
--- a/Engine/Source/Programs/UnrealBuildTool/Platform/Linux/LinuxToolChain.cs
+++ b/Engine/Source/Programs/UnrealBuildTool/Platform/Linux/LinuxToolChain.cs
@@ -1016,6 +1016,15 @@ namespace UnrealBuildTool
 				}
 			}
 
+			if (LinkEnvironment.Configuration == CppConfiguration.Shipping)
+			{
+				Result += " -Wl,--icf=all"; // Enables ICF (Identical Code Folding). [all, safe] safe == fold functions that can be proven not to have their address taken.
+				if (!UsingLld(LinkEnvironment.Architecture))
+				{
+					Result += " -Wl,--icf-iterations=3";
+				}
+			}
+
 			// Profile Guided Optimization (PGO) and Link Time Optimization (LTO)
 			// Whether we actually can enable that is checked in CanUseAdvancedLinkerFeatures() earlier
 			if (LinkEnvironment.bPGOOptimize)
