Allow build targets to use config file attributes

diff --git a/Engine/Source/Programs/UnrealBuildTool/System/ConfigFileAttribute.cs b/Engine/Source/Programs/UnrealBuildTool/System/ConfigFileAttribute.cs
--- a/Engine/Source/Programs/UnrealBuildTool/System/ConfigFileAttribute.cs
+++ b/Engine/Source/Programs/UnrealBuildTool/System/ConfigFileAttribute.cs
@@ -1,4 +1,4 @@
-﻿// Copyright Epic Games, Inc. All Rights Reserved.
+// Copyright Epic Games, Inc. All Rights Reserved.
 
 using System;
 using System.Collections.Generic;
@@ -12,7 +12,7 @@ namespace UnrealBuildTool
 	/// Attribute indicating a value which should be populated from a UE4 .ini config file
 	/// </summary>
 	[AttributeUsage(AttributeTargets.Field)]
-	class ConfigFileAttribute : Attribute
+	public class ConfigFileAttribute : Attribute
 	{
 		/// <summary>
 		/// Name of the config hierarchy to read from
