Ignore empty AAR directories on Android (contributed by [REDACTED_PERSON])

#PR #6676
#ue4
#android

diff --git a/Engine/Source/Programs/UnrealBuildTool/Platform/Android/AndroidAARHandler.cs b/Engine/Source/Programs/UnrealBuildTool/Platform/Android/AndroidAARHandler.cs
--- a/Engine/Source/Programs/UnrealBuildTool/Platform/Android/AndroidAARHandler.cs
+++ b/Engine/Source/Programs/UnrealBuildTool/Platform/Android/AndroidAARHandler.cs
@@ -80,26 +80,33 @@ namespace UnrealBuildTool
 		/// <param name="SearchPattern">Search pattern to match</param>
 		public void AddRepositories(string RepositoryPath, string SearchPattern)
 		{
-			List<string> ToCheck = new List<string>();
-			ToCheck.Add(RepositoryPath);
-			while (ToCheck.Count > 0)
-			{
-				int LastIndex = ToCheck.Count - 1;
-				string CurrentDir = ToCheck[LastIndex];
-				ToCheck.RemoveAt(LastIndex);
-				foreach (string SearchPath in Directory.GetDirectories(CurrentDir))
+			if (Directory.Exists(RepositoryPath))
+			{
+				List<string> ToCheck = new List<string>();
+				ToCheck.Add(RepositoryPath);
+				while (ToCheck.Count > 0)
 				{
-					if (SearchPath.Contains(SearchPattern))
-					{
-						Log.TraceInformation("Added repository: {0}", SearchPath);
-						Repositories.Add(SearchPath);
-					}
-					else
+					int LastIndex = ToCheck.Count - 1;
+					string CurrentDir = ToCheck[LastIndex];
+					ToCheck.RemoveAt(LastIndex);
+					foreach (string SearchPath in Directory.GetDirectories(CurrentDir))
 					{
-						ToCheck.Add(SearchPath);
+						if (SearchPath.Contains(SearchPattern))
+						{
+							Log.TraceInformation("Added repository: {0}", SearchPath);
+							Repositories.Add(SearchPath);
+						}
+						else
+						{
+							ToCheck.Add(SearchPath);
+						}
 					}
 				}
 			}
+			else
+			{
+				Log.TraceInformation("AddRepositories: Directory {0} not found; ignored", RepositoryPath);
+			}
 		}
 
 		public void DumpAAR()
