Another fix for UE-88874 (exception generating project files from binary release).

#rnx

diff --git a/Engine/Source/Programs/UnrealBuildTool/ProjectFiles/VisualStudio/VCProject.cs b/Engine/Source/Programs/UnrealBuildTool/ProjectFiles/VisualStudio/VCProject.cs
--- a/Engine/Source/Programs/UnrealBuildTool/ProjectFiles/VisualStudio/VCProject.cs
+++ b/Engine/Source/Programs/UnrealBuildTool/ProjectFiles/VisualStudio/VCProject.cs
@@ -569,7 +569,7 @@ namespace UnrealBuildTool
 			CppStandardVersion Version = CppStandardVersion.Default;
 			foreach (ProjectConfigAndTargetCombination Combination in ProjectConfigAndTargetCombinations)
 			{
-				if (Combination.ProjectTarget.TargetRules != null && Combination.ProjectTarget.TargetRules.CppStandard > Version)
+				if (Combination.ProjectTarget != null && Combination.ProjectTarget.TargetRules != null && Combination.ProjectTarget.TargetRules.CppStandard > Version)
 				{
 					Version = Combination.ProjectTarget.TargetRules.CppStandard;
 				}
