Linux: Always use lld linker with clang v9 and above

diff --git a/Engine/Source/Programs/UnrealBuildTool/Platform/Linux/LinuxToolChain.cs b/Engine/Source/Programs/UnrealBuildTool/Platform/Linux/LinuxToolChain.cs
--- a/Engine/Source/Programs/UnrealBuildTool/Platform/Linux/LinuxToolChain.cs
+++ b/Engine/Source/Programs/UnrealBuildTool/Platform/Linux/LinuxToolChain.cs
@@ -911,7 +911,7 @@ namespace UnrealBuildTool
 		{
 			string Result = "";
 
-			if (UsingLld(LinkEnvironment.Architecture) && !LinkEnvironment.bIsBuildingDLL)
+			if (UsingLld(LinkEnvironment.Architecture) && (!LinkEnvironment.bIsBuildingDLL || (CompilerVersionMajor >= 9)))
 			{
 				Result += (BuildHostPlatform.Current.Platform == UnrealTargetPlatform.Win64) ? " -fuse-ld=lld.exe" : " -fuse-ld=lld";
 			}
@@ -1441,7 +1441,7 @@ namespace UnrealBuildTool
 
 		bool UsingLld(string Architecture)
 		{
-			return bUseLld && Architecture.StartsWith("x86_64");
+			return bUseLld && (Architecture.StartsWith("x86_64") || (CompilerVersionMajor >= 9));
 		}
 
 		/// <summary>
