Linux: Build with -gdb-index with Clang v9 and above

On my machine, this cuts symbol loading time (breaking at main) from 45 seconds
to 17 seconds (with gdb v8.3.1).

diff --git a/Engine/Source/Programs/UnrealBuildTool/Platform/Linux/LinuxToolChain.cs b/Engine/Source/Programs/UnrealBuildTool/Platform/Linux/LinuxToolChain.cs
--- a/Engine/Source/Programs/UnrealBuildTool/Platform/Linux/LinuxToolChain.cs
+++ b/Engine/Source/Programs/UnrealBuildTool/Platform/Linux/LinuxToolChain.cs
@@ -729,6 +729,13 @@ namespace UnrealBuildTool
 			{
 				// libdwarf (from elftoolchain 0.6.1) doesn't support DWARF4. If we need to go back to depending on elftoolchain revert this back to dwarf-3
 				Result += " -gdwarf-4";
+
+				if (CompilerVersionMajor >= 9)
+				{
+					// Generate .debug_pubnames and .debug_pubtypes sections in a format suitable for conversion into a
+					// GDB index. This option is only useful with a linker that can produce GDB index version 7.
+					Result += " -ggnu-pubnames";
+				}
 			}
 
 			// optimization level
@@ -968,6 +975,13 @@ namespace UnrealBuildTool
 				}
 			}
 
+			if (UsingLld(Architecture) && LinkEnvironment.bCreateDebugInfo && (CompilerVersionMajor >= 9))
+			{
+				// Generate .gdb_index section. On my machine, this cuts symbol loading time (breaking at main) from 45
+				// seconds to 17 seconds (with gdb v8.3.1).
+				Result += " -Wl,--gdb-index";
+			}
+
 			// RPATH for third party libs
 			Result += " -Wl,-rpath=${ORIGIN}";
 			Result += " -Wl,-rpath-link=${ORIGIN}";
