Always compile PushModel in all Editor Builds so they can be shared across projects.
Clients and Servers can still opt in via Build.cs files.

diff --git a/Engine/Source/Programs/UnrealBuildTool/Configuration/TargetRules.cs b/Engine/Source/Programs/UnrealBuildTool/Configuration/TargetRules.cs
--- a/Engine/Source/Programs/UnrealBuildTool/Configuration/TargetRules.cs
+++ b/Engine/Source/Programs/UnrealBuildTool/Configuration/TargetRules.cs
@@ -587,10 +587,23 @@ namespace UnrealBuildTool
 		private bool? bWithServerCodeOverride;
 
 		/// <summary>
-		/// When enabled, Push Model Networking will be used on the server. This can help reduce CPU overhead of networking, at the cost of more memory.
+		/// When enabled, Push Model Networking support will be compiled in.
+		/// This can help reduce CPU overhead of networking, at the cost of more memory.
+		/// Always enabled in editor builds.
 		/// </summary>
 		[RequiresUniqueBuildEnvironment]
-		public bool bWithPushModel = false;
+		public bool bWithPushModel
+		{
+			get
+			{
+				return bWithPushModelOverride ?? (bBuildEditor);
+			}
+			set
+			{
+				bWithPushModelOverride = value;
+			}
+		}
+		private bool? bWithPushModelOverride;
 
 		/// <summary>
 		/// Whether to include stats support even without the engine.
