Add bGdbIndexSection bool to LinuxToolChain

diff --git a/Engine/Source/Programs/UnrealBuildTool/Platform/Linux/LinuxToolChain.cs b/Engine/Source/Programs/UnrealBuildTool/Platform/Linux/LinuxToolChain.cs
--- a/Engine/Source/Programs/UnrealBuildTool/Platform/Linux/LinuxToolChain.cs
+++ b/Engine/Source/Programs/UnrealBuildTool/Platform/Linux/LinuxToolChain.cs
@@ -67,6 +67,9 @@ namespace UnrealBuildTool
 		/** Whether or not to preserve the portable symbol file produced by dump_syms */
 		bool bPreservePSYM = false;
 
+		/** Pass --gdb-index option to linker to generate .gdb_index section. */
+		bool bGdbIndexSection = true;
+
 		/** Platform SDK to use */
 		protected LinuxPlatformSDK PlatformSDK;
 
@@ -184,6 +187,9 @@ namespace UnrealBuildTool
 			// trust lld only for clang 5.x and above (FIXME: also find if present on the system?)
 			// NOTE: with early version you can run into errors like "failed to compute relocation:" and others
 			bUseLld = (CompilerVersionMajor >= 5);
+
+			// Add --gdb-index for Clang 9.0 and higher
+			bGdbIndexSection = (CompilerVersionMajor >= 9);
 		}
 
 		public LinuxToolChain(UnrealTargetPlatform InPlatform, string InArchitecture, LinuxPlatformSDK InSDK, bool InPreservePSYM = false, LinuxToolChainOptions InOptions = LinuxToolChainOptions.None)
@@ -727,10 +733,9 @@ namespace UnrealBuildTool
 			// bCreateDebugInfo is normally set for all configurations, including Shipping - this is needed to enable callstack in Shipping builds (proper resolution: UEPLAT-205, separate files with debug info)
 			if (CompileEnvironment.bCreateDebugInfo)
 			{
-				// libdwarf (from elftoolchain 0.6.1) doesn't support DWARF4. If we need to go back to depending on elftoolchain revert this back to dwarf-3
 				Result += " -gdwarf-4";
 
-				if (CompilerVersionMajor >= 9)
+				if (bGdbIndexSection)
 				{
 					// Generate .debug_pubnames and .debug_pubtypes sections in a format suitable for conversion into a
 					// GDB index. This option is only useful with a linker that can produce GDB index version 7.
@@ -975,7 +980,7 @@ namespace UnrealBuildTool
 				}
 			}
 
-			if (UsingLld(Architecture) && LinkEnvironment.bCreateDebugInfo && (CompilerVersionMajor >= 9))
+			if (UsingLld(Architecture) && LinkEnvironment.bCreateDebugInfo && bGdbIndexSection)
 			{
 				// Generate .gdb_index section. On my machine, this cuts symbol loading time (breaking at main) from 45
 				// seconds to 17 seconds (with gdb v8.3.1).
