Fixed incorrect dsymutil being used with 11.0.0 and higher.

diff --git a/Engine/Source/Programs/UnrealBuildTool/ToolChain/AppleToolChain.cs b/Engine/Source/Programs/UnrealBuildTool/ToolChain/AppleToolChain.cs
--- a/Engine/Source/Programs/UnrealBuildTool/ToolChain/AppleToolChain.cs
+++ b/Engine/Source/Programs/UnrealBuildTool/ToolChain/AppleToolChain.cs
@@ -219,11 +219,11 @@ namespace UnrealBuildTool
 						Log.TraceInformationOnce("Unable to parse version tokens: {0}", Tokens[3]);
 					}
 					else
-					{
-						Log.TraceInformationOnce("Parsed dsymutil version as {0}.{1}.{2}", Major, Minor, Patch);
-
-						if (Major < 10 || (Minor == 0 && Patch == 0))
+					{				
+						// 10.0.0 and earlier have an issue with large binaries
+						if (Major < 10 || (Major == 10 && Minor == 0 && Patch == 0))
 						{
+							Log.TraceInformationOnce("Parsed dsymutil version as {0}.{1}.{2}. Using local UE version if available.", Major, Minor, Patch);
 							bUseInstalledDsymutil = false;
 						}
 					}
