Add protected to bGdbIndexSection bool

diff --git a/Engine/Source/Programs/UnrealBuildTool/Platform/Linux/LinuxToolChain.cs b/Engine/Source/Programs/UnrealBuildTool/Platform/Linux/LinuxToolChain.cs
--- a/Engine/Source/Programs/UnrealBuildTool/Platform/Linux/LinuxToolChain.cs
+++ b/Engine/Source/Programs/UnrealBuildTool/Platform/Linux/LinuxToolChain.cs
@@ -68,7 +68,7 @@ namespace UnrealBuildTool
 		bool bPreservePSYM = false;
 
 		/** Pass --gdb-index option to linker to generate .gdb_index section. */
-		bool bGdbIndexSection = true;
+		protected bool bGdbIndexSection = true;
 
 		/** Platform SDK to use */
 		protected LinuxPlatformSDK PlatformSDK;
