PR #6656: Fix crash on exception from Steamworks module (Contributed by [REDACTED_PERSON])

diff --git a/Engine/Source/Programs/UnrealBuildTool/ProjectFiles/Rider/RiderProjectFile.cs b/Engine/Source/Programs/UnrealBuildTool/ProjectFiles/Rider/RiderProjectFile.cs
--- a/Engine/Source/Programs/UnrealBuildTool/ProjectFiles/Rider/RiderProjectFile.cs
+++ b/Engine/Source/Programs/UnrealBuildTool/ProjectFiles/Rider/RiderProjectFile.cs
@@ -72,9 +72,17 @@ namespace UnrealBuildTool
 						TargetDescriptor TargetDesc = new TargetDescriptor(ProjectTarget.UnrealProjectFilePath, ProjectTarget.Name,
 							BuildHostPlatform.Current.Platform, UnrealTargetConfiguration.Development,
 							DefaultArchitecture, Arguments);
-						UEBuildTarget BuildTarget = UEBuildTarget.Create(TargetDesc, false, false);
-						FileReference OutputFile = FileReference.Combine(TargetFolder, $"{ProjectName}.json");
-						fileToTarget.Add(Tuple.Create(OutputFile, BuildTarget));
+						try
+						{
+							UEBuildTarget BuildTarget = UEBuildTarget.Create(TargetDesc, false, false);
+						
+							FileReference OutputFile = FileReference.Combine(TargetFolder, $"{ProjectName}.json");
+							fileToTarget.Add(Tuple.Create(OutputFile, BuildTarget));
+						}
+						catch(Exception Ex)
+						{
+							Log.TraceWarning("Exception while generating include data for {0}: {1}", TargetDesc.Name, Ex.ToString());
+						}
 					}
 				}
 			}
