Fix Rider project file generation for projects with multiple targets of the same type.

diff --git a/Engine/Source/Programs/UnrealBuildTool/ProjectFiles/Rider/RiderProjectFileGenerator.cs b/Engine/Source/Programs/UnrealBuildTool/ProjectFiles/Rider/RiderProjectFileGenerator.cs
--- a/Engine/Source/Programs/UnrealBuildTool/ProjectFiles/Rider/RiderProjectFileGenerator.cs
+++ b/Engine/Source/Programs/UnrealBuildTool/ProjectFiles/Rider/RiderProjectFileGenerator.cs
@@ -137,6 +137,18 @@ namespace UnrealBuildTool
 			return true;
 		}
 
+		private FileReference GetRiderProjectLocation(TargetRules TargetRulesObject, DirectoryReference GameFolder, string GeneratedProjectName)
+		{
+			if (TargetRulesObject.Type != TargetType.Program && GameFolder != null)
+			{
+				return FileReference.Combine(GameFolder, "Intermediate", "ProjectFiles", GeneratedProjectName + ProjectFileExtension);
+			}
+			else
+			{
+				return GetProjectLocation(GeneratedProjectName);
+			}
+		}
+
 		private void AddProjectsForAllTargets(
 			PlatformProjectGeneratorCollection PlatformProjectGenerators,
 			List<FileReference> AllGames,
@@ -257,8 +269,7 @@ namespace UnrealBuildTool
 					if (GeneratedProjectName == null)
 					{
 						ProjectFile ExistingProjectFile;
-						if (ProjectFileMap.TryGetValue(GetProjectLocation(ProjectFileNameBase),
-							    out ExistingProjectFile) &&
+						if (ProjectFileMap.TryGetValue(GetRiderProjectLocation(TargetRulesObject, GameFolder, ProjectFileNameBase), out ExistingProjectFile) &&
 						    ExistingProjectFile.ProjectTargets.Any(x => x.TargetRules.Type == TargetRulesObject.Type))
 						{
 							GeneratedProjectName = TargetRulesObject.Name;
@@ -269,17 +280,7 @@ namespace UnrealBuildTool
 						}
 					}
 
-					FileReference ProjectFilePath;
-					if (TargetRulesObject.Type != TargetType.Program && GameFolder != null)
-					{
-						ProjectFilePath = FileReference.Combine(GameFolder, "Intermediate", "ProjectFiles",
-							GeneratedProjectName + ProjectFileExtension);
-					}
-					else
-					{
-						ProjectFilePath = GetProjectLocation(GeneratedProjectName);
-					}
-
+					FileReference ProjectFilePath = GetRiderProjectLocation(TargetRulesObject, GameFolder, GeneratedProjectName);
 					if (TargetRulesObject.Type == TargetType.Game || TargetRulesObject.Type == TargetType.Client ||
 					    TargetRulesObject.Type == TargetType.Server)
 					{
