UBT: Add a utility function to create an action that copies from one file to another.

#rnx

diff --git a/Engine/Source/Programs/UnrealBuildTool/Configuration/UEBuildTarget.cs b/Engine/Source/Programs/UnrealBuildTool/Configuration/UEBuildTarget.cs
--- a/Engine/Source/Programs/UnrealBuildTool/Configuration/UEBuildTarget.cs
+++ b/Engine/Source/Programs/UnrealBuildTool/Configuration/UEBuildTarget.cs
@@ -2067,24 +2067,7 @@ namespace UnrealBuildTool
 			FileItem SourceFileItem = FileItem.GetItemByFileReference(SourceFile);
 			FileItem TargetFileItem = FileItem.GetItemByFileReference(TargetFile);
 
-			Action CopyAction = new Action(ActionType.BuildProject);
-			CopyAction.CommandDescription = "Copy";
-			CopyAction.CommandPath = BuildHostPlatform.Current.Shell;
-			if(BuildHostPlatform.Current.ShellType == ShellType.Cmd)
-			{
-				CopyAction.CommandArguments = String.Format("/C \"copy /Y \"{0}\" \"{1}\" 1>nul\"", SourceFile, TargetFile);
-			}
-			else
-			{
-				CopyAction.CommandArguments = String.Format("-c 'cp -f \"{0}\" \"{1}\"'", SourceFile.FullName, TargetFile.FullName);
-			}
-			CopyAction.WorkingDirectory = UnrealBuildTool.EngineSourceDirectory;
-			CopyAction.PrerequisiteItems.Add(SourceFileItem);
-			CopyAction.ProducedItems.Add(TargetFileItem);
-			CopyAction.DeleteItems.Add(TargetFileItem);
-			CopyAction.StatusDescription = TargetFileItem.Location.GetFileName();
-			CopyAction.bCanExecuteRemotely = false;
-			Actions.Add(CopyAction);
+			Actions.Add(Action.CreateCopyAction(SourceFileItem, TargetFileItem));
 
 			return TargetFileItem;
 		}
diff --git a/Engine/Source/Programs/UnrealBuildTool/System/Action.cs b/Engine/Source/Programs/UnrealBuildTool/System/Action.cs
--- a/Engine/Source/Programs/UnrealBuildTool/System/Action.cs
+++ b/Engine/Source/Programs/UnrealBuildTool/System/Action.cs
@@ -366,6 +366,34 @@ namespace UnrealBuildTool
 			}
 		}
 
+		/// <summary>
+		/// Creates an action which copies a file from one location to another
+		/// </summary>
+		/// <param name="SourceFile">The source file location</param>
+		/// <param name="TargetFile">The target file location</param>
+		/// <returns>File item for the output file</returns>
+		public static Action CreateCopyAction(FileItem SourceFile, FileItem TargetFile)
+		{
+			Action CopyAction = new Action(ActionType.BuildProject);
+			CopyAction.CommandDescription = "Copy";
+			CopyAction.CommandPath = BuildHostPlatform.Current.Shell;
+			if (BuildHostPlatform.Current.ShellType == ShellType.Cmd)
+			{
+				CopyAction.CommandArguments = String.Format("/C \"copy /Y \"{0}\" \"{1}\" 1>nul\"", SourceFile.AbsolutePath, TargetFile.AbsolutePath);
+			}
+			else
+			{
+				CopyAction.CommandArguments = String.Format("-c 'cp -f \"{0}\" \"{1}\"'", SourceFile.AbsolutePath, TargetFile.AbsolutePath);
+			}
+			CopyAction.WorkingDirectory = UnrealBuildTool.EngineSourceDirectory;
+			CopyAction.PrerequisiteItems.Add(SourceFile);
+			CopyAction.ProducedItems.Add(TargetFile);
+			CopyAction.DeleteItems.Add(TargetFile);
+			CopyAction.StatusDescription = TargetFile.Location.GetFileName();
+			CopyAction.bCanExecuteRemotely = false;
+			return CopyAction;
+		}
+
 		/// <summary>
 		/// Creates an action which calls UBT recursively
 		/// </summary>
