Allow single file compile to work with generated cpp files.

diff --git a/Engine/Source/Programs/UnrealBuildTool/Configuration/UEBuildModule.cs b/Engine/Source/Programs/UnrealBuildTool/Configuration/UEBuildModule.cs
--- a/Engine/Source/Programs/UnrealBuildTool/Configuration/UEBuildModule.cs
+++ b/Engine/Source/Programs/UnrealBuildTool/Configuration/UEBuildModule.cs
@@ -243,6 +243,16 @@ namespace UnrealBuildTool
 			ModuleDirectories = MergedDirectories.ToArray();
 		}
 
+		/// <summary>
+		/// Determines if a file is part of the given module
+		/// </summary>
+		/// <param name="Location">Path to the file</param>
+		/// <returns>True if the file is part of this module</returns>
+		public virtual bool ContainsFile(FileReference Location)
+		{
+			return Location.IsUnderDirectory(ModuleDirectory);
+		}
+
 		/// <summary>
 		/// Returns a list of this module's dependencies.
 		/// </summary>
diff --git a/Engine/Source/Programs/UnrealBuildTool/Configuration/UEBuildModuleCPP.cs b/Engine/Source/Programs/UnrealBuildTool/Configuration/UEBuildModuleCPP.cs
--- a/Engine/Source/Programs/UnrealBuildTool/Configuration/UEBuildModuleCPP.cs
+++ b/Engine/Source/Programs/UnrealBuildTool/Configuration/UEBuildModuleCPP.cs
@@ -158,6 +158,24 @@ namespace UnrealBuildTool
 			AddDefaultIncludePaths();
 		}
 
+		/// <summary>
+		/// Determines if a file is part of the given module
+		/// </summary>
+		/// <param name="Location">Path to the file</param>
+		/// <returns>True if the file is part of this module</returns>
+		public override bool ContainsFile(FileReference Location)
+		{
+			if (base.ContainsFile(Location))
+			{
+				return true;
+			}
+			if (GeneratedCodeDirectory != null && Location.IsUnderDirectory(GeneratedCodeDirectory))
+			{
+				return true;
+			}
+			return false;
+		}
+
 		/// <summary>
 		/// Add the default include paths for this module to its settings
 		/// </summary>
@@ -294,7 +312,7 @@ namespace UnrealBuildTool
 				InputFiles.CCFiles.RemoveAll(x => x.Location != SingleFileToCompile);
 				InputFiles.CFiles.RemoveAll(x => x.Location != SingleFileToCompile);
 
-				if(InputFiles.CPPFiles.Count == 0 && InputFiles.CCFiles.Count == 0 && InputFiles.CFiles.Count == 0)
+				if(InputFiles.CPPFiles.Count == 0 && InputFiles.CCFiles.Count == 0 && InputFiles.CFiles.Count == 0 && !ContainsFile(SingleFileToCompile))
 				{
 					return new List<FileItem>();
 				}
@@ -397,7 +415,7 @@ namespace UnrealBuildTool
 			}
 
 			// Compile all the generated CPP files
-			if (GeneratedCodeWildcard != null && !CompileEnvironment.bHackHeaderGenerator && SingleFileToCompile == null)
+			if (GeneratedCodeWildcard != null && !CompileEnvironment.bHackHeaderGenerator)
 			{
 				string[] GeneratedFiles = Directory.GetFiles(Path.GetDirectoryName(GeneratedCodeWildcard), Path.GetFileName(GeneratedCodeWildcard));
 				if(GeneratedFiles.Length > 0)
@@ -428,9 +446,10 @@ namespace UnrealBuildTool
 					foreach (string GeneratedFilename in GeneratedFiles)
 					{
 						FileItem GeneratedCppFileItem = FileItem.GetItemByPath(GeneratedFilename);
-
-						// @todo ubtmake: Check for ALL other places where we might be injecting .cpp or .rc files for compiling without caching CachedCPPIncludeInfo first (anything platform specific?)
-						GeneratedFileItems.Add(GeneratedCppFileItem);
+						if (SingleFileToCompile == null || GeneratedCppFileItem.Location == SingleFileToCompile)
+						{
+							GeneratedFileItems.Add(GeneratedCppFileItem);
+						}
 					}
 
 					if (bModuleUsesUnityBuild)
diff --git a/Engine/Source/Programs/UnrealBuildTool/Configuration/UEBuildTarget.cs b/Engine/Source/Programs/UnrealBuildTool/Configuration/UEBuildTarget.cs
--- a/Engine/Source/Programs/UnrealBuildTool/Configuration/UEBuildTarget.cs
+++ b/Engine/Source/Programs/UnrealBuildTool/Configuration/UEBuildTarget.cs
@@ -1715,7 +1715,7 @@ namespace UnrealBuildTool
 			List<UEBuildBinary> BuildBinaries = Binaries;
 			if (SingleFileToCompile != null)
 			{
-				BuildBinaries = Binaries.Where(x => x.Modules.Any(y => SingleFileToCompile.IsUnderDirectory(y.ModuleDirectory))).ToList();
+				BuildBinaries = Binaries.Where(x => x.Modules.Any(y => y.ContainsFile(SingleFileToCompile))).ToList();
 				if (BuildBinaries.Count == 0)
 				{
 					throw new BuildException("Couldn't find any module containing {0} in {1}.", SingleFileToCompile, TargetName);
