UBT: Fix error running UBT with -JsonExport if the target instantiates (but decides not to link) a module.

diff --git a/Engine/Source/Programs/UnrealBuildTool/Configuration/UEBuildTarget.cs b/Engine/Source/Programs/UnrealBuildTool/Configuration/UEBuildTarget.cs
--- a/Engine/Source/Programs/UnrealBuildTool/Configuration/UEBuildTarget.cs
+++ b/Engine/Source/Programs/UnrealBuildTool/Configuration/UEBuildTarget.cs
@@ -2188,11 +2188,14 @@ namespace UnrealBuildTool
 				Writer.WriteArrayEnd();
 
 				Writer.WriteObjectStart("Modules");
-				foreach(UEBuildModule Module in Modules.Values)
+				foreach (UEBuildBinary Binary in Binaries)
 				{
-					Writer.WriteObjectStart(Module.Name);
-					Module.ExportJson(Module.Binary?.OutputDir, GetExecutableDir(), Writer);
-					Writer.WriteObjectEnd();
+					foreach (UEBuildModule Module in Binary.Modules)
+					{
+						Writer.WriteObjectStart(Module.Name);
+						Module.ExportJson(Binary.OutputDir, GetExecutableDir(), Writer);
+						Writer.WriteObjectEnd();
+					}
 				}
 				Writer.WriteObjectEnd();
 
