UBT: Fix log file being written to Engine directory, even in installed builds.

UE-87654

diff --git a/Engine/Source/Programs/UnrealBuildTool/Modes/BuildMode.cs b/Engine/Source/Programs/UnrealBuildTool/Modes/BuildMode.cs
--- a/Engine/Source/Programs/UnrealBuildTool/Modes/BuildMode.cs
+++ b/Engine/Source/Programs/UnrealBuildTool/Modes/BuildMode.cs
@@ -44,7 +44,7 @@ namespace UnrealBuildTool
 		/// Specifies the file to use for logging.
 		/// </summary>
 		[XmlConfigFile(Category = "BuildConfiguration")]
-		public string BaseLogFileName = "../Programs/UnrealBuildTool/Log.txt";
+		public string BaseLogFileName;
 
 		/// <summary>
 		/// Whether to skip checking for files identified by the junk manifest.
@@ -97,8 +97,14 @@ namespace UnrealBuildTool
 			// Read the XML configuration files
 			XmlConfig.ApplyTo(this);
 
+			// Fixup the log path if it wasn't overridden by a config file
+			if (BaseLogFileName == null)
+			{
+				BaseLogFileName = FileReference.Combine(UnrealBuildTool.EngineProgramSavedDirectory, "UnrealBuildTool", "Log.txt").FullName;
+			}
+
 			// Create the log file, and flush the startup listener to it
-			if(!Arguments.HasOption("-NoLog") && !Log.HasFileWriter())
+			if (!Arguments.HasOption("-NoLog") && !Log.HasFileWriter())
 			{
 				FileReference LogFile = new FileReference(BaseLogFileName);
 				foreach(string LogSuffix in Arguments.GetValues("-LogSuffix="))
diff --git a/Engine/Source/Programs/UnrealBuildTool/UnrealBuildTool.cs b/Engine/Source/Programs/UnrealBuildTool/UnrealBuildTool.cs
--- a/Engine/Source/Programs/UnrealBuildTool/UnrealBuildTool.cs
+++ b/Engine/Source/Programs/UnrealBuildTool/UnrealBuildTool.cs
@@ -39,6 +39,11 @@ namespace UnrealBuildTool
 		/// </summary>
 		static FileReference InstalledProjectFile;
 
+		/// <summary>
+		/// Directory for saved application settings (typically Engine/Programs)
+		/// </summary>
+		static DirectoryReference CachedEngineProgramSavedDirectory;
+
 		/// <summary>
 		/// The path to UBT
 		/// </summary>
@@ -109,6 +114,21 @@ namespace UnrealBuildTool
 		/// </summary>
 		public static readonly DirectoryReference EnterpriseIntermediateDirectory = DirectoryReference.Combine(EnterpriseDirectory, "Intermediate");
 
+		/// <summary>
+		/// The engine programs directory
+		/// </summary>
+		public static DirectoryReference EngineProgramSavedDirectory
+		{
+			get
+			{
+				if (CachedEngineProgramSavedDirectory == null)
+				{
+					CachedEngineProgramSavedDirectory = Utils.GetUserSettingDirectory() ?? DirectoryReference.Combine(EngineDirectory, "Programs");
+				}
+				return CachedEngineProgramSavedDirectory;
+			}
+		}
+
 		/// <summary>
 		/// Returns the root location of platform extensions within the given project
 		/// </summary>
