UBT: Always embed an application manifest during linking, rather than relying on the resource file to include one. This ensures that projects that have a custom resource file for icons, versioning metadata, etc... still operate correctly with manifested OS versions.

To embed a custom manifest, set the WindowsPlatform.ManifestFile to a vaild path in a *.Target.cs file. To disable the embedding of the default manifest, set it to null.

diff --git a/Engine/Build/Windows/Resources/Default-Win32-Debug.manifest b/Engine/Build/Windows/Resources/Default-Win32-Debug.manifest
deleted file mode 100644
diff --git a/Engine/Build/Windows/Resources/Default-Win64-Debug.manifest b/Engine/Build/Windows/Resources/Default-Win64-Debug.manifest
deleted file mode 100644
diff --git a/Engine/Build/Windows/Resources/Default.rc2 b/Engine/Build/Windows/Resources/Default.rc2
--- a/Engine/Build/Windows/Resources/Default.rc2
+++ b/Engine/Build/Windows/Resources/Default.rc2
@@ -95,31 +95,6 @@ END
 IDICON_UE4Game          ICON                    BUILD_ICON_FILE_NAME
 #endif
 
-/////////////////////////////////////////////////////////////////////////////
-//
-// Application Manifest
-//
-
-#if _DEBUG
-
-// Debug manifest
-#if _WIN64
-1						RT_MANIFEST				"Default-Win64-Debug.manifest"
-#else
-1						RT_MANIFEST				"Default-Win32-Debug.manifest"
-#endif
-
-#else
-
-// Release manifest
-#if _WIN64
-1						RT_MANIFEST				"Default-Win64.manifest"
-#else
-1						RT_MANIFEST				"Default-Win32.manifest"
-#endif
-
-#endif
-
 #endif    // English (United States) resources
 /////////////////////////////////////////////////////////////////////////////
 
diff --git a/Engine/Source/Programs/UnrealBuildTool/Platform/Windows/UEBuildWindows.cs b/Engine/Source/Programs/UnrealBuildTool/Platform/Windows/UEBuildWindows.cs
--- a/Engine/Source/Programs/UnrealBuildTool/Platform/Windows/UEBuildWindows.cs
+++ b/Engine/Source/Programs/UnrealBuildTool/Platform/Windows/UEBuildWindows.cs
@@ -199,6 +199,12 @@ namespace UnrealBuildTool
 		/// </summary>
 		public string ModuleDefinitionFile;
 
+		/// <summary>
+		/// Specifies the path to a manifest file for the linker to embed. Defaults to the manifest in Engine/Build/Windows/Resources. Can be assigned to null
+		/// if the target wants to specify its own manifest.
+		/// </summary>
+		public string ManifestFile;
+
 		/// <summary>
 		/// Enables strict standard conformance mode (/permissive-) in VS2017+.
 		/// </summary>
@@ -363,6 +369,8 @@ namespace UnrealBuildTool
 		internal WindowsTargetRules(TargetRules Target)
 		{
 			this.Target = Target;
+
+			ManifestFile = FileReference.Combine(UnrealBuildTool.EngineDirectory, "Build", "Windows", "Resources", String.Format("Default-{0}.manifest", Target.Platform)).FullName;
 		}
 	}
 
@@ -462,6 +470,11 @@ namespace UnrealBuildTool
 			get { return Inner.ModuleDefinitionFile; }
 		}
 
+		public string ManifestFile
+		{
+			get { return Inner.ManifestFile; }
+		}
+
 		public bool bNeedsLegacyStdioDefinitionsLib
 		{
 			get { return Inner.bNeedsLegacyStdioDefinitionsLib; }
diff --git a/Engine/Source/Programs/UnrealBuildTool/Platform/Windows/VCEnvironment.cs b/Engine/Source/Programs/UnrealBuildTool/Platform/Windows/VCEnvironment.cs
--- a/Engine/Source/Programs/UnrealBuildTool/Platform/Windows/VCEnvironment.cs
+++ b/Engine/Source/Programs/UnrealBuildTool/Platform/Windows/VCEnvironment.cs
@@ -162,6 +162,12 @@ namespace UnrealBuildTool
 				AddDirectoryToPath(GetVCToolPath(ToolChain, ToolChainDir, WindowsArchitecture.x86));
 				AddDirectoryToPath(GetVCToolPath(ToolChain, ToolChainDir, WindowsArchitecture.x64));
 			}
+
+			// Add the Windows SDK directory to the path too, for mt.exe.
+			if (WindowsSdkVersion >= new VersionNumber(10))
+			{
+				AddDirectoryToPath(DirectoryReference.Combine(WindowsSdkDir, "bin", WindowsSdkVersion.ToString(), Architecture.ToString()));
+			}
 		}
 
 		/// <summary>
diff --git a/Engine/Source/Programs/UnrealBuildTool/Platform/Windows/VCToolChain.cs b/Engine/Source/Programs/UnrealBuildTool/Platform/Windows/VCToolChain.cs
--- a/Engine/Source/Programs/UnrealBuildTool/Platform/Windows/VCToolChain.cs
+++ b/Engine/Source/Programs/UnrealBuildTool/Platform/Windows/VCToolChain.cs
@@ -708,7 +708,15 @@ namespace UnrealBuildTool
 			}
 
 			// Don't create a side-by-side manifest file for the executable.
-			Arguments.Add("/MANIFEST:NO");
+			if (Target.WindowsPlatform.ManifestFile == null)
+			{
+				Arguments.Add("/MANIFEST:NO");
+			}
+			else
+			{
+				Arguments.Add("/MANIFEST:EMBED");
+				Arguments.Add(String.Format("/MANIFESTINPUT:{0}", Utils.MakePathSafeToUseWithCommandLine(Target.WindowsPlatform.ManifestFile)));
+			}
 
 			// Prevents the linker from displaying its logo for each invocation.
 			Arguments.Add("/NOLOGO");
