Prevent writing compiled assemblies to installed engine directory. Write them to the AppData folder instead.

diff --git a/Engine/Source/Programs/UnrealBuildTool/System/RulesCompiler.cs b/Engine/Source/Programs/UnrealBuildTool/System/RulesCompiler.cs
--- a/Engine/Source/Programs/UnrealBuildTool/System/RulesCompiler.cs
+++ b/Engine/Source/Programs/UnrealBuildTool/System/RulesCompiler.cs
@@ -471,9 +471,21 @@ namespace UnrealBuildTool
 				FindModuleRulesForPlugins(Plugins, PluginsModuleContext, ModuleFileToContext);
 			}
 
+			// Get the path to store any generated assemblies
+			DirectoryReference AssemblyDir = RootDirectories[0];
+			if (UnrealBuildTool.IsFileInstalled(FileReference.Combine(AssemblyDir, AssemblyPrefix)))
+			{
+				DirectoryReference UserDir = Utils.GetUserSettingDirectory();
+				if (UserDir != null)
+				{
+					ReadOnlyBuildVersion Version = ReadOnlyBuildVersion.Current;
+					AssemblyDir = DirectoryReference.Combine(UserDir, "UnrealEngine", String.Format("{0}.{1}", Version.MajorVersion, Version.MinorVersion));
+				}
+			}
+
 			// Create the assembly
-			FileReference EngineAssemblyFileName = FileReference.Combine(RootDirectories[0], "Intermediate", "Build", "BuildRules", AssemblyPrefix + "Rules" + FrameworkAssemblyExtension);
-			RulesAssembly EngineAssembly = new RulesAssembly(Scope, RootDirectories[0], Plugins, ModuleFileToContext, new List<FileReference>(), EngineAssemblyFileName, bContainsEngineModules: true, DefaultBuildSettings: BuildSettingsVersion.Latest, bReadOnly: bReadOnly, bSkipCompile: bSkipCompile, Parent: Parent);
+			FileReference EngineAssemblyFileName = FileReference.Combine(AssemblyDir, "Intermediate", "Build", "BuildRules", AssemblyPrefix + "Rules" + FrameworkAssemblyExtension);
+			RulesAssembly EngineAssembly = new RulesAssembly(Scope, AssemblyDir, Plugins, ModuleFileToContext, new List<FileReference>(), EngineAssemblyFileName, bContainsEngineModules: true, DefaultBuildSettings: BuildSettingsVersion.Latest, bReadOnly: bReadOnly, bSkipCompile: bSkipCompile, Parent: Parent);
 
 			List<FileReference> ProgramTargetFiles = new List<FileReference>();
 			Dictionary<FileReference, ModuleRulesContext> ProgramModuleFiles = new Dictionary<FileReference, ModuleRulesContext>();
@@ -499,8 +511,8 @@ namespace UnrealBuildTool
 			}
 
 			// Create a path to the assembly that we'll either load or compile
-			FileReference ProgramAssemblyFileName = FileReference.Combine(RootDirectories[0], "Intermediate", "Build", "BuildRules", AssemblyPrefix + "ProgramRules" + FrameworkAssemblyExtension);
-			RulesAssembly ProgramAssembly = new RulesAssembly(ProgramsScope, RootDirectories[0], new List<PluginInfo>().AsReadOnly(), ProgramModuleFiles, ProgramTargetFiles, ProgramAssemblyFileName, bContainsEngineModules: false, DefaultBuildSettings: BuildSettingsVersion.Latest, bReadOnly: bReadOnly, bSkipCompile: bSkipCompile, Parent: EngineAssembly);
+			FileReference ProgramAssemblyFileName = FileReference.Combine(AssemblyDir, "Intermediate", "Build", "BuildRules", AssemblyPrefix + "ProgramRules" + FrameworkAssemblyExtension);
+			RulesAssembly ProgramAssembly = new RulesAssembly(ProgramsScope, AssemblyDir, new List<PluginInfo>().AsReadOnly(), ProgramModuleFiles, ProgramTargetFiles, ProgramAssemblyFileName, bContainsEngineModules: false, DefaultBuildSettings: BuildSettingsVersion.Latest, bReadOnly: bReadOnly, bSkipCompile: bSkipCompile, Parent: EngineAssembly);
 
 			// Return the combined assembly
 			return ProgramAssembly;
