Fix UBT always writing logs to user folder, rather than just in installed builds.

#rnx

diff --git a/Engine/Source/Programs/UnrealBuildTool/UnrealBuildTool.cs b/Engine/Source/Programs/UnrealBuildTool/UnrealBuildTool.cs
--- a/Engine/Source/Programs/UnrealBuildTool/UnrealBuildTool.cs
+++ b/Engine/Source/Programs/UnrealBuildTool/UnrealBuildTool.cs
@@ -123,7 +123,14 @@ namespace UnrealBuildTool
 			{
 				if (CachedEngineProgramSavedDirectory == null)
 				{
-					CachedEngineProgramSavedDirectory = Utils.GetUserSettingDirectory() ?? DirectoryReference.Combine(EngineDirectory, "Programs");
+					if (IsEngineInstalled())
+					{
+						CachedEngineProgramSavedDirectory = Utils.GetUserSettingDirectory() ?? DirectoryReference.Combine(EngineDirectory, "Programs");
+					}
+					else
+					{
+						CachedEngineProgramSavedDirectory = DirectoryReference.Combine(EngineDirectory, "Programs");
+					}
 				}
 				return CachedEngineProgramSavedDirectory;
 			}
