UBT: Output a warning if any intermediate path is > 200 characters under the UE root directory.

diff --git a/Engine/Source/Programs/UnrealBuildTool/Configuration/BuildConfiguration.cs b/Engine/Source/Programs/UnrealBuildTool/Configuration/BuildConfiguration.cs
--- a/Engine/Source/Programs/UnrealBuildTool/Configuration/BuildConfiguration.cs
+++ b/Engine/Source/Programs/UnrealBuildTool/Configuration/BuildConfiguration.cs
@@ -146,5 +146,10 @@ namespace UnrealBuildTool
 		/// </summary>
 		[CommandLine("-SkipRulesCompile")]
 		public bool bSkipRulesCompile = false;
+
+		/// <summary>
+		/// Maximum length of a path relative to the root directory. Used on Windows to ensure paths are portable between machines. Defaults to off.
+		/// </summary>
+		public int MaxNestedPathLength = 170;
 	}
 }
diff --git a/Engine/Source/Programs/UnrealBuildTool/Modes/BuildMode.cs b/Engine/Source/Programs/UnrealBuildTool/Modes/BuildMode.cs
--- a/Engine/Source/Programs/UnrealBuildTool/Modes/BuildMode.cs
+++ b/Engine/Source/Programs/UnrealBuildTool/Modes/BuildMode.cs
@@ -243,6 +243,12 @@ namespace UnrealBuildTool
 				// Make sure that none of the actions conflict with any other (producing output files differently, etc...)
 				ActionGraph.CheckForConflicts(Makefiles.SelectMany(x => x.Actions));
 
+				// Check we don't exceed the nominal max path length
+				using (Timeline.ScopeEvent("ActionGraph.CheckPathLengths"))
+				{
+					ActionGraph.CheckPathLengths(BuildConfiguration, Makefiles.SelectMany(x => x.Actions));
+				}
+
 				// Find all the actions to be executed
 				HashSet<Action>[] ActionsToExecute = new HashSet<Action>[TargetDescriptors.Count];
 				for(int TargetIdx = 0; TargetIdx < TargetDescriptors.Count; TargetIdx++)
@@ -265,7 +271,7 @@ namespace UnrealBuildTool
 				ActionGraph.Link(MergedActionsToExecute);
 
 				// Make sure the appropriate executor is selected
-				foreach(TargetDescriptor TargetDescriptor in TargetDescriptors)
+				foreach (TargetDescriptor TargetDescriptor in TargetDescriptors)
 				{
 					UEBuildPlatform BuildPlatform = UEBuildPlatform.GetBuildPlatform(TargetDescriptor.Platform);
 					BuildConfiguration.bAllowXGE &= BuildPlatform.CanUseXGE();
diff --git a/Engine/Source/Programs/UnrealBuildTool/System/ActionGraph.cs b/Engine/Source/Programs/UnrealBuildTool/System/ActionGraph.cs
--- a/Engine/Source/Programs/UnrealBuildTool/System/ActionGraph.cs
+++ b/Engine/Source/Programs/UnrealBuildTool/System/ActionGraph.cs
@@ -108,6 +108,42 @@ namespace UnrealBuildTool
 			return ActionsToExecute;
 		}
 
+		/// <summary>
+		/// Checks that there aren't any intermediate files longer than the max allowed path length
+		/// </summary>
+		/// <param name="BuildConfiguration">The build configuration</param>
+		/// <param name="Actions">List of actions in the graph</param>
+		public static void CheckPathLengths(BuildConfiguration BuildConfiguration, IEnumerable<Action> Actions)
+		{
+			if (BuildHostPlatform.Current.Platform == UnrealTargetPlatform.Win64)
+			{
+				List<FileReference> Paths = new List<FileReference>();
+				foreach (Action Action in Actions)
+				{
+					foreach (FileItem ProducedItem in Action.ProducedItems)
+					{
+						if (ProducedItem.Location.FullName.Length > UnrealBuildTool.RootDirectory.FullName.Length + BuildConfiguration.MaxNestedPathLength && ProducedItem.Location.IsUnderDirectory(UnrealBuildTool.RootDirectory))
+						{
+							Paths.Add(ProducedItem.Location);
+						}
+					}
+				}
+
+				if (Paths.Count > 0)
+				{
+					StringBuilder Message = new StringBuilder();
+					Message.AppendFormat("Detected paths more than {0} characters below UE root directory. This may cause portability issues due to the 260 character maximum path length on Windows:\n", BuildConfiguration.MaxNestedPathLength);
+					foreach (FileReference Path in Paths)
+					{
+						string RelativePath = Path.MakeRelativeTo(UnrealBuildTool.RootDirectory);
+						Message.AppendFormat("\n[{0} characters] {1}", RelativePath.Length, RelativePath);
+					}
+					Message.AppendFormat("\n\nConsider setting {0} = ... in module *.Build.cs files to use alternative names for intermediate paths.", nameof(ModuleRules.ShortName));
+					Log.TraceWarning(Message.ToString());
+				}
+			}
+		}
+
 		/// <summary>
 		/// Executes a list of actions.
 		/// </summary>
