Output a warning if the root path is longer than 50 characters.

diff --git a/Engine/Source/Programs/UnrealBuildTool/Configuration/BuildConfiguration.cs b/Engine/Source/Programs/UnrealBuildTool/Configuration/BuildConfiguration.cs
--- a/Engine/Source/Programs/UnrealBuildTool/Configuration/BuildConfiguration.cs
+++ b/Engine/Source/Programs/UnrealBuildTool/Configuration/BuildConfiguration.cs
@@ -147,9 +147,16 @@ namespace UnrealBuildTool
 		[CommandLine("-SkipRulesCompile")]
 		public bool bSkipRulesCompile = false;
 
+		/// <summary>
+		/// Maximum recommended root path length.
+		/// </summary>
+		[XmlConfigFile(Category = "WindowsPlatform")]
+		public int MaxRootPathLength = 50;
+
 		/// <summary>
 		/// Maximum length of a path relative to the root directory. Used on Windows to ensure paths are portable between machines. Defaults to off.
 		/// </summary>
-		public int MaxNestedPathLength = 170;
+		[XmlConfigFile(Category = "WindowsPlatform")]
+		public int MaxNestedPathLength = 200;
 	}
 }
diff --git a/Engine/Source/Programs/UnrealBuildTool/Modes/BuildMode.cs b/Engine/Source/Programs/UnrealBuildTool/Modes/BuildMode.cs
--- a/Engine/Source/Programs/UnrealBuildTool/Modes/BuildMode.cs
+++ b/Engine/Source/Programs/UnrealBuildTool/Modes/BuildMode.cs
@@ -121,7 +121,13 @@ namespace UnrealBuildTool
 			BuildConfiguration BuildConfiguration = new BuildConfiguration();
 			XmlConfig.ApplyTo(BuildConfiguration);
 			Arguments.ApplyTo(BuildConfiguration);
-			
+
+			// Check the root path length isn't too long
+			if (BuildHostPlatform.Current.Platform == UnrealTargetPlatform.Win64 && UnrealBuildTool.RootDirectory.FullName.Length > BuildConfiguration.MaxRootPathLength)
+			{
+				Log.TraceWarning("Running from a path with a long directory name (\"{0}\" = {1} characters). Root paths shorter than {2} characters are recommended to avoid exceeding maximum path lengths on Windows.", UnrealBuildTool.RootDirectory, UnrealBuildTool.RootDirectory.FullName.Length, BuildConfiguration.MaxRootPathLength);
+			}
+
 			// now that we know the available platforms, we can delete other platforms' junk. if we're only building specific modules from the editor, don't touch anything else (it may be in use).
 			if (!bIgnoreJunk && !UnrealBuildTool.IsEngineInstalled())
 			{
