UBT: Remove read-only files when cleaning a target.

diff --git a/Engine/Source/Programs/UnrealBuildTool/Modes/CleanMode.cs b/Engine/Source/Programs/UnrealBuildTool/Modes/CleanMode.cs
--- a/Engine/Source/Programs/UnrealBuildTool/Modes/CleanMode.cs
+++ b/Engine/Source/Programs/UnrealBuildTool/Modes/CleanMode.cs
@@ -185,7 +185,7 @@ namespace UnrealBuildTool
 					Log.TraceVerbose("    Deleting {0}{1}...", DirectoryToDelete, Path.DirectorySeparatorChar);
 					try
 					{
-						DirectoryReference.Delete(DirectoryToDelete, true);
+						FileUtils.ForceDeleteDirectory(DirectoryToDelete);
 					}
 					catch (Exception Ex)
 					{
@@ -201,7 +201,7 @@ namespace UnrealBuildTool
 					Log.TraceVerbose("    Deleting " + FileToDelete);
 					try
 					{
-						FileReference.Delete(FileToDelete);
+						FileUtils.ForceDeleteFile(FileToDelete);
 					}
 					catch (Exception Ex)
 					{
