Fix incorrect force included files in intellisense data.

diff --git a/Engine/Source/Programs/UnrealBuildTool/Configuration/UEBuildModuleCPP.cs b/Engine/Source/Programs/UnrealBuildTool/Configuration/UEBuildModuleCPP.cs
--- a/Engine/Source/Programs/UnrealBuildTool/Configuration/UEBuildModuleCPP.cs
+++ b/Engine/Source/Programs/UnrealBuildTool/Configuration/UEBuildModuleCPP.cs
@@ -960,7 +960,7 @@ namespace UnrealBuildTool
 				}
 
 				// Try to find a suitable shared PCH for this module
-				if (CompileEnvironment.PrecompiledHeaderFile == null && CompileEnvironment.SharedPCHs.Count > 0 && !CompileEnvironment.bIsBuildingLibrary && Rules.PCHUsage != ModuleRules.PCHUsageMode.NoSharedPCHs)
+				if (CompileEnvironment.PrecompiledHeaderIncludeFilename == null && CompileEnvironment.SharedPCHs.Count > 0 && !CompileEnvironment.bIsBuildingLibrary && Rules.PCHUsage != ModuleRules.PCHUsageMode.NoSharedPCHs)
 				{
 					// Find all the dependencies of this module
 					HashSet<UEBuildModule> ReferencedModules = new HashSet<UEBuildModule>();
@@ -1201,6 +1201,11 @@ namespace UnrealBuildTool
 
 		public CppCompileEnvironment CreateCompileEnvironmentForIntellisense(ReadOnlyTargetRules Target, CppCompileEnvironment BaseCompileEnvironment)
 		{
+			if(Name == "Engine")
+			{
+				Console.Write("");
+			}
+
 			CppCompileEnvironment CompileEnvironment = CreateModuleCompileEnvironment(Target, BaseCompileEnvironment);
 			CompileEnvironment = SetupPrecompiledHeaders(Target, null, CompileEnvironment, new List<FileItem>(), new List<Action>());
 			CreateHeaderForDefinitions(CompileEnvironment, IntermediateDirectory, null);
