Don't allow UBT to try and compile engine modules when the engine is already running.

diff --git a/Engine/Source/Developer/DesktopPlatform/Private/DesktopPlatformBase.cpp b/Engine/Source/Developer/DesktopPlatform/Private/DesktopPlatformBase.cpp
--- a/Engine/Source/Developer/DesktopPlatform/Private/DesktopPlatformBase.cpp
+++ b/Engine/Source/Developer/DesktopPlatform/Private/DesktopPlatformBase.cpp
@@ -576,7 +576,7 @@ bool FDesktopPlatformBase::CompileGameProject(const FString& RootDir, const FStr
 	}
 
 	// Append any other options
-	Arguments += " -Progress -NoHotReloadFromIDE";
+	Arguments += " -Progress -NoEngineChanges -NoHotReloadFromIDE";
 
 	// Run UBT
 	bool bResult = RunUnrealBuildTool(LOCTEXT("CompilingProject", "Compiling project..."), RootDir, Arguments, Warn);
diff --git a/Engine/Source/Programs/UnrealBuildTool/Modes/BuildMode.cs b/Engine/Source/Programs/UnrealBuildTool/Modes/BuildMode.cs
--- a/Engine/Source/Programs/UnrealBuildTool/Modes/BuildMode.cs
+++ b/Engine/Source/Programs/UnrealBuildTool/Modes/BuildMode.cs
@@ -32,6 +32,11 @@ namespace UnrealBuildTool
 		/// Just output a list of XGE actions; don't build anything
 		/// </summary>
 		XGEExport = 2,
+
+		/// <summary>
+		/// Fail if any engine files would be modified by the build
+		/// </summary>
+		NoEngineChanges = 4,
 	}
 
 	/// <summary>
@@ -65,6 +70,12 @@ namespace UnrealBuildTool
 		[CommandLine("-XGEExport")]
 		public bool bXGEExport = false;
 
+		/// <summary>
+		/// Do not allow any engine files to be output (used by compile on startup functionality)
+		/// </summary>
+		[CommandLine("-NoEngineChanges")]
+		public bool bNoEngineChanges = false;
+
 		/// <summary>
 		/// Whether we should just export the outdated actions list
 		/// </summary>
@@ -197,6 +208,10 @@ namespace UnrealBuildTool
 					{
 						Options |= BuildOptions.XGEExport;
 					}
+					if(bNoEngineChanges)
+					{
+						Options |= BuildOptions.NoEngineChanges;
+					}
 
 					// Create the working set provider
 					using (ISourceFileWorkingSet WorkingSet = SourceFileWorkingSet.Create(UnrealBuildTool.RootDirectory, ProjectDirs))
@@ -276,6 +291,22 @@ namespace UnrealBuildTool
 				// Link all the actions together
 				ActionGraph.Link(MergedActionsToExecute);
 
+				// Make sure we're not modifying any engine files
+				if ((Options & BuildOptions.NoEngineChanges) != 0)
+				{
+					List<FileItem> EngineChanges = MergedActionsToExecute.SelectMany(x => x.ProducedItems).Where(x => x.Location.IsUnderDirectory(UnrealBuildTool.EngineDirectory)).Distinct().OrderBy(x => x.FullName).ToList();
+					if (EngineChanges.Count > 0)
+					{
+						StringBuilder Result = new StringBuilder("Building would modify the following engine files:\n");
+						foreach (FileItem EngineChange in EngineChanges)
+						{
+							Result.AppendFormat("\n{0}", EngineChange.FullName);
+						}
+						Result.Append("\n\nPlease rebuild from an IDE instead.");
+						throw new BuildException(Result.ToString());
+					}
+				}
+
 				// Make sure the appropriate executor is selected
 				foreach (TargetDescriptor TargetDescriptor in TargetDescriptors)
 				{
diff --git a/Engine/Source/Runtime/Launch/Private/LaunchEngineLoop.cpp b/Engine/Source/Runtime/Launch/Private/LaunchEngineLoop.cpp
--- a/Engine/Source/Runtime/Launch/Private/LaunchEngineLoop.cpp
+++ b/Engine/Source/Runtime/Launch/Private/LaunchEngineLoop.cpp
@@ -5420,16 +5420,26 @@ bool FEngineLoop::AppInit( )
 					ModulesList += FString::Printf(TEXT("  (+%d others, see log for details)\n"), IncompatibleFiles.Num() - NumModulesToDisplay);
 				}
 
-				ModulesList += TEXT("\nWould you like to rebuild them now?");
-
 				// If we're running with -stdout, assume that we're a non interactive process and about to fail
 				if (FApp::IsUnattended() || FParse::Param(FCommandLine::Get(), TEXT("stdout")))
 				{
 					return false;
 				}
 
+				// If there are any engine modules that need building, force the user to build through the IDE
+				for (const FString& IncompatibleFile : IncompatibleFiles)
+				{
+					if (FPaths::IsUnderDirectory(IncompatibleFile, FPaths::EngineDir()))
+					{
+						FString CompileForbidden = ModulesList + TEXT("\nEngine modules cannot be compiled at runtime. Please build through your IDE.");
+						FPlatformMisc::MessageBoxExt(EAppMsgType::Ok, *CompileForbidden, TEXT("Missing Modules"));
+						return false;
+					}
+				}
+
 				// Ask whether to compile before continuing
-				if (FPlatformMisc::MessageBoxExt(EAppMsgType::YesNo, *ModulesList, *FString::Printf(TEXT("Missing %s Modules"), FApp::GetProjectName())) == EAppReturnType::No)
+				FString CompilePrompt = ModulesList + TEXT("\nWould you like to rebuild them now?");
+				if (FPlatformMisc::MessageBoxExt(EAppMsgType::YesNo, *CompilePrompt, *FString::Printf(TEXT("Missing %s Modules"), FApp::GetProjectName())) == EAppReturnType::No)
 				{
 					return false;
 				}
