UBT: Produce an error if any output path is longer than MAX_PATH on Windows.

diff --git a/Engine/Source/Programs/UnrealBuildTool/System/ActionGraph.cs b/Engine/Source/Programs/UnrealBuildTool/System/ActionGraph.cs
--- a/Engine/Source/Programs/UnrealBuildTool/System/ActionGraph.cs
+++ b/Engine/Source/Programs/UnrealBuildTool/System/ActionGraph.cs
@@ -117,23 +117,48 @@ namespace UnrealBuildTool
 		{
 			if (BuildHostPlatform.Current.Platform == UnrealTargetPlatform.Win64)
 			{
-				List<FileReference> Paths = new List<FileReference>();
+				const int MAX_PATH = 260;
+
+				List<FileReference> FailPaths = new List<FileReference>();
+				List<FileReference> WarnPaths = new List<FileReference>();
 				foreach (Action Action in Actions)
 				{
+					foreach (FileItem PrerequisiteItem in Action.PrerequisiteItems)
+					{
+						if (PrerequisiteItem.Location.FullName.Length >= MAX_PATH)
+						{
+							FailPaths.Add(PrerequisiteItem.Location);
+						}
+					}
 					foreach (FileItem ProducedItem in Action.ProducedItems)
 					{
+						if (ProducedItem.Location.FullName.Length >= MAX_PATH)
+						{
+							FailPaths.Add(ProducedItem.Location);
+						}
 						if (ProducedItem.Location.FullName.Length > UnrealBuildTool.RootDirectory.FullName.Length + BuildConfiguration.MaxNestedPathLength && ProducedItem.Location.IsUnderDirectory(UnrealBuildTool.RootDirectory))
 						{
-							Paths.Add(ProducedItem.Location);
+							WarnPaths.Add(ProducedItem.Location);
 						}
 					}
 				}
 
-				if (Paths.Count > 0)
+				if (FailPaths.Count > 0)
+				{
+					StringBuilder Message = new StringBuilder();
+					Message.AppendFormat("The following output paths are longer than {0} characters. Please move the engine to a directory with a shorter path.", MAX_PATH);
+					foreach (FileReference Path in FailPaths)
+					{
+						Message.AppendFormat("\n[{0} characters] {1}", Path.FullName.Length, Path);
+					}
+					throw new BuildException(Message.ToString());
+				}
+
+				if (WarnPaths.Count > 0)
 				{
 					StringBuilder Message = new StringBuilder();
-					Message.AppendFormat("Detected paths more than {0} characters below UE root directory. This may cause portability issues due to the 260 character maximum path length on Windows:\n", BuildConfiguration.MaxNestedPathLength);
-					foreach (FileReference Path in Paths)
+					Message.AppendFormat("Detected paths more than {0} characters below UE root directory. This may cause portability issues due to the {1} character maximum path length on Windows:\n", BuildConfiguration.MaxNestedPathLength, MAX_PATH);
+					foreach (FileReference Path in WarnPaths)
 					{
 						string RelativePath = Path.MakeRelativeTo(UnrealBuildTool.RootDirectory);
 						Message.AppendFormat("\n[{0} characters] {1}", RelativePath.Length, RelativePath);
