Fix NDK20 fallback tool paths for some architectures and disable ld.gold for ARM64

#ue4
#android

diff --git a/Engine/Source/Programs/UnrealBuildTool/Platform/Android/AndroidToolChain.cs b/Engine/Source/Programs/UnrealBuildTool/Platform/Android/AndroidToolChain.cs
--- a/Engine/Source/Programs/UnrealBuildTool/Platform/Android/AndroidToolChain.cs
+++ b/Engine/Source/Programs/UnrealBuildTool/Platform/Android/AndroidToolChain.cs
@@ -297,10 +297,10 @@ namespace UnrealBuildTool
 			if (NDKDefineInt < 210000)
 			{
 				// use ld before r21
-				ArPathArm = Utils.CollapseRelativeDirectories(Path.Combine(NDKPath, @"toolchains/x86-4.9", ArchitecturePath, @"bin/armv7a-linux-androideabi-ar" + ExeExtension));
-				ArPathArm64 = Utils.CollapseRelativeDirectories(Path.Combine(NDKPath, @"toolchains/x86-4.9", ArchitecturePath, @"bin/aarch64-linux-android-ar" + ExeExtension));
+				ArPathArm = Utils.CollapseRelativeDirectories(Path.Combine(NDKPath, @"toolchains/arm-linux-androideabi-4.9", ArchitecturePath, @"bin/armv7a-linux-androideabi-ar" + ExeExtension));
+				ArPathArm64 = Utils.CollapseRelativeDirectories(Path.Combine(NDKPath, @"toolchains/aarch64-linux-android-4.9", ArchitecturePath, @"bin/aarch64-linux-android-ar" + ExeExtension));
 				ArPathx86 = Utils.CollapseRelativeDirectories(Path.Combine(NDKPath, @"toolchains/x86-4.9", ArchitecturePath, @"bin/i686-linux-android-ar" + ExeExtension));
-				ArPathx64 = Utils.CollapseRelativeDirectories(Path.Combine(NDKPath, @"toolchains/x86_4.9", ArchitecturePath, @"bin/x86_64-linux-android-ar" + ExeExtension));
+				ArPathx64 = Utils.CollapseRelativeDirectories(Path.Combine(NDKPath, @"toolchains/x86_64-4.9", ArchitecturePath, @"bin/x86_64-linux-android-ar" + ExeExtension));
 			}
 			else
 			{
@@ -836,10 +836,12 @@ namespace UnrealBuildTool
 			}
 
 			bool bUseLLD = NDKDefineInt >= 210000;
+			bool bAllowLdGold = true;
 			if (Architecture == "-arm64")
 			{
 				Result += ToolchainLinkParamsArm64;
 				Result += " -march=armv8-a";
+				bAllowLdGold = false;       // NDK issue 70838247
 			}
 			else if (Architecture == "-x86")
 			{
@@ -873,7 +875,7 @@ namespace UnrealBuildTool
 			}
 			else
 			{
-				if (bUseLdGold)
+				if (bAllowLdGold && bUseLdGold)
 				{
 					// use ld.gold as linker (requires strip)
 					Result += " -fuse-ld=gold";
