Fix exception running IncludeTool due to include path modules no longer being included in JSON export.

diff --git a/Engine/Source/Programs/UnrealBuildTool/Configuration/UEBuildModule.cs b/Engine/Source/Programs/UnrealBuildTool/Configuration/UEBuildModule.cs
--- a/Engine/Source/Programs/UnrealBuildTool/Configuration/UEBuildModule.cs
+++ b/Engine/Source/Programs/UnrealBuildTool/Configuration/UEBuildModule.cs
@@ -951,19 +951,23 @@ namespace UnrealBuildTool
 			}
 			Writer.WriteArrayEnd();
 
-			Writer.WriteArrayStart("RuntimeDependencies");
-			foreach(ModuleRules.RuntimeDependency RuntimeDependency in Rules.RuntimeDependencies.Inner)
+			// Don't add runtime dependencies for modules that aren't being linked in. They may reference BinaryOutputDir, which is invalid.
+			if (Binary != null)
 			{
-				Writer.WriteObjectStart();
-				Writer.WriteValue("Path", ExpandPathVariables(RuntimeDependency.Path, BinaryOutputDir, TargetOutputDir));
-				if(RuntimeDependency.SourcePath != null)
+				Writer.WriteArrayStart("RuntimeDependencies");
+				foreach (ModuleRules.RuntimeDependency RuntimeDependency in Rules.RuntimeDependencies.Inner)
 				{
-					Writer.WriteValue("SourcePath", ExpandPathVariables(RuntimeDependency.SourcePath, BinaryOutputDir, TargetOutputDir));
+					Writer.WriteObjectStart();
+					Writer.WriteValue("Path", ExpandPathVariables(RuntimeDependency.Path, BinaryOutputDir, TargetOutputDir));
+					if (RuntimeDependency.SourcePath != null)
+					{
+						Writer.WriteValue("SourcePath", ExpandPathVariables(RuntimeDependency.SourcePath, BinaryOutputDir, TargetOutputDir));
+					}
+					Writer.WriteValue("Type", RuntimeDependency.Type.ToString());
+					Writer.WriteObjectEnd();
 				}
-				Writer.WriteValue("Type", RuntimeDependency.Type.ToString());
-				Writer.WriteObjectEnd();
+				Writer.WriteArrayEnd();
 			}
-			Writer.WriteArrayEnd();
 		}
 
 		/// <summary>
diff --git a/Engine/Source/Programs/UnrealBuildTool/Configuration/UEBuildTarget.cs b/Engine/Source/Programs/UnrealBuildTool/Configuration/UEBuildTarget.cs
--- a/Engine/Source/Programs/UnrealBuildTool/Configuration/UEBuildTarget.cs
+++ b/Engine/Source/Programs/UnrealBuildTool/Configuration/UEBuildTarget.cs
@@ -2188,14 +2188,11 @@ namespace UnrealBuildTool
 				Writer.WriteArrayEnd();
 
 				Writer.WriteObjectStart("Modules");
-				foreach (UEBuildBinary Binary in Binaries)
+				foreach (UEBuildModule Module in Modules.Values)
 				{
-					foreach (UEBuildModule Module in Binary.Modules)
-					{
-						Writer.WriteObjectStart(Module.Name);
-						Module.ExportJson(Binary.OutputDir, GetExecutableDir(), Writer);
-						Writer.WriteObjectEnd();
-					}
+					Writer.WriteObjectStart(Module.Name);
+					Module.ExportJson(Module.Binary?.OutputDir, GetExecutableDir(), Writer);
+					Writer.WriteObjectEnd();
 				}
 				Writer.WriteObjectEnd();
 
