Add option to force use of legacy ld instead of lld for Android

#ue4
#android

diff --git a/Engine/Config/BaseEngine.ini b/Engine/Config/BaseEngine.ini
--- a/Engine/Config/BaseEngine.ini
+++ b/Engine/Config/BaseEngine.ini
@@ -2099,6 +2099,7 @@ bBundleABISplit=True
 bBundleLanguageSplit=True
 bBundleDensitySplit=True
 bFullScreen=True
+bForceLDLinker=False
 
 [/Script/AndroidPlatformEditor.AndroidSDKSettings]
 SDKAPILevel=latest
diff --git a/Engine/Source/Programs/UnrealBuildTool/Platform/Android/AndroidToolChain.cs b/Engine/Source/Programs/UnrealBuildTool/Platform/Android/AndroidToolChain.cs
--- a/Engine/Source/Programs/UnrealBuildTool/Platform/Android/AndroidToolChain.cs
+++ b/Engine/Source/Programs/UnrealBuildTool/Platform/Android/AndroidToolChain.cs
@@ -294,7 +294,7 @@ namespace UnrealBuildTool
 
 			// set up the path to our toolchains
 			ClangPath = Utils.CollapseRelativeDirectories(Path.Combine(NDKPath, @"toolchains/llvm", ArchitecturePath, @"bin/clang++" + ExeExtension));
-			if (NDKDefineInt < 210000)
+			if (NDKDefineInt < 210000 || ForceLDLinker())
 			{
 				// use ld before r21
 				ArPathArm = Utils.CollapseRelativeDirectories(Path.Combine(NDKPath, @"toolchains/arm-linux-androideabi-4.9", ArchitecturePath, @"bin/armv7a-linux-androideabi-ar" + ExeExtension));
@@ -401,6 +401,13 @@ namespace UnrealBuildTool
 			return CompileEnvironment.Configuration == CppConfiguration.Shipping && (Ini.GetBool("/Script/AndroidRuntimeSettings.AndroidRuntimeSettings", "bBuildWithHiddenSymbolVisibility", out bBuild) && bBuild);
 		}
 
+		private bool ForceLDLinker()
+		{
+			ConfigHierarchy Ini = ConfigCache.ReadHierarchy(ConfigHierarchyType.Engine, DirectoryReference.FromFile(ProjectFile), UnrealTargetPlatform.Android);
+			bool bForceLDLinker = false;
+			return Ini.GetBool("/Script/AndroidRuntimeSettings.AndroidRuntimeSettings", "bForceLDLinker", out bForceLDLinker) && bForceLDLinker;
+		}
+
 		public override void SetUpGlobalEnvironment(ReadOnlyTargetRules Target)
 		{
 			base.SetUpGlobalEnvironment(Target);
@@ -824,7 +831,6 @@ namespace UnrealBuildTool
 			Result += " -no-canonical-prefixes";
 			Result += " -Wl,-shared,-Bsymbolic";
 			Result += " -Wl,--no-undefined";
-			Result += " -Wl,-no-pie";
 			if (bEnableGcSections)
 			{
 				Result += " -Wl,-gc-sections"; // Enable garbage collection of unused input sections. works best with -ffunction-sections, -fdata-sections
@@ -835,7 +841,7 @@ namespace UnrealBuildTool
 				Result += " -Wl,--strip-debug";
 			}
 
-			bool bUseLLD = NDKDefineInt >= 210000;
+			bool bUseLLD = NDKDefineInt >= 210000 && !ForceLDLinker();
 			bool bAllowLdGold = true;
 			if (Architecture == "-arm64")
 			{
@@ -870,6 +876,8 @@ namespace UnrealBuildTool
 
 			if (bUseLLD)
 			{
+				Result += " -Wl,-no-pie";
+
 				// use lld as linker (requires llvm-strip)
 				Result += " -fuse-ld=lld";
 			}
diff --git a/Engine/Source/Runtime/Android/AndroidRuntimeSettings/Classes/AndroidRuntimeSettings.h b/Engine/Source/Runtime/Android/AndroidRuntimeSettings/Classes/AndroidRuntimeSettings.h
--- a/Engine/Source/Runtime/Android/AndroidRuntimeSettings/Classes/AndroidRuntimeSettings.h
+++ b/Engine/Source/Runtime/Android/AndroidRuntimeSettings/Classes/AndroidRuntimeSettings.h
@@ -399,6 +399,10 @@ public:
 	UPROPERTY(GlobalConfig, EditAnywhere, Category = AdvancedBuild, meta = (DisplayName = "Always save a copy of the libUE4.so with symbols. [Experimental]"))
 	bool bSaveSymbols;
 
+	// Use legacy ld instead of new lld linker.
+	UPROPERTY(GlobalConfig, EditAnywhere, Category = AdvancedBuild, meta = (DisplayName = "Force linking to use ld instead of lld"))
+	bool bForceLDLinker;
+
 	// If selected, the checked architectures will be split into separate .apk files [CURRENTLY FOR FULL SOURCE GAMES ONLY]
 	// @todo android fat binary: Currently, there isn't much utility in merging multiple .so's into a single .apk except for debugging,
 	// but we can't properly handle multiple GPU architectures in a single .apk, so we are disabling the feature for now
