Add support for whitelisting additional plugins and adding additional modules to plugin extensions.
Additionally allow monolithic programs inside platform extensions to ouput to the Binaries directory in the extension.

diff --git a/Engine/Source/Programs/AutomationTool/AutomationUtils/ProjectUtils.cs b/Engine/Source/Programs/AutomationTool/AutomationUtils/ProjectUtils.cs
--- a/Engine/Source/Programs/AutomationTool/AutomationUtils/ProjectUtils.cs
+++ b/Engine/Source/Programs/AutomationTool/AutomationUtils/ProjectUtils.cs
@@ -320,8 +320,7 @@ namespace AutomationTool
 			bool bRequiresCookedData = (TargetType != TargetType.Editor);
 			foreach (KeyValuePair<string, PluginInfo> Pair in EnabledPlugins)
 			{
-				ModuleDescriptor[] Modules = Pair.Value.Descriptor.Modules;
-				if (Modules != null)
+				if (Pair.Value.Descriptor.Modules != null)
 				{
 					foreach (ModuleDescriptor Module in Pair.Value.Descriptor.Modules)
 					{
diff --git a/Engine/Source/Programs/UnrealBuildTool/Configuration/UEBuildTarget.cs b/Engine/Source/Programs/UnrealBuildTool/Configuration/UEBuildTarget.cs
--- a/Engine/Source/Programs/UnrealBuildTool/Configuration/UEBuildTarget.cs
+++ b/Engine/Source/Programs/UnrealBuildTool/Configuration/UEBuildTarget.cs
@@ -3237,7 +3237,8 @@ namespace UnrealBuildTool
 							UEBuildModuleCPP Module = FindOrCreateCppModuleByName(ModuleInfo.Name, PluginReferenceChain);
 							if(!Instance.Modules.Contains(Module))
 							{
-								if (!Module.RulesFile.IsUnderDirectory(Info.Directory))
+								// This could be in a child plugin so scan thorugh those as well
+								if (!Module.RulesFile.IsUnderDirectory(Info.Directory) && !Info.ChildFiles.Any(ChildFile => Module.RulesFile.IsUnderDirectory(ChildFile.Directory)))
 								{
 									throw new BuildException("Plugin '{0}' (referenced via {1}) does not contain the '{2}' module, but lists it in '{3}'.", Info.Name, ReferenceChain, ModuleInfo.Name, Info.File);
 								}
@@ -3322,14 +3323,22 @@ namespace UnrealBuildTool
 			}
 
 			// Construct the output paths for this target's executable
-			DirectoryReference OutputDirectory;
-			if (bCompileMonolithic || !bUseSharedBuildEnvironment)
+			DirectoryReference OutputDirectory = UnrealBuildTool.EngineDirectory;
+
+			bool bOutputToPlatformExtensionDirectory = Rules.File.IsUnderDirectory(UnrealBuildTool.EnginePlatformExtensionsDirectory) || Rules.File.IsUnderDirectory(UnrealBuildTool.ProjectPlatformExtensionsDirectory(ProjectDirectory));
+			bool bOutputToProjectDirectory = (bCompileMonolithic || !bUseSharedBuildEnvironment) && (ProjectDirectory != UnrealBuildTool.EngineDirectory);
+			
+			if (bOutputToPlatformExtensionDirectory && bOutputToProjectDirectory)
 			{
-				OutputDirectory = ProjectDirectory;
+				OutputDirectory = UnrealBuildTool.GetAllProjectDirectories(ProjectDirectory).First(x => x != ProjectDirectory && Rules.File.IsUnderDirectory(x));
 			}
-			else
+			else if(bOutputToPlatformExtensionDirectory)
 			{
-				OutputDirectory = UnrealBuildTool.EngineDirectory;
+				OutputDirectory = UnrealBuildTool.GetAllEngineDirectories().First(x => x != UnrealBuildTool.EngineDirectory && Rules.File.IsUnderDirectory(x));
+			} 
+			else if (bOutputToProjectDirectory)
+			{
+				OutputDirectory = ProjectDirectory;
 			}
 
 			bool bCompileAsDLL = Rules.bShouldCompileAsDLL && bCompileMonolithic;
diff --git a/Engine/Source/Programs/UnrealBuildTool/ProjectFiles/ProjectFileGenerator.cs b/Engine/Source/Programs/UnrealBuildTool/ProjectFiles/ProjectFileGenerator.cs
--- a/Engine/Source/Programs/UnrealBuildTool/ProjectFiles/ProjectFileGenerator.cs
+++ b/Engine/Source/Programs/UnrealBuildTool/ProjectFiles/ProjectFileGenerator.cs
@@ -2466,7 +2466,7 @@ namespace UnrealBuildTool
 				ProjectFile GameProject = GameProjects.First();
 				foreach(PluginInfo PluginInfo in Plugins.ReadProjectPlugins(GameProject.BaseDir))
 				{
-					if (PluginInfo.Descriptor.Modules != null && PluginInfo.Descriptor.Modules.Length > 0 && PluginInfo.Type == PluginType.Mod)
+					if (PluginInfo.Descriptor.Modules != null && PluginInfo.Descriptor.Modules.Count > 0 && PluginInfo.Type == PluginType.Mod)
 					{
 						FileReference ModProjectFilePath = FileReference.Combine(PluginInfo.Directory, "Mods", PluginInfo.Name + ProjectFileExtension);
 
diff --git a/Engine/Source/Programs/UnrealBuildTool/System/ModuleDescriptor.cs b/Engine/Source/Programs/UnrealBuildTool/System/ModuleDescriptor.cs
--- a/Engine/Source/Programs/UnrealBuildTool/System/ModuleDescriptor.cs
+++ b/Engine/Source/Programs/UnrealBuildTool/System/ModuleDescriptor.cs
@@ -170,12 +170,12 @@ namespace UnrealBuildTool
 		/// <summary>
 		/// List of allowed platforms
 		/// </summary>
-		public UnrealTargetPlatform[] WhitelistPlatforms;
+		public List<UnrealTargetPlatform> WhitelistPlatforms;
 
 		/// <summary>
 		/// List of disallowed platforms
 		/// </summary>
-		public UnrealTargetPlatform[] BlacklistPlatforms;
+		public List<UnrealTargetPlatform> BlacklistPlatforms;
 
 		/// <summary>
 		/// List of allowed targets
@@ -243,13 +243,13 @@ namespace UnrealBuildTool
 				string[] WhitelistPlatforms;
 				if (InObject.TryGetStringArrayField("WhitelistPlatforms", out WhitelistPlatforms))
 				{
-					Module.WhitelistPlatforms = Array.ConvertAll(WhitelistPlatforms, x => UnrealTargetPlatform.Parse(x));
+					Module.WhitelistPlatforms = Array.ConvertAll(WhitelistPlatforms, x => UnrealTargetPlatform.Parse(x)).ToList();
 				}
 
 				string[] BlacklistPlatforms;
 				if (InObject.TryGetStringArrayField("BlacklistPlatforms", out BlacklistPlatforms))
 				{
-					Module.BlacklistPlatforms = Array.ConvertAll(BlacklistPlatforms, x => UnrealTargetPlatform.Parse(x));
+					Module.BlacklistPlatforms = Array.ConvertAll(BlacklistPlatforms, x => UnrealTargetPlatform.Parse(x)).ToList();
 				}
 			}
 			catch (BuildException Ex)
@@ -313,7 +313,7 @@ namespace UnrealBuildTool
 			Writer.WriteValue("Name", Name);
 			Writer.WriteValue("Type", Type.ToString());
 			Writer.WriteValue("LoadingPhase", LoadingPhase.ToString());
-			if (WhitelistPlatforms != null && WhitelistPlatforms.Length > 0)
+			if (WhitelistPlatforms != null && WhitelistPlatforms.Count > 0)
 			{
 				Writer.WriteArrayStart("WhitelistPlatforms");
 				foreach (UnrealTargetPlatform WhitelistPlatform in WhitelistPlatforms)
@@ -322,7 +322,7 @@ namespace UnrealBuildTool
 				}
 				Writer.WriteArrayEnd();
 			}
-			if (BlacklistPlatforms != null && BlacklistPlatforms.Length > 0)
+			if (BlacklistPlatforms != null && BlacklistPlatforms.Count > 0)
 			{
 				Writer.WriteArrayStart("BlacklistPlatforms");
 				foreach (UnrealTargetPlatform BlacklistPlatform in BlacklistPlatforms)
@@ -431,7 +431,7 @@ namespace UnrealBuildTool
 		public bool IsCompiledInConfiguration(UnrealTargetPlatform Platform, UnrealTargetConfiguration Configuration, string TargetName, TargetType TargetType, bool bBuildDeveloperTools, bool bBuildRequiresCookedData)
 		{
 			// Check the platform is whitelisted
-			if (WhitelistPlatforms != null && WhitelistPlatforms.Length > 0 && !WhitelistPlatforms.Contains(Platform))
+			if (WhitelistPlatforms != null && WhitelistPlatforms.Count > 0 && !WhitelistPlatforms.Contains(Platform))
 			{
 				return false;
 			}
diff --git a/Engine/Source/Programs/UnrealBuildTool/System/PluginDescriptor.cs b/Engine/Source/Programs/UnrealBuildTool/System/PluginDescriptor.cs
--- a/Engine/Source/Programs/UnrealBuildTool/System/PluginDescriptor.cs
+++ b/Engine/Source/Programs/UnrealBuildTool/System/PluginDescriptor.cs
@@ -122,7 +122,7 @@ namespace UnrealBuildTool
 		/// <summary>
 		/// List of platforms supported by this plugin. This list will be copied to any plugin reference from a project file, to allow filtering entire plugins from staged builds.
 		/// </summary>
-		public UnrealTargetPlatform[] SupportedTargetPlatforms;
+		public List<UnrealTargetPlatform> SupportedTargetPlatforms;
 
 		/// <summary>
 		/// List of programs supported by this plugin.
@@ -132,7 +132,7 @@ namespace UnrealBuildTool
 		/// <summary>
 		/// List of all modules associated with this plugin
 		/// </summary>
-		public ModuleDescriptor[] Modules;
+		public List<ModuleDescriptor> Modules;
 
 		/// <summary>
 		/// List of all localization targets associated with this plugin
@@ -187,7 +187,7 @@ namespace UnrealBuildTool
 		/// <summary>
 		/// Additional plugins that this plugin depends on
 		/// </summary>
-		public PluginReferenceDescriptor[] Plugins;
+		public List<PluginReferenceDescriptor> Plugins;
 
 		/// <summary>
 		/// Private constructor. This object should not be created directly; read it from disk using FromFile() instead.
@@ -252,7 +252,7 @@ namespace UnrealBuildTool
 				string[] SupportedTargetPlatformNames;
 				if (RawObject.TryGetStringArrayField("SupportedTargetPlatforms", out SupportedTargetPlatformNames))
 				{
-					SupportedTargetPlatforms = Array.ConvertAll(SupportedTargetPlatformNames, x => UnrealTargetPlatform.Parse(x));
+					SupportedTargetPlatforms = Array.ConvertAll(SupportedTargetPlatformNames, x => UnrealTargetPlatform.Parse(x)).ToList();
 				}
 			}
 			catch (BuildException Ex)
@@ -265,7 +265,7 @@ namespace UnrealBuildTool
 			JsonObject[] ModulesArray;
 			if (RawObject.TryGetObjectArrayField("Modules", out ModulesArray))
 			{
-				Modules = Array.ConvertAll(ModulesArray, x => ModuleDescriptor.FromJsonObject(x));
+				Modules = Array.ConvertAll(ModulesArray, x => ModuleDescriptor.FromJsonObject(x)).ToList();
 			}
 
 			JsonObject[] LocalizationTargetsArray;
@@ -301,7 +301,7 @@ namespace UnrealBuildTool
 			JsonObject[] PluginsArray;
 			if(RawObject.TryGetObjectArrayField("Plugins", out PluginsArray))
 			{
-				Plugins = Array.ConvertAll(PluginsArray, x => PluginReferenceDescriptor.FromJsonObject(x));
+				Plugins = Array.ConvertAll(PluginsArray, x => PluginReferenceDescriptor.FromJsonObject(x)).ToList();
 			}
 		}
 
@@ -394,7 +394,7 @@ namespace UnrealBuildTool
 				Writer.WriteValue("ExplicitlyLoaded", bExplicitlyLoaded);
 			}
 
-			if(SupportedTargetPlatforms != null && SupportedTargetPlatforms.Length > 0)
+			if(SupportedTargetPlatforms != null && SupportedTargetPlatforms.Count > 0)
 			{
 				Writer.WriteStringArrayField("SupportedTargetPlatforms", SupportedTargetPlatforms.Select<UnrealTargetPlatform, string>(x => x.ToString()).ToArray());
 			}
@@ -408,7 +408,7 @@ namespace UnrealBuildTool
 				Writer.WriteValue("bIsPluginExtension", bIsPluginExtension);
 			}
 
-			ModuleDescriptor.WriteArray(Writer, "Modules", Modules);
+			ModuleDescriptor.WriteArray(Writer, "Modules", Modules.ToArray());
 
 			LocalizationTargetDescriptor.WriteArray(Writer, "LocalizationTargets", LocalizationTargets);
 
@@ -422,7 +422,7 @@ namespace UnrealBuildTool
 				PostBuildSteps.Write(Writer, "PostBuildSteps");
 			}
 
-			PluginReferenceDescriptor.WriteArray(Writer, "Plugins", Plugins);
+			PluginReferenceDescriptor.WriteArray(Writer, "Plugins", Plugins.ToArray());
 		}
 
 		/// <summary>
@@ -432,7 +432,7 @@ namespace UnrealBuildTool
 		/// <returns>True if the plugin should be enabled</returns>
 		public bool SupportsTargetPlatform(UnrealTargetPlatform Platform)
 		{
-			return SupportedTargetPlatforms == null || SupportedTargetPlatforms.Length == 0 || SupportedTargetPlatforms.Contains(Platform);
+			return SupportedTargetPlatforms == null || SupportedTargetPlatforms.Count == 0 || SupportedTargetPlatforms.Contains(Platform);
 		}
 	}
 }
diff --git a/Engine/Source/Programs/UnrealBuildTool/System/PluginReferenceDescriptor.cs b/Engine/Source/Programs/UnrealBuildTool/System/PluginReferenceDescriptor.cs
--- a/Engine/Source/Programs/UnrealBuildTool/System/PluginReferenceDescriptor.cs
+++ b/Engine/Source/Programs/UnrealBuildTool/System/PluginReferenceDescriptor.cs
@@ -44,12 +44,12 @@ namespace UnrealBuildTool
 		/// <summary>
 		/// If enabled, list of platforms for which the plugin should be enabled (or all platforms if blank).
 		/// </summary>
-		public UnrealTargetPlatform[] WhitelistPlatforms;
+		public List<UnrealTargetPlatform> WhitelistPlatforms;
 
 		/// <summary>
 		/// If enabled, list of platforms for which the plugin should be disabled.
 		/// </summary>
-		public UnrealTargetPlatform[] BlacklistPlatforms;
+		public List<UnrealTargetPlatform> BlacklistPlatforms;
 
 		/// <summary>
 		/// If enabled, list of target configurations for which the plugin should be enabled (or all target configurations if blank).
@@ -74,7 +74,7 @@ namespace UnrealBuildTool
 		/// <summary>
 		/// The list of supported platforms for this plugin. This field is copied from the plugin descriptor, and supplements the user's whitelisted and blacklisted platforms.
 		/// </summary>
-		public UnrealTargetPlatform[] SupportedTargetPlatforms;
+		public List<UnrealTargetPlatform> SupportedTargetPlatforms;
 
 		/// <summary>
 		/// Constructor
@@ -110,11 +110,11 @@ namespace UnrealBuildTool
 			{
 				Writer.WriteValue("MarketplaceURL", MarketplaceURL);
 			}
-			if(WhitelistPlatforms != null && WhitelistPlatforms.Length > 0)
+			if(WhitelistPlatforms != null && WhitelistPlatforms.Count > 0)
 			{
 				Writer.WriteStringArrayField("WhitelistPlatforms", WhitelistPlatforms.Select(x => x.ToString()).ToArray());
 			}
-			if(BlacklistPlatforms != null && BlacklistPlatforms.Length > 0)
+			if(BlacklistPlatforms != null && BlacklistPlatforms.Count > 0)
 			{
 				Writer.WriteStringArrayField("BlacklistPlatforms", BlacklistPlatforms.Select(x => x.ToString()).ToArray());
 			}
@@ -134,7 +134,7 @@ namespace UnrealBuildTool
 			{
 				Writer.WriteEnumArrayField("BlacklistTargets", BlacklistTargets);
 			}
-			if(SupportedTargetPlatforms != null && SupportedTargetPlatforms.Length > 0)
+			if(SupportedTargetPlatforms != null && SupportedTargetPlatforms.Count > 0)
 			{
 				Writer.WriteStringArrayField("SupportedTargetPlatforms", SupportedTargetPlatforms.Select(x => x.ToString()).ToArray());
 			}
@@ -193,15 +193,15 @@ namespace UnrealBuildTool
 				// convert string array to UnrealTargetPlatform arrays
 				if (WhitelistPlatformNames != null)
 				{
-					Descriptor.WhitelistPlatforms = WhitelistPlatformNames.Select(x => UnrealTargetPlatform.Parse(x)).ToArray();
+					Descriptor.WhitelistPlatforms = WhitelistPlatformNames.Select(x => UnrealTargetPlatform.Parse(x)).ToList();
 				}
 				if (BlacklistPlatformNames != null)
 				{
-					Descriptor.BlacklistPlatforms = BlacklistPlatformNames.Select(x => UnrealTargetPlatform.Parse(x)).ToArray();
+					Descriptor.BlacklistPlatforms = BlacklistPlatformNames.Select(x => UnrealTargetPlatform.Parse(x)).ToList();
 				}
 				if (SupportedTargetPlatformNames != null)
 				{
-					Descriptor.SupportedTargetPlatforms = SupportedTargetPlatformNames.Select(x => UnrealTargetPlatform.Parse(x)).ToArray();
+					Descriptor.SupportedTargetPlatforms = SupportedTargetPlatformNames.Select(x => UnrealTargetPlatform.Parse(x)).ToList();
 				}
 			}
 			catch (BuildException Ex)
@@ -225,7 +225,7 @@ namespace UnrealBuildTool
 			{
 				return false;
 			}
-			if (WhitelistPlatforms != null && WhitelistPlatforms.Length > 0 && !WhitelistPlatforms.Contains(Platform))
+			if (WhitelistPlatforms != null && WhitelistPlatforms.Count > 0 && !WhitelistPlatforms.Contains(Platform))
 			{
 				return false;
 			}
@@ -287,7 +287,7 @@ namespace UnrealBuildTool
 		/// <returns>True if the plugin for this target platform</returns>
 		public bool IsSupportedTargetPlatform(UnrealTargetPlatform Platform)
 		{
-			return SupportedTargetPlatforms == null || SupportedTargetPlatforms.Length == 0 || SupportedTargetPlatforms.Contains(Platform);
+			return SupportedTargetPlatforms == null || SupportedTargetPlatforms.Count == 0 || SupportedTargetPlatforms.Contains(Platform);
 		}
 	}
 }
diff --git a/Engine/Source/Programs/UnrealBuildTool/System/Plugins.cs b/Engine/Source/Programs/UnrealBuildTool/System/Plugins.cs
--- a/Engine/Source/Programs/UnrealBuildTool/System/Plugins.cs
+++ b/Engine/Source/Programs/UnrealBuildTool/System/Plugins.cs
@@ -372,50 +372,105 @@ namespace UnrealBuildTool
 			UnrealTargetPlatform Platform = UnrealTargetPlatform.Parse(PlatformName);
 
 			// merge the supported platforms
-			if (Parent.Descriptor.SupportedTargetPlatforms != null)
+			if (Child.Descriptor.SupportedTargetPlatforms != null)
 			{
-				List<UnrealTargetPlatform> SupportedTargetPlatforms = Parent.Descriptor.SupportedTargetPlatforms.ToList();
-				if (Child.Descriptor.SupportedTargetPlatforms != null && Child.Descriptor.SupportedTargetPlatforms.Length > 0)
+				if (Parent.Descriptor.SupportedTargetPlatforms == null)
 				{
-					SupportedTargetPlatforms = SupportedTargetPlatforms.Union(Child.Descriptor.SupportedTargetPlatforms).ToList();
+					Parent.Descriptor.SupportedTargetPlatforms = Child.Descriptor.SupportedTargetPlatforms;
 				}
-				else if (!Parent.Descriptor.SupportedTargetPlatforms.Contains(Platform))
+				else
 				{
-					SupportedTargetPlatforms.Add(Platform);
+					Parent.Descriptor.SupportedTargetPlatforms = Parent.Descriptor.SupportedTargetPlatforms.Union(Child.Descriptor.SupportedTargetPlatforms).ToList();
 				}
-				Parent.Descriptor.SupportedTargetPlatforms = SupportedTargetPlatforms.ToArray();
 			}
 
-			// make sure we are whitelisted for any modules we list, if the parent had a whitelist
+			// make sure we are whitelisted for any modules we list
 			if (Child.Descriptor.Modules != null)
 			{
-				foreach (ModuleDescriptor ChildModule in Child.Descriptor.Modules)
+				if (Parent.Descriptor.Modules == null)
+				{
+					Parent.Descriptor.Modules = Child.Descriptor.Modules;
+				}
+				else
 				{
-					ModuleDescriptor ParentModule = Parent.Descriptor.Modules.FirstOrDefault(x => x.Name.Equals(ChildModule.Name) && x.Type == ChildModule.Type);
-					if (ParentModule != null)
+					foreach (ModuleDescriptor ChildModule in Child.Descriptor.Modules)
 					{
-						// merge white/blacklists (if the parent had a list, and child didn't specify a list, just add the child platform to the parent list - for white and black!)
-						if (ParentModule.WhitelistPlatforms != null && ParentModule.WhitelistPlatforms.Length > 0)
+						ModuleDescriptor ParentModule = Parent.Descriptor.Modules.FirstOrDefault(x => x.Name.Equals(ChildModule.Name) && x.Type == ChildModule.Type);
+						if (ParentModule != null)
 						{
-							List<UnrealTargetPlatform> Whitelist = ParentModule.WhitelistPlatforms.ToList();
-							if (ChildModule.WhitelistPlatforms != null && ChildModule.WhitelistPlatforms.Length > 0)
+							// merge white/blacklists (if the parent had a list, and child didn't specify a list, just add the child platform to the parent list - for white and black!)
+							if (ChildModule.WhitelistPlatforms != null)
 							{
-								Whitelist.AddRange(ChildModule.WhitelistPlatforms);
+								if (ParentModule.WhitelistPlatforms == null)
+								{
+									ParentModule.WhitelistPlatforms = ChildModule.WhitelistPlatforms;
+								}
+								else
+								{
+									ParentModule.WhitelistPlatforms = ParentModule.WhitelistPlatforms.Union(ChildModule.WhitelistPlatforms).ToList();
+								}
 							}
-							else
+							if (ChildModule.BlacklistPlatforms != null)
 							{
-								Whitelist.Add(Platform);
+								if (ParentModule.BlacklistPlatforms == null)
+								{
+									ParentModule.BlacklistPlatforms = ChildModule.BlacklistPlatforms;
+								}
+								else
+								{
+									ParentModule.BlacklistPlatforms = ParentModule.BlacklistPlatforms.Union(ChildModule.BlacklistPlatforms).ToList();
+								}
 							}
-							ParentModule.WhitelistPlatforms = Whitelist.ToArray();
 						}
-						if (ParentModule.BlacklistPlatforms != null && ParentModule.BlacklistPlatforms.Length > 0)
+						else
 						{
-							if (ChildModule.BlacklistPlatforms != null && ChildModule.BlacklistPlatforms.Length > 0)
+							Parent.Descriptor.Modules.Add(ChildModule);
+						}
+					}
+				}
+			}
+
+			// make sure we are whitelisted for any plugins we list
+			if (Child.Descriptor.Plugins != null)
+			{
+				if (Parent.Descriptor.Plugins == null)
+				{
+					Parent.Descriptor.Plugins = Child.Descriptor.Plugins;
+				}
+				else
+				{ 
+					foreach (PluginReferenceDescriptor ChildPluginReference in Child.Descriptor.Plugins)
+					{
+						PluginReferenceDescriptor ParentPluginReference = Parent.Descriptor.Plugins.FirstOrDefault(x => x.Name.Equals(ChildPluginReference.Name));
+						if (ParentPluginReference != null)
+						{
+							// merge white/blacklists (if the parent had a list, and child didn't specify a list, just add the child platform to the parent list - for white and black!)
+							if (ChildPluginReference.WhitelistPlatforms != null)
 							{
-								List<UnrealTargetPlatform> Blacklist = ParentModule.BlacklistPlatforms.ToList();
-								Blacklist.AddRange(ChildModule.BlacklistPlatforms);
-								ParentModule.BlacklistPlatforms = Blacklist.ToArray();
+								if (ParentPluginReference.WhitelistPlatforms == null)
+								{
+									ParentPluginReference.WhitelistPlatforms = ChildPluginReference.WhitelistPlatforms;
+								}
+								else 
+								{
+									ParentPluginReference.WhitelistPlatforms = ParentPluginReference.WhitelistPlatforms.Union(ChildPluginReference.WhitelistPlatforms).ToList();
+								}
 							}
+							if (ChildPluginReference.BlacklistPlatforms != null)
+							{
+								if (ParentPluginReference.BlacklistPlatforms == null)
+								{
+									ParentPluginReference.BlacklistPlatforms = ChildPluginReference.BlacklistPlatforms;
+								}
+								else
+								{
+									ParentPluginReference.BlacklistPlatforms = ParentPluginReference.BlacklistPlatforms.Union(ChildPluginReference.BlacklistPlatforms).ToList();
+								}
+							}
+						}
+						else
+						{
+							Parent.Descriptor.Plugins.Add(ChildPluginReference);
 						}
 					}
 				}
