Fix exception due to directory not existing when creating VSCode project.

#rnx

diff --git a/Engine/Source/Programs/UnrealBuildTool/ProjectFiles/VisualStudioCode/VSCodeProjectFileGenerator.cs b/Engine/Source/Programs/UnrealBuildTool/ProjectFiles/VisualStudioCode/VSCodeProjectFileGenerator.cs
--- a/Engine/Source/Programs/UnrealBuildTool/ProjectFiles/VisualStudioCode/VSCodeProjectFileGenerator.cs
+++ b/Engine/Source/Programs/UnrealBuildTool/ProjectFiles/VisualStudioCode/VSCodeProjectFileGenerator.cs
@@ -600,6 +600,8 @@ namespace UnrealBuildTool
 
 		private void WriteCppPropertiesFile(DirectoryReference OutputDirectory, bool RewriteIncludePaths, ProjectData Projects)
 		{
+			DirectoryReference.CreateDirectory(OutputDirectory);
+
 			JsonFile OutFile = new JsonFile();
 
 			OutFile.BeginRootObject();
