Add missing null checks to PluginDescriptor.cs

- Previously this would have been handled by WriteArray, now Plugins and Modules are a List and ToArray is called, the check needs to be moved up.

diff --git a/Engine/Source/Programs/UnrealBuildTool/System/PluginDescriptor.cs b/Engine/Source/Programs/UnrealBuildTool/System/PluginDescriptor.cs
--- a/Engine/Source/Programs/UnrealBuildTool/System/PluginDescriptor.cs
+++ b/Engine/Source/Programs/UnrealBuildTool/System/PluginDescriptor.cs
@@ -408,7 +408,10 @@ namespace UnrealBuildTool
 				Writer.WriteValue("bIsPluginExtension", bIsPluginExtension);
 			}
 
-			ModuleDescriptor.WriteArray(Writer, "Modules", Modules.ToArray());
+			if (Modules != null && Modules.Count > 0)
+			{
+				ModuleDescriptor.WriteArray(Writer, "Modules", Modules.ToArray());
+			}
 
 			LocalizationTargetDescriptor.WriteArray(Writer, "LocalizationTargets", LocalizationTargets);
 
@@ -422,7 +425,10 @@ namespace UnrealBuildTool
 				PostBuildSteps.Write(Writer, "PostBuildSteps");
 			}
 
-			PluginReferenceDescriptor.WriteArray(Writer, "Plugins", Plugins.ToArray());
+			if (Plugins != null && Plugins.Count > 0)
+			{
+				PluginReferenceDescriptor.WriteArray(Writer, "Plugins", Plugins.ToArray());
+			}
 		}
 
 		/// <summary>
