- Removed extra Engine directory under the UE4 project above the Platforms directory (Platforms was the only folder in the Engine folder)

diff --git a/Engine/Source/Programs/UnrealBuildTool/ProjectFiles/ProjectFileGenerator.cs b/Engine/Source/Programs/UnrealBuildTool/ProjectFiles/ProjectFileGenerator.cs
--- a/Engine/Source/Programs/UnrealBuildTool/ProjectFiles/ProjectFileGenerator.cs
+++ b/Engine/Source/Programs/UnrealBuildTool/ProjectFiles/ProjectFileGenerator.cs
@@ -1355,7 +1355,7 @@ namespace UnrealBuildTool
 				SubdirectoryNamesToExclude.Add("Binaries");
 				SubdirectoryNamesToExclude.Add("Content");
 
-				EngineProject.AddFilesToProject(SourceFileSearch.FindFiles(PlatformExtensionsDirectory, SubdirectoryNamesToExclude), UnrealBuildTool.RootDirectory);
+				EngineProject.AddFilesToProject(SourceFileSearch.FindFiles(PlatformExtensionsDirectory, SubdirectoryNamesToExclude), UnrealBuildTool.EngineDirectory);
 			}
 		}
 
@@ -2176,7 +2176,7 @@ namespace UnrealBuildTool
 				}
 				if (Path == UnrealBuildTool.EnginePlatformExtensionsDirectory)
 				{
-					BaseFolder = UnrealBuildTool.RootDirectory;
+					BaseFolder = UnrealBuildTool.EngineDirectory;
 					return FindOrAddProjectHelper(EngineProjectFileNameBase, BaseFolder);
 				}
 				if (Path == UnrealBuildTool.EnterpriseDirectory)
