Add Thread_UseAllCpuGroups to UnrealBuildTool config to enable UBT to use both CPU groups on high-core systems such as 64-core ThreadRippers.

See - https://docs.microsoft.com/en-us/dotnet/framework/configure-apps/file-schema/runtime/thread-useallcpugroups-element

Also added MaxParallelActions command-line flag to UBT that can be used to restrict the number of cores used for local (e.g. nonxge).

Added -cores=X+Y=Z option to benchmark script for profiling

diff --git a/Engine/Source/Programs/AutomationTool/Scripts/BenchmarkBuild/BenchmarkBuild.Automation.cs b/Engine/Source/Programs/AutomationTool/Scripts/BenchmarkBuild/BenchmarkBuild.Automation.cs
--- a/Engine/Source/Programs/AutomationTool/Scripts/BenchmarkBuild/BenchmarkBuild.Automation.cs
+++ b/Engine/Source/Programs/AutomationTool/Scripts/BenchmarkBuild/BenchmarkBuild.Automation.cs
@@ -29,7 +29,7 @@ namespace AutomationTool.Benchmark
 	[Help("noxge", "Do a compile without XGE / FASTBuild")]
 	[Help("singlecompile", "Do a single-file compile")]
 	[Help("nopcompile", "Do a nothing-needs-compiled compile")]
-	//[Help("processors=X+Y+Z", "Do noxge builds with these processor counts (default is Env.ProcCount)")]
+	[Help("cores=X+Y+Z", "Do noxge builds with these processor counts (default is Environment.ProcessorCount)")]
 	[Help("cook", "Do a cook for the specified platform")]
 	[Help("cookcold", "When cooking clear the local ddc before each run")]
 	[Help("cooknoshaderddc", "Do a cook test with no ddc for shaders")]
@@ -66,7 +66,7 @@ namespace AutomationTool.Benchmark
 			bool DoNoAcceleratedCompile = AllCompile | ParseParam("noxge") | ParseParam("nofastbuild");
 
 			bool DoCookTests = AllThings | ParseParam("cook");
-			bool DoColdCook = AllThings | ParseParam("coldcook");
+			bool DoColdCook = AllThings | ParseParam("cookcold");
 			bool DoHotDDC = AllThings | ParseParam("hotddc");
 			bool DoNoShaderDDC = AllThings | ParseParam("cooknoshaderddc");
 			bool DoNoDDC = ParseParam("cooknoddc");
@@ -104,7 +104,7 @@ namespace AutomationTool.Benchmark
 				ClientPlatforms.Add(EditorPlatform);
 			}
 
-			string ProcessorArg = ParseParamValue("processors", "");
+			string ProcessorArg = ParseParamValue("cores", "");
 
 			IEnumerable<int> ProcessorCounts = new[] { Environment.ProcessorCount };
 
diff --git a/Engine/Source/Programs/AutomationTool/Scripts/BenchmarkBuild/BenchmarkBuildTests.cs b/Engine/Source/Programs/AutomationTool/Scripts/BenchmarkBuild/BenchmarkBuildTests.cs
--- a/Engine/Source/Programs/AutomationTool/Scripts/BenchmarkBuild/BenchmarkBuildTests.cs
+++ b/Engine/Source/Programs/AutomationTool/Scripts/BenchmarkBuild/BenchmarkBuildTests.cs
@@ -53,7 +53,7 @@ namespace AutomationTool.Benchmark
 			}
 		}
 
-		public BenchmarkBuildTask(string InProject, string InTarget, UnrealTargetPlatform InPlatform, BuildOptions InOptions, int ProcessorCount=0)
+		public BenchmarkBuildTask(string InProject, string InTarget, UnrealTargetPlatform InPlatform, BuildOptions InOptions, int CoreCount=0)
 		{
 			bool IsVanillaUE4 = InProject == null || string.Equals(InProject, "UE4", StringComparison.OrdinalIgnoreCase);
 
@@ -80,11 +80,11 @@ namespace AutomationTool.Benchmark
 				TaskModifiers.Add(Arg);
 				Command.Params = new[] { Arg }; // need to also pass it to this
 
-				if (ProcessorCount > 0)
+				if (CoreCount > 0)
 				{
-					TaskModifiers.Add(string.Format("{0}c", ProcessorCount));
+					TaskModifiers.Add(string.Format("{0}c", CoreCount));
 
-					Command.UBTArgs += string.Format(" -ParallelActions={0}", ProcessorCount);
+					Command.UBTArgs += string.Format(" -MaxParallelActions={0}", CoreCount);
 				}
 			}
 			else
diff --git a/Engine/Source/Programs/UnrealBuildTool/Configuration/BuildConfiguration.cs b/Engine/Source/Programs/UnrealBuildTool/Configuration/BuildConfiguration.cs
--- a/Engine/Source/Programs/UnrealBuildTool/Configuration/BuildConfiguration.cs
+++ b/Engine/Source/Programs/UnrealBuildTool/Configuration/BuildConfiguration.cs
@@ -107,6 +107,14 @@ namespace UnrealBuildTool
 		[XmlConfigFile]
 		public bool bAllowParallelExecutor = true;
 
+		/// <summary>
+		/// Number of actions that can be executed in parallel. If 0 then code will pick a default based
+		/// on the number of cores available. Only applies to the ParallelExecutor
+		/// </summary>
+		[XmlConfigFile]
+		[CommandLine("-MaxParallelActions")]
+		public int MaxParallelActions = 0;
+
 		/// <summary>
 		/// If true, force header regeneration. Intended for the build machine.
 		/// </summary>
diff --git a/Engine/Source/Programs/UnrealBuildTool/Executors/HybridExecutor.cs b/Engine/Source/Programs/UnrealBuildTool/Executors/HybridExecutor.cs
--- a/Engine/Source/Programs/UnrealBuildTool/Executors/HybridExecutor.cs
+++ b/Engine/Source/Programs/UnrealBuildTool/Executors/HybridExecutor.cs
@@ -33,9 +33,10 @@ namespace UnrealBuildTool
 		/// <summary>
 		/// Constructor
 		/// </summary>
-		public HybridExecutor()
+		public HybridExecutor(int InMaxLocalActions)
 		{
-			this.LocalExecutor = new ParallelExecutor();
+			MaxLocalActions = InMaxLocalActions;
+			this.LocalExecutor = new ParallelExecutor(MaxLocalActions);
 			this.RemoteExecutor = new XGE();
 
 			XmlConfig.ApplyTo(this);
diff --git a/Engine/Source/Programs/UnrealBuildTool/Executors/ParallelExecutor.cs b/Engine/Source/Programs/UnrealBuildTool/Executors/ParallelExecutor.cs
--- a/Engine/Source/Programs/UnrealBuildTool/Executors/ParallelExecutor.cs
+++ b/Engine/Source/Programs/UnrealBuildTool/Executors/ParallelExecutor.cs
@@ -33,16 +33,16 @@ namespace UnrealBuildTool
 		}
 
 		/// <summary>
-		/// Processor count multiplier for local execution. Can be below 1 to reserve CPU for other tasks.
+		/// Maximum processor count for local execution. 
 		/// </summary>
 		[XmlConfigFile]
-		double ProcessorCountMultiplier = 1.0;
+		int MaxProcessorCount = int.MaxValue;
 
 		/// <summary>
-		/// Maximum processor count for local execution. 
+		/// Processor count multiplier for local execution. Can be below 1 to reserve CPU for other tasks.
 		/// </summary>
 		[XmlConfigFile]
-		int MaxProcessorCount = int.MaxValue;
+		double ProcessorCountMultiplier = 1.0;
 
 		/// <summary>
 		/// When enabled, will stop compiling targets after a compile error occurs.
@@ -50,12 +50,29 @@ namespace UnrealBuildTool
 		[XmlConfigFile]
 		bool bStopCompilationAfterErrors = false;
 
+		/// <summary>
+		/// How many processes that will be executed in parallel
+		/// </summary>
+		int MaxProcesses;
+
 		/// <summary>
 		/// Constructor
 		/// </summary>
-		public ParallelExecutor()
+		/// <param name="MaxLocalActions">How many actions to execute in parallel</param>
+		public ParallelExecutor(int MaxLocalActions)
 		{
 			XmlConfig.ApplyTo(this);
+
+			// if specified this caps how many processors we can use
+			if (MaxLocalActions > 0)
+			{
+				MaxProcesses = MaxLocalActions;
+			}
+			else
+			{
+				// Figure out how many processors to use
+				MaxProcesses = Math.Min((int)(Environment.ProcessorCount * ProcessorCountMultiplier), MaxProcessorCount);
+			}
 		}
 
 		/// <summary>
@@ -81,8 +98,6 @@ namespace UnrealBuildTool
 		/// <returns>True if all the tasks successfully executed, or false if any of them failed.</returns>
 		public override bool ExecuteActions(List<Action> InputActions, bool bLogDetailedActionStats)
 		{
-			// Figure out how many processors to use
-			int MaxProcesses = Math.Min((int)(Environment.ProcessorCount * ProcessorCountMultiplier), MaxProcessorCount);
 			Log.TraceInformation("Building {0} {1} with {2} {3}...", InputActions.Count, (InputActions.Count == 1) ? "action" : "actions", MaxProcesses, (MaxProcesses == 1)? "process" : "processes");
 
 			// Create actions with all our internal metadata
diff --git a/Engine/Source/Programs/UnrealBuildTool/System/ActionGraph.cs b/Engine/Source/Programs/UnrealBuildTool/System/ActionGraph.cs
--- a/Engine/Source/Programs/UnrealBuildTool/System/ActionGraph.cs
+++ b/Engine/Source/Programs/UnrealBuildTool/System/ActionGraph.cs
@@ -184,7 +184,7 @@ namespace UnrealBuildTool
 				ActionExecutor Executor;
 				if (BuildConfiguration.bAllowHybridExecutor && HybridExecutor.IsAvailable())
 				{
-					Executor = new HybridExecutor();
+					Executor = new HybridExecutor(BuildConfiguration.MaxParallelActions);
 				}
 				else if (BuildConfiguration.bAllowXGE && XGE.IsAvailable())
 				{
@@ -206,7 +206,7 @@ namespace UnrealBuildTool
 				}
 				else if(BuildConfiguration.bAllowParallelExecutor && ParallelExecutor.IsAvailable())
 				{
-					Executor = new ParallelExecutor();
+					Executor = new ParallelExecutor(BuildConfiguration.MaxParallelActions);
 				}
 				else
 				{
diff --git a/Engine/Source/Programs/UnrealBuildTool/app.config b/Engine/Source/Programs/UnrealBuildTool/app.config
--- a/Engine/Source/Programs/UnrealBuildTool/app.config
+++ b/Engine/Source/Programs/UnrealBuildTool/app.config
@@ -1,3 +1,11 @@
 <?xml version="1.0"?>
 <configuration>
-<startup><supportedRuntime version="v4.0" sku=".NETFramework,Version=v4.6.2"/></startup></configuration>
+  <startup>
+    <supportedRuntime version="v4.0" sku=".NETFramework,Version=v4.6.2"/>
+  </startup>
+  <runtime>
+    <Thread_UseAllCpuGroups enabled="true"/>
+    <GCCpuGroup enabled="true"/>
+    <gcServer enabled="true"/>
+  </runtime>
+</configuration>
