Update default min Win10 SDK for HL packaging
SDK's < 17763 will fail to package if one of these SDKs is installed on the build machine (eg: 16299)

diff --git a/Engine/Source/Programs/UnrealBuildTool/Platform/HoloLens/ManifestGenerator.cs b/Engine/Source/Programs/UnrealBuildTool/Platform/HoloLens/ManifestGenerator.cs
--- a/Engine/Source/Programs/UnrealBuildTool/Platform/HoloLens/ManifestGenerator.cs
+++ b/Engine/Source/Programs/UnrealBuildTool/Platform/HoloLens/ManifestGenerator.cs
@@ -1155,7 +1155,7 @@ namespace UnrealBuildTool
 					TargetDeviceFamily.Attributes.Append(NameAttribute);
 
 					XmlAttribute MinVersionAttribute = AppxManifestXmlDocument.CreateAttribute("MinVersion");
-					string versionString = CreateStringValue("MinimumPlatformVersion", "Package.Dependencies.TargetDeviceFamily[0].MinVersion", "MinimumPlatformVersion", "MinVersion", "10.0.10240.0");
+					string versionString = CreateStringValue("MinimumPlatformVersion", "Package.Dependencies.TargetDeviceFamily[0].MinVersion", "MinimumPlatformVersion", "MinVersion", "10.0.17763.0");
 					MinVersionAttribute.Value = versionString;
 					TargetDeviceFamily.Attributes.Append(MinVersionAttribute);
 
