Prevent store logo from falling back to the wrong size logo

The application logo is not a valid fallback if the package logo is not available. This will cause it to fail WACK, and prevent it being uploaded to the store / downloaded on a user's device.

diff --git a/Engine/Source/Programs/UnrealBuildTool/Platform/HoloLens/ManifestGenerator.cs b/Engine/Source/Programs/UnrealBuildTool/Platform/HoloLens/ManifestGenerator.cs
--- a/Engine/Source/Programs/UnrealBuildTool/Platform/HoloLens/ManifestGenerator.cs
+++ b/Engine/Source/Programs/UnrealBuildTool/Platform/HoloLens/ManifestGenerator.cs
@@ -1180,21 +1180,11 @@ namespace UnrealBuildTool
 					Properties.AppendChild(PublisherDisplayName);
 
 					XmlElement PackageLogo = AppxManifestXmlDocument.CreateElement("Logo");
-					// Some applications may not have a package logo and use the application logo instead.
-					// Try logos in the following order:
-					//   1. Project package logo
-					//   2. Project application logo
-					//   3. Engine application logo (the engine always uses a single logo for package and application)
-					if (CopyAndReplaceBinaryIntermediate("StoreLogo.png", false))
-					{
+                    if (CopyAndReplaceBinaryIntermediate("StoreLogo.png"))
+                    {
 						PackageLogo.InnerText = BuildResourceSubPath + "\\StoreLogo.png";
 						Properties.AppendChild(PackageLogo);
 					}
-					else if (CopyAndReplaceBinaryIntermediate("Logo.png"))
-					{
-						PackageLogo.InnerText = BuildResourceSubPath + "\\Logo.png";
-						Properties.AppendChild(PackageLogo);
-					}
 					else
 					{
 						Log.TraceError("Unable to stage package logo.");
@@ -1732,21 +1722,11 @@ namespace UnrealBuildTool
 			AddResourceEntry("PackageDescription", "PackageDescription", "Package.Properties.Description", "/Script/EngineSettings.GeneralProjectSettings", "Description", "");
 
 			XmlElement PackageLogo = AppxManifestXmlDocument.CreateElement("Logo");
-			// Some applications may not have a package logo and use the application logo instead.
-			// Try logos in the following order:
-			//   1. Project package logo
-			//   2. Project application logo
-			//   3. Engine application logo (the engine always uses a single logo for package and application)
-			if (CopyAndReplaceBinaryIntermediate("StoreLogo.png", false))
+			if (CopyAndReplaceBinaryIntermediate("StoreLogo.png"))
 			{
 				PackageLogo.InnerText = BuildResourceSubPath + "\\StoreLogo.png";
 				Properties.AppendChild(PackageLogo);
 			}
-			else if (CopyAndReplaceBinaryIntermediate("Logo.png"))
-			{
-				PackageLogo.InnerText = BuildResourceSubPath + "\\Logo.png";
-				Properties.AppendChild(PackageLogo);
-			}
 			else
 			{
 				Log.TraceError("Unable to stage package logo.");
