Set default resource entries to pass WACK
Previously, resources.resw contained empty entries which would fail WACK

diff --git a/Engine/Source/Programs/UnrealBuildTool/Platform/HoloLens/ManifestGenerator.cs b/Engine/Source/Programs/UnrealBuildTool/Platform/HoloLens/ManifestGenerator.cs
--- a/Engine/Source/Programs/UnrealBuildTool/Platform/HoloLens/ManifestGenerator.cs
+++ b/Engine/Source/Programs/UnrealBuildTool/Platform/HoloLens/ManifestGenerator.cs
@@ -1530,9 +1530,22 @@ namespace UnrealBuildTool
 				}
 			}
 
-			// Any culture with a culture-specific value will override the neutral value,
-			// even for unrelated cultures.  So propagate the neutral value to avoid this happening.
-			if (IsEverLocalized)
+            // Values cannot be empty in the resource file, or the appx will fail WACK.
+            if (string.IsNullOrEmpty(NeutralValue.Trim()))
+            {
+                if (!string.IsNullOrEmpty(DefaultValue.Trim()))
+                {
+                    NeutralValue = DefaultValue.Trim();
+                }
+                else
+                {
+                    NeutralValue = "Missing Entry";
+                }
+            }
+
+            // Any culture with a culture-specific value will override the neutral value,
+            // even for unrelated cultures.  So propagate the neutral value to avoid this happening.
+            if (IsEverLocalized)
 			{
 				for (int i = 0; i < CulturesToStage.Count; ++i)
 				{
@@ -1719,7 +1732,7 @@ namespace UnrealBuildTool
 			XmlElement PackageDescription = AppxManifestXmlDocument.CreateElement("Description");
 			PackageDescription.InnerText = "ms-resource:PackageDescription";
 			Properties.AppendChild(PackageDescription);
-			AddResourceEntry("PackageDescription", "PackageDescription", "Package.Properties.Description", "/Script/EngineSettings.GeneralProjectSettings", "Description", "");
+			AddResourceEntry("PackageDescription", "PackageDescription", "Package.Properties.Description", "/Script/EngineSettings.GeneralProjectSettings", "Description", "No Description");
 
 			XmlElement PackageLogo = AppxManifestXmlDocument.CreateElement("Logo");
 			if (CopyAndReplaceBinaryIntermediate("StoreLogo.png"))
