Fix ModuleRules so PhysX module is never added if Chaos is enabled.
Fixes NonUnity build.

diff --git a/Engine/Source/Programs/UnrealBuildTool/Configuration/ModuleRules.cs b/Engine/Source/Programs/UnrealBuildTool/Configuration/ModuleRules.cs
--- a/Engine/Source/Programs/UnrealBuildTool/Configuration/ModuleRules.cs
+++ b/Engine/Source/Programs/UnrealBuildTool/Configuration/ModuleRules.cs
@@ -1121,7 +1121,7 @@ namespace UnrealBuildTool
             // definitions used outside of PhysX/APEX need to be set here, not in PhysX.Build.cs or APEX.Build.cs, 
             // since we need to make sure we always set it, even to 0 (because these are Private dependencies, the
             // defines inside their Build.cs files won't leak out)
-            if (Target.bCompilePhysX == true)
+            if (Target.bCompilePhysX == true && Target.bCompileChaos == false && Target.bUseChaos == false)
 			{
 				PrivateDependencyModuleNames.Add("PhysX");
 			}
