From Microsoft:
Autoincrement project version after each build
Opt-in in HoloLens project settings under packaging/ Auto Increment Version

The store requires unique version numbers for each update, this makes it easier to track.

diff --git a/Engine/Plugins/Runtime/AR/Microsoft/HoloLensAR/Source/HoloLensPlatformEditor/Private/HoloLensTargetSettings.h b/Engine/Plugins/Runtime/AR/Microsoft/HoloLensAR/Source/HoloLensPlatformEditor/Private/HoloLensTargetSettings.h
--- a/Engine/Plugins/Runtime/AR/Microsoft/HoloLensAR/Source/HoloLensPlatformEditor/Private/HoloLensTargetSettings.h
+++ b/Engine/Plugins/Runtime/AR/Microsoft/HoloLensAR/Source/HoloLensPlatformEditor/Private/HoloLensTargetSettings.h
@@ -51,6 +51,9 @@ public:
 	 */
 	UPROPERTY(EditAnywhere, config, Category = "Packaging", AdvancedDisplay, Meta = (DisplayName = "Use Retail Windows Store Environment"))
 	bool bBuildForRetailWindowsStore;
+	
+	UPROPERTY(EditAnywhere, Config, Category = "Packaging")
+	bool bAutoIncrementVersion;
 
 	UPROPERTY(EditAnywhere, config, Category = "Rendering", meta = (DisplayName = "Enable PIX Profiling"))
 	bool bEnablePIXProfiling;
diff --git a/Engine/Source/Programs/UnrealBuildTool/Platform/HoloLens/ManifestGenerator.cs b/Engine/Source/Programs/UnrealBuildTool/Platform/HoloLens/ManifestGenerator.cs
--- a/Engine/Source/Programs/UnrealBuildTool/Platform/HoloLens/ManifestGenerator.cs
+++ b/Engine/Source/Programs/UnrealBuildTool/Platform/HoloLens/ManifestGenerator.cs
@@ -615,21 +615,119 @@ namespace UnrealBuildTool
 
 		}
 
-		/// <summary>
-		/// Kicks off manifest generation. Will always attempt to fully calculate a new manifest but will not update the output
-		/// file unless there are changes (to avoid unnecessary copies when deploying).
-		/// </summary>
-		/// <param name="TargetPlatform">The platform we're generating a manifest for.</param>
-		/// <param name="TargetArchitecture">The architecture we're generating a manifest for.</param>
-		/// <param name="InOutputPath">Path to write manifest files to.</param>
-		/// <param name="InIntermediatePath">Path to store temporary intermediate data (e.g. XML resource file).</param>
-		/// <param name="InProjectFile">Path to the uproject file</param>
-		/// <param name="InProjectDirectory">Directory containing the uproject file or the base engine path if no project file is specified (for content only builds).</param>
-		/// <param name="InTargetConfigs">Configurations to build manifest data for. Each configuration will generate it's own application entry.</param>
-		/// <param name="InExecutables">The launch executable for each configuration. Must match the length and order of InTargetConfigs.</param>
-		/// <param name="InWinMDReferences">The WinMD references that should be added as activatable types</param>
-		/// <returns>A list of all updated target files</returns>
-		public List<string> CreateManifest(UnrealTargetPlatform TargetPlatform, WindowsArchitecture TargetArchitecture, string InOutputPath, string InIntermediatePath, FileReference InProjectFile, string InProjectDirectory, List<UnrealTargetConfiguration> InTargetConfigs, List<string> InExecutables, IEnumerable<WinMDRegistrationInfo> InWinMDReferences)
+        // Add a new key with the given value to the input ini file.  The new value will take effect on the next build since the ini has already been ingested.
+        bool WriteProjectIniString(string iniFilePath, string SectionName, string keyName, string newString)
+        {
+            if (!File.Exists(iniFilePath))
+            {
+                return false;
+            }
+
+            string[] lines = File.ReadAllLines(iniFilePath);
+            List<string> updatedLines = new List<string>();
+            for (int i = 0; i < lines.Length; i++)
+            {
+                string line = lines[i].Trim();
+				if (line.Equals("[" + SectionName.Trim() + "]"))
+                {
+                    updatedLines = lines.ToList();
+					updatedLines.Insert(i + 1, string.Concat(keyName.Trim(), "=", newString.Trim()));
+
+                    break;
+                }
+            }
+
+            if (updatedLines.Count > 0)
+            {
+                File.WriteAllLines(iniFilePath, updatedLines.ToArray());
+                return true;
+            }
+
+            return false;
+        }
+
+        // Update any ini keyName entries in the project directory with newString.  The new value will take effect on the next build since the ini has already been ingested.
+        bool UpdateProjectIniString(FileReference InProjectFile, UnrealTargetPlatform TargetPlatform, string SectionName, string keyName, string newString)
+        {
+            DirectoryReference IniDirRef = DirectoryReference.FromFile(InProjectFile);
+            bool iniUpdated = false;
+
+            List<string> projectIniFilePaths = new List<string>();
+            List<ConfigFile> Files = new List<ConfigFile>();
+            foreach (FileReference IniFileName in ConfigHierarchy.EnumerateConfigFileLocations(ConfigHierarchyType.Game, IniDirRef, TargetPlatform))
+            {
+                // If ini file does not exist, or is not in the project directory, move on to the next one.
+                if (!File.Exists(IniFileName.FullName)
+					|| !IniFileName.FullName.StartsWith(InProjectFile.Directory.FullName))
+                {
+                    continue;
+                }
+
+                Log.TraceLog("Looking at file: " + IniFileName.FullName);
+                projectIniFilePaths.Add(IniFileName.FullName);
+
+                bool isInDesiredSection = false;
+                bool rewriteCurrentFile = false;
+
+                string[] lines = File.ReadAllLines(IniFileName.FullName);
+                for (int i = 0; i < lines.Length; i++)
+                {
+                    string line = lines[i].Trim();
+                    if (line.StartsWith("[") && line.EndsWith("]"))
+                    {
+						isInDesiredSection = line.Equals("[" + SectionName.Trim() + "]");
+                    }
+                    if (!isInDesiredSection)
+                    {
+                        continue;
+                    }
+
+					if (line.StartsWith(keyName.Trim() + "="))
+                    {
+                        Log.TraceLog("Found string match.");
+                        iniUpdated = true;
+                        rewriteCurrentFile = true;
+
+						lines[i] = string.Concat(keyName.Trim(), "=", newString.Trim());
+                    }
+                }
+
+                if (rewriteCurrentFile)
+                {
+                    File.WriteAllLines(IniFileName.FullName, lines);
+                }
+            }
+
+            if (!iniUpdated)
+            {
+                // The project ini did not have a project version entry.  This will happen on the first build which defaults to an entry in the engine ini.
+                foreach (string iniPath in projectIniFilePaths)
+                {
+                    if (WriteProjectIniString(iniPath, SectionName, keyName, newString))
+                    {
+                        iniUpdated = true;
+                    }
+                }
+            }
+
+            return iniUpdated;
+        }
+
+        /// <summary>
+        /// Kicks off manifest generation. Will always attempt to fully calculate a new manifest but will not update the output
+        /// file unless there are changes (to avoid unnecessary copies when deploying).
+        /// </summary>
+        /// <param name="TargetPlatform">The platform we're generating a manifest for.</param>
+        /// <param name="TargetArchitecture">The architecture we're generating a manifest for.</param>
+        /// <param name="InOutputPath">Path to write manifest files to.</param>
+        /// <param name="InIntermediatePath">Path to store temporary intermediate data (e.g. XML resource file).</param>
+        /// <param name="InProjectFile">Path to the uproject file</param>
+        /// <param name="InProjectDirectory">Directory containing the uproject file or the base engine path if no project file is specified (for content only builds).</param>
+        /// <param name="InTargetConfigs">Configurations to build manifest data for. Each configuration will generate it's own application entry.</param>
+        /// <param name="InExecutables">The launch executable for each configuration. Must match the length and order of InTargetConfigs.</param>
+        /// <param name="InWinMDReferences">The WinMD references that should be added as activatable types</param>
+        /// <returns>A list of all updated target files</returns>
+        public List<string> CreateManifest(UnrealTargetPlatform TargetPlatform, WindowsArchitecture TargetArchitecture, string InOutputPath, string InIntermediatePath, FileReference InProjectFile, string InProjectDirectory, List<UnrealTargetConfiguration> InTargetConfigs, List<string> InExecutables, IEnumerable<WinMDRegistrationInfo> InWinMDReferences)
 		{
 			// Check parameter values are valid
 			if (InTargetConfigs.Count != InExecutables.Count)
@@ -788,8 +886,43 @@ namespace UnrealBuildTool
 				return null;
 			}
 
-			// Create the appxmanifest document
-			AppxManifestXmlDocument = new XmlDocument();
+            // Check autoincrement flag
+            bool bAutoIncrementVersion;
+            if (EngineIni.GetBool(TargetSettings, "bAutoIncrementVersion", out bAutoIncrementVersion) && bAutoIncrementVersion)
+            {
+                string currentVersion;
+                if (GameIni.GetString("/Script/EngineSettings.GeneralProjectSettings", "ProjectVersion", out currentVersion))
+                {
+                    Log.TraceLog("Automatically incrementing version. Starting version is " + currentVersion);
+
+                    string[] versionEntries = currentVersion.Split(new char[] { '.' });
+                    if (versionEntries.Length == 4)
+                    {
+                        uint versionEntryToIncrement = Convert.ToUInt32(versionEntries[2]);
+                        versionEntryToIncrement++;
+
+                        string newVersion = string.Concat(versionEntries[0], ".", versionEntries[1], ".", versionEntryToIncrement, ".", versionEntries[3]);
+                        Log.TraceLog("Writing new version string: " + newVersion);
+
+                        // Update GameIni with newVersion.
+                        if (!UpdateProjectIniString(InProjectFile, TargetPlatform, "/Script/EngineSettings.GeneralProjectSettings", "ProjectVersion", newVersion))
+                        {
+                            Log.TraceWarning("Auto incrementing the project version was unsuccessful.");
+                        }
+                    }
+                    else
+                    {
+                        Log.TraceWarning("Auto increment was desired, but the number of decimals in the version string was unexpected.");
+                    }
+                }
+                else
+                {
+                    Log.TraceWarning("Auto increment was desired, but the existing version could not be identified.");
+                }
+            }
+
+            // Create the appxmanifest document
+            AppxManifestXmlDocument = new XmlDocument();
 			XmlDeclaration Declaration = AppxManifestXmlDocument.CreateXmlDeclaration("1.0", Encoding.UTF8.BodyName, null);
 			AppxManifestXmlDocument.AppendChild(Declaration);
 			
@@ -1635,7 +1768,7 @@ namespace UnrealBuildTool
 			XmlNode Properties = GetProperties();
 			AddElementIfValid(Package, Properties, true);
 
-			XmlNode Dependencies = GetDependencies();
+			XmlNode Dependencies = GetDependencies(TargetConfigs);
 			AddElementIfValid(Package, Dependencies, true);
 
 			XmlNode Resources = GetResources();
@@ -1751,7 +1884,7 @@ namespace UnrealBuildTool
 		/// <summary>
 		/// Gather all information for the Dependencies element of the manifest.
 		/// </summary>
-		private XmlNode GetDependencies()
+		private XmlNode GetDependencies(List<UnrealTargetConfiguration> TargetConfigs)
 		{
 			XmlElement Dependencies = AppxManifestXmlDocument.CreateElement("Dependencies");
 			
@@ -1777,8 +1910,16 @@ namespace UnrealBuildTool
 				Dependencies.AppendChild(PackageDependency);
 
 				XmlAttribute NameAttribute = AppxManifestXmlDocument.CreateAttribute("Name");
-				NameAttribute.Value = "Microsoft.VCLibs.140.00";
-				PackageDependency.Attributes.Append(NameAttribute);
+                if (TargetConfigs.Contains(UnrealTargetConfiguration.Debug)
+                    && HoloLensToolChain.UseDebugCRT())
+                {
+                    NameAttribute.Value = "Microsoft.VCLibs.140.00.Debug";
+                }
+                else
+                {
+                    NameAttribute.Value = "Microsoft.VCLibs.140.00";
+                }
+                PackageDependency.Attributes.Append(NameAttribute);
 
 				XmlAttribute PublisherAttribute = AppxManifestXmlDocument.CreateAttribute("Publisher");
 				PublisherAttribute.Value = "CN=Microsoft Corporation, O=Microsoft Corporation, L=Redmond, S=Washington, C=US";
diff --git a/Engine/Source/Programs/UnrealBuildTool/Platform/HoloLens/UEBuildHoloLens.cs b/Engine/Source/Programs/UnrealBuildTool/Platform/HoloLens/UEBuildHoloLens.cs
--- a/Engine/Source/Programs/UnrealBuildTool/Platform/HoloLens/UEBuildHoloLens.cs
+++ b/Engine/Source/Programs/UnrealBuildTool/Platform/HoloLens/UEBuildHoloLens.cs
@@ -52,7 +52,13 @@ namespace UnrealBuildTool
 		public string Win10SDKVersionString = null;
 
 		internal Version Win10SDKVersion = null;
-	}
+
+        /// <summary>
+        /// Automatically increment the project version after each build.
+        /// </summary>
+        [ConfigFile(ConfigHierarchyType.Engine, "/Script/HoloLensPlatformEditor.HoloLensTargetSettings", "bAutoIncrementVersion")]
+        public bool bAutoIncrementVersion = false;
+    }
 
 	/// <summary>
 	/// Read-only wrapper for HoloLens-specific target settings
