Fix build break in HL manifest generator

diff --git a/Engine/Source/Programs/UnrealBuildTool/Platform/HoloLens/ManifestGenerator.cs b/Engine/Source/Programs/UnrealBuildTool/Platform/HoloLens/ManifestGenerator.cs
--- a/Engine/Source/Programs/UnrealBuildTool/Platform/HoloLens/ManifestGenerator.cs
+++ b/Engine/Source/Programs/UnrealBuildTool/Platform/HoloLens/ManifestGenerator.cs
@@ -1910,16 +1910,8 @@ namespace UnrealBuildTool
 				Dependencies.AppendChild(PackageDependency);
 
 				XmlAttribute NameAttribute = AppxManifestXmlDocument.CreateAttribute("Name");
-                if (TargetConfigs.Contains(UnrealTargetConfiguration.Debug)
-                    && HoloLensToolChain.UseDebugCRT())
-                {
-                    NameAttribute.Value = "Microsoft.VCLibs.140.00.Debug";
-                }
-                else
-                {
-                    NameAttribute.Value = "Microsoft.VCLibs.140.00";
-                }
-                PackageDependency.Attributes.Append(NameAttribute);
+				NameAttribute.Value = "Microsoft.VCLibs.140.00";
+				PackageDependency.Attributes.Append(NameAttribute);
 
 				XmlAttribute PublisherAttribute = AppxManifestXmlDocument.CreateAttribute("Publisher");
 				PublisherAttribute.Value = "CN=Microsoft Corporation, O=Microsoft Corporation, L=Redmond, S=Washington, C=US";
