On iOS builds, copy LaunchScreen.storyboardc into the app bundle if the project setting is enabled.

diff --git a/Engine/Source/Programs/AutomationTool/IOS/IOSPlatform.Automation.cs b/Engine/Source/Programs/AutomationTool/IOS/IOSPlatform.Automation.cs
--- a/Engine/Source/Programs/AutomationTool/IOS/IOSPlatform.Automation.cs
+++ b/Engine/Source/Programs/AutomationTool/IOS/IOSPlatform.Automation.cs
@@ -880,14 +880,14 @@ public class IOSPlatform : Platform
 
 					// copy the icons/launch screens from the engine
 					{
-						DirectoryReference DataPath = DirectoryReference.Combine(SC.EngineRoot, "Build", "IOS", "Resources", "Graphics");
-						StageImageAndIconFiles(DataPath, bSupportsPortrait, bSupportsLandscape, SC, bSkipIcons);
+						DirectoryReference GraphicsDataPath = DirectoryReference.Combine(SC.EngineRoot, "Build", "IOS", "Resources", "Graphics");
+						StageImageAndIconFiles(Params, GraphicsDataPath, bSupportsPortrait, bSupportsLandscape, SC, bSkipIcons);
 					}
 
 					// copy the icons/launch screens from the game (may stomp the engine copies)
 					{
-						DirectoryReference DataPath = DirectoryReference.Combine(SC.ProjectRoot, "Build", "IOS", "Resources", "Graphics");
-						StageImageAndIconFiles(DataPath, bSupportsPortrait, bSupportsLandscape, SC, bSkipIcons);
+						DirectoryReference GraphicsDataPath = DirectoryReference.Combine(SC.ProjectRoot, "Build", "IOS", "Resources", "Graphics");
+						StageImageAndIconFiles(Params, GraphicsDataPath, bSupportsPortrait, bSupportsLandscape, SC, bSkipIcons);
 					}
 				}
 
@@ -947,9 +947,18 @@ public class IOSPlatform : Platform
 		}
 	}
 
-	private void StageImageAndIconFiles(DirectoryReference DataPath, bool bSupportsPortrait, bool bSupportsLandscape, DeploymentContext SC, bool bSkipIcons)
+	private void StageImageAndIconFiles(ProjectParams Params, DirectoryReference GraphicsDataPath, bool bSupportsPortrait, bool bSupportsLandscape, DeploymentContext SC, bool bSkipIcons)
 	{
-		if (DirectoryReference.Exists(DataPath))
+
+		bool bLaunchscreenStoryboard = false;
+		ConfigHierarchy PlatformGameConfig;
+		bool bXCArchive = false;
+		if (Params.EngineConfigs.TryGetValue(SC.StageTargetPlatform.PlatformType, out PlatformGameConfig))
+		{
+			PlatformGameConfig.GetBool("/Script/IOSRuntimeSettings.IOSRuntimeSettings", "bLaunchscreenStoryboard", out bLaunchscreenStoryboard);
+		}
+
+		if (DirectoryReference.Exists(GraphicsDataPath) && !bLaunchscreenStoryboard)
 		{
 			List<string> ImageFileNames = new List<string>();
 			if (bSupportsPortrait)
@@ -981,7 +990,7 @@ public class IOSPlatform : Platform
 
 			foreach (string ImageFileName in ImageFileNames)
 			{
-				FileReference ImageFile = FileReference.Combine(DataPath, ImageFileName);
+				FileReference ImageFile = FileReference.Combine(GraphicsDataPath, ImageFileName);
 				if (FileReference.Exists(ImageFile))
 				{
 					SC.StageFile(StagedFileType.SystemNonUFS, ImageFile, new StagedFileReference(ImageFileName));
@@ -990,7 +999,7 @@ public class IOSPlatform : Platform
 
 			if (!bSkipIcons)
 			{
-				SC.StageFiles(StagedFileType.SystemNonUFS, DataPath, "Icon*.png", StageFilesSearch.TopDirectoryOnly, StagedDirectoryReference.Root);
+				SC.StageFiles(StagedFileType.SystemNonUFS, GraphicsDataPath, "Icon*.png", StageFilesSearch.TopDirectoryOnly, StagedDirectoryReference.Root);
 			}
 		}
 	}
diff --git a/Engine/Source/Programs/UnrealBuildTool/Platform/IOS/UEDeployIOS.cs b/Engine/Source/Programs/UnrealBuildTool/Platform/IOS/UEDeployIOS.cs
--- a/Engine/Source/Programs/UnrealBuildTool/Platform/IOS/UEDeployIOS.cs
+++ b/Engine/Source/Programs/UnrealBuildTool/Platform/IOS/UEDeployIOS.cs
@@ -28,6 +28,7 @@
 			set { bForDistribution = value; }
 		}
 		bool bForDistribution = false;
+		public static string LaunchStoryboardPath;
 
 		protected class VersionUtilities
 		{
@@ -42,7 +43,7 @@
 				set;
 			}
 
-			
+			public static bool bLaunchscreenStoryboard = false;
 
 			static string RunningVersionFilename
 			{
@@ -210,11 +211,11 @@
 			string OldPListData = File.Exists(PListFile) ? File.ReadAllText(PListFile) : "";
 
 			// determine if there is a Lanchscreen Storyboard THEN a launch.xib
-			string LaunchStoryboard = InEngineDir + "/Build/IOS/Resources/Interface/LaunchScreen.storyboard";
+			LaunchStoryboardPath = InEngineDir + "/Build/IOS/Resources/Interface";
 			string LaunchXib = InEngineDir + "/Build/IOS/Resources/Interface/LaunchScreen.xib";
-			if (File.Exists(BuildDirectory + "/Resources/Interface/LaunchScreen.storyboard"))
-			{
-				LaunchStoryboard = BuildDirectory + "/Resources/Interface/LaunchScreen.storyboard";
+			if (Directory.Exists(BuildDirectory + "/Resources/Interface/LaunchScreen.storyboardc"))
+			{
+				LaunchStoryboardPath = BuildDirectory + "/Resources/Interface";
 			}
 			else if (File.Exists(BuildDirectory + "/Resources/Interface/LaunchScreen.xib"))
 			{
@@ -401,6 +402,8 @@
 			string ExtraData = "";
 			Ini.GetString("/Script/IOSRuntimeSettings.IOSRuntimeSettings", "AdditionalPlistData", out ExtraData);
 
+			Ini.GetBool("/Script/IOSRuntimeSettings.IOSRuntimeSettings", "bLaunchscreenStoryboard", out VersionUtilities.bLaunchscreenStoryboard);
+
 			// generate the plist file
 			StringBuilder Text = new StringBuilder();
 			Text.AppendLine("<?xml version=\"1.0\" encoding=\"UTF-8\"?>");
@@ -577,7 +580,7 @@
 				Text.AppendLine("\t\t</dict>");
 				Text.AppendLine("\t</dict>");
 			}
-			if (File.Exists(LaunchStoryboard) || File.Exists(LaunchXib))
+			if (Directory.Exists(LaunchStoryboardPath + "/LaunchScreen.storyboardc") || File.Exists(LaunchXib) && VersionUtilities.bLaunchscreenStoryboard)
 			{
 				Text.AppendLine("\t<key>UILaunchStoryboardName</key>");
 				Text.AppendLine("\t<string>LaunchScreen</string>");
@@ -878,8 +881,15 @@
 			// copy engine assets in (IOS and TVOS shared in IOS)
 			if (bSkipDefaultPNGs)
 			{
-				//Copy storyboard or xib
-				CopyFiles(InEngineDir + "/Build/IOS/Resources/Interface", AppDirectory, "*", true);
+				if (VersionUtilities.bLaunchscreenStoryboard)
+				{
+					if(!Directory.Exists(AppDirectory + "/LaunchScreen.storyboardc"))
+					{
+						Directory.CreateDirectory(AppDirectory + "/LaunchScreen.storyboardc");
+					}
+					
+					CopyFolder(LaunchStoryboardPath + "/LaunchScreen.storyboardc", AppDirectory + "/LaunchScreen.storyboardc", true);
+				}
 				// we still want default icons
 				if (!bSkipIcons)
 				{
@@ -927,7 +937,7 @@
 				{
 					CopyFiles(BuildDirectory + "/Resources/Graphics", AppDirectory, "Icon*.png", true);
 				}
-				if (bSupportsPortrait)
+				if (bSupportsPortrait && !VersionUtilities.bLaunchscreenStoryboard)
 				{
 					CopyFiles(BuildDirectory + "/Resources/Graphics", AppDirectory, "Default-IPhone6.png", true);
 					CopyFiles(BuildDirectory + "/Resources/Graphics", AppDirectory, "Default-IPhone6Plus-Portrait.png", true);
@@ -939,7 +949,7 @@
 					CopyFiles(BuildDirectory + "/Resources/Graphics", AppDirectory, "Default-IPhoneXSMax-Portrait.png", true);
 					CopyFiles(BuildDirectory + "/Resources/Graphics", AppDirectory, "Default-IPhoneXR-Portrait.png", true);
 				}
-				if (bSupportsLandscape)
+				if (bSupportsLandscape && !VersionUtilities.bLaunchscreenStoryboard)
 				{
 					CopyFiles(BuildDirectory + "/Resources/Graphics", AppDirectory, "Default-IPhone6-Landscape.png", true);
 					CopyFiles(BuildDirectory + "/Resources/Graphics", AppDirectory, "Default-IPhone6Plus-Landscape.png", true);
@@ -951,8 +961,11 @@
 					CopyFiles(BuildDirectory + "/Resources/Graphics", AppDirectory, "Default-IPhoneXSMax-Landscape.png", true);
 					CopyFiles(BuildDirectory + "/Resources/Graphics", AppDirectory, "Default-IPhoneXR-Landscape.png", true);
 				}
-				CopyFiles(BuildDirectory + "/Resources/Graphics", AppDirectory, "[REDACTED_EMAIL]", true);
-				CopyFiles(BuildDirectory + "/Resources/Graphics", AppDirectory, "[REDACTED_EMAIL]", true);
+				if (!VersionUtilities.bLaunchscreenStoryboard)
+				{
+					CopyFiles(BuildDirectory + "/Resources/Graphics", AppDirectory, "[REDACTED_EMAIL]", true);
+					CopyFiles(BuildDirectory + "/Resources/Graphics", AppDirectory, "[REDACTED_EMAIL]", true);
+				}
 			}
 		}
 		protected virtual void CopyLocalizationsResources(string InEngineDir, string AppDirectory, string BuildDirectory, string IntermediateDir)
diff --git a/Engine/Source/Runtime/IOS/IOSRuntimeSettings/Classes/IOSRuntimeSettings.h b/Engine/Source/Runtime/IOS/IOSRuntimeSettings/Classes/IOSRuntimeSettings.h
--- a/Engine/Source/Runtime/IOS/IOSRuntimeSettings/Classes/IOSRuntimeSettings.h
+++ b/Engine/Source/Runtime/IOS/IOSRuntimeSettings/Classes/IOSRuntimeSettings.h
@@ -423,6 +423,14 @@ public:
 	// Any additional plist key/value data utilizing \n for a new line
 	UPROPERTY(GlobalConfig, EditAnywhere, Category = ExtraData)
 	FString AdditionalPlistData;
+	
+	/**
+	 * Chose wether to use a storyboard as a Launchscreen or not. If True, LaunchScreen.storyboard in Build/IOS/Resources/Interface
+	 * in the Engine folder will be copied to the bundle app and the Launchscreen images WILL NOT be included.
+	 * When this setting in changed, delete the Binaries/IOS/Payload folder before recompiling.
+	 */
+	UPROPERTY(GlobalConfig, EditAnywhere, Category = LaunchScreen)
+	bool 	bLaunchscreenStoryboard;
 
 	// Whether the app supports Facebook
 	UPROPERTY(GlobalConfig, EditAnywhere, Category = Online)
diff --git a/Engine/Source/Runtime/IOS/IOSRuntimeSettings/Private/IOSRuntimeSettings.cpp b/Engine/Source/Runtime/IOS/IOSRuntimeSettings/Private/IOSRuntimeSettings.cpp
--- a/Engine/Source/Runtime/IOS/IOSRuntimeSettings/Private/IOSRuntimeSettings.cpp
+++ b/Engine/Source/Runtime/IOS/IOSRuntimeSettings/Private/IOSRuntimeSettings.cpp
@@ -37,6 +37,7 @@ UIOSRuntimeSettings::UIOSRuntimeSettings(const FObjectInitializer& ObjectInitial
 	bShipForArmV7S = false;
 	bShipForBitcode = true;
 	bUseRSync = true;
+	bLaunchscreenStoryboard = false;
 	AdditionalPlistData = TEXT("");
 	AdditionalLinkerFlags = TEXT("");
 	AdditionalShippingLinkerFlags = TEXT("");
