Unblocking Robomerge for "Add gameActivityForceQuitAdditions to allow android plugins to be notified when an application requests to quit."
#rnx

diff --git a/Engine/Build/Android/Java/src/com/epicgames/ue4/GameActivity.java.template b/Engine/Build/Android/Java/src/com/epicgames/ue4/GameActivity.java.template
--- a/Engine/Build/Android/Java/src/com/epicgames/ue4/GameActivity.java.template
+++ b/Engine/Build/Android/Java/src/com/epicgames/ue4/GameActivity.java.template
@@ -4374,6 +4374,7 @@ public class GameActivity extends $${gameActivitySuperClass}$$ implements Surfac
 
 	public void AndroidThunkJava_ForceQuit()
 	{
+//$${gameActivityForceQuitAdditions}$$
 		System.exit(0);
 		// finish();
 	}
diff --git a/Engine/Source/Programs/UnrealBuildTool/Platform/Android/UEDeployAndroid.cs b/Engine/Source/Programs/UnrealBuildTool/Platform/Android/UEDeployAndroid.cs
--- a/Engine/Source/Programs/UnrealBuildTool/Platform/Android/UEDeployAndroid.cs
+++ b/Engine/Source/Programs/UnrealBuildTool/Platform/Android/UEDeployAndroid.cs
@@ -4453,7 +4453,8 @@ namespace UnrealBuildTool
 				{ "//$${gameActivityGetLoginIdAdditions}$$", UPL.ProcessPluginNode(NDKArch, "gameActivityGetLoginIdAdditions", "")},
 				{ "//$${gameActivityGetFunnelIdAdditions}$$", UPL.ProcessPluginNode(NDKArch, "gameActivityGetFunnelIdAdditions", "")},
 				{ "//$${gameActivityAllowedRemoteNotificationsAdditions}$$", UPL.ProcessPluginNode(NDKArch, "gameActivityAllowedRemoteNotificationsAdditions", "")},
-				{ "//$${gameActivityOnRestartApplicationAdditions}$$", UPL.ProcessPluginNode(NDKArch, "gameActivityOnRestartApplicationAdditions", "")},				
+				{ "//$${gameActivityOnRestartApplicationAdditions}$$", UPL.ProcessPluginNode(NDKArch, "gameActivityOnRestartApplicationAdditions", "")},
+				{ "//$${gameActivityForceQuitAdditions}$$", UPL.ProcessPluginNode(NDKArch, "gameActivityForceQuitAdditions", "")},
 				{ "//$${soLoadLibrary}$$", UPL.ProcessPluginNode(NDKArch, "soLoadLibrary", LoadLibraryDefaults)},
 				{ "$${gameActivitySuperClass}$$", SuperClassDefault},
 			};
