Fix Android arch excluding for lib paths that specify arch and version

diff --git a/Engine/Source/Programs/UnrealBuildTool/Platform/Android/AndroidToolChain.cs b/Engine/Source/Programs/UnrealBuildTool/Platform/Android/AndroidToolChain.cs
--- a/Engine/Source/Programs/UnrealBuildTool/Platform/Android/AndroidToolChain.cs
+++ b/Engine/Source/Programs/UnrealBuildTool/Platform/Android/AndroidToolChain.cs
@@ -948,7 +948,7 @@ namespace UnrealBuildTool
 					foreach (string ArchName in Pair.Value)
 					{
 						// if there's a directory in the path with a bad architecture name, reject it
-						if (Regex.IsMatch(Dir, "/" + ArchName + "$") || Regex.IsMatch(Dir, "/" + ArchName + "/") || Regex.IsMatch(Dir, "/" + ArchName + "_API[0-9]+_NDK[0-9]+"))
+						if (Regex.IsMatch(Dir, "/" + ArchName + "$") || Regex.IsMatch(Dir, "/" + ArchName + "/") || Regex.IsMatch(Dir, "/" + ArchName + "_API[0-9]+_NDK[0-9]+", RegexOptions.IgnoreCase))
 						{
 							return false;
 						}
