Refactor UBT to use an interface (IActionGraphBuilder) to construct the action graph. Response files are also constructed through this interface, allowing them to be tracked as dependencies (will be added in subsequent change).

diff --git a/Engine/Platforms/XXX/Source/Programs/UnrealBuildTool/XXXToolChain.cs b/Engine/Platforms/XXX/Source/Programs/UnrealBuildTool/XXXToolChain.cs
--- a/Engine/Platforms/XXX/Source/Programs/UnrealBuildTool/XXXToolChain.cs
+++ b/Engine/Platforms/XXX/Source/Programs/UnrealBuildTool/XXXToolChain.cs
@@ -17,7 +17,7 @@ namespace UnrealBuildTool
 		{
 		}
 
-		public override CPPOutput CompileCPPFiles(CppCompileEnvironment CompileEnvironment, List<FileItem> InputFiles, DirectoryReference OutputDir, string ModuleName, List<Action> Actions)
+		public override CPPOutput CompileCPPFiles(CppCompileEnvironment CompileEnvironment, List<FileItem> InputFiles, DirectoryReference OutputDir, string ModuleName, IActionGraphBuilder Graph)
 		{
 			// Create a compile action for each source file.
 			CPPOutput Result = new CPPOutput();
@@ -25,7 +25,7 @@ namespace UnrealBuildTool
 			{
 				FileItem ObjectFile = FileItem.GetItemByFileReference(FileReference.Combine(OutputDir, Path.GetFileName(SourceFile.AbsolutePath) + ".xo"));
 
-				Action CompileAction = new Action(ActionType.Compile);
+				Action CompileAction = Graph.CreateAction(ActionType.Compile);
 				CompileAction.CommandDescription = "FakeCompile";
 				CompileAction.CommandPath = BuildHostPlatform.Current.Shell;
 				// we use type/cat instead of copy so that timestamp gets updated
@@ -48,19 +48,16 @@ namespace UnrealBuildTool
 				{
 					CompileAction.PrerequisiteItems.Add(ForceIncludeFile);
 				}
-
-
-				Actions.Add(CompileAction);
 			}
 
 			return Result;
 		}
 
-		public override FileItem LinkFiles(LinkEnvironment LinkEnvironment, bool bBuildImportLibraryOnly, List<Action> Actions)
+		public override FileItem LinkFiles(LinkEnvironment LinkEnvironment, bool bBuildImportLibraryOnly, IActionGraphBuilder Graph)
 		{
 			FileItem OutputFile = FileItem.GetItemByFileReference(LinkEnvironment.OutputFilePath);
 
-			Action LinkAction = new Action(ActionType.Link);
+			Action LinkAction = Graph.CreateAction(ActionType.Link);
 			LinkAction.CommandDescription = "FakeCompile";
 			LinkAction.CommandPath = BuildHostPlatform.Current.Shell;
 			if (BuildHostPlatform.Current.Platform == UnrealTargetPlatform.Win64)
@@ -80,7 +77,6 @@ namespace UnrealBuildTool
 			LinkAction.DeleteItems.Add(OutputFile);
 			LinkAction.StatusDescription = OutputFile.Location.GetFileName();
 			LinkAction.bCanExecuteRemotely = false;
-			Actions.Add(LinkAction);
 
 			return OutputFile;
 		}
diff --git a/Engine/Source/Programs/UnrealBuildTool/Configuration/UEBuildBinary.cs b/Engine/Source/Programs/UnrealBuildTool/Configuration/UEBuildBinary.cs
--- a/Engine/Source/Programs/UnrealBuildTool/Configuration/UEBuildBinary.cs
+++ b/Engine/Source/Programs/UnrealBuildTool/Configuration/UEBuildBinary.cs
@@ -172,9 +172,9 @@ namespace UnrealBuildTool
 		/// <param name="SingleFileToCompile">If non-null, specifies a single cpp file to be compiled</param>
 		/// <param name="WorkingSet">The working set of source files</param>
 		/// <param name="ExeDir">Directory containing the output executable</param>
-		/// <param name="Makefile">The makefile to add actions to</param>
+		/// <param name="Graph">The graph being built</param>
 		/// <returns>Set of built products</returns>
-		public List<FileItem> Build(ReadOnlyTargetRules Target, UEToolChain ToolChain, CppCompileEnvironment CompileEnvironment, LinkEnvironment LinkEnvironment, FileReference SingleFileToCompile, ISourceFileWorkingSet WorkingSet, DirectoryReference ExeDir, TargetMakefile Makefile)
+		public List<FileItem> Build(ReadOnlyTargetRules Target, UEToolChain ToolChain, CppCompileEnvironment CompileEnvironment, LinkEnvironment LinkEnvironment, FileReference SingleFileToCompile, ISourceFileWorkingSet WorkingSet, DirectoryReference ExeDir, IActionGraphBuilder Graph)
 		{
 			// Return nothing if we're using precompiled binaries. If we're not linking, we might want just one module to be compiled (eg. a foreign plugin), so allow any actions to run.
 			if (bUsePrecompiled && !Target.bDisableLinking)
@@ -183,7 +183,7 @@ namespace UnrealBuildTool
 			}
 
 			// Setup linking environment.
-			LinkEnvironment BinaryLinkEnvironment = SetupBinaryLinkEnvironment(Target, ToolChain, LinkEnvironment, CompileEnvironment, SingleFileToCompile, WorkingSet, ExeDir, Makefile);
+			LinkEnvironment BinaryLinkEnvironment = SetupBinaryLinkEnvironment(Target, ToolChain, LinkEnvironment, CompileEnvironment, SingleFileToCompile, WorkingSet, ExeDir, Graph);
 
 			// If we're generating projects, we only need include paths and definitions, there is no need to run the linking logic.
 			if (ProjectFileGenerator.bGenerateProjectFiles)
@@ -209,18 +209,18 @@ namespace UnrealBuildTool
 					BinaryLinkEnvironment.Platform != UnrealTargetPlatform.LinuxAArch64)
 				{
 					// Create the import library.
-					OutputFiles.AddRange(ToolChain.LinkAllFiles(BinaryLinkEnvironment, true, Makefile.Actions));
+					OutputFiles.AddRange(ToolChain.LinkAllFiles(BinaryLinkEnvironment, true, Graph));
 				}
 			}
 
 			// Link the binary.
-			FileItem[] Executables = ToolChain.LinkAllFiles(BinaryLinkEnvironment, false, Makefile.Actions);
+			FileItem[] Executables = ToolChain.LinkAllFiles(BinaryLinkEnvironment, false, Graph);
 			OutputFiles.AddRange(Executables);
 
 			// Save all the output items for this binary. This is used for hot-reload, and excludes any items added in PostBuild (such as additional files copied into the app).
 			if(Target.LinkType == TargetLinkType.Modular)
 			{
-				Makefile.ModuleNameToOutputItems[PrimaryModule.Name] = OutputFiles.ToArray();
+				Graph.SetOutputItemsForModule(PrimaryModule.Name, OutputFiles.ToArray());
 			}
 
 			// Produce additional console app if requested
@@ -233,18 +233,18 @@ namespace UnrealBuildTool
 				ConsoleAppLinkEvironment.OutputFilePaths = ConsoleAppLinkEvironment.OutputFilePaths.Select(Path => GetAdditionalConsoleAppPath(Path)).ToList();
 
 				// Link the console app executable
-				FileItem[] ConsoleAppOutputFiles = ToolChain.LinkAllFiles(ConsoleAppLinkEvironment, false, Makefile.Actions);
+				FileItem[] ConsoleAppOutputFiles = ToolChain.LinkAllFiles(ConsoleAppLinkEvironment, false, Graph);
 				OutputFiles.AddRange(ConsoleAppOutputFiles);
 
 				foreach (FileItem Executable in ConsoleAppOutputFiles)
 				{
-					OutputFiles.AddRange(ToolChain.PostBuild(Executable, ConsoleAppLinkEvironment, Makefile.Actions));
+					OutputFiles.AddRange(ToolChain.PostBuild(Executable, ConsoleAppLinkEvironment, Graph));
 				}
 			}
 
 			foreach (FileItem Executable in Executables)
 			{
-				OutputFiles.AddRange(ToolChain.PostBuild(Executable, BinaryLinkEnvironment, Makefile.Actions));
+				OutputFiles.AddRange(ToolChain.PostBuild(Executable, BinaryLinkEnvironment, Graph));
 			}
 
 			return OutputFiles;
@@ -674,7 +674,7 @@ namespace UnrealBuildTool
 			return BinaryCompileEnvironment;
 		}
 
-		private LinkEnvironment SetupBinaryLinkEnvironment(ReadOnlyTargetRules Target, UEToolChain ToolChain, LinkEnvironment LinkEnvironment, CppCompileEnvironment CompileEnvironment, FileReference SingleFileToCompile, ISourceFileWorkingSet WorkingSet, DirectoryReference ExeDir, TargetMakefile Makefile)
+		private LinkEnvironment SetupBinaryLinkEnvironment(ReadOnlyTargetRules Target, UEToolChain ToolChain, LinkEnvironment LinkEnvironment, CppCompileEnvironment CompileEnvironment, FileReference SingleFileToCompile, ISourceFileWorkingSet WorkingSet, DirectoryReference ExeDir, IActionGraphBuilder Graph)
 		{
 			LinkEnvironment BinaryLinkEnvironment = new LinkEnvironment(LinkEnvironment);
 			HashSet<UEBuildModule> LinkEnvironmentVisitedModules = new HashSet<UEBuildModule>();
@@ -694,12 +694,12 @@ namespace UnrealBuildTool
 				{
 					// Compile each module.
 					Log.TraceVerbose("Compile module: " + Module.Name);
-					LinkInputFiles = Module.Compile(Target, ToolChain, BinaryCompileEnvironment, SingleFileToCompile, WorkingSet, Makefile);
+					LinkInputFiles = Module.Compile(Target, ToolChain, BinaryCompileEnvironment, SingleFileToCompile, WorkingSet, Graph);
 
 					// Save the module outputs. In monolithic builds, this is just the object files.
 					if (Target.LinkType == TargetLinkType.Monolithic)
 					{
-						Makefile.ModuleNameToOutputItems[Module.Name] = LinkInputFiles.ToArray();
+						Graph.SetOutputItemsForModule(Module.Name, LinkInputFiles.ToArray());
 					}
 
 					// NOTE: Because of 'Shared PCHs', in monolithic builds the same PCH file may appear as a link input
@@ -781,7 +781,7 @@ namespace UnrealBuildTool
 
 						// Otherwise compile the default resource file per-binary, so that it gets the correct ORIGINAL_FILE_NAME macro.
 						FileItem DefaultResourceFile = FileItem.GetItemByFileReference(FileReference.Combine(UnrealBuildTool.EngineDirectory, "Build", "Windows", "Resources", "Default.rc2"));
-						CPPOutput DefaultResourceOutput = ToolChain.CompileRCFiles(ResourceCompileEnvironment, new List<FileItem> { DefaultResourceFile }, ResourceIntermediateDirectory, Makefile.Actions);
+						CPPOutput DefaultResourceOutput = ToolChain.CompileRCFiles(ResourceCompileEnvironment, new List<FileItem> { DefaultResourceFile }, ResourceIntermediateDirectory, Graph);
 						BinaryLinkEnvironment.InputFiles.AddRange(DefaultResourceOutput.ObjectFiles);
 					}
 				}
diff --git a/Engine/Source/Programs/UnrealBuildTool/Configuration/UEBuildModule.cs b/Engine/Source/Programs/UnrealBuildTool/Configuration/UEBuildModule.cs
--- a/Engine/Source/Programs/UnrealBuildTool/Configuration/UEBuildModule.cs
+++ b/Engine/Source/Programs/UnrealBuildTool/Configuration/UEBuildModule.cs
@@ -786,13 +786,13 @@ namespace UnrealBuildTool
 		/// <summary>
 		/// Compiles the module, and returns a list of files output by the compiler.
 		/// </summary>
-		public virtual List<FileItem> Compile(ReadOnlyTargetRules Target, UEToolChain ToolChain, CppCompileEnvironment CompileEnvironment, FileReference SingleFileToCompile, ISourceFileWorkingSet WorkingSet, TargetMakefile Makefile)
+		public virtual List<FileItem> Compile(ReadOnlyTargetRules Target, UEToolChain ToolChain, CppCompileEnvironment CompileEnvironment, FileReference SingleFileToCompile, ISourceFileWorkingSet WorkingSet, IActionGraphBuilder Graph)
 		{
 			// Generate type libraries for Windows
 			foreach(ModuleRules.TypeLibrary TypeLibrary in Rules.TypeLibraries)
 			{
 				FileReference OutputFile = FileReference.Combine(IntermediateDirectory, TypeLibrary.Header);
-				ToolChain.GenerateTypeLibraryHeader(CompileEnvironment, TypeLibrary, OutputFile, Makefile.Actions);
+				ToolChain.GenerateTypeLibraryHeader(CompileEnvironment, TypeLibrary, OutputFile, Graph);
 			}
 
 			return new List<FileItem>();
diff --git a/Engine/Source/Programs/UnrealBuildTool/Configuration/UEBuildModuleCPP.cs b/Engine/Source/Programs/UnrealBuildTool/Configuration/UEBuildModuleCPP.cs
--- a/Engine/Source/Programs/UnrealBuildTool/Configuration/UEBuildModuleCPP.cs
+++ b/Engine/Source/Programs/UnrealBuildTool/Configuration/UEBuildModuleCPP.cs
@@ -265,11 +265,11 @@ namespace UnrealBuildTool
 		}
 
 		// UEBuildModule interface.
-		public override List<FileItem> Compile(ReadOnlyTargetRules Target, UEToolChain ToolChain, CppCompileEnvironment BinaryCompileEnvironment, FileReference SingleFileToCompile, ISourceFileWorkingSet WorkingSet, TargetMakefile Makefile)
+		public override List<FileItem> Compile(ReadOnlyTargetRules Target, UEToolChain ToolChain, CppCompileEnvironment BinaryCompileEnvironment, FileReference SingleFileToCompile, ISourceFileWorkingSet WorkingSet, IActionGraphBuilder Graph)
 		{
 			UEBuildPlatform BuildPlatform = UEBuildPlatform.GetBuildPlatform(BinaryCompileEnvironment.Platform);
 
-			List<FileItem> LinkInputFiles = base.Compile(Target, ToolChain, BinaryCompileEnvironment, SingleFileToCompile, WorkingSet, Makefile);
+			List<FileItem> LinkInputFiles = base.Compile(Target, ToolChain, BinaryCompileEnvironment, SingleFileToCompile, WorkingSet, Graph);
 
 			CppCompileEnvironment ModuleCompileEnvironment = CreateModuleCompileEnvironment(Target, BinaryCompileEnvironment);
 
@@ -298,11 +298,17 @@ namespace UnrealBuildTool
 			foreach (DirectoryReference ModuleDirectory in ModuleDirectories)
 			{
 				DirectoryItem ModuleDirectoryItem = DirectoryItem.GetItemByDirectoryReference(ModuleDirectory);
-				Makefile.SourceDirectories.Add(ModuleDirectoryItem);
+				Graph.AddSourceDir(ModuleDirectoryItem);
 			}
 
 			// Find all the input files
-			InputFileCollection InputFiles = FindInputFiles(Target.Platform, Makefile.DirectoryToSourceFiles);
+			Dictionary<DirectoryItem, FileItem[]> DirectoryToSourceFiles = new Dictionary<DirectoryItem, FileItem[]>();
+			InputFileCollection InputFiles = FindInputFiles(Target.Platform, DirectoryToSourceFiles);
+
+			foreach (KeyValuePair<DirectoryItem, FileItem[]> Pair in DirectoryToSourceFiles)
+			{
+				Graph.AddSourceFiles(Pair.Key, Pair.Value);
+			}
 
 			// If we're compiling a single file, strip out anything else. This prevents us clobbering response files for anything we're 
 			// not going to build, triggering a larger build than necessary when we do a regular build again.
@@ -390,14 +396,14 @@ namespace UnrealBuildTool
 			// Generate ISPC headers first so C++ can consume them
 			if (InputFiles.ISPCFiles.Count > 0)
 			{
-				CreateHeadersForISPC(ToolChain, CompileEnvironment, InputFiles.ISPCFiles, IntermediateDirectory, Makefile.Actions);
+				CreateHeadersForISPC(ToolChain, CompileEnvironment, InputFiles.ISPCFiles, IntermediateDirectory, Graph);
 			}
 
 			// Configure the precompiled headers for this module
-			CompileEnvironment = SetupPrecompiledHeaders(Target, ToolChain, CompileEnvironment, LinkInputFiles, Makefile.Actions);
+			CompileEnvironment = SetupPrecompiledHeaders(Target, ToolChain, CompileEnvironment, LinkInputFiles, Graph);
 
 			// Write all the definitions to a separate file
-			CreateHeaderForDefinitions(CompileEnvironment, IntermediateDirectory, null);
+			CreateHeaderForDefinitions(CompileEnvironment, IntermediateDirectory, null, Graph);
 
 			// Mapping of source file to unity file. We output this to intermediate directories for other tools (eg. live coding) to use.
 			Dictionary<FileItem, FileItem> SourceFileToUnityFile = new Dictionary<FileItem, FileItem>();
@@ -406,12 +412,12 @@ namespace UnrealBuildTool
 			List<FileItem> CPPFilesToCompile = InputFiles.CPPFiles;
 			if (bModuleUsesUnityBuild)
 			{
-				CPPFilesToCompile = Unity.GenerateUnityCPPs(Target, CPPFilesToCompile, CompileEnvironment, WorkingSet, Rules.ShortName ?? Name, IntermediateDirectory, Makefile, SourceFileToUnityFile);
-				LinkInputFiles.AddRange(CompileUnityFilesWithToolChain(Target, ToolChain, CompileEnvironment, ModuleCompileEnvironment, CPPFilesToCompile, Makefile.Actions).ObjectFiles);
+				CPPFilesToCompile = Unity.GenerateUnityCPPs(Target, CPPFilesToCompile, CompileEnvironment, WorkingSet, Rules.ShortName ?? Name, IntermediateDirectory, Graph, SourceFileToUnityFile);
+				LinkInputFiles.AddRange(CompileUnityFilesWithToolChain(Target, ToolChain, CompileEnvironment, ModuleCompileEnvironment, CPPFilesToCompile, Graph).ObjectFiles);
 			}
 			else
 			{
-				LinkInputFiles.AddRange(ToolChain.CompileCPPFiles(CompileEnvironment, CPPFilesToCompile, IntermediateDirectory, Name, Makefile.Actions).ObjectFiles);
+				LinkInputFiles.AddRange(ToolChain.CompileCPPFiles(CompileEnvironment, CPPFilesToCompile, IntermediateDirectory, Name, Graph).ObjectFiles);
 			}
 
 			// Compile all the generated CPP files
@@ -454,12 +460,12 @@ namespace UnrealBuildTool
 
 					if (bModuleUsesUnityBuild)
 					{
-						GeneratedFileItems = Unity.GenerateUnityCPPs(Target, GeneratedFileItems, GeneratedCPPCompileEnvironment, WorkingSet, (Rules.ShortName ?? Name) + ".gen", IntermediateDirectory, Makefile, SourceFileToUnityFile);
-						LinkInputFiles.AddRange(CompileUnityFilesWithToolChain(Target, ToolChain, GeneratedCPPCompileEnvironment, ModuleCompileEnvironment, GeneratedFileItems, Makefile.Actions).ObjectFiles);
+						GeneratedFileItems = Unity.GenerateUnityCPPs(Target, GeneratedFileItems, GeneratedCPPCompileEnvironment, WorkingSet, (Rules.ShortName ?? Name) + ".gen", IntermediateDirectory, Graph, SourceFileToUnityFile);
+						LinkInputFiles.AddRange(CompileUnityFilesWithToolChain(Target, ToolChain, GeneratedCPPCompileEnvironment, ModuleCompileEnvironment, GeneratedFileItems, Graph).ObjectFiles);
 					}
 					else
 					{
-						LinkInputFiles.AddRange(ToolChain.CompileCPPFiles(GeneratedCPPCompileEnvironment, GeneratedFileItems, IntermediateDirectory, Name, Makefile.Actions).ObjectFiles);
+						LinkInputFiles.AddRange(ToolChain.CompileCPPFiles(GeneratedCPPCompileEnvironment, GeneratedFileItems, IntermediateDirectory, Name, Graph).ObjectFiles);
 					}
 				}
 			}
@@ -467,25 +473,25 @@ namespace UnrealBuildTool
 			// Compile ISPC files directly
 			if (InputFiles.ISPCFiles.Count > 0)
 			{
-				LinkInputFiles.AddRange(ToolChain.CompileISPCFiles(CompileEnvironment, InputFiles.ISPCFiles, IntermediateDirectory, Makefile.Actions).ObjectFiles);
+				LinkInputFiles.AddRange(ToolChain.CompileISPCFiles(CompileEnvironment, InputFiles.ISPCFiles, IntermediateDirectory, Graph).ObjectFiles);
 			}
 
 			// Compile C files directly. Do not use a PCH here, because a C++ PCH is not compatible with C source files.
 			if(InputFiles.CFiles.Count > 0)
 			{
-				LinkInputFiles.AddRange(ToolChain.CompileCPPFiles(ModuleCompileEnvironment, InputFiles.CFiles, IntermediateDirectory, Name, Makefile.Actions).ObjectFiles);
+				LinkInputFiles.AddRange(ToolChain.CompileCPPFiles(ModuleCompileEnvironment, InputFiles.CFiles, IntermediateDirectory, Name, Graph).ObjectFiles);
 			}
 
 			// Compile CC files directly.
 			if(InputFiles.CCFiles.Count > 0)
 			{
-				LinkInputFiles.AddRange(ToolChain.CompileCPPFiles(CompileEnvironment, InputFiles.CCFiles, IntermediateDirectory, Name, Makefile.Actions).ObjectFiles);
+				LinkInputFiles.AddRange(ToolChain.CompileCPPFiles(CompileEnvironment, InputFiles.CCFiles, IntermediateDirectory, Name, Graph).ObjectFiles);
 			}
 
 			// Compile MM files directly.
 			if(InputFiles.MMFiles.Count > 0)
 			{
-				LinkInputFiles.AddRange(ToolChain.CompileCPPFiles(CompileEnvironment, InputFiles.MMFiles, IntermediateDirectory, Name, Makefile.Actions).ObjectFiles);
+				LinkInputFiles.AddRange(ToolChain.CompileCPPFiles(CompileEnvironment, InputFiles.MMFiles, IntermediateDirectory, Name, Graph).ObjectFiles);
 			}
 
 			// Compile RC files. The resource compiler does not work with response files, and using the regular compile environment can easily result in the 
@@ -499,7 +505,7 @@ namespace UnrealBuildTool
 					// @todo: This should be in some Windows code somewhere...
 					ResourceCompileEnvironment.Definitions.Add("ORIGINAL_FILE_NAME=\"" + Binary.OutputFilePaths[0].GetFileName() + "\"");
 				}
-				LinkInputFiles.AddRange(ToolChain.CompileRCFiles(ResourceCompileEnvironment, InputFiles.RCFiles, IntermediateDirectory, Makefile.Actions).ObjectFiles);
+				LinkInputFiles.AddRange(ToolChain.CompileRCFiles(ResourceCompileEnvironment, InputFiles.RCFiles, IntermediateDirectory, Graph).ObjectFiles);
 			}
 
 			// Write the compiled manifest
@@ -567,13 +573,13 @@ namespace UnrealBuildTool
 		/// <param name="ToolChain">The toolchain to generate the PCH</param>
 		/// <param name="HeaderFile"></param>
 		/// <param name="ModuleCompileEnvironment"></param>
-		/// <param name="Actions">List of actions to be executed. Additional actions will be added to this list.</param>
+		/// <param name="Graph">List of actions to be executed. Additional actions will be added to this list.</param>
 		/// <returns>The created PCH instance.</returns>
-		private PrecompiledHeaderInstance CreatePrivatePCH(UEToolChain ToolChain, FileItem HeaderFile, CppCompileEnvironment ModuleCompileEnvironment, List<Action> Actions)
+		private PrecompiledHeaderInstance CreatePrivatePCH(UEToolChain ToolChain, FileItem HeaderFile, CppCompileEnvironment ModuleCompileEnvironment, IActionGraphBuilder Graph)
 		{
 			// Create the wrapper file, which sets all the definitions needed to compile it
 			FileReference WrapperLocation = FileReference.Combine(IntermediateDirectory, String.Format("PCH.{0}.h", Name));
-			FileItem WrapperFile = CreatePCHWrapperFile(WrapperLocation, ModuleCompileEnvironment.Definitions, HeaderFile);
+			FileItem WrapperFile = CreatePCHWrapperFile(WrapperLocation, ModuleCompileEnvironment.Definitions, HeaderFile, Graph);
 
 			// Create a new C++ environment that is used to create the PCH.
 			CppCompileEnvironment CompileEnvironment = new CppCompileEnvironment(ModuleCompileEnvironment);
@@ -590,7 +596,7 @@ namespace UnrealBuildTool
 			}
 			else
 			{
-				Output = ToolChain.CompileCPPFiles(CompileEnvironment, new List<FileItem>() { WrapperFile }, IntermediateDirectory, Name, Actions);
+				Output = ToolChain.CompileCPPFiles(CompileEnvironment, new List<FileItem>() { WrapperFile }, IntermediateDirectory, Name, Graph);
 			}
 			return new PrecompiledHeaderInstance(WrapperFile, CompileEnvironment, Output);
 		}
@@ -601,9 +607,9 @@ namespace UnrealBuildTool
 		/// <param name="ToolChain">The toolchain being used to build this module</param>
 		/// <param name="Template">The PCH template</param>
 		/// <param name="ModuleCompileEnvironment">Compile environment for the current module</param>
-		/// <param name="Actions">List of actions to be executed. Additional actions will be added to this list.</param>
+		/// <param name="Graph">List of actions to be executed. Additional actions will be added to this list.</param>
 		/// <returns>Instance of a PCH</returns>
-		public PrecompiledHeaderInstance FindOrCreateSharedPCH(UEToolChain ToolChain, PrecompiledHeaderTemplate Template, CppCompileEnvironment ModuleCompileEnvironment, List<Action> Actions)
+		public PrecompiledHeaderInstance FindOrCreateSharedPCH(UEToolChain ToolChain, PrecompiledHeaderTemplate Template, CppCompileEnvironment ModuleCompileEnvironment, IActionGraphBuilder Graph)
 		{
 			PrecompiledHeaderInstance Instance = Template.Instances.Find(x => IsCompatibleForSharedPCH(x.CompileEnvironment, ModuleCompileEnvironment));
 			if(Instance == null)
@@ -613,7 +619,7 @@ namespace UnrealBuildTool
 
 				// Create the wrapper file, which sets all the definitions needed to compile it
 				FileReference WrapperLocation = FileReference.Combine(Template.OutputDir, String.Format("SharedPCH.{0}{1}.h", Template.Module.Name, Variant));
-				FileItem WrapperFile = CreatePCHWrapperFile(WrapperLocation, Template.BaseCompileEnvironment.Definitions, Template.HeaderFile);
+				FileItem WrapperFile = CreatePCHWrapperFile(WrapperLocation, Template.BaseCompileEnvironment.Definitions, Template.HeaderFile, Graph);
 
 				// Create the compile environment for this PCH
 				CppCompileEnvironment CompileEnvironment = new CppCompileEnvironment(Template.BaseCompileEnvironment);
@@ -630,7 +636,7 @@ namespace UnrealBuildTool
 				}
 				else
 				{
-					Output = ToolChain.CompileCPPFiles(CompileEnvironment, new List<FileItem>() { WrapperFile }, Template.OutputDir, "Shared", Actions);
+					Output = ToolChain.CompileCPPFiles(CompileEnvironment, new List<FileItem>() { WrapperFile }, Template.OutputDir, "Shared", Graph);
 				}
 				Instance = new PrecompiledHeaderInstance(WrapperFile, CompileEnvironment, Output);
 				Template.Instances.Add(Instance);
@@ -779,7 +785,7 @@ namespace UnrealBuildTool
 		/// <summary>
 		/// Compiles the provided CPP unity files. Will
 		/// </summary>
-		private CPPOutput CompileUnityFilesWithToolChain(ReadOnlyTargetRules Target, UEToolChain ToolChain, CppCompileEnvironment CompileEnvironment, CppCompileEnvironment ModuleCompileEnvironment, List<FileItem> SourceFiles, List<Action> Actions)
+		private CPPOutput CompileUnityFilesWithToolChain(ReadOnlyTargetRules Target, UEToolChain ToolChain, CppCompileEnvironment CompileEnvironment, CppCompileEnvironment ModuleCompileEnvironment, List<FileItem> SourceFiles, IActionGraphBuilder Graph)
 		{
 			List<FileItem> NormalFiles = new List<FileItem>();
 			List<FileItem> AdaptiveFiles = new List<FileItem>();
@@ -821,7 +827,7 @@ namespace UnrealBuildTool
 
 			if (NormalFiles.Count > 0)
 			{
-				OutputFiles = ToolChain.CompileCPPFiles(CompileEnvironment, NormalFiles, IntermediateDirectory, Name, Actions);
+				OutputFiles = ToolChain.CompileCPPFiles(CompileEnvironment, NormalFiles, IntermediateDirectory, Name, Graph);
 			}
 
 			if (AdaptiveFiles.Count > 0)
@@ -841,19 +847,19 @@ namespace UnrealBuildTool
 				CPPOutput AdaptiveOutput;
 				if(Target.bAdaptiveUnityCreatesDedicatedPCH)
 				{
-					AdaptiveOutput = CompileAdaptiveNonUnityFilesWithDedicatedPCH(ToolChain, AdaptiveUnityEnvironment, AdaptiveFiles, IntermediateDirectory, Name, Actions);
+					AdaptiveOutput = CompileAdaptiveNonUnityFilesWithDedicatedPCH(ToolChain, AdaptiveUnityEnvironment, AdaptiveFiles, IntermediateDirectory, Name, Graph);
 				}
 				else if(bAdaptiveUnityDisablesPCH)
 				{
-					AdaptiveOutput = CompileAdaptiveNonUnityFilesWithoutPCH(ToolChain, AdaptiveUnityEnvironment, AdaptiveFiles, IntermediateDirectory, Name, Actions);
+					AdaptiveOutput = CompileAdaptiveNonUnityFilesWithoutPCH(ToolChain, AdaptiveUnityEnvironment, AdaptiveFiles, IntermediateDirectory, Name, Graph);
 				}
 				else if(AdaptiveUnityEnvironment.bOptimizeCode != CompileEnvironment.bOptimizeCode || AdaptiveUnityEnvironment.bSupportEditAndContinue != CompileEnvironment.bSupportEditAndContinue)
 				{
-					AdaptiveOutput = CompileAdaptiveNonUnityFiles(ToolChain, AdaptiveUnityEnvironment, AdaptiveFiles, IntermediateDirectory, Name, Actions);
+					AdaptiveOutput = CompileAdaptiveNonUnityFiles(ToolChain, AdaptiveUnityEnvironment, AdaptiveFiles, IntermediateDirectory, Name, Graph);
 				}
 				else
 				{
-					AdaptiveOutput = CompileAdaptiveNonUnityFiles(ToolChain, CompileEnvironment, AdaptiveFiles, IntermediateDirectory, Name, Actions);
+					AdaptiveOutput = CompileAdaptiveNonUnityFiles(ToolChain, CompileEnvironment, AdaptiveFiles, IntermediateDirectory, Name, Graph);
 				}
 
 				// Merge output
@@ -864,28 +870,28 @@ namespace UnrealBuildTool
 			return OutputFiles;
 		}
 
-		static CPPOutput CompileAdaptiveNonUnityFiles(UEToolChain ToolChain, CppCompileEnvironment CompileEnvironment, List<FileItem> Files, DirectoryReference IntermediateDirectory, string ModuleName, List<Action> Actions)
+		static CPPOutput CompileAdaptiveNonUnityFiles(UEToolChain ToolChain, CppCompileEnvironment CompileEnvironment, List<FileItem> Files, DirectoryReference IntermediateDirectory, string ModuleName, IActionGraphBuilder Graph)
 		{
 			// Write all the definitions out to a separate file
-			CreateHeaderForDefinitions(CompileEnvironment, IntermediateDirectory, "Adaptive");
+			CreateHeaderForDefinitions(CompileEnvironment, IntermediateDirectory, "Adaptive", Graph);
 
 			// Compile the files
-			return ToolChain.CompileCPPFiles(CompileEnvironment, Files, IntermediateDirectory, ModuleName, Actions);
+			return ToolChain.CompileCPPFiles(CompileEnvironment, Files, IntermediateDirectory, ModuleName, Graph);
 		}
 
-		static CPPOutput CompileAdaptiveNonUnityFilesWithoutPCH(UEToolChain ToolChain, CppCompileEnvironment CompileEnvironment, List<FileItem> Files, DirectoryReference IntermediateDirectory, string ModuleName, List<Action> Actions)
+		static CPPOutput CompileAdaptiveNonUnityFilesWithoutPCH(UEToolChain ToolChain, CppCompileEnvironment CompileEnvironment, List<FileItem> Files, DirectoryReference IntermediateDirectory, string ModuleName, IActionGraphBuilder Graph)
 		{
 			// Disable precompiled headers
 			CompileEnvironment.PrecompiledHeaderAction = PrecompiledHeaderAction.None;
 
 			// Write all the definitions out to a separate file
-			CreateHeaderForDefinitions(CompileEnvironment, IntermediateDirectory, "Adaptive");
+			CreateHeaderForDefinitions(CompileEnvironment, IntermediateDirectory, "Adaptive", Graph);
 
 			// Compile the files
-			return ToolChain.CompileCPPFiles(CompileEnvironment, Files, IntermediateDirectory, ModuleName, Actions);
+			return ToolChain.CompileCPPFiles(CompileEnvironment, Files, IntermediateDirectory, ModuleName, Graph);
 		}
 
-		static CPPOutput CompileAdaptiveNonUnityFilesWithDedicatedPCH(UEToolChain ToolChain, CppCompileEnvironment CompileEnvironment, List<FileItem> Files, DirectoryReference IntermediateDirectory, string ModuleName, List<Action> Actions)
+		static CPPOutput CompileAdaptiveNonUnityFilesWithDedicatedPCH(UEToolChain ToolChain, CppCompileEnvironment CompileEnvironment, List<FileItem> Files, DirectoryReference IntermediateDirectory, string ModuleName, IActionGraphBuilder Graph)
 		{
 			CPPOutput Output = new CPPOutput();
 			foreach(FileItem File in Files)
@@ -906,7 +912,7 @@ namespace UnrealBuildTool
 
 				// Write the PCH header
 				FileReference DedicatedPchLocation = FileReference.Combine(IntermediateDirectory, String.Format("PCH.Dedicated.{0}.h", File.Location.GetFileNameWithoutExtension()));
-				FileItem DedicatedPchFile = FileItem.CreateIntermediateTextFile(DedicatedPchLocation, WrapperContents.ToString());
+				FileItem DedicatedPchFile = Graph.CreateIntermediateTextFile(DedicatedPchLocation, WrapperContents.ToString());
 
 				// Create a new C++ environment to compile the PCH
 				CppCompileEnvironment PchEnvironment = new CppCompileEnvironment(CompileEnvironment);
@@ -916,7 +922,7 @@ namespace UnrealBuildTool
 				PchEnvironment.PrecompiledHeaderIncludeFilename = DedicatedPchFile.Location;
 
 				// Create the action to compile the PCH file.
-				CPPOutput PchOutput = ToolChain.CompileCPPFiles(PchEnvironment, new List<FileItem>() { DedicatedPchFile }, IntermediateDirectory, ModuleName, Actions);
+				CPPOutput PchOutput = ToolChain.CompileCPPFiles(PchEnvironment, new List<FileItem>() { DedicatedPchFile }, IntermediateDirectory, ModuleName, Graph);
 				Output.ObjectFiles.AddRange(PchOutput.ObjectFiles);
 
 				// Create a new C++ environment to compile the original file
@@ -927,7 +933,7 @@ namespace UnrealBuildTool
 				FileEnvironment.PrecompiledHeaderFile = PchOutput.PrecompiledHeaderFile;
 
 				// Create the action to compile the PCH file.
-				CPPOutput FileOutput = ToolChain.CompileCPPFiles(FileEnvironment, new List<FileItem>() { File }, IntermediateDirectory, ModuleName, Actions);
+				CPPOutput FileOutput = ToolChain.CompileCPPFiles(FileEnvironment, new List<FileItem>() { File }, IntermediateDirectory, ModuleName, Graph);
 				Output.ObjectFiles.AddRange(FileOutput.ObjectFiles);
 			}
 			return Output;
@@ -940,15 +946,15 @@ namespace UnrealBuildTool
 		/// <param name="ToolChain">The toolchain to build with</param>
 		/// <param name="CompileEnvironment">The current compile environment</param>
 		/// <param name="LinkInputFiles">List of files that will be linked for the target</param>
-		/// <param name="Actions">List of build actions</param>
-		CppCompileEnvironment SetupPrecompiledHeaders(ReadOnlyTargetRules Target, UEToolChain ToolChain, CppCompileEnvironment CompileEnvironment, List<FileItem> LinkInputFiles, List<Action> Actions)
+		/// <param name="Graph">List of build actions</param>
+		CppCompileEnvironment SetupPrecompiledHeaders(ReadOnlyTargetRules Target, UEToolChain ToolChain, CppCompileEnvironment CompileEnvironment, List<FileItem> LinkInputFiles, IActionGraphBuilder Graph)
 		{
 			if (Target.bUsePCHFiles && Rules.PCHUsage != ModuleRules.PCHUsageMode.NoPCHs)
 			{
 				// If this module doesn't need a shared PCH, configure that
 				if(Rules.PrivatePCHHeaderFile != null && (Rules.PCHUsage == ModuleRules.PCHUsageMode.NoSharedPCHs || Rules.PCHUsage == ModuleRules.PCHUsageMode.UseExplicitOrSharedPCHs))
 				{
-					PrecompiledHeaderInstance Instance = CreatePrivatePCH(ToolChain, FileItem.GetItemByFileReference(FileReference.Combine(ModuleDirectory, Rules.PrivatePCHHeaderFile)), CompileEnvironment, Actions);
+					PrecompiledHeaderInstance Instance = CreatePrivatePCH(ToolChain, FileItem.GetItemByFileReference(FileReference.Combine(ModuleDirectory, Rules.PrivatePCHHeaderFile)), CompileEnvironment, Graph);
 
 					CompileEnvironment = new CppCompileEnvironment(CompileEnvironment);
 					CompileEnvironment.Definitions.Clear();
@@ -970,7 +976,7 @@ namespace UnrealBuildTool
 					PrecompiledHeaderTemplate Template = CompileEnvironment.SharedPCHs.FirstOrDefault(x => ReferencedModules.Contains(x.Module));
 					if(Template != null && Template.IsValidFor(CompileEnvironment))
 					{
-						PrecompiledHeaderInstance Instance = FindOrCreateSharedPCH(ToolChain, Template, CompileEnvironment, Actions);
+						PrecompiledHeaderInstance Instance = FindOrCreateSharedPCH(ToolChain, Template, CompileEnvironment, Graph);
 
 						FileReference PrivateDefinitionsFile = FileReference.Combine(IntermediateDirectory, String.Format("Definitions.{0}.h", Name));
 
@@ -991,7 +997,7 @@ namespace UnrealBuildTool
 							}
 
 							WriteDefinitions(CompileEnvironment.Definitions, Writer);
-							PrivateDefinitionsFileItem = FileItem.CreateIntermediateTextFile(PrivateDefinitionsFile, Writer.ToString());
+							PrivateDefinitionsFileItem = Graph.CreateIntermediateTextFile(PrivateDefinitionsFile, Writer.ToString());
 						}
 
 						CompileEnvironment = new CppCompileEnvironment(CompileEnvironment);
@@ -1016,7 +1022,8 @@ namespace UnrealBuildTool
 		/// <param name="CompileEnvironment">The compile environment</param>
 		/// <param name="IntermediateDirectory">Directory to create the intermediate file</param>
 		/// <param name="HeaderSuffix">Suffix for the included file</param>
-		static void CreateHeaderForDefinitions(CppCompileEnvironment CompileEnvironment, DirectoryReference IntermediateDirectory, string HeaderSuffix)
+		/// <param name="Graph">The action graph being built</param>
+		static void CreateHeaderForDefinitions(CppCompileEnvironment CompileEnvironment, DirectoryReference IntermediateDirectory, string HeaderSuffix, IActionGraphBuilder Graph)
 		{
 			if(CompileEnvironment.Definitions.Count > 0)
 			{
@@ -1034,7 +1041,7 @@ namespace UnrealBuildTool
 					WriteDefinitions(CompileEnvironment.Definitions, Writer);
 					CompileEnvironment.Definitions.Clear();
 
-					FileItem PrivateDefinitionsFileItem = FileItem.CreateIntermediateTextFile(PrivateDefinitionsFile, Writer.ToString());
+					FileItem PrivateDefinitionsFileItem = Graph.CreateIntermediateTextFile(PrivateDefinitionsFile, Writer.ToString());
 					CompileEnvironment.ForceIncludeFiles.Add(PrivateDefinitionsFileItem);
 				}
 			}
@@ -1047,10 +1054,10 @@ namespace UnrealBuildTool
 		/// <param name="CompileEnvironment">Compile environment</param>
 		/// <param name="InputFiles">List of ISPC source files</param>
 		/// <param name="IntermediateDirectory">Directory to create the intermediate file</param>
-		/// <param name="Actions">List of actions to be executed. Additional actions will be added to this list.</param>
-		static void CreateHeadersForISPC(UEToolChain ToolChain, CppCompileEnvironment CompileEnvironment, List<FileItem> InputFiles, DirectoryReference IntermediateDirectory, List<Action> Actions)
+		/// <param name="Graph">List of actions to be executed. Additional actions will be added to this list.</param>
+		static void CreateHeadersForISPC(UEToolChain ToolChain, CppCompileEnvironment CompileEnvironment, List<FileItem> InputFiles, DirectoryReference IntermediateDirectory, IActionGraphBuilder Graph)
 		{
-			CPPOutput Output = ToolChain.GenerateISPCHeaders(CompileEnvironment, InputFiles, IntermediateDirectory, Actions);
+			CPPOutput Output = ToolChain.GenerateISPCHeaders(CompileEnvironment, InputFiles, IntermediateDirectory, Graph);
 
 			CompileEnvironment.AdditionalPrerequisites.AddRange(Output.GeneratedHeaderFiles);
 			CompileEnvironment.UserIncludePaths.Add(IntermediateDirectory);
@@ -1064,8 +1071,9 @@ namespace UnrealBuildTool
 		/// <param name="OutputFile">The output file to create</param>
 		/// <param name="Definitions">Definitions required by the PCH</param>
 		/// <param name="IncludedFile">The PCH file to include</param>
+		/// <param name="Graph">The action graph builder</param>
 		/// <returns>FileItem for the created file</returns>
-		static FileItem CreatePCHWrapperFile(FileReference OutputFile, IEnumerable<string> Definitions, FileItem IncludedFile)
+		static FileItem CreatePCHWrapperFile(FileReference OutputFile, IEnumerable<string> Definitions, FileItem IncludedFile, IActionGraphBuilder Graph)
 		{
 			// Build the contents of the wrapper file
 			StringBuilder WrapperContents = new StringBuilder();
@@ -1077,7 +1085,7 @@ namespace UnrealBuildTool
 			}
 
 			// Create the item
-			FileItem WrapperFile = FileItem.CreateIntermediateTextFile(OutputFile, WrapperContents.ToString());
+			FileItem WrapperFile = Graph.CreateIntermediateTextFile(OutputFile, WrapperContents.ToString());
 
 			// Touch it if the included file is newer, to make sure our timestamp dependency checking is accurate.
 			if (IncludedFile.LastWriteTimeUtc > WrapperFile.LastWriteTimeUtc)
@@ -1202,8 +1210,8 @@ namespace UnrealBuildTool
 		public CppCompileEnvironment CreateCompileEnvironmentForIntellisense(ReadOnlyTargetRules Target, CppCompileEnvironment BaseCompileEnvironment)
 		{
 			CppCompileEnvironment CompileEnvironment = CreateModuleCompileEnvironment(Target, BaseCompileEnvironment);
-			CompileEnvironment = SetupPrecompiledHeaders(Target, null, CompileEnvironment, new List<FileItem>(), new List<Action>());
-			CreateHeaderForDefinitions(CompileEnvironment, IntermediateDirectory, null);
+			CompileEnvironment = SetupPrecompiledHeaders(Target, null, CompileEnvironment, new List<FileItem>(), new NullActionGraphBuilder());
+			CreateHeaderForDefinitions(CompileEnvironment, IntermediateDirectory, null, new NullActionGraphBuilder());
 			return CompileEnvironment;
 		}
 
diff --git a/Engine/Source/Programs/UnrealBuildTool/Configuration/UEBuildTarget.cs b/Engine/Source/Programs/UnrealBuildTool/Configuration/UEBuildTarget.cs
--- a/Engine/Source/Programs/UnrealBuildTool/Configuration/UEBuildTarget.cs
+++ b/Engine/Source/Programs/UnrealBuildTool/Configuration/UEBuildTarget.cs
@@ -1705,7 +1705,7 @@ namespace UnrealBuildTool
 
 						FileItem DefaultResourceFile = FileItem.GetItemByFileReference(FileReference.Combine(UnrealBuildTool.EngineDirectory, "Build", "Windows", "Resources", "Default.rc2"));
 
-						CPPOutput DefaultResourceOutput = TargetToolChain.CompileRCFiles(DefaultResourceCompileEnvironment, new List<FileItem> { DefaultResourceFile }, EngineIntermediateDirectory, Makefile.Actions);
+						CPPOutput DefaultResourceOutput = TargetToolChain.CompileRCFiles(DefaultResourceCompileEnvironment, new List<FileItem> { DefaultResourceFile }, EngineIntermediateDirectory, Makefile);
 						GlobalLinkEnvironment.DefaultResourceFiles.AddRange(DefaultResourceOutput.ObjectFiles);
 					}
 				}
@@ -1746,7 +1746,7 @@ namespace UnrealBuildTool
 			{
 				if(!UnrealBuildTool.IsFileInstalled(Pair.Key))
 				{
-					Makefile.OutputItems.Add(CreateCopyAction(Pair.Value, Pair.Key, Makefile.Actions));
+					Makefile.OutputItems.Add(Makefile.CreateCopyAction(Pair.Value, Pair.Key));
 				}
 			}
 
@@ -1852,14 +1852,13 @@ namespace UnrealBuildTool
 					WriteMetadataArguments.Append(" -NoManifestChanges");
 				}
 
-				Action WriteMetadataAction = Action.CreateRecursiveAction<WriteMetadataMode>(ActionType.WriteMetadata, WriteMetadataArguments.ToString());
+				Action WriteMetadataAction = Makefile.CreateRecursiveAction<WriteMetadataMode>(ActionType.WriteMetadata, WriteMetadataArguments.ToString());
 				WriteMetadataAction.WorkingDirectory = UnrealBuildTool.EngineSourceDirectory;
 				WriteMetadataAction.StatusDescription = ReceiptFileName.GetFileName();
 				WriteMetadataAction.bCanExecuteRemotely = false;
 				WriteMetadataAction.PrerequisiteItems.Add(FileItem.GetItemByFileReference(MetadataTargetFile));
 				WriteMetadataAction.PrerequisiteItems.AddRange(Makefile.OutputItems);
 				WriteMetadataAction.ProducedItems.Add(FileItem.GetItemByFileReference(ReceiptFileName));
-				Makefile.Actions.Add(WriteMetadataAction);
 
 				Makefile.OutputItems.AddRange(WriteMetadataAction.ProducedItems);
 
@@ -1869,7 +1868,7 @@ namespace UnrealBuildTool
 				{
 					FileReference OutputFile = new FileReference(PostBuildScript.FullName + ".ran");
 
-					Action PostBuildStepAction = new Action(ActionType.PostBuildStep);
+					Action PostBuildStepAction = Makefile.CreateAction(ActionType.PostBuildStep);
 					PostBuildStepAction.CommandPath = BuildHostPlatform.Current.Shell;
 					if(BuildHostPlatform.Current.ShellType == ShellType.Cmd)
 					{
@@ -1884,7 +1883,6 @@ namespace UnrealBuildTool
 					PostBuildStepAction.bCanExecuteRemotely = false;
 					PostBuildStepAction.PrerequisiteItems.Add(FileItem.GetItemByFileReference(ReceiptFileName));
 					PostBuildStepAction.ProducedItems.Add(FileItem.GetItemByFileReference(OutputFile));
-					Makefile.Actions.Add(PostBuildStepAction);
 
 					Makefile.OutputItems.AddRange(PostBuildStepAction.ProducedItems);
 				}
@@ -2055,23 +2053,6 @@ namespace UnrealBuildTool
 			return RestrictedFolders.FindPermittedRestrictedFolderReferences(BaseDir, File.Directory);
 		}
 
-		/// <summary>
-		/// Creates an action which copies a file from one location to another
-		/// </summary>
-		/// <param name="SourceFile">The source file location</param>
-		/// <param name="TargetFile">The target file location</param>
-		/// <param name="Actions">List of actions to be executed. Additional actions will be added to this list.</param>
-		/// <returns>File item for the output file</returns>
-		static FileItem CreateCopyAction(FileReference SourceFile, FileReference TargetFile, List<Action> Actions)
-		{
-			FileItem SourceFileItem = FileItem.GetItemByFileReference(SourceFile);
-			FileItem TargetFileItem = FileItem.GetItemByFileReference(TargetFile);
-
-			Actions.Add(Action.CreateCopyAction(SourceFileItem, TargetFileItem));
-
-			return TargetFileItem;
-		}
-
 		/// <summary>
 		/// Creates a toolchain for the current target. May be overridden by the target rules.
 		/// </summary>
diff --git a/Engine/Source/Programs/UnrealBuildTool/Executors/SNDBS.cs b/Engine/Source/Programs/UnrealBuildTool/Executors/SNDBS.cs
--- a/Engine/Source/Programs/UnrealBuildTool/Executors/SNDBS.cs
+++ b/Engine/Source/Programs/UnrealBuildTool/Executors/SNDBS.cs
@@ -257,7 +257,8 @@ namespace UnrealBuildTool
 									Writer.WriteLine("#endif");
 								}
 
-								FileItem DummyPCHIncludeFileDependency = FileItem.CreateIntermediateTextFile(new FileReference(DummyPCHIncludeFile + ".dummy.h"), WrapperContents.ToString());
+								FileReference DummyPCHIncludeFileDependency = new FileReference(DummyPCHIncludeFile + ".dummy.h");
+								Utils.WriteFileIfChanged(DummyPCHIncludeFileDependency, WrapperContents.ToString(), StringComparison.OrdinalIgnoreCase);
 								AdditionalStubIncludes = string.Format("/FI\"{0}\"", DummyPCHIncludeFileDependency);
 							}
 
diff --git a/Engine/Source/Programs/UnrealBuildTool/Platform/Android/AndroidToolChain.cs b/Engine/Source/Programs/UnrealBuildTool/Platform/Android/AndroidToolChain.cs
--- a/Engine/Source/Programs/UnrealBuildTool/Platform/Android/AndroidToolChain.cs
+++ b/Engine/Source/Programs/UnrealBuildTool/Platform/Android/AndroidToolChain.cs
@@ -1116,7 +1116,7 @@ namespace UnrealBuildTool
 			}
 		}
 
-		void GenerateEmptyLinkFunctionsForRemovedModules(List<FileItem> SourceFiles, string ModuleName, DirectoryReference OutputDirectory)
+		void GenerateEmptyLinkFunctionsForRemovedModules(List<FileItem> SourceFiles, string ModuleName, DirectoryReference OutputDirectory, IActionGraphBuilder Graph)
 		{
 			// Only add to Launch module
 			if (!ModuleName.Equals("Launch"))
@@ -1131,7 +1131,7 @@ namespace UnrealBuildTool
 			if (!FileReference.Exists(LinkerExceptionsCPPFilename))
 			{
 				// Create a dummy file in case it doesn't exist yet so that the module does not complain it's not there
-				FileItem.CreateIntermediateTextFile(LinkerExceptionsCPPFilename, new List<string>());
+				Graph.CreateIntermediateTextFile(LinkerExceptionsCPPFilename, new List<string>());
 			}
 
 			List<string> Result = new List<string>();
@@ -1172,7 +1172,7 @@ namespace UnrealBuildTool
 			// If we determined that we should write the file, write it now.
 			if (bShouldWriteFile)
 			{
-				FileItem.CreateIntermediateTextFile(LinkerExceptionsCPPFilename, Result);
+				Graph.CreateIntermediateTextFile(LinkerExceptionsCPPFilename, Result);
 			}
 
 			SourceFiles.Add(FileItem.GetItemByFileReference(LinkerExceptionsCPPFilename));
@@ -1210,7 +1210,7 @@ namespace UnrealBuildTool
 
 		static private bool bHasPrintedApiLevel = false;
 		static private bool bHasHandledLaunchModule = false;
-		public override CPPOutput CompileCPPFiles(CppCompileEnvironment CompileEnvironment, List<FileItem> InputFiles, DirectoryReference OutputDir, string ModuleName, List<Action> Actions)
+		public override CPPOutput CompileCPPFiles(CppCompileEnvironment CompileEnvironment, List<FileItem> InputFiles, DirectoryReference OutputDir, string ModuleName, IActionGraphBuilder Graph)
 		{
 			if (Arches.Count == 0)
 			{
@@ -1279,7 +1279,7 @@ namespace UnrealBuildTool
 				// Directly added NDK files for NDK extensions
 				ModifySourceFiles(CompileEnvironment, InputFiles, ModuleName);
 				// Deal with dynamic modules removed by architecture
-				GenerateEmptyLinkFunctionsForRemovedModules(InputFiles, ModuleName, OutputDir);
+				GenerateEmptyLinkFunctionsForRemovedModules(InputFiles, ModuleName, OutputDir, Graph);
 
 				bHasHandledLaunchModule = true;
 			}
@@ -1372,7 +1372,7 @@ namespace UnrealBuildTool
 
 					foreach (FileItem SourceFile in InputFiles)
 					{
-						Action CompileAction = new Action(ActionType.Compile);
+						Action CompileAction = Graph.CreateAction(ActionType.Compile);
 						CompileAction.PrerequisiteItems.AddRange(CompileEnvironment.ForceIncludeFiles);
 						CompileAction.PrerequisiteItems.AddRange(CompileEnvironment.AdditionalPrerequisites);
 
@@ -1505,7 +1505,7 @@ namespace UnrealBuildTool
 
 						// Create the response file
 						FileReference ResponseFileName = CompileAction.ProducedItems[0].Location + ".rsp";
-						FileItem ResponseFileItem = FileItem.CreateIntermediateTextFile(ResponseFileName, new List<string> { AllArguments });
+						FileItem ResponseFileItem = Graph.CreateIntermediateTextFile(ResponseFileName, new List<string> { AllArguments });
 						string ResponseArgument = string.Format("@\"{0}\"", ResponseFileName);
 
 						CompileAction.WorkingDirectory = UnrealBuildTool.EngineSourceDirectory;
@@ -1536,8 +1536,6 @@ namespace UnrealBuildTool
 						CompileAction.bCanExecuteRemotely =
 							CompileEnvironment.PrecompiledHeaderAction != PrecompiledHeaderAction.Create ||
 							CompileEnvironment.bAllowRemotelyCompiledPCHs;
-
-						Actions.Add(CompileAction);
 					}
 				}
 			}
@@ -1545,7 +1543,7 @@ namespace UnrealBuildTool
 			return Result;
 		}
 
-		public override FileItem LinkFiles(LinkEnvironment LinkEnvironment, bool bBuildImportLibraryOnly, List<Action> Actions)
+		public override FileItem LinkFiles(LinkEnvironment LinkEnvironment, bool bBuildImportLibraryOnly, IActionGraphBuilder Graph)
 		{
 			return null;
 		}
@@ -1580,7 +1578,7 @@ namespace UnrealBuildTool
 			return DirectoryReference.Combine(PathRef, "include", Arch.Replace("-", "") + GPUArchitecture);
 		}
 
-		public override CPPOutput GenerateISPCHeaders(CppCompileEnvironment CompileEnvironment, List<FileItem> InputFiles, DirectoryReference OutputDir, List<Action> Actions)
+		public override CPPOutput GenerateISPCHeaders(CppCompileEnvironment CompileEnvironment, List<FileItem> InputFiles, DirectoryReference OutputDir, IActionGraphBuilder Graph)
 		{
 			if (Arches.Count == 0)
 			{
@@ -1609,7 +1607,7 @@ namespace UnrealBuildTool
 
 					foreach (FileItem ISPCFile in InputFiles)
 					{
-						Action CompileAction = new Action(ActionType.Compile);
+						Action CompileAction = Graph.CreateAction(ActionType.Compile);
 						CompileAction.CommandDescription = "Compile";
 						CompileAction.WorkingDirectory = UnrealBuildTool.EngineSourceDirectory;
 						CompileAction.CommandPath = new FileReference(GetISPCHostCompilerPath(BuildHostPlatform.Current.Platform));
@@ -1684,8 +1682,6 @@ namespace UnrealBuildTool
 						CompileAction.PrerequisiteItems.Add(ISPCFile);
 						CompileAction.StatusDescription = string.Format("{0} [{1}]", Path.GetFileName(ISPCFile.AbsolutePath), Arch.Replace("-", ""));
 
-						Actions.Add(CompileAction);
-
 						FileItem ISPCFinalHeaderFile = FileItem.GetItemByFileReference(
 							FileReference.Combine(
 								InlineArchIncludeFolder(OutputDir, Arch, GPUArchitecture),
@@ -1700,7 +1696,7 @@ namespace UnrealBuildTool
 						FileItem SourceFileItem = FileItem.GetItemByFileReference(SourceFile);
 						FileItem TargetFileItem = FileItem.GetItemByFileReference(TargetFile);
 
-						Action CopyAction = new Action(ActionType.BuildProject);
+						Action CopyAction = Graph.CreateAction(ActionType.BuildProject);
 						CopyAction.CommandDescription = "Copy";
 						CopyAction.CommandPath = BuildHostPlatform.Current.Shell;
 						if (BuildHostPlatform.Current.ShellType == ShellType.Cmd)
@@ -1717,7 +1713,6 @@ namespace UnrealBuildTool
 						CopyAction.StatusDescription = TargetFileItem.Location.GetFileName();
 						CopyAction.bCanExecuteRemotely = false;
 						CopyAction.bShouldOutputStatusDescription = false;
-						Actions.Add(CopyAction);
 
 						Result.GeneratedHeaderFiles.Add(TargetFileItem);
 
@@ -1729,7 +1724,7 @@ namespace UnrealBuildTool
 			return Result;
 		}
 		
-		public override CPPOutput CompileISPCFiles(CppCompileEnvironment CompileEnvironment, List<FileItem> InputFiles, DirectoryReference OutputDir, List<Action> Actions)
+		public override CPPOutput CompileISPCFiles(CppCompileEnvironment CompileEnvironment, List<FileItem> InputFiles, DirectoryReference OutputDir, IActionGraphBuilder Graph)
 		{
 			if (Arches.Count == 0)
 			{
@@ -1756,7 +1751,7 @@ namespace UnrealBuildTool
 
 					foreach (FileItem ISPCFile in InputFiles)
 					{
-						Action CompileAction = new Action(ActionType.Compile);
+						Action CompileAction = Graph.CreateAction(ActionType.Compile);
 						CompileAction.CommandDescription = "Compile";
 						CompileAction.WorkingDirectory = UnrealBuildTool.EngineSourceDirectory;
 						CompileAction.CommandPath = new FileReference(GetISPCHostCompilerPath(BuildHostPlatform.Current.Platform));
@@ -1906,8 +1901,6 @@ namespace UnrealBuildTool
 							CompileAction.StatusDescription = string.Format("{0} [{1}]", Path.GetFileName(ISPCFile.AbsolutePath), Arch.Replace("-", ""));
 						}
 
-						Actions.Add(CompileAction);
-
 						for(int i = 0; i < CompiledISPCObjFiles.Count; i++)
 						{
 							// ISPC compiler can't add suffix on the end of the arch, so copy to put into what linker expects
@@ -1922,7 +1915,7 @@ namespace UnrealBuildTool
 							FileItem SourceFileItem = FileItem.GetItemByFileReference(SourceFile);
 							FileItem TargetFileItem = FileItem.GetItemByFileReference(TargetFile);
 
-							Action CopyAction = new Action(ActionType.BuildProject);
+							Action CopyAction = Graph.CreateAction(ActionType.BuildProject);
 							CopyAction.CommandDescription = "Copy";
 							CopyAction.CommandPath = BuildHostPlatform.Current.Shell;
 							if (BuildHostPlatform.Current.ShellType == ShellType.Cmd)
@@ -1939,8 +1932,6 @@ namespace UnrealBuildTool
 							CopyAction.StatusDescription = TargetFileItem.Location.GetFileName();
 							CopyAction.bCanExecuteRemotely = false;
 							CopyAction.bShouldOutputStatusDescription = false;
-
-							Actions.Add(CopyAction);
 						}
 
 						Result.ObjectFiles.AddRange(FinalISPCObjFiles);
@@ -1953,7 +1944,7 @@ namespace UnrealBuildTool
 			return Result;
 		}
 
-		public override FileItem[] LinkAllFiles(LinkEnvironment LinkEnvironment, bool bBuildImportLibraryOnly, List<Action> Actions)
+		public override FileItem[] LinkAllFiles(LinkEnvironment LinkEnvironment, bool bBuildImportLibraryOnly, IActionGraphBuilder Graph)
 		{
 			List<FileItem> Outputs = new List<FileItem>();
 
@@ -1981,7 +1972,7 @@ namespace UnrealBuildTool
 					}
 
 					// Create an action that invokes the linker.
-					Action LinkAction = new Action(ActionType.Link);
+					Action LinkAction = Graph.CreateAction(ActionType.Link);
 					LinkAction.WorkingDirectory = UnrealBuildTool.EngineSourceDirectory;
 
 					if (LinkEnvironment.bIsBuildingLibrary)
@@ -2156,7 +2147,7 @@ namespace UnrealBuildTool
 					FileReference ResponseFileName = GetResponseFileName(LinkEnvironment, OutputFile);
 					InputFileNames.Add(LinkResponseArguments.Replace("\\", "/"));
 
-					FileItem ResponseFileItem = FileItem.CreateIntermediateTextFile(ResponseFileName, InputFileNames);
+					FileItem ResponseFileItem = Graph.CreateIntermediateTextFile(ResponseFileName, InputFileNames);
 
 					LinkAction.CommandArguments += string.Format(" @\"{0}\"", ResponseFileName);
 					LinkAction.PrerequisiteItems.Add(ResponseFileItem);
@@ -2171,7 +2162,6 @@ namespace UnrealBuildTool
 					{
 						SetupActionToExecuteCompilerThroughShell(ref LinkAction, LinkAction.CommandPath.FullName, LinkAction.CommandArguments, "Link");
 					}
-					Actions.Add(LinkAction);
 
 					// Windows can run into an issue with too long of a commandline when clang tries to call ld to link.
 					// To work around this we call clang to just get the command it would execute and generate a
@@ -2214,7 +2204,7 @@ namespace UnrealBuildTool
 
 								// now create a response file for the full command using ld directly
 								FileReference FinalResponseFileName = FileReference.Combine(LinkEnvironment.IntermediateDirectory, OutputFile.Location.GetFileName() + ".responseFinal");
-								FileItem FinalResponseFileItem = FileItem.CreateIntermediateTextFile(FinalResponseFileName, LinkAction.CommandArguments);
+								FileItem FinalResponseFileItem = Graph.CreateIntermediateTextFile(FinalResponseFileName, LinkAction.CommandArguments);
 								LinkAction.CommandArguments = string.Format("@\"{0}\"", FinalResponseFileName);
 								LinkAction.PrerequisiteItems.Add(FinalResponseFileItem);
 							}
diff --git a/Engine/Source/Programs/UnrealBuildTool/Platform/HoloLens/HoloLensToolChain.cs b/Engine/Source/Programs/UnrealBuildTool/Platform/HoloLens/HoloLensToolChain.cs
--- a/Engine/Source/Programs/UnrealBuildTool/Platform/HoloLens/HoloLensToolChain.cs
+++ b/Engine/Source/Programs/UnrealBuildTool/Platform/HoloLens/HoloLensToolChain.cs
@@ -434,7 +434,7 @@ namespace UnrealBuildTool
 			Arguments.Add("/ignore:4264");
 		}
 
-		public override CPPOutput CompileCPPFiles(CppCompileEnvironment CompileEnvironment, List<FileItem> InputFiles, DirectoryReference OutputDir, string ModuleName, List<Action> Actions)
+		public override CPPOutput CompileCPPFiles(CppCompileEnvironment CompileEnvironment, List<FileItem> InputFiles, DirectoryReference OutputDir, string ModuleName, IActionGraphBuilder Graph)
 		{
 			List<string> SharedArguments = new List<string>();
 			AppendCLArguments_Global(CompileEnvironment, EnvVars, SharedArguments);
@@ -466,8 +466,7 @@ namespace UnrealBuildTool
 			CPPOutput Result = new CPPOutput();
 			foreach (FileItem SourceFile in InputFiles)
 			{
-				Action CompileAction = new Action(ActionType.Compile);
-				Actions.Add(CompileAction);
+				Action CompileAction = Graph.CreateAction(ActionType.Compile);
 				CompileAction.CommandDescription = "Compile";
 
 				List<string> FileArguments = new List<string>();
@@ -481,7 +480,7 @@ namespace UnrealBuildTool
 				{
 					// Generate a CPP File that just includes the precompiled header.
 					FileReference PCHCPPPath = CompileEnvironment.PrecompiledHeaderIncludeFilename.ChangeExtension(".cpp");
-					FileItem PCHCPPFile = FileItem.CreateIntermediateTextFile(
+					FileItem PCHCPPFile = Graph.CreateIntermediateTextFile(
 						PCHCPPPath,
 						string.Format("#include \"{0}\"\r\n", CompileEnvironment.PrecompiledHeaderIncludeFilename.FullName.Replace('\\', '/'))
 						);
@@ -629,7 +628,7 @@ namespace UnrealBuildTool
 				{
 					FileItem TargetFile = CompileAction.ProducedItems[0];
 					FileReference ResponseFileName = new FileReference(TargetFile.AbsolutePath + ".response");
-					FileItem ResponseFileItem = FileItem.CreateIntermediateTextFile(ResponseFileName, SharedArguments.Concat(FileArguments).Concat(AdditionalArguments).Select(x => Utils.ExpandVariables(x)));
+					FileItem ResponseFileItem = Graph.CreateIntermediateTextFile(ResponseFileName, SharedArguments.Concat(FileArguments).Concat(AdditionalArguments).Select(x => Utils.ExpandVariables(x)));
 					CompileAction.CommandArguments = " @\"" + ResponseFileName + "\"";
 					CompileAction.PrerequisiteItems.Add(ResponseFileItem);
 				}
@@ -678,13 +677,13 @@ namespace UnrealBuildTool
 			return Result;
 		}
 
-		public override CPPOutput CompileRCFiles(CppCompileEnvironment CompileEnvironment, List<FileItem> InputFiles, DirectoryReference OutputDir, List<Action> Actions)
+		public override CPPOutput CompileRCFiles(CppCompileEnvironment CompileEnvironment, List<FileItem> InputFiles, DirectoryReference OutputDir, IActionGraphBuilder Graph)
 		{
 			CPPOutput Result = new CPPOutput();
 
 			foreach (FileItem RCFile in InputFiles)
 			{
-				Action CompileAction = new Action(ActionType.Compile);
+				Action CompileAction = Graph.CreateAction(ActionType.Compile);
 				CompileAction.CommandDescription = "Resource";
 				CompileAction.WorkingDirectory = UnrealBuildTool.EngineSourceDirectory;
 				CompileAction.CommandPath = EnvVars.ResourceCompilerPath;
@@ -752,13 +751,12 @@ namespace UnrealBuildTool
 
 				// Add the C++ source file and its included files to the prerequisite item list.
 				CompileAction.PrerequisiteItems.Add(RCFile);
-				Actions.Add(CompileAction);
 			}
 
 			return Result;
 		}
 
-		public override FileItem LinkFiles(LinkEnvironment LinkEnvironment, bool bBuildImportLibraryOnly, List<Action> Actions) //(LinkEnvironment LinkEnvironment, bool bBuildImportLibraryOnly)
+		public override FileItem LinkFiles(LinkEnvironment LinkEnvironment, bool bBuildImportLibraryOnly, IActionGraphBuilder Graph) //(LinkEnvironment LinkEnvironment, bool bBuildImportLibraryOnly)
 		{
 			if (LinkEnvironment.bIsBuildingDotNetAssembly)
 			{
@@ -951,13 +949,12 @@ namespace UnrealBuildTool
 			FileReference ResponseFileName = GetResponseFileName(LinkEnvironment, OutputFile);
 			if (!ProjectFileGenerator.bGenerateProjectFiles)
 			{
-				FileItem ResponseFile = FileItem.CreateIntermediateTextFile(ResponseFileName, String.Join(Environment.NewLine, Arguments));
+				FileItem ResponseFile = Graph.CreateIntermediateTextFile(ResponseFileName, String.Join(Environment.NewLine, Arguments));
 				PrerequisiteItems.Add(ResponseFile);
 			}
 
 			// Create an action that invokes the linker.
-			Action LinkAction = new Action(ActionType.Link);
-			Actions.Add(LinkAction);
+			Action LinkAction = Graph.CreateAction(ActionType.Link);
 			LinkAction.CommandDescription = "Link";
 			LinkAction.WorkingDirectory = UnrealBuildTool.EngineSourceDirectory;
 			LinkAction.CommandPath = bIsBuildingLibrary ? EnvVars.LibraryManagerPath : EnvVars.LinkerPath;
diff --git a/Engine/Source/Programs/UnrealBuildTool/Platform/IOS/IOSToolChain.cs b/Engine/Source/Programs/UnrealBuildTool/Platform/IOS/IOSToolChain.cs
--- a/Engine/Source/Programs/UnrealBuildTool/Platform/IOS/IOSToolChain.cs
+++ b/Engine/Source/Programs/UnrealBuildTool/Platform/IOS/IOSToolChain.cs
@@ -596,7 +596,7 @@ namespace UnrealBuildTool
 			return Result;
 		}
 
-		public override CPPOutput CompileCPPFiles(CppCompileEnvironment CompileEnvironment, List<FileItem> InputFiles, DirectoryReference OutputDir, string ModuleName, List<Action> Actions)
+		public override CPPOutput CompileCPPFiles(CppCompileEnvironment CompileEnvironment, List<FileItem> InputFiles, DirectoryReference OutputDir, string ModuleName, IActionGraphBuilder Graph)
 		{
 			string Arguments = GetCompileArguments_Global(CompileEnvironment);
 			string PCHArguments = "";
@@ -628,7 +628,7 @@ namespace UnrealBuildTool
 			// Create a compile action for each source file.
 			foreach (FileItem SourceFile in InputFiles)
 			{
-				Action CompileAction = new Action(ActionType.Compile);
+				Action CompileAction = Graph.CreateAction(ActionType.Compile);
 				string FilePCHArguments = "";
 				string FileArguments = "";
 				string Extension = Path.GetExtension(SourceFile.AbsolutePath).ToUpperInvariant();
@@ -772,23 +772,21 @@ namespace UnrealBuildTool
 				{
 					if (Framework.ZipFile != null)
 					{
-						FileItem ExtractedTokenFile = ExtractFramework(Framework, Actions);
+						FileItem ExtractedTokenFile = ExtractFramework(Framework, Graph);
 						CompileAction.PrerequisiteItems.Add(ExtractedTokenFile);
 					}
 				}
-
-				Actions.Add(CompileAction);
 			}
 			return Result;
 		}
 
-		public override FileItem LinkFiles(LinkEnvironment LinkEnvironment, bool bBuildImportLibraryOnly, List<Action> Actions)
+		public override FileItem LinkFiles(LinkEnvironment LinkEnvironment, bool bBuildImportLibraryOnly, IActionGraphBuilder Graph)
 		{
 			string LinkerPath = Settings.Value.ToolchainDir +
 				(LinkEnvironment.bIsBuildingLibrary ? IOSArchiver : IOSLinker);
 
 			// Create an action that invokes the linker.
-			Action LinkAction = new Action(ActionType.Link);
+			Action LinkAction = Graph.CreateAction(ActionType.Link);
 
 			// RPC utility parameters are in terms of the Mac side
 			LinkAction.WorkingDirectory = GetMacDevSrcRoot();
@@ -841,7 +839,7 @@ namespace UnrealBuildTool
 			{
 				if (Framework.ZipFile != null)
 				{
-					FileItem ExtractedTokenFile = ExtractFramework(Framework, Actions);
+					FileItem ExtractedTokenFile = ExtractFramework(Framework, Graph);
 					LinkAction.PrerequisiteItems.Add(ExtractedTokenFile);
 				}
 			}
@@ -883,7 +881,7 @@ namespace UnrealBuildTool
 			{
 				bool bIsUE4Game = LinkEnvironment.OutputFilePath.FullName.Contains("UE4Game");
 				FileReference ResponsePath = FileReference.Combine(((!bIsUE4Game && ProjectFile != null) ? ProjectFile.Directory : UnrealBuildTool.EngineDirectory), "Intermediate", "Build", LinkEnvironment.Platform.ToString(), "LinkFileList_" + LinkEnvironment.OutputFilePath.GetFileNameWithoutExtension() + ".tmp");
-				FileItem.CreateIntermediateTextFile(ResponsePath, InputFileNames);
+				Graph.CreateIntermediateTextFile(ResponsePath, InputFileNames);
 				LinkCommandArguments += string.Format(" @\"{0}\"", ResponsePath.FullName);
 			}
 
@@ -936,7 +934,6 @@ namespace UnrealBuildTool
 				// This is not a shipping build so no need to delete the output file since symbols will not have been stripped from it.
 				LinkAction.CommandArguments = string.Format("-c '\"{0}\" {1}'", LinkerPath, LinkCommandArguments);
 			}
-			Actions.Add(LinkAction);
 
 			return OutputFile;
 		}
@@ -1025,9 +1022,9 @@ namespace UnrealBuildTool
 		/// Generates debug info for a given executable
 		/// </summary>
 		/// <param name="Executable">FileItem describing the executable to generate debug info for</param>
-		/// <param name="Actions">List of actions to be executed. Additional actions will be added to this list.</param>
 		/// <param name="bIsForLTOBuild">Was this build made with LTO enabled?</param>
-		private List<FileItem> GenerateDebugInfo(FileItem Executable, List<Action> Actions, bool bIsForLTOBuild)
+		/// <param name="Graph">List of actions to be executed. Additional actions will be added to this list.</param>
+		private List<FileItem> GenerateDebugInfo(FileItem Executable, bool bIsForLTOBuild, IActionGraphBuilder Graph)
 		{
 			// Make a file item for the source and destination files
 			string FullDestPathRoot = GetdSYMPath(Executable);
@@ -1036,7 +1033,7 @@ namespace UnrealBuildTool
 			FileItem ZipOutputFile = FileItem.GetItemByPath(FullDestPathRoot + ".zip");
 
 			// Make the compile action
-			Action GenDebugAction = new Action(ActionType.GenerateDebugInfo);
+			Action GenDebugAction = Graph.CreateAction(ActionType.GenerateDebugInfo);
 
 			GenDebugAction.WorkingDirectory = GetMacDevSrcRoot();
 			GenDebugAction.CommandPath = BuildHostPlatform.Current.Shell;
@@ -1065,7 +1062,6 @@ namespace UnrealBuildTool
 			GenDebugAction.ProducedItems.Add(OutputFile);
 			GenDebugAction.StatusDescription = GenDebugAction.CommandArguments;// string.Format("Generating debug info for {0}", Path.GetFileName(Executable.AbsolutePath));
 			GenDebugAction.bCanExecuteRemotely = false;
-			Actions.Add(GenDebugAction);
 
 			return GenDebugAction.ProducedItems; // (ProjectSettings.bGeneratedSYMBundle ? ZipOutputFile : OutputFile);
 		}
@@ -1074,8 +1070,8 @@ namespace UnrealBuildTool
 		/// Generates pseudo pdb info for a given executable
 		/// </summary>
 		/// <param name="Executable">FileItem describing the executable to generate debug info for</param>
-		/// <param name="Actions">List of actions to be executed. Additional actions will be added to this list.</param>
-		public FileItem GeneratePseudoPDB(FileItem Executable, List<Action> Actions)
+		/// <param name="Graph">List of actions to be executed. Additional actions will be added to this list.</param>
+		public FileItem GeneratePseudoPDB(FileItem Executable, IActionGraphBuilder Graph)
 		{
 			// Make a file item for the source and destination files
 			string FulldSYMPathRoot = GetdSYMPath(Executable);
@@ -1089,7 +1085,7 @@ namespace UnrealBuildTool
 			FileItem OutputFile = FileItem.GetItemByPath(FullDestPathRoot);
 
 			// Make the compile action
-			Action GenDebugAction = new Action(ActionType.GenerateDebugInfo);
+			Action GenDebugAction = Graph.CreateAction(ActionType.GenerateDebugInfo);
 			GenDebugAction.WorkingDirectory = DirectoryReference.Combine(UnrealBuildTool.EngineDirectory, "Binaries", "Mac");
 
 			GenDebugAction.CommandPath = BuildHostPlatform.Current.Shell;
@@ -1102,7 +1098,6 @@ namespace UnrealBuildTool
 			GenDebugAction.ProducedItems.Add(OutputFile);
 			GenDebugAction.StatusDescription = GenDebugAction.CommandArguments;// string.Format("Generating debug info for {0}", Path.GetFileName(Executable.AbsolutePath));
 			GenDebugAction.bCanExecuteRemotely = false;
-			Actions.Add(GenDebugAction);
 
 			return OutputFile;
 		}
@@ -1221,7 +1216,7 @@ namespace UnrealBuildTool
 			}
 		}
 
-		FileItem ExtractFramework(UEBuildFramework Framework, List<Action> Actions)
+		FileItem ExtractFramework(UEBuildFramework Framework, IActionGraphBuilder Graph)
 		{
 			if (Framework.ZipFile == null)
 			{
@@ -1240,9 +1235,9 @@ namespace UnrealBuildTool
 				ExtractScript.AppendLine(String.Format("unzip -q -o {0} -d {1}", Utils.MakePathSafeToUseWithCommandLine(Framework.ZipFile.FullName), Utils.MakePathSafeToUseWithCommandLine(Framework.OutputDirectory.ParentDirectory.FullName))); // Zip contains folder with the same name, hence ParentDirectory
 				ExtractScript.AppendLine(String.Format("touch {0}", Utils.MakePathSafeToUseWithCommandLine(Framework.ExtractedTokenFile.AbsolutePath)));
 
-				FileItem ExtractScriptFileItem = FileItem.CreateIntermediateTextFile(new FileReference(Framework.OutputDirectory.FullName + ".sh"), ExtractScript.ToString());
+				FileItem ExtractScriptFileItem = Graph.CreateIntermediateTextFile(new FileReference(Framework.OutputDirectory.FullName + ".sh"), ExtractScript.ToString());
 
-				Action UnzipAction = new Action(ActionType.BuildProject);
+				Action UnzipAction = Graph.CreateAction(ActionType.BuildProject);
 				UnzipAction.CommandPath = new FileReference("/bin/sh");
 				UnzipAction.CommandArguments = Utils.MakePathSafeToUseWithCommandLine(ExtractScriptFileItem.AbsolutePath);
 				UnzipAction.WorkingDirectory = UnrealBuildTool.EngineDirectory;
@@ -1252,7 +1247,6 @@ namespace UnrealBuildTool
 				UnzipAction.DeleteItems.Add(Framework.ExtractedTokenFile);
 				UnzipAction.StatusDescription = String.Format("Unzipping : {0} -> {1}", Framework.ZipFile, Framework.OutputDirectory);
 				UnzipAction.bCanExecuteRemotely = false;
-				Actions.Add(UnzipAction);
 			}
 			return Framework.ExtractedTokenFile;
 		}
@@ -1402,9 +1396,9 @@ namespace UnrealBuildTool
 			return new DirectoryReference(ResourcesDir);
 		}
 
-		public override ICollection<FileItem> PostBuild(FileItem Executable, LinkEnvironment BinaryLinkEnvironment, List<Action> Actions)
+		public override ICollection<FileItem> PostBuild(FileItem Executable, LinkEnvironment BinaryLinkEnvironment, IActionGraphBuilder Graph)
 		{
-			List<FileItem> OutputFiles = new List<FileItem>(base.PostBuild(Executable, BinaryLinkEnvironment, Actions));
+			List<FileItem> OutputFiles = new List<FileItem>(base.PostBuild(Executable, BinaryLinkEnvironment, Graph));
 
 			if (BinaryLinkEnvironment.bIsBuildingLibrary)
 			{
@@ -1414,14 +1408,14 @@ namespace UnrealBuildTool
 			// For IOS/tvOS, generate the dSYM file if needed or requested
 			if (Target.IOSPlatform.bGeneratedSYM)
 			{
-				List<FileItem> Files = GenerateDebugInfo(Executable, Actions, BinaryLinkEnvironment.bAllowLTCG);
+				List<FileItem> Files = GenerateDebugInfo(Executable, BinaryLinkEnvironment.bAllowLTCG, Graph);
 				foreach (FileItem item in Files)
 				{
 					OutputFiles.Add(item);
 				}
 				if (ProjectSettings.bGenerateCrashReportSymbols || Target.bUseMallocProfiler)
 				{
-					OutputFiles.Add(GeneratePseudoPDB(Executable, Actions));
+					OutputFiles.Add(GeneratePseudoPDB(Executable, Graph));
 				}
 			}
 
@@ -1433,7 +1427,7 @@ namespace UnrealBuildTool
 				// If building a framework we can only strip local symbols, need to leave global in place
 				string StripArguments = BinaryLinkEnvironment.bIsBuildingDLL ? "-x" : "";
 
-				Action StripAction = new Action(ActionType.CreateAppBundle);
+				Action StripAction = Graph.CreateAction(ActionType.CreateAppBundle);
 				StripAction.WorkingDirectory = GetMacDevSrcRoot();
 				StripAction.CommandPath = BuildHostPlatform.Current.Shell;
 				StripAction.CommandArguments = String.Format("-c '\"{0}strip\" {1} \"{2}\" && touch \"{3}\"'", Settings.Value.ToolchainDir, StripArguments, Executable.Location, StripCompleteFile);
@@ -1442,7 +1436,6 @@ namespace UnrealBuildTool
 				StripAction.ProducedItems.Add(StripCompleteFile);
 				StripAction.StatusDescription = String.Format("Stripping symbols from {0}", Executable.AbsolutePath);
 				StripAction.bCanExecuteRemotely = false;
-				Actions.Add(StripAction);
 
 				OutputFiles.Add(StripCompleteFile);
 			}
@@ -1457,7 +1450,7 @@ namespace UnrealBuildTool
 				FileItem AssetCatalogFile = FileItem.GetItemByFileReference(GetAssetCatalogFile(BinaryLinkEnvironment.Platform, Executable.Location));
 
 				// Make the compile action
-				Action CompileAssetAction = new Action(ActionType.CreateAppBundle);
+				Action CompileAssetAction = Graph.CreateAction(ActionType.CreateAppBundle);
 				CompileAssetAction.WorkingDirectory = GetMacDevSrcRoot();
 				CompileAssetAction.CommandPath = new FileReference("/usr/bin/xcrun");
 				CompileAssetAction.CommandArguments = GetAssetCatalogArgs(BinaryLinkEnvironment.Platform, ResourcesDir.FullName, Path.GetDirectoryName(AssetCatalogFile.AbsolutePath));
@@ -1466,7 +1459,6 @@ namespace UnrealBuildTool
 				CompileAssetAction.DeleteItems.Add(AssetCatalogFile);
 				CompileAssetAction.StatusDescription = CompileAssetAction.CommandArguments;// string.Format("Generating debug info for {0}", Path.GetFileName(Executable.AbsolutePath));
 				CompileAssetAction.bCanExecuteRemotely = false;
-				Actions.Add(CompileAssetAction);
 
 				// Add it to the output files so it's always built
 				OutputFiles.Add(AssetCatalogFile);
@@ -1495,7 +1487,7 @@ namespace UnrealBuildTool
 				FileReference StagedExecutablePath = GetStagedExecutablePath(Executable.Location, Target.Name);
 				foreach (UEBuildBundleResource Resource in BinaryLinkEnvironment.AdditionalBundleResources)
 				{
-					OutputFiles.Add(CopyBundleResource(Resource, Executable, StagedExecutablePath.Directory, Actions));
+					OutputFiles.Add(CopyBundleResource(Resource, Executable, StagedExecutablePath.Directory, Graph));
 				}
 
 				IOSPostBuildSyncTarget PostBuildSyncTarget = new IOSPostBuildSyncTarget(Target, Executable.Location, BinaryLinkEnvironment.IntermediateDirectory, UPLScripts, SdkVersion, FrameworkNameToSourceDir);
@@ -1504,7 +1496,7 @@ namespace UnrealBuildTool
 
 				string PostBuildSyncArguments = String.Format("-Input=\"{0}\" -XmlConfigCache=\"{1}\" -remoteini=\"{2}\"", PostBuildSyncFile, XmlConfig.CacheFile, UnrealBuildTool.GetRemoteIniPath());
 
-				Action PostBuildSyncAction = Action.CreateRecursiveAction<IOSPostBuildSyncMode>(ActionType.CreateAppBundle, PostBuildSyncArguments);
+				Action PostBuildSyncAction = Graph.CreateRecursiveAction<IOSPostBuildSyncMode>(ActionType.CreateAppBundle, PostBuildSyncArguments);
 				PostBuildSyncAction.WorkingDirectory = UnrealBuildTool.EngineSourceDirectory;
 				PostBuildSyncAction.PrerequisiteItems.Add(Executable);
 				PostBuildSyncAction.PrerequisiteItems.AddRange(OutputFiles);
@@ -1512,7 +1504,6 @@ namespace UnrealBuildTool
 				PostBuildSyncAction.DeleteItems.AddRange(PostBuildSyncAction.ProducedItems);
 				PostBuildSyncAction.StatusDescription = "Executing PostBuildSync";
 				PostBuildSyncAction.bCanExecuteRemotely = false;
-				Actions.Add(PostBuildSyncAction);
 
 				OutputFiles.AddRange(PostBuildSyncAction.ProducedItems);
 			}
@@ -2125,9 +2116,9 @@ namespace UnrealBuildTool
 			StripSymbolsWithXcode(SourceFile, TargetFile, Settings.Value.ToolchainDir);
 		}
 
-		FileItem CopyBundleResource(UEBuildBundleResource Resource, FileItem Executable, DirectoryReference BundleDirectory, List<Action> Actions)
+		FileItem CopyBundleResource(UEBuildBundleResource Resource, FileItem Executable, DirectoryReference BundleDirectory, IActionGraphBuilder Graph)
 		{
-			Action CopyAction = new Action(ActionType.CreateAppBundle);
+			Action CopyAction = Graph.CreateAction(ActionType.CreateAppBundle);
 			CopyAction.WorkingDirectory = GetMacDevSrcRoot(); // Path.GetFullPath(".");
 			CopyAction.CommandPath = BuildHostPlatform.Current.Shell;
 			CopyAction.CommandDescription = "";
@@ -2144,7 +2135,6 @@ namespace UnrealBuildTool
 			CopyAction.bShouldOutputStatusDescription = Resource.bShouldLog;
 			CopyAction.StatusDescription = string.Format("Copying {0} to app bundle", Path.GetFileName(Resource.ResourcePath));
 			CopyAction.bCanExecuteRemotely = false;
-			Actions.Add(CopyAction);
 
 			return TargetItem;
 		}
diff --git a/Engine/Source/Programs/UnrealBuildTool/Platform/Linux/LinuxToolChain.cs b/Engine/Source/Programs/UnrealBuildTool/Platform/Linux/LinuxToolChain.cs
--- a/Engine/Source/Programs/UnrealBuildTool/Platform/Linux/LinuxToolChain.cs
+++ b/Engine/Source/Programs/UnrealBuildTool/Platform/Linux/LinuxToolChain.cs
@@ -1273,7 +1273,7 @@ namespace UnrealBuildTool
 			return false;
 		}
 
-		public override CPPOutput CompileCPPFiles(CppCompileEnvironment CompileEnvironment, List<FileItem> InputFiles, DirectoryReference OutputDir, string ModuleName, List<Action> Actions)
+		public override CPPOutput CompileCPPFiles(CppCompileEnvironment CompileEnvironment, List<FileItem> InputFiles, DirectoryReference OutputDir, string ModuleName, IActionGraphBuilder Graph)
 		{
 			string Arguments = GetCLArguments_Global(CompileEnvironment);
 			string PCHArguments = "";
@@ -1341,7 +1341,7 @@ namespace UnrealBuildTool
 			CPPOutput Result = new CPPOutput();
 			foreach (FileItem SourceFile in InputFiles)
 			{
-				Action CompileAction = new Action(ActionType.Compile);
+				Action CompileAction = Graph.CreateAction(ActionType.Compile);
 				CompileAction.PrerequisiteItems.AddRange(CompileEnvironment.ForceIncludeFiles);
 				CompileAction.PrerequisiteItems.AddRange(CompileEnvironment.AdditionalPrerequisites);
 
@@ -1439,7 +1439,7 @@ namespace UnrealBuildTool
 				Debug.Assert(CompileAction.ProducedItems.Count > 0);
 
 				FileReference CompilerResponseFileName = CompileAction.ProducedItems[0].Location + ".rsp";
-				FileItem CompilerResponseFileItem = FileItem.CreateIntermediateTextFile(CompilerResponseFileName, AllArguments);
+				FileItem CompilerResponseFileItem = Graph.CreateIntermediateTextFile(CompilerResponseFileName, AllArguments);
 
 				CompileAction.CommandArguments = string.Format(" @\"{0}\"", CompilerResponseFileName);
 				CompileAction.PrerequisiteItems.Add(CompilerResponseFileItem);
@@ -1451,8 +1451,6 @@ namespace UnrealBuildTool
 				CompileAction.bCanExecuteRemotely =
 					CompileEnvironment.PrecompiledHeaderAction != PrecompiledHeaderAction.Create ||
 					CompileEnvironment.bAllowRemotelyCompiledPCHs;
-
-				Actions.Add(CompileAction);
 			}
 
 			return Result;
@@ -1466,10 +1464,10 @@ namespace UnrealBuildTool
 		/// <summary>
 		/// Creates an action to archive all the .o files into single .a file
 		/// </summary>
-		public FileItem CreateArchiveAndIndex(LinkEnvironment LinkEnvironment, List<Action> Actions)
+		public FileItem CreateArchiveAndIndex(LinkEnvironment LinkEnvironment, IActionGraphBuilder Graph)
 		{
 			// Create an archive action
-			Action ArchiveAction = new Action(ActionType.Link);
+			Action ArchiveAction = Graph.CreateAction(ActionType.Link);
 			ArchiveAction.WorkingDirectory = UnrealBuildTool.EngineSourceDirectory;
 			ArchiveAction.CommandPath = BuildHostPlatform.Current.Shell;
 
@@ -1505,7 +1503,7 @@ namespace UnrealBuildTool
 			FileReference ResponsePath = GetResponseFileName(LinkEnvironment, OutputFile);
 			if (!ProjectFileGenerator.bGenerateProjectFiles)
 			{
-				FileItem ResponseFileItem = FileItem.CreateIntermediateTextFile(ResponsePath, InputFileNames);
+				FileItem ResponseFileItem = Graph.CreateIntermediateTextFile(ResponsePath, InputFileNames);
 				ArchiveAction.PrerequisiteItems.Add(ResponseFileItem);
 			}
 			ArchiveAction.CommandArguments += string.Format(" @\"{0}\"", ResponsePath.FullName);
@@ -1531,12 +1529,11 @@ namespace UnrealBuildTool
 
 			// Only execute linking on the local PC.
 			ArchiveAction.bCanExecuteRemotely = false;
-			Actions.Add(ArchiveAction);
 
 			return OutputFile;
 		}
 
-		public FileItem FixDependencies(LinkEnvironment LinkEnvironment, FileItem Executable, List<Action> Actions)
+		public FileItem FixDependencies(LinkEnvironment LinkEnvironment, FileItem Executable, IActionGraphBuilder Graph)
 		{
 			if (bUseFixdeps)
 			{
@@ -1554,7 +1551,7 @@ namespace UnrealBuildTool
 
 				FileItem FixDepsScript = FileItem.GetItemByFileReference(FileReference.Combine(LinkEnvironment.LocalShadowDirectory, ScriptName));
 
-				Action PostLinkAction = new Action(ActionType.Link);
+				Action PostLinkAction = Graph.CreateAction(ActionType.Link);
 				PostLinkAction.WorkingDirectory = UnrealBuildTool.EngineSourceDirectory;
 				PostLinkAction.CommandPath = ShellBinary;
 				PostLinkAction.StatusDescription = string.Format("{0}", Path.GetFileName(Executable.AbsolutePath));
@@ -1585,7 +1582,6 @@ namespace UnrealBuildTool
 				System.Console.WriteLine("{0} {1}", PostLinkAction.CommandPath, PostLinkAction.CommandArguments);
 
 				PostLinkAction.ProducedItems.Add(OutputFile);
-				Actions.Add(PostLinkAction);
 				return OutputFile;
 			}
 			else
@@ -1601,7 +1597,7 @@ namespace UnrealBuildTool
 		}
 
 
-		public override FileItem LinkFiles(LinkEnvironment LinkEnvironment, bool bBuildImportLibraryOnly, List<Action> Actions)
+		public override FileItem LinkFiles(LinkEnvironment LinkEnvironment, bool bBuildImportLibraryOnly, IActionGraphBuilder Graph)
 		{
 			Debug.Assert(!bBuildImportLibraryOnly);
 
@@ -1614,11 +1610,11 @@ namespace UnrealBuildTool
 
 			if (LinkEnvironment.bIsBuildingLibrary || bBuildImportLibraryOnly)
 			{
-				return CreateArchiveAndIndex(LinkEnvironment, Actions);
+				return CreateArchiveAndIndex(LinkEnvironment, Graph);
 			}
 
 			// Create an action that invokes the linker.
-			Action LinkAction = new Action(ActionType.Link);
+			Action LinkAction = Graph.CreateAction(ActionType.Link);
 			LinkAction.WorkingDirectory = UnrealBuildTool.EngineSourceDirectory;
 
 			string LinkCommandString;
@@ -1856,7 +1852,7 @@ namespace UnrealBuildTool
 			ResponseLines.Add(" --end-group");
 
 			FileReference ResponseFileName = GetResponseFileName(LinkEnvironment, OutputFile);
-			FileItem ResponseFileItem = FileItem.CreateIntermediateTextFile(ResponseFileName, ResponseLines);
+			FileItem ResponseFileItem = Graph.CreateIntermediateTextFile(ResponseFileName, ResponseLines);
 
 			LinkCommandString += string.Format(" -Wl,@\"{0}\"", ResponseFileName);
 			LinkAction.PrerequisiteItems.Add(ResponseFileItem);
@@ -1970,7 +1966,6 @@ namespace UnrealBuildTool
 
 			// Only execute linking on the local PC.
 			LinkAction.bCanExecuteRemotely = false;
-			Actions.Add(LinkAction);
 
 			// Prepare a script that will run later, once all shared libraries and the executable
 			// are created. This script will be called by action created in FixDependencies()
@@ -2039,7 +2034,7 @@ namespace UnrealBuildTool
 				else
 				{
 					// Create the action to relink the library. This actions does not overwrite the source file so it can be executed in parallel
-					Action RelinkAction = new Action(ActionType.Link);
+					Action RelinkAction = Graph.CreateAction(ActionType.Link);
 					RelinkAction.WorkingDirectory = LinkAction.WorkingDirectory;
 					RelinkAction.StatusDescription = LinkAction.StatusDescription;
 					RelinkAction.CommandDescription = "Relink";
@@ -2125,7 +2120,6 @@ namespace UnrealBuildTool
 
 					RelinkAction.CommandPath = ShellBinary;
 					RelinkAction.CommandArguments = ExecuteSwitch + " \"" + RelinkScriptFullPath + "\"";
-					Actions.Add(RelinkAction);
 				}
 			}
 			return OutputFile;
@@ -2142,9 +2136,9 @@ namespace UnrealBuildTool
 			}
 		}
 
-		public override ICollection<FileItem> PostBuild(FileItem Executable, LinkEnvironment BinaryLinkEnvironment, List<Action> Actions)
+		public override ICollection<FileItem> PostBuild(FileItem Executable, LinkEnvironment BinaryLinkEnvironment, IActionGraphBuilder Graph)
 		{
-			ICollection<FileItem> OutputFiles = base.PostBuild(Executable, BinaryLinkEnvironment, Actions);
+			ICollection<FileItem> OutputFiles = base.PostBuild(Executable, BinaryLinkEnvironment, Graph);
 
 			if (bUseFixdeps)
 			{
@@ -2153,7 +2147,7 @@ namespace UnrealBuildTool
 					return OutputFiles;
 				}
 
-				FileItem FixDepsOutputFile = FixDependencies(BinaryLinkEnvironment, Executable, Actions);
+				FileItem FixDepsOutputFile = FixDependencies(BinaryLinkEnvironment, Executable, Graph);
 				if (FixDepsOutputFile != null)
 				{
 					OutputFiles.Add(FixDepsOutputFile);
diff --git a/Engine/Source/Programs/UnrealBuildTool/Platform/Mac/MacToolChain.cs b/Engine/Source/Programs/UnrealBuildTool/Platform/Mac/MacToolChain.cs
--- a/Engine/Source/Programs/UnrealBuildTool/Platform/Mac/MacToolChain.cs
+++ b/Engine/Source/Programs/UnrealBuildTool/Platform/Mac/MacToolChain.cs
@@ -401,7 +401,7 @@ namespace UnrealBuildTool
 			return Result;
 		}
 
-		public override CPPOutput CompileCPPFiles(CppCompileEnvironment CompileEnvironment, List<FileItem> InputFiles, DirectoryReference OutputDir, string ModuleName, List<Action> Actions)
+		public override CPPOutput CompileCPPFiles(CppCompileEnvironment CompileEnvironment, List<FileItem> InputFiles, DirectoryReference OutputDir, string ModuleName, IActionGraphBuilder Graph)
 		{
 			StringBuilder Arguments = new StringBuilder();
 			StringBuilder PCHArguments = new StringBuilder();
@@ -452,7 +452,7 @@ namespace UnrealBuildTool
 			// Create a compile action for each source file.
 			foreach (FileItem SourceFile in InputFiles)
 			{
-				Action CompileAction = new Action(ActionType.Compile);
+				Action CompileAction = Graph.CreateAction(ActionType.Compile);
 				CompileAction.PrerequisiteItems.AddRange(CompileEnvironment.ForceIncludeFiles);
 				CompileAction.PrerequisiteItems.AddRange(CompileEnvironment.AdditionalPrerequisites);
 
@@ -558,7 +558,6 @@ namespace UnrealBuildTool
 				// We're already distributing the command by execution on Mac.
 				CompileAction.bCanExecuteRemotely = Extension != ".C";
 				CompileAction.bShouldOutputStatusDescription = true;
-				Actions.Add(CompileAction);
 			}
 			return Result;
 		}
@@ -662,12 +661,12 @@ namespace UnrealBuildTool
 			}
 		}
 
-		public override FileItem LinkFiles(LinkEnvironment LinkEnvironment, bool bBuildImportLibraryOnly, List<Action> Actions)
+		public override FileItem LinkFiles(LinkEnvironment LinkEnvironment, bool bBuildImportLibraryOnly, IActionGraphBuilder Graph)
 		{
 			bool bIsBuildingLibrary = LinkEnvironment.bIsBuildingLibrary || bBuildImportLibraryOnly;
 
 			// Create an action that invokes the linker.
-			Action LinkAction = new Action(ActionType.Link);
+			Action LinkAction = Graph.CreateAction(ActionType.Link);
 
 			LinkAction.WorkingDirectory = GetMacDevSrcRoot();
 			LinkAction.CommandPath = BuildHostPlatform.Current.Shell;
@@ -867,8 +866,6 @@ namespace UnrealBuildTool
 
 			LinkAction.ProducedItems.Add(OutputFile);
 
-			Actions.Add(LinkAction);
-
 			if (!DirectoryReference.Exists(LinkEnvironment.IntermediateDirectory))
 			{
 				return OutputFile;
@@ -1056,9 +1053,9 @@ namespace UnrealBuildTool
 			return String.Format("rsync --checksum \"{0}\" \"{1}\"", SourceFile, TargetFile);
 		}
 
-		FileItem FixDylibDependencies(LinkEnvironment LinkEnvironment, FileItem Executable, List<Action> Actions)
+		FileItem FixDylibDependencies(LinkEnvironment LinkEnvironment, FileItem Executable, IActionGraphBuilder Graph)
 		{
-			Action FixDylibAction = new Action(ActionType.PostBuildStep);
+			Action FixDylibAction = Graph.CreateAction(ActionType.PostBuildStep);
 			FixDylibAction.WorkingDirectory = UnrealBuildTool.EngineSourceDirectory;
 			FixDylibAction.CommandPath = BuildHostPlatform.Current.Shell;
 			FixDylibAction.CommandDescription = "";
@@ -1086,7 +1083,6 @@ namespace UnrealBuildTool
 			FixDylibAction.CommandArguments += "'";
 
 			FixDylibAction.ProducedItems.Add(OutputFile);
-			Actions.Add(FixDylibAction);
 
 			return OutputFile;
 		}
@@ -1096,8 +1092,8 @@ namespace UnrealBuildTool
 		/// </summary>
 		/// <param name="MachOBinary">FileItem describing the executable or dylib to generate debug info for</param>
 		/// <param name="LinkEnvironment"></param>
-		/// <param name="Actions">List of actions to be executed. Additional actions will be added to this list.</param>
-		public FileItem GenerateDebugInfo(FileItem MachOBinary, LinkEnvironment LinkEnvironment, List<Action> Actions)
+		/// <param name="Graph">List of actions to be executed. Additional actions will be added to this list.</param>
+		public FileItem GenerateDebugInfo(FileItem MachOBinary, LinkEnvironment LinkEnvironment, IActionGraphBuilder Graph)
 		{
 			string BinaryPath = MachOBinary.AbsolutePath;
 			if (BinaryPath.Contains(".app"))
@@ -1122,7 +1118,7 @@ namespace UnrealBuildTool
 			}
 
 			// Make the compile action
-			Action GenDebugAction = new Action(ActionType.GenerateDebugInfo);
+			Action GenDebugAction = Graph.CreateAction(ActionType.GenerateDebugInfo);
 			GenDebugAction.WorkingDirectory = GetMacDevSrcRoot();
 			GenDebugAction.CommandPath = BuildHostPlatform.Current.Shell;
 
@@ -1145,7 +1141,6 @@ namespace UnrealBuildTool
 			GenDebugAction.CommandDescription = "";
 			GenDebugAction.StatusDescription = "Generating " + Path.GetFileName(BinaryPath);
 			GenDebugAction.bCanExecuteRemotely = false;
-			Actions.Add(GenDebugAction);
 
 			return OutputFile;
 		}
@@ -1156,15 +1151,15 @@ namespace UnrealBuildTool
 		/// <param name="LinkEnvironment"></param>
 		/// <param name="Executable">FileItem describing the executable to generate app bundle for</param>
 		/// <param name="FixDylibOutputFile"></param>
-		/// <param name="Actions">List of actions to be executed. Additional actions will be added to this list.</param>
-		FileItem FinalizeAppBundle(LinkEnvironment LinkEnvironment, FileItem Executable, FileItem FixDylibOutputFile, List<Action> Actions)
+		/// <param name="Graph">List of actions to be executed. Additional actions will be added to this list.</param>
+		FileItem FinalizeAppBundle(LinkEnvironment LinkEnvironment, FileItem Executable, FileItem FixDylibOutputFile, IActionGraphBuilder Graph)
 		{
 			// Make a file item for the source and destination files
 			string FullDestPath = Executable.AbsolutePath.Substring(0, Executable.AbsolutePath.IndexOf(".app") + 4);
 			FileItem DestFile = FileItem.GetItemByPath(FullDestPath);
 
 			// Make the compile action
-			Action FinalizeAppBundleAction = new Action(ActionType.CreateAppBundle);
+			Action FinalizeAppBundleAction = Graph.CreateAction(ActionType.CreateAppBundle);
 			FinalizeAppBundleAction.WorkingDirectory = GetMacDevSrcRoot(); // Path.GetFullPath(".");
 			FinalizeAppBundleAction.CommandPath = BuildHostPlatform.Current.Shell;
 			FinalizeAppBundleAction.CommandDescription = "";
@@ -1177,14 +1172,13 @@ namespace UnrealBuildTool
 			FinalizeAppBundleAction.ProducedItems.Add(DestFile);
 			FinalizeAppBundleAction.StatusDescription = string.Format("Finalizing app bundle: {0}.app", Path.GetFileName(Executable.AbsolutePath));
 			FinalizeAppBundleAction.bCanExecuteRemotely = false;
-			Actions.Add(FinalizeAppBundleAction);
 
 			return DestFile;
 		}
 
-		FileItem CopyBundleResource(UEBuildBundleResource Resource, FileItem Executable, DirectoryReference BundleDirectory, List<Action> Actions)
+		FileItem CopyBundleResource(UEBuildBundleResource Resource, FileItem Executable, DirectoryReference BundleDirectory, IActionGraphBuilder Graph)
 		{
-			Action CopyAction = new Action(ActionType.CreateAppBundle);
+			Action CopyAction = Graph.CreateAction(ActionType.CreateAppBundle);
 			CopyAction.WorkingDirectory = GetMacDevSrcRoot(); // Path.GetFullPath(".");
 			CopyAction.CommandPath = BuildHostPlatform.Current.Shell;
 			CopyAction.CommandDescription = "";
@@ -1201,7 +1195,6 @@ namespace UnrealBuildTool
 			CopyAction.bShouldOutputStatusDescription = Resource.bShouldLog;
 			CopyAction.StatusDescription = string.Format("Copying {0} to app bundle", Path.GetFileName(Resource.ResourcePath));
 			CopyAction.bCanExecuteRemotely = false;
-			Actions.Add(CopyAction);
 
 			return TargetItem;
 		}
@@ -1315,9 +1308,9 @@ namespace UnrealBuildTool
 			}
 		}
 
-		public override ICollection<FileItem> PostBuild(FileItem Executable, LinkEnvironment BinaryLinkEnvironment, List<Action> Actions)
+		public override ICollection<FileItem> PostBuild(FileItem Executable, LinkEnvironment BinaryLinkEnvironment, IActionGraphBuilder Graph)
 		{
-			ICollection<FileItem> OutputFiles = base.PostBuild(Executable, BinaryLinkEnvironment, Actions);
+			ICollection<FileItem> OutputFiles = base.PostBuild(Executable, BinaryLinkEnvironment, Graph);
 
 			if (BinaryLinkEnvironment.bIsBuildingLibrary)
 			{
@@ -1328,7 +1321,7 @@ namespace UnrealBuildTool
 			{
 				foreach (UEBuildBundleResource Resource in BinaryLinkEnvironment.AdditionalBundleResources)
 				{
-					OutputFiles.Add(CopyBundleResource(Resource, Executable, BinaryLinkEnvironment.BundleDirectory, Actions));
+					OutputFiles.Add(CopyBundleResource(Resource, Executable, BinaryLinkEnvironment.BundleDirectory, Graph));
 				}
 			}
 
@@ -1338,7 +1331,7 @@ namespace UnrealBuildTool
 				// We want dsyms to be created after all dylib dependencies are fixed. If FixDylibDependencies action was not created yet, save the info for later.
 				if (FixDylibOutputFile != null)
 				{
-					OutputFiles.Add(GenerateDebugInfo(Executable, BinaryLinkEnvironment, Actions));
+					OutputFiles.Add(GenerateDebugInfo(Executable, BinaryLinkEnvironment, Graph));
 				}
 				else
 				{
@@ -1351,17 +1344,17 @@ namespace UnrealBuildTool
 				return OutputFiles;
 			}
 
-			FixDylibOutputFile = FixDylibDependencies(BinaryLinkEnvironment, Executable, Actions);
+			FixDylibOutputFile = FixDylibDependencies(BinaryLinkEnvironment, Executable, Graph);
 			OutputFiles.Add(FixDylibOutputFile);
 			if (!BinaryLinkEnvironment.bIsBuildingConsoleApplication)
 			{
-				OutputFiles.Add(FinalizeAppBundle(BinaryLinkEnvironment, Executable, FixDylibOutputFile, Actions));
+				OutputFiles.Add(FinalizeAppBundle(BinaryLinkEnvironment, Executable, FixDylibOutputFile, Graph));
 			}
 
 			// Add dsyms that we couldn't add before FixDylibDependencies action was created
 			foreach (FileItem Exe in ExecutablesThatNeedDsyms)
 			{
-				OutputFiles.Add(GenerateDebugInfo(Exe, BinaryLinkEnvironment, Actions));
+				OutputFiles.Add(GenerateDebugInfo(Exe, BinaryLinkEnvironment, Graph));
 			}
 			ExecutablesThatNeedDsyms.Clear();
 
diff --git a/Engine/Source/Programs/UnrealBuildTool/Platform/Windows/PVSToolChain.cs b/Engine/Source/Programs/UnrealBuildTool/Platform/Windows/PVSToolChain.cs
--- a/Engine/Source/Programs/UnrealBuildTool/Platform/Windows/PVSToolChain.cs
+++ b/Engine/Source/Programs/UnrealBuildTool/Platform/Windows/PVSToolChain.cs
@@ -178,7 +178,7 @@ namespace UnrealBuildTool
 			if(ApplicationSettings != null && !String.IsNullOrEmpty(ApplicationSettings.UserName) && !String.IsNullOrEmpty(ApplicationSettings.SerialNumber))
 			{
 				LicenseFile = FileReference.Combine(UnrealBuildTool.EngineDirectory, "Intermediate", "PVS", "PVS-Studio.lic");
-				FileItem.CreateIntermediateTextFile(LicenseFile, new string[]{ ApplicationSettings.UserName, ApplicationSettings.SerialNumber });
+				Utils.WriteFileIfChanged(LicenseFile, String.Format("{0}\n{1}\n", ApplicationSettings.UserName, ApplicationSettings.SerialNumber), StringComparison.Ordinal);
 			}
 			else
 			{
@@ -195,7 +195,25 @@ namespace UnrealBuildTool
 			return Path.GetFullPath(Utils.ExpandVariables(IncludePath));
 		}
 
-		public override CPPOutput CompileCPPFiles(CppCompileEnvironment CompileEnvironment, List<FileItem> InputFiles, DirectoryReference OutputDir, string ModuleName, List<Action> Actions)
+		class ActionGraphCapture : ForwardingActionGraphBuilder
+		{
+			List<Action> Actions;
+
+			public ActionGraphCapture(IActionGraphBuilder Inner, List<Action> Actions)
+				: base(Inner)
+			{
+				this.Actions = Actions;
+			}
+
+			public override Action CreateAction(ActionType Type)
+			{
+				Action Action = base.CreateAction(Type);
+				Actions.Add(Action);
+				return Action;
+			}
+		}
+
+		public override CPPOutput CompileCPPFiles(CppCompileEnvironment CompileEnvironment, List<FileItem> InputFiles, DirectoryReference OutputDir, string ModuleName, IActionGraphBuilder Graph)
 		{
 			// Preprocess the source files with the regular toolchain
 			CppCompileEnvironment PreprocessCompileEnvironment = new CppCompileEnvironment(CompileEnvironment);
@@ -204,8 +222,7 @@ namespace UnrealBuildTool
 			PreprocessCompileEnvironment.Definitions.Add("PVS_STUDIO");
 
 			List<Action> PreprocessActions = new List<Action>();
-			CPPOutput Result = InnerToolChain.CompileCPPFiles(PreprocessCompileEnvironment, InputFiles, OutputDir, ModuleName, PreprocessActions);
-			Actions.AddRange(PreprocessActions);
+			CPPOutput Result = InnerToolChain.CompileCPPFiles(PreprocessCompileEnvironment, InputFiles, OutputDir, ModuleName, new ActionGraphCapture(Graph, PreprocessActions));
 
 			// Run the source files through PVS-Studio
 			foreach(Action PreprocessAction in PreprocessActions)
@@ -276,13 +293,13 @@ namespace UnrealBuildTool
 				string BaseFileName = PreprocessedFileItem.Location.GetFileNameWithoutExtension();
 
 				FileReference ConfigFileLocation = FileReference.Combine(OutputDir, BaseFileName + ".cfg");
-				FileItem ConfigFileItem = FileItem.CreateIntermediateTextFile(ConfigFileLocation, ConfigFileContents.ToString());
+				FileItem ConfigFileItem = Graph.CreateIntermediateTextFile(ConfigFileLocation, ConfigFileContents.ToString());
 
 				// Run the analzyer on the preprocessed source file
 				FileReference OutputFileLocation = FileReference.Combine(OutputDir, BaseFileName + ".pvslog");
 				FileItem OutputFileItem = FileItem.GetItemByFileReference(OutputFileLocation);
 
-				Action AnalyzeAction = new Action(ActionType.Compile);
+				Action AnalyzeAction = Graph.CreateAction(ActionType.Compile);
 				AnalyzeAction.CommandDescription = "Analyzing";
 				AnalyzeAction.StatusDescription = BaseFileName;
 				AnalyzeAction.WorkingDirectory = UnrealBuildTool.EngineSourceDirectory;
@@ -298,14 +315,13 @@ namespace UnrealBuildTool
 				AnalyzeAction.PrerequisiteItems.AddRange(InputFiles); // Add the InputFiles as PrerequisiteItems so that in SingleFileCompile mode the PVSAnalyze step is not filtered out
 				AnalyzeAction.ProducedItems.Add(OutputFileItem);
 				AnalyzeAction.DeleteItems.Add(OutputFileItem); // PVS Studio will append by default, so need to delete produced items
-				Actions.Add(AnalyzeAction);
 
 				Result.ObjectFiles.AddRange(AnalyzeAction.ProducedItems);
 			}
 			return Result;
 		}
 
-		public override FileItem LinkFiles(LinkEnvironment LinkEnvironment, bool bBuildImportLibraryOnly, List<Action> Actions)
+		public override FileItem LinkFiles(LinkEnvironment LinkEnvironment, bool bBuildImportLibraryOnly, IActionGraphBuilder Graph)
 		{
 			throw new BuildException("Unable to link with PVS toolchain.");
 		}
@@ -327,9 +343,9 @@ namespace UnrealBuildTool
 			// Collect the prerequisite items off of the Compile action added in CompileCPPFiles so that in SingleFileCompile mode the PVSGather step is also not filtered out
 			List<FileItem> AnalyzeActionPrerequisiteItems = Makefile.Actions.Where(x => x.ActionType == ActionType.Compile).SelectMany(x => x.PrerequisiteItems).ToList();
 
-			FileItem InputFileListItem = FileItem.CreateIntermediateTextFile(OutputFile.ChangeExtension(".input"), InputFiles.Select(x => x.FullName));
+			FileItem InputFileListItem = Makefile.CreateIntermediateTextFile(OutputFile.ChangeExtension(".input"), InputFiles.Select(x => x.FullName));
 
-			Action AnalyzeAction = new Action(ActionType.Compile);
+			Action AnalyzeAction = Makefile.CreateAction(ActionType.Compile);
 			AnalyzeAction.CommandPath = UnrealBuildTool.GetUBTPath();
 			AnalyzeAction.CommandArguments = String.Format("-Mode=PVSGather -Input=\"{0}\" -Output=\"{1}\"", InputFileListItem.Location, OutputFile);
 			AnalyzeAction.WorkingDirectory = UnrealBuildTool.EngineSourceDirectory;
diff --git a/Engine/Source/Programs/UnrealBuildTool/Platform/Windows/VCToolChain.cs b/Engine/Source/Programs/UnrealBuildTool/Platform/Windows/VCToolChain.cs
--- a/Engine/Source/Programs/UnrealBuildTool/Platform/Windows/VCToolChain.cs
+++ b/Engine/Source/Programs/UnrealBuildTool/Platform/Windows/VCToolChain.cs
@@ -923,7 +923,7 @@ namespace UnrealBuildTool
 			}
 		}
 
-		public override CPPOutput CompileCPPFiles(CppCompileEnvironment CompileEnvironment, List<FileItem> InputFiles, DirectoryReference OutputDir, string ModuleName, List<Action> Actions)
+		public override CPPOutput CompileCPPFiles(CppCompileEnvironment CompileEnvironment, List<FileItem> InputFiles, DirectoryReference OutputDir, string ModuleName, IActionGraphBuilder Graph)
 		{
 			List<string> SharedArguments = new List<string>();
 			AppendCLArguments_Global(CompileEnvironment, SharedArguments);
@@ -974,7 +974,7 @@ namespace UnrealBuildTool
 			CPPOutput Result = new CPPOutput();
 			foreach (FileItem SourceFile in InputFiles)
 			{
-				Action CompileAction = new Action(ActionType.Compile);
+				Action CompileAction = Graph.CreateAction(ActionType.Compile);
 				CompileAction.CommandDescription = "Compile";
 
 				// ensure compiler timings are captured when we execute the action.
@@ -994,7 +994,7 @@ namespace UnrealBuildTool
 				{
 					// Generate a CPP File that just includes the precompiled header.
 					FileReference PCHCPPPath = CompileEnvironment.PrecompiledHeaderIncludeFilename.ChangeExtension(".cpp");
-					FileItem PCHCPPFile = FileItem.CreateIntermediateTextFile(
+					FileItem PCHCPPFile = Graph.CreateIntermediateTextFile(
 						PCHCPPPath,
 						string.Format("#include \"{0}\"\r\n", CompileEnvironment.PrecompiledHeaderIncludeFilename.FullName.Replace('\\', '/'))
 						);
@@ -1170,7 +1170,7 @@ namespace UnrealBuildTool
 				{
 					FileItem TargetFile = CompileAction.ProducedItems[0];
 					FileReference ResponseFileName = new FileReference(TargetFile.AbsolutePath + ".response");
-					FileItem ResponseFileItem = FileItem.CreateIntermediateTextFile(ResponseFileName, SharedArguments.Concat(FileArguments).Concat(AdditionalArguments).Select(x => Utils.ExpandVariables(x)));
+					FileItem ResponseFileItem = Graph.CreateIntermediateTextFile(ResponseFileName, SharedArguments.Concat(FileArguments).Concat(AdditionalArguments).Select(x => Utils.ExpandVariables(x)));
 					CompileAction.CommandArguments = " @\"" + ResponseFileName + "\"";
 					CompileAction.PrerequisiteItems.Add(ResponseFileItem);
 				}
@@ -1181,7 +1181,7 @@ namespace UnrealBuildTool
 
 				if(CompileEnvironment.bGenerateDependenciesFile)
 				{
-					GenerateDependenciesFile(Target, OutputDir, Result, SourceFile, Actions, CompileAction);
+					GenerateDependenciesFile(Target, OutputDir, Result, SourceFile, CompileAction, Graph);
 				}
 
 				if (CompileEnvironment.PrecompiledHeaderAction == PrecompiledHeaderAction.Create)
@@ -1210,14 +1210,12 @@ namespace UnrealBuildTool
 				{
 					CompileAction.bCanExecuteRemotelyWithSNDBS = false;
 				}
-
-				Actions.Add(CompileAction);
 			}
 
 			return Result;
 		}
 
-		private Action GenerateParseTimingInfoAction(FileItem SourceFile, Action CompileAction, CPPOutput Result)
+		private Action GenerateParseTimingInfoAction(FileItem SourceFile, Action CompileAction, CPPOutput Result, IActionGraphBuilder Graph)
 		{
 			FileItem TimingJsonFile = FileItem.GetItemByPath(Path.ChangeExtension(CompileAction.TimingFile.AbsolutePath, ".timing.bin"));
 			Result.DebugDataFiles.Add(TimingJsonFile);
@@ -1229,7 +1227,7 @@ namespace UnrealBuildTool
 				ParseTimingArguments += " -Tracing";
 			}
 
-			Action ParseTimingInfoAction = Action.CreateRecursiveAction<ParseMsvcTimingInfoMode>(ActionType.ParseTimingInfo, ParseTimingArguments);
+			Action ParseTimingInfoAction = Graph.CreateRecursiveAction<ParseMsvcTimingInfoMode>(ActionType.ParseTimingInfo, ParseTimingArguments);
 			ParseTimingInfoAction.WorkingDirectory = UnrealBuildTool.EngineSourceDirectory;
 			ParseTimingInfoAction.StatusDescription = Path.GetFileName(CompileAction.TimingFile.AbsolutePath);
 			ParseTimingInfoAction.bCanExecuteRemotely = true;
@@ -1240,7 +1238,7 @@ namespace UnrealBuildTool
 			return ParseTimingInfoAction;
 		}
 
-		private void GenerateDependenciesFile(ReadOnlyTargetRules Target, DirectoryReference OutputDir, CPPOutput Result, FileItem SourceFile, List<Action> Actions, Action CompileAction)
+		private void GenerateDependenciesFile(ReadOnlyTargetRules Target, DirectoryReference OutputDir, CPPOutput Result, FileItem SourceFile, Action CompileAction, IActionGraphBuilder Graph)
 		{
 			List<string> CommandArguments = new List<string>();
 			CompileAction.DependencyListFile = FileItem.GetItemByFileReference(FileReference.Combine(OutputDir, String.Format("{0}.txt", SourceFile.Location.GetFileName())));
@@ -1251,8 +1249,7 @@ namespace UnrealBuildTool
 			{
 				CompileAction.TimingFile = FileItem.GetItemByFileReference(FileReference.Combine(OutputDir, String.Format("{0}.timing.txt", SourceFile.Location.GetFileName())));
 				CompileAction.ProducedItems.Add(CompileAction.TimingFile);
-				Action ParseTimingInfoAction = GenerateParseTimingInfoAction(SourceFile, CompileAction, Result);
-				Actions.Add(ParseTimingInfoAction);
+				Action ParseTimingInfoAction = GenerateParseTimingInfoAction(SourceFile, CompileAction, Result, Graph);
 				CommandArguments.Add(string.Format("-timing={0}", Utils.MakePathSafeToUseWithCommandLine(CompileAction.TimingFile.Location)));
 			}
 
@@ -1297,7 +1294,7 @@ namespace UnrealBuildTool
 						String.Format("-CompileTimingFile={0}", ExpectedCompileTimeFile),
 					};
 
-					Action AggregateTimingInfoAction = Action.CreateRecursiveAction<AggregateParsedTimingInfo>(ActionType.ParseTimingInfo, string.Join(" ", ActionArgs));
+					Action AggregateTimingInfoAction = Makefile.CreateRecursiveAction<AggregateParsedTimingInfo>(ActionType.ParseTimingInfo, string.Join(" ", ActionArgs));
 					AggregateTimingInfoAction.WorkingDirectory = UnrealBuildTool.EngineSourceDirectory;
 					AggregateTimingInfoAction.StatusDescription = $"Aggregating {TimingJsonFiles.Count} Timing File(s)";
 					AggregateTimingInfoAction.bCanExecuteRemotely = false;
@@ -1307,18 +1304,17 @@ namespace UnrealBuildTool
 					FileItem AggregateOutputFile = FileItem.GetItemByFileReference(FileReference.Combine(Makefile.ProjectIntermediateDirectory, $"{Target.Name}.timing.bin"));
 					AggregateTimingInfoAction.ProducedItems.Add(AggregateOutputFile);
 					Makefile.OutputItems.Add(AggregateOutputFile);
-					Makefile.Actions.Add(AggregateTimingInfoAction);
 				}
 			}
 		}
 
-		public override CPPOutput CompileRCFiles(CppCompileEnvironment CompileEnvironment, List<FileItem> InputFiles, DirectoryReference OutputDir, List<Action> Actions)
+		public override CPPOutput CompileRCFiles(CppCompileEnvironment CompileEnvironment, List<FileItem> InputFiles, DirectoryReference OutputDir, IActionGraphBuilder Graph)
 		{
 			CPPOutput Result = new CPPOutput();
 
 			foreach (FileItem RCFile in InputFiles)
 			{
-				Action CompileAction = new Action(ActionType.Compile);
+				Action CompileAction = Graph.CreateAction(ActionType.Compile);
 				CompileAction.CommandDescription = "Resource";
 				CompileAction.WorkingDirectory = UnrealBuildTool.EngineSourceDirectory;
 				CompileAction.CommandPath = EnvVars.ResourceCompilerPath;
@@ -1427,14 +1423,12 @@ namespace UnrealBuildTool
 
 				// Add the C++ source file and its included files to the prerequisite item list.
 				CompileAction.PrerequisiteItems.Add(RCFile);
-
-				Actions.Add(CompileAction);
 			}
 
 			return Result;
 		}
 
-		public override void GenerateTypeLibraryHeader(CppCompileEnvironment CompileEnvironment, ModuleRules.TypeLibrary TypeLibrary, FileReference OutputFile, List<Action> Actions)
+		public override void GenerateTypeLibraryHeader(CppCompileEnvironment CompileEnvironment, ModuleRules.TypeLibrary TypeLibrary, FileReference OutputFile, IActionGraphBuilder Graph)
 		{
 			// Create the input file
 			StringBuilder Contents = new StringBuilder();
@@ -1450,7 +1444,7 @@ namespace UnrealBuildTool
 			}
 			Contents.AppendLine();
 
-			FileItem InputFile = FileItem.CreateIntermediateTextFile(OutputFile.ChangeExtension(".cpp"), Contents.ToString());
+			FileItem InputFile = Graph.CreateIntermediateTextFile(OutputFile.ChangeExtension(".cpp"), Contents.ToString());
 
 			// Build the argument list
 			FileItem ObjectFile = FileItem.GetItemByFileReference(OutputFile.ChangeExtension(".obj"));
@@ -1477,7 +1471,7 @@ namespace UnrealBuildTool
 			}
 
 			// Create the compile action. Only mark the object file as an output, because we need to touch the generated header afterwards.
-			Action CompileAction = new Action(ActionType.Compile);
+			Action CompileAction = Graph.CreateAction(ActionType.Compile);
 			CompileAction.CommandDescription = "GenerateTLH";
 			CompileAction.PrerequisiteItems.Add(InputFile);
 			CompileAction.ProducedItems.Add(ObjectFile);
@@ -1488,10 +1482,9 @@ namespace UnrealBuildTool
 			CompileAction.CommandArguments = String.Join(" ", Arguments);
 			CompileAction.bShouldOutputStatusDescription = (Target.WindowsPlatform.Compiler == WindowsCompiler.Clang);
 			CompileAction.bCanExecuteRemotely = false; // Incompatible with SN-DBS
-			Actions.Add(CompileAction);
 
 			// Touch the output header
-			Action TouchAction = new Action(ActionType.BuildProject);
+			Action TouchAction = Graph.CreateAction(ActionType.BuildProject);
 			TouchAction.CommandDescription = "Touch";
 			TouchAction.CommandPath = BuildHostPlatform.Current.Shell;
 			TouchAction.CommandArguments = String.Format("/C \"copy /b \"{0}\"+,, \"{0}\" 1>nul:\"", OutputFile.FullName);
@@ -1500,7 +1493,6 @@ namespace UnrealBuildTool
 			TouchAction.ProducedItems.Add(FileItem.GetItemByFileReference(OutputFile));
 			TouchAction.StatusDescription = OutputFile.GetFileName();
 			TouchAction.bCanExecuteRemotely = false;
-			Actions.Add(TouchAction);
 		}
 
 		/// <summary>
@@ -1522,7 +1514,7 @@ namespace UnrealBuildTool
 			return Result.ToString();
 		}
 
-		public override FileItem LinkFiles(LinkEnvironment LinkEnvironment, bool bBuildImportLibraryOnly, List<Action> Actions)
+		public override FileItem LinkFiles(LinkEnvironment LinkEnvironment, bool bBuildImportLibraryOnly, IActionGraphBuilder Graph)
 		{
 			if (LinkEnvironment.bIsBuildingDotNetAssembly)
 			{
@@ -1744,12 +1736,12 @@ namespace UnrealBuildTool
 			FileReference ResponseFileName = GetResponseFileName(LinkEnvironment, OutputFile);
 			if (!ProjectFileGenerator.bGenerateProjectFiles)
 			{
-				FileItem ResponseFile = FileItem.CreateIntermediateTextFile(ResponseFileName, String.Join(Environment.NewLine, Arguments));
+				FileItem ResponseFile = Graph.CreateIntermediateTextFile(ResponseFileName, String.Join(Environment.NewLine, Arguments));
 				PrerequisiteItems.Add(ResponseFile);
 			}
 
 			// Create an action that invokes the linker.
-			Action LinkAction = new Action(ActionType.Link);
+			Action LinkAction = Graph.CreateAction(ActionType.Link);
 			LinkAction.CommandDescription = "Link";
 			LinkAction.WorkingDirectory = UnrealBuildTool.EngineSourceDirectory;
 			if(bIsBuildingLibraryOrImportLibrary)
@@ -1790,7 +1782,6 @@ namespace UnrealBuildTool
 
 			// Allow remote linking.  Especially in modular builds with many small DLL files, this is almost always very efficient
 			LinkAction.bCanExecuteRemotely = true;
-			Actions.Add(LinkAction);
 
 			Log.TraceVerbose("     Linking: " + LinkAction.StatusDescription);
 			Log.TraceVerbose("     Command: " + LinkAction.CommandArguments);
diff --git a/Engine/Source/Programs/UnrealBuildTool/System/Action.cs b/Engine/Source/Programs/UnrealBuildTool/System/Action.cs
--- a/Engine/Source/Programs/UnrealBuildTool/System/Action.cs
+++ b/Engine/Source/Programs/UnrealBuildTool/System/Action.cs
@@ -366,54 +366,6 @@ namespace UnrealBuildTool
 			}
 		}
 
-		/// <summary>
-		/// Creates an action which copies a file from one location to another
-		/// </summary>
-		/// <param name="SourceFile">The source file location</param>
-		/// <param name="TargetFile">The target file location</param>
-		/// <returns>File item for the output file</returns>
-		public static Action CreateCopyAction(FileItem SourceFile, FileItem TargetFile)
-		{
-			Action CopyAction = new Action(ActionType.BuildProject);
-			CopyAction.CommandDescription = "Copy";
-			CopyAction.CommandPath = BuildHostPlatform.Current.Shell;
-			if (BuildHostPlatform.Current.ShellType == ShellType.Cmd)
-			{
-				CopyAction.CommandArguments = String.Format("/C \"copy /Y \"{0}\" \"{1}\" 1>nul\"", SourceFile.AbsolutePath, TargetFile.AbsolutePath);
-			}
-			else
-			{
-				CopyAction.CommandArguments = String.Format("-c 'cp -f \"{0}\" \"{1}\"'", SourceFile.AbsolutePath, TargetFile.AbsolutePath);
-			}
-			CopyAction.WorkingDirectory = UnrealBuildTool.EngineSourceDirectory;
-			CopyAction.PrerequisiteItems.Add(SourceFile);
-			CopyAction.ProducedItems.Add(TargetFile);
-			CopyAction.DeleteItems.Add(TargetFile);
-			CopyAction.StatusDescription = TargetFile.Location.GetFileName();
-			CopyAction.bCanExecuteRemotely = false;
-			return CopyAction;
-		}
-
-		/// <summary>
-		/// Creates an action which calls UBT recursively
-		/// </summary>
-		/// <param name="Type">Type of the action</param>
-		/// <param name="Arguments">Arguments for the action</param>
-		/// <returns>New action instance</returns>
-		public static Action CreateRecursiveAction<T>(ActionType Type, string Arguments) where T : ToolMode
-		{
-			ToolModeAttribute Attribute = typeof(T).GetCustomAttribute<ToolModeAttribute>();
-			if(Attribute == null)
-			{
-				throw new BuildException("Missing ToolModeAttribute on {0}", typeof(T).Name);
-			}
-
-			Action NewAction = new Action(Type);
-			NewAction.CommandPath = UnrealBuildTool.GetUBTPath();
-			NewAction.CommandArguments = String.Format("-Mode={0} {1}", Attribute.Name, Arguments);
-			return NewAction;
-		}
-
 		/// <summary>
 		/// Finds conflicts betwee two actions, and prints them to the log
 		/// </summary>
diff --git a/Engine/Source/Programs/UnrealBuildTool/System/ActionGraphBuilder.cs b/Engine/Source/Programs/UnrealBuildTool/System/ActionGraphBuilder.cs
new file mode 100644
--- /dev/null
+++ b/Engine/Source/Programs/UnrealBuildTool/System/ActionGraphBuilder.cs
@@ -0,0 +1,272 @@
+using System;
+using System.Collections.Generic;
+using System.Linq;
+using System.Reflection;
+using System.Text;
+using System.Threading.Tasks;
+using Tools.DotNETCommon;
+
+namespace UnrealBuildTool
+{
+	/// <summary>
+	/// Interface for toolchain operations that produce output
+	/// </summary>
+	interface IActionGraphBuilder
+	{
+		/// <summary>
+		/// Creates a new action to be built as part of this target
+		/// </summary>
+		/// <param name="Type">Type of action to create</param>
+		/// <returns>New action</returns>
+		Action CreateAction(ActionType Type);
+
+		/// <summary>
+		/// Creates a response file for use in the action graph
+		/// </summary>
+		/// <param name="Location">Location of the response file</param>
+		/// <param name="Contents">Contents of the file</param>
+		/// <returns>New file item</returns>
+		FileItem CreateIntermediateTextFile(FileReference Location, string Contents);
+
+		/// <summary>
+		/// Adds a file which is in the non-unity working set
+		/// </summary>
+		/// <param name="File">The file to add to the working set</param>
+		void AddFileToWorkingSet(FileItem File);
+
+		/// <summary>
+		/// Adds a file which is a candidate for being in the non-unity working set
+		/// </summary>
+		/// <param name="File">The file to add to the working set</param>
+		void AddCandidateForWorkingSet(FileItem File);
+
+		/// <summary>
+		/// Adds a source directory. These folders are scanned recursively for C++ source files.
+		/// </summary>
+		/// <param name="SourceDir">Base source directory</param>
+		void AddSourceDir(DirectoryItem SourceDir);
+
+		/// <summary>
+		/// Adds the given source files as dependencies
+		/// </summary>
+		/// <param name="SourceDir">Source directory containing files to build</param>
+		/// <param name="SourceFiles">Contents of the directory</param>
+		void AddSourceFiles(DirectoryItem SourceDir, FileItem[] SourceFiles);
+
+		/// <summary>
+		/// Sets the output items which belong to a particular module
+		/// </summary>
+		/// <param name="ModuleName">Name of the module</param>
+		/// <param name="OutputItems">Array of output items for this module</param>
+		void SetOutputItemsForModule(string ModuleName, FileItem[] OutputItems);
+
+		/// <summary>
+		/// Adds a diagnostic message
+		/// </summary>
+		/// <param name="Message">Message to display</param>
+		void AddDiagnostic(string Message);
+	}
+
+	/// <summary>
+	/// Implementation of IActionGraphBuilder which discards all unnecessary operations
+	/// </summary>
+	class NullActionGraphBuilder : IActionGraphBuilder
+	{
+		/// <inheritdoc/>
+		public Action CreateAction(ActionType Type)
+		{
+			return new Action(Type);
+		}
+
+		/// <inheritdoc/>
+		public virtual FileItem CreateIntermediateTextFile(FileReference Location, string Contents)
+		{
+			Utils.WriteFileIfChanged(Location, Contents, StringComparison.OrdinalIgnoreCase);
+			return FileItem.GetItemByFileReference(Location);
+		}
+
+		/// <inheritdoc/>
+		public virtual void AddSourceDir(DirectoryItem SourceDir)
+		{
+		}
+
+		/// <inheritdoc/>
+		public virtual void AddSourceFiles(DirectoryItem SourceDir, FileItem[] SourceFiles)
+		{
+		}
+
+		/// <inheritdoc/>
+		public virtual void AddFileToWorkingSet(FileItem File)
+		{
+		}
+
+		/// <inheritdoc/>
+		public virtual void AddCandidateForWorkingSet(FileItem File)
+		{
+		}
+
+		/// <inheritdoc/>
+		public virtual void AddDiagnostic(string Message)
+		{
+		}
+
+		/// <inheritdoc/>
+		public virtual void SetOutputItemsForModule(string ModuleName, FileItem[] OutputItems)
+		{
+		}
+	}
+
+	/// <summary>
+	/// Implementation of IActionGraphBuilder which forwards calls to an underlying implementation, allowing derived classes to intercept certain calls
+	/// </summary>
+	class ForwardingActionGraphBuilder : IActionGraphBuilder
+	{
+		/// <summary>
+		/// The inner graph builder
+		/// </summary>
+		IActionGraphBuilder Inner;
+
+		/// <summary>
+		/// Constructor
+		/// </summary>
+		/// <param name="Inner">Builder to pass all calls to</param>
+		public ForwardingActionGraphBuilder(IActionGraphBuilder Inner)
+		{
+			this.Inner = Inner;
+		}
+
+		/// <inheritdoc/>
+		public virtual Action CreateAction(ActionType Type)
+		{
+			return Inner.CreateAction(Type);
+		}
+
+		/// <inheritdoc/>
+		public virtual FileItem CreateIntermediateTextFile(FileReference Location, string Contents)
+		{
+			return Inner.CreateIntermediateTextFile(Location, Contents);
+		}
+
+		/// <inheritdoc/>
+		public virtual void AddSourceDir(DirectoryItem SourceDir)
+		{
+			Inner.AddSourceDir(SourceDir);
+		}
+
+		/// <inheritdoc/>
+		public virtual void AddSourceFiles(DirectoryItem SourceDir, FileItem[] SourceFiles)
+		{
+			Inner.AddSourceFiles(SourceDir, SourceFiles);
+		}
+
+		/// <inheritdoc/>
+		public virtual void AddFileToWorkingSet(FileItem File)
+		{
+			Inner.AddFileToWorkingSet(File);
+		}
+
+		/// <inheritdoc/>
+		public virtual void AddCandidateForWorkingSet(FileItem File)
+		{
+			Inner.AddCandidateForWorkingSet(File);
+		}
+
+		/// <inheritdoc/>
+		public virtual void AddDiagnostic(string Message)
+		{
+			Inner.AddDiagnostic(Message);
+		}
+
+		/// <inheritdoc/>
+		public virtual void SetOutputItemsForModule(string ModuleName, FileItem[] OutputItems)
+		{
+			Inner.SetOutputItemsForModule(ModuleName, OutputItems);
+		}
+	}
+
+	/// <summary>
+	/// Extension methods for IActionGraphBuilder classes
+	/// </summary>
+	static class ActionGraphBuilderExtensions
+	{
+		/// <summary>
+		/// Creates an action which copies a file from one location to another
+		/// </summary>
+		/// <param name="Graph">The action graph</param>
+		/// <param name="SourceFile">The source file location</param>
+		/// <param name="TargetFile">The target file location</param>
+		/// <returns>File item for the output file</returns>
+		public static Action CreateCopyAction(this IActionGraphBuilder Graph, FileItem SourceFile, FileItem TargetFile)
+		{
+			Action CopyAction = Graph.CreateAction(ActionType.BuildProject);
+			CopyAction.CommandDescription = "Copy";
+			CopyAction.CommandPath = BuildHostPlatform.Current.Shell;
+			if (BuildHostPlatform.Current.ShellType == ShellType.Cmd)
+			{
+				CopyAction.CommandArguments = String.Format("/C \"copy /Y \"{0}\" \"{1}\" 1>nul\"", SourceFile.AbsolutePath, TargetFile.AbsolutePath);
+			}
+			else
+			{
+				CopyAction.CommandArguments = String.Format("-c 'cp -f \"{0}\" \"{1}\"'", SourceFile.AbsolutePath, TargetFile.AbsolutePath);
+			}
+			CopyAction.WorkingDirectory = UnrealBuildTool.EngineSourceDirectory;
+			CopyAction.PrerequisiteItems.Add(SourceFile);
+			CopyAction.ProducedItems.Add(TargetFile);
+			CopyAction.DeleteItems.Add(TargetFile);
+			CopyAction.StatusDescription = TargetFile.Location.GetFileName();
+			CopyAction.bCanExecuteRemotely = false;
+			return CopyAction;
+		}
+
+		/// <summary>
+		/// Creates an action which copies a file from one location to another
+		/// </summary>
+		/// <param name="Graph">List of actions to be executed. Additional actions will be added to this list.</param>
+		/// <param name="SourceFile">The source file location</param>
+		/// <param name="TargetFile">The target file location</param>
+		/// <returns>File item for the output file</returns>
+		public static FileItem CreateCopyAction(this IActionGraphBuilder Graph, FileReference SourceFile, FileReference TargetFile)
+		{
+			FileItem SourceFileItem = FileItem.GetItemByFileReference(SourceFile);
+			FileItem TargetFileItem = FileItem.GetItemByFileReference(TargetFile);
+
+			Graph.CreateCopyAction(SourceFileItem, TargetFileItem);
+
+			return TargetFileItem;
+		}
+
+		/// <summary>
+		/// Creates an action which calls UBT recursively
+		/// </summary>
+		/// <param name="Graph">The action graph</param>
+		/// <param name="Type">Type of the action</param>
+		/// <param name="Arguments">Arguments for the action</param>
+		/// <returns>New action instance</returns>
+		public static Action CreateRecursiveAction<T>(this IActionGraphBuilder Graph, ActionType Type, string Arguments) where T : ToolMode
+		{
+			ToolModeAttribute Attribute = typeof(T).GetCustomAttribute<ToolModeAttribute>();
+			if (Attribute == null)
+			{
+				throw new BuildException("Missing ToolModeAttribute on {0}", typeof(T).Name);
+			}
+
+			Action NewAction = Graph.CreateAction(Type);
+			NewAction.CommandPath = UnrealBuildTool.GetUBTPath();
+			NewAction.CommandArguments = String.Format("-Mode={0} {1}", Attribute.Name, Arguments);
+			return NewAction;
+		}
+
+		/// <summary>
+		/// Creates a text file with the given contents.  If the contents of the text file aren't changed, it won't write the new contents to
+		/// the file to avoid causing an action to be considered outdated.
+		/// </summary>
+		/// <param name="Graph">The action graph</param>
+		/// <param name="AbsolutePath">Path to the intermediate file to create</param>
+		/// <param name="Contents">Contents of the new file</param>
+		/// <returns>File item for the newly created file</returns>
+		public static FileItem CreateIntermediateTextFile(this IActionGraphBuilder Graph, FileReference AbsolutePath, IEnumerable<string> Contents)
+		{
+			return Graph.CreateIntermediateTextFile(AbsolutePath, string.Join(Environment.NewLine, Contents));
+		}
+	}
+}
diff --git a/Engine/Source/Programs/UnrealBuildTool/System/FileItem.cs b/Engine/Source/Programs/UnrealBuildTool/System/FileItem.cs
--- a/Engine/Source/Programs/UnrealBuildTool/System/FileItem.cs
+++ b/Engine/Source/Programs/UnrealBuildTool/System/FileItem.cs
@@ -207,72 +207,6 @@ namespace UnrealBuildTool
 			return Result;
 		}
 
-		/// <summary>
-		/// Determines the appropriate encoding for a string: either ASCII or UTF-8.
-		/// </summary>
-		/// <param name="Str">The string to test.</param>
-		/// <returns>Either System.Text.Encoding.ASCII or System.Text.Encoding.UTF8, depending on whether or not the string contains non-ASCII characters.</returns>
-		private static Encoding GetEncodingForString(string Str)
-		{
-			// If the string length is equivalent to the encoded length, then no non-ASCII characters were present in the string.
-			// Don't write BOM as it messes with clang when loading response files.
-			return (Encoding.UTF8.GetByteCount(Str) == Str.Length) ? Encoding.ASCII : new UTF8Encoding(false);
-		}
-
-		/// <summary>
-		/// Creates a text file with the given contents.  If the contents of the text file aren't changed, it won't write the new contents to
-		/// the file to avoid causing an action to be considered outdated.
-		/// </summary>
-		/// <param name="Location">Path to the intermediate file to create</param>
-		/// <param name="Contents">Contents of the new file</param>
-		/// <returns>File item for the newly created file</returns>
-		public static FileItem CreateIntermediateTextFile(FileReference Location, string Contents)
-		{
-			// Only write the file if its contents have changed.
-			if (!FileReference.Exists(Location))
-			{
-				DirectoryReference.CreateDirectory(Location.Directory);
-				FileReference.WriteAllText(Location, Contents, GetEncodingForString(Contents));
-			}
-			else
-			{
-				string CurrentContents = Utils.ReadAllText(Location.FullName);
-				if(!String.Equals(CurrentContents, Contents, StringComparison.InvariantCultureIgnoreCase))
-				{
-					FileReference BackupFile = new FileReference(Location.FullName + ".old");
-					try
-					{
-						Log.TraceLog("Updating {0}: contents have changed. Saving previous version to {1}.", Location, BackupFile);
-						FileReference.Delete(BackupFile);
-						FileReference.Move(Location, BackupFile);
-					}
-					catch(Exception Ex)
-					{
-						Log.TraceWarning("Unable to rename {0} to {1}", Location, BackupFile);
-						Log.TraceLog("{0}", ExceptionUtils.FormatExceptionDetails(Ex));
-					}
-					FileReference.WriteAllText(Location, Contents, GetEncodingForString(Contents));
-				}
-			}
-
-			// Reset the file info, in case it already knows about the old file
-			FileItem Item = GetItemByFileReference(Location);
-			Item.ResetCachedInfo();
-			return Item;
-		}
-
-		/// <summary>
-		/// Creates a text file with the given contents.  If the contents of the text file aren't changed, it won't write the new contents to
-		/// the file to avoid causing an action to be considered outdated.
-		/// </summary>
-		/// <param name="AbsolutePath">Path to the intermediate file to create</param>
-		/// <param name="Contents">Contents of the new file</param>
-		/// <returns>File item for the newly created file</returns>
-		public static FileItem CreateIntermediateTextFile(FileReference AbsolutePath, IEnumerable<string> Contents)
-		{
-			return CreateIntermediateTextFile(AbsolutePath, string.Join(Environment.NewLine, Contents));
-		}
-
 		/// <summary>
 		/// Deletes the file.
 		/// </summary>
diff --git a/Engine/Source/Programs/UnrealBuildTool/System/TargetMakefile.cs b/Engine/Source/Programs/UnrealBuildTool/System/TargetMakefile.cs
--- a/Engine/Source/Programs/UnrealBuildTool/System/TargetMakefile.cs
+++ b/Engine/Source/Programs/UnrealBuildTool/System/TargetMakefile.cs
@@ -16,7 +16,7 @@ namespace UnrealBuildTool
 	/// <summary>
 	/// Cached list of actions that need to be executed to build a target, along with the information needed to determine whether they are valid.
 	/// </summary>
-	class TargetMakefile
+	class TargetMakefile : IActionGraphBuilder
 	{
 		/// <summary>
 		/// The version number to write
@@ -637,5 +637,64 @@ namespace UnrealBuildTool
 			DirectoryReference BaseDirectory = DirectoryReference.FromFile(ProjectFile) ?? UnrealBuildTool.EngineDirectory;
 			return FileReference.Combine(BaseDirectory, "Intermediate", "Build", Platform.ToString(), TargetName, Configuration.ToString(), "Makefile.bin");
 		}
+
+		/// <inheritdoc/>
+		public Action CreateAction(ActionType Type)
+		{
+			Action Action = new Action(Type);
+			Actions.Add(Action);
+			return Action;
+		}
+
+		/// <inheritdoc/>
+		public FileItem CreateIntermediateTextFile(FileReference Location, string Contents)
+		{
+			// Write the file
+			Utils.WriteFileIfChanged(Location, Contents, StringComparison.InvariantCultureIgnoreCase);
+
+			// Reset the file info, in case it already knows about the old file
+			FileItem Item = FileItem.GetItemByFileReference(Location);
+			Item.ResetCachedInfo();
+			return Item;
+		}
+
+		/// <inheritdoc/>
+		public void AddSourceDir(DirectoryItem SourceDir)
+		{
+			SourceDirectories.Add(SourceDir);
+		}
+
+		/// <inheritdoc/>
+		public void AddSourceFiles(DirectoryItem SourceDir, FileItem[] SourceFiles)
+		{
+			DirectoryToSourceFiles[SourceDir] = SourceFiles;
+		}
+
+		/// <inheritdoc/>
+		public void AddDiagnostic(string Message)
+		{
+			if(!Diagnostics.Contains(Message))
+			{
+				Diagnostics.Add(Message);
+			}
+		}
+
+		/// <inheritdoc/>
+		public void AddFileToWorkingSet(FileItem File)
+		{
+			WorkingSet.Add(File);
+		}
+
+		/// <inheritdoc/>
+		public void AddCandidateForWorkingSet(FileItem File)
+		{
+			CandidatesForWorkingSet.Add(File);
+		}
+
+		/// <inheritdoc/>
+		public void SetOutputItemsForModule(string ModuleName, FileItem[] OutputItems)
+		{
+			ModuleNameToOutputItems[ModuleName] = OutputItems;
+		}
 	}
 }
diff --git a/Engine/Source/Programs/UnrealBuildTool/System/Unity.cs b/Engine/Source/Programs/UnrealBuildTool/System/Unity.cs
--- a/Engine/Source/Programs/UnrealBuildTool/System/Unity.cs
+++ b/Engine/Source/Programs/UnrealBuildTool/System/Unity.cs
@@ -152,7 +152,7 @@ namespace UnrealBuildTool
 		/// <param name="WorkingSet">Interface to query files which belong to the working set</param>
 		/// <param name="BaseName">Base name to use for the Unity files</param>
 		/// <param name="IntermediateDirectory">Intermediate directory for unity cpp files</param>
-		/// <param name="Makefile">The makefile being built</param>
+		/// <param name="Graph">The makefile being built</param>
 		/// <param name="SourceFileToUnityFile">Receives a mapping of source file to unity file</param>
 		/// <returns>The "unity" C++ files.</returns>
 		public static List<FileItem> GenerateUnityCPPs(
@@ -162,7 +162,7 @@ namespace UnrealBuildTool
 			ISourceFileWorkingSet WorkingSet,
 			string BaseName,
 			DirectoryReference IntermediateDirectory,
-			TargetMakefile Makefile,
+			IActionGraphBuilder Graph,
 			Dictionary<FileItem, FileItem> SourceFileToUnityFile
 			)
 		{
@@ -202,6 +202,7 @@ namespace UnrealBuildTool
 
 				// Figure out whether we REALLY want to use adaptive unity for this module.  If nearly every file in the module appears in the working
 				// set, we'll just go ahead and let unity build do its thing.
+				HashSet<FileItem> FilesInWorkingSet = new HashSet<FileItem>();
 				if (bUseAdaptiveUnityBuild)
 				{
 					int CandidateWorkingSetSourceFileCount = 0;
@@ -218,7 +219,8 @@ namespace UnrealBuildTool
 							// Mark this file as part of the working set.  This will be saved into the UBT Makefile so that
 							// the assembler can automatically invalidate the Makefile when the working set changes (allowing this
 							// code to run again, to build up new unity blobs.)
-							Makefile.WorkingSet.Add(CPPFile);
+							FilesInWorkingSet.Add(CPPFile);
+							Graph.AddFileToWorkingSet(CPPFile);
 						}
 					}
 
@@ -240,7 +242,7 @@ namespace UnrealBuildTool
 					}
 
 					// When adaptive unity is enabled, go ahead and exclude any source files that we're actively working with
-					if (bUseAdaptiveUnityBuild && Makefile.WorkingSet.Contains(CPPFile))
+					if (bUseAdaptiveUnityBuild && FilesInWorkingSet.Contains(CPPFile))
 					{
 						// Just compile this file normally, not as part of the unity blob
 						NewCPPFiles.Add(CPPFile);
@@ -266,7 +268,7 @@ namespace UnrealBuildTool
 						// If adaptive unity build is enabled for this module, add this source file to the set that will invalidate the makefile
 						if(bUseAdaptiveUnityBuild)
 						{
-							Makefile.CandidatesForWorkingSet.Add(CPPFile);
+							Graph.AddCandidateForWorkingSet(CPPFile);
 						}
 
 						// Compile this file as part of the unity blob
@@ -278,23 +280,23 @@ namespace UnrealBuildTool
 				{
 					if (Target.bAdaptiveUnityCreatesDedicatedPCH)
 					{
-						AddUniqueDiagnostic(Makefile, "[Adaptive unity build] Creating dedicated PCH for each excluded file. Set bAdaptiveUnityCreatesDedicatedPCH to false in BuildConfiguration.xml to change this behavior.");
+						Graph.AddDiagnostic("[Adaptive unity build] Creating dedicated PCH for each excluded file. Set bAdaptiveUnityCreatesDedicatedPCH to false in BuildConfiguration.xml to change this behavior.");
 					}
 					else if (Target.bAdaptiveUnityDisablesPCH)
 					{
-						AddUniqueDiagnostic(Makefile, "[Adaptive unity build] Disabling PCH for excluded files. Set bAdaptiveUnityDisablesPCH to false in BuildConfiguration.xml to change this behavior.");
+						Graph.AddDiagnostic("[Adaptive unity build] Disabling PCH for excluded files. Set bAdaptiveUnityDisablesPCH to false in BuildConfiguration.xml to change this behavior.");
 					}
 
 					if (Target.bAdaptiveUnityDisablesOptimizations)
 					{
-						AddUniqueDiagnostic(Makefile, "[Adaptive unity build] Disabling optimizations for excluded files. Set bAdaptiveUnityDisablesOptimizations to false in BuildConfiguration.xml to change this behavior.");
+						Graph.AddDiagnostic("[Adaptive unity build] Disabling optimizations for excluded files. Set bAdaptiveUnityDisablesOptimizations to false in BuildConfiguration.xml to change this behavior.");
 					}
 					if (Target.bAdaptiveUnityEnablesEditAndContinue)
 					{
-						AddUniqueDiagnostic(Makefile, "[Adaptive unity build] Enabling Edit & Continue for excluded files. Set bAdaptiveUnityEnablesEditAndContinue to false in BuildConfiguration.xml to change this behavior.");
+						Graph.AddDiagnostic("[Adaptive unity build] Enabling Edit & Continue for excluded files. Set bAdaptiveUnityEnablesEditAndContinue to false in BuildConfiguration.xml to change this behavior.");
 					}
 
-					Makefile.Diagnostics.Add(AdaptiveUnityBuildInfoString.ToString());
+					Graph.AddDiagnostic(AdaptiveUnityBuildInfoString.ToString());
 				}
 
 				AllUnityFiles = CPPUnityFileBuilder.GetUnityFiles();
@@ -330,7 +332,7 @@ namespace UnrealBuildTool
 				FileReference UnityCPPFilePath = FileReference.Combine(IntermediateDirectory, UnityCPPFileName);
 
 				// Write the unity file to the intermediate folder.
-				FileItem UnityCPPFile = FileItem.CreateIntermediateTextFile(UnityCPPFilePath, OutputUnityCPPWriter.ToString());
+				FileItem UnityCPPFile = Graph.CreateIntermediateTextFile(UnityCPPFilePath, OutputUnityCPPWriter.ToString());
 				NewCPPFiles.Add(UnityCPPFile);
 
 				// Store the mapping of source files to unity files in the makefile
diff --git a/Engine/Source/Programs/UnrealBuildTool/System/Utils.cs b/Engine/Source/Programs/UnrealBuildTool/System/Utils.cs
--- a/Engine/Source/Programs/UnrealBuildTool/System/Utils.cs
+++ b/Engine/Source/Programs/UnrealBuildTool/System/Utils.cs
@@ -1246,5 +1246,53 @@ namespace UnrealBuildTool
 			}
 			return CommandLine.ToString();
 		}
+
+		/// <summary>
+		/// Writes a file if the contents have changed
+		/// </summary>
+		/// <param name="Location">Location of the file</param>
+		/// <param name="Contents">New contents of the file</param>
+		/// <param name="Comparison">The type of string comparison to use</param>
+		public static void WriteFileIfChanged(FileReference Location, string Contents, StringComparison Comparison)
+		{
+			// Only write the file if its contents have changed.
+			if (!FileReference.Exists(Location))
+			{
+				DirectoryReference.CreateDirectory(Location.Directory);
+				FileReference.WriteAllText(Location, Contents, GetEncodingForString(Contents));
+			}
+			else
+			{
+				string CurrentContents = Utils.ReadAllText(Location.FullName);
+				if (!String.Equals(CurrentContents, Contents, Comparison))
+				{
+					FileReference BackupFile = new FileReference(Location.FullName + ".old");
+					try
+					{
+						Log.TraceLog("Updating {0}: contents have changed. Saving previous version to {1}.", Location, BackupFile);
+						FileReference.Delete(BackupFile);
+						FileReference.Move(Location, BackupFile);
+					}
+					catch (Exception Ex)
+					{
+						Log.TraceWarning("Unable to rename {0} to {1}", Location, BackupFile);
+						Log.TraceLog("{0}", ExceptionUtils.FormatExceptionDetails(Ex));
+					}
+					FileReference.WriteAllText(Location, Contents, GetEncodingForString(Contents));
+				}
+			}
+		}
+
+		/// <summary>
+		/// Determines the appropriate encoding for a string: either ASCII or UTF-8.
+		/// </summary>
+		/// <param name="Str">The string to test.</param>
+		/// <returns>Either System.Text.Encoding.ASCII or System.Text.Encoding.UTF8, depending on whether or not the string contains non-ASCII characters.</returns>
+		private static Encoding GetEncodingForString(string Str)
+		{
+			// If the string length is equivalent to the encoded length, then no non-ASCII characters were present in the string.
+			// Don't write BOM as it messes with clang when loading response files.
+			return (Encoding.UTF8.GetByteCount(Str) == Str.Length) ? Encoding.ASCII : new UTF8Encoding(false);
+		}
 	}
 }
diff --git a/Engine/Source/Programs/UnrealBuildTool/ToolChain/UEToolChain.cs b/Engine/Source/Programs/UnrealBuildTool/ToolChain/UEToolChain.cs
--- a/Engine/Source/Programs/UnrealBuildTool/ToolChain/UEToolChain.cs
+++ b/Engine/Source/Programs/UnrealBuildTool/ToolChain/UEToolChain.cs
@@ -23,34 +23,34 @@ namespace UnrealBuildTool
 		{
 		}
 
-		public abstract CPPOutput CompileCPPFiles(CppCompileEnvironment CompileEnvironment, List<FileItem> InputFiles, DirectoryReference OutputDir, string ModuleName, List<Action> Actions);
+		public abstract CPPOutput CompileCPPFiles(CppCompileEnvironment CompileEnvironment, List<FileItem> InputFiles, DirectoryReference OutputDir, string ModuleName, IActionGraphBuilder Graph);
 
-		public virtual CPPOutput CompileRCFiles(CppCompileEnvironment Environment, List<FileItem> InputFiles, DirectoryReference OutputDir, List<Action> Actions)
+		public virtual CPPOutput CompileRCFiles(CppCompileEnvironment Environment, List<FileItem> InputFiles, DirectoryReference OutputDir, IActionGraphBuilder Graph)
 		{
 			CPPOutput Result = new CPPOutput();
 			return Result;
 		}
 
-		public virtual CPPOutput CompileISPCFiles(CppCompileEnvironment Environment, List<FileItem> InputFiles, DirectoryReference OutputDir,List<Action> Actions)
+		public virtual CPPOutput CompileISPCFiles(CppCompileEnvironment Environment, List<FileItem> InputFiles, DirectoryReference OutputDir, IActionGraphBuilder Graph)
 		{
 			CPPOutput Result = new CPPOutput();
 			return Result;
 		}
-		public virtual CPPOutput GenerateISPCHeaders(CppCompileEnvironment Environment, List<FileItem> InputFiles, DirectoryReference OutputDir, List<Action> Actions)
+		public virtual CPPOutput GenerateISPCHeaders(CppCompileEnvironment Environment, List<FileItem> InputFiles, DirectoryReference OutputDir, IActionGraphBuilder Graph)
 		{
 			CPPOutput Result = new CPPOutput();
 			return Result;
 		}
 
-		public virtual void GenerateTypeLibraryHeader(CppCompileEnvironment CompileEnvironment, ModuleRules.TypeLibrary TypeLibrary, FileReference OutputFile, List<Action> Actions)
+		public virtual void GenerateTypeLibraryHeader(CppCompileEnvironment CompileEnvironment, ModuleRules.TypeLibrary TypeLibrary, FileReference OutputFile, IActionGraphBuilder Graph)
 		{
 			throw new NotSupportedException("This platform does not support type libraries.");
 		}
 
-		public abstract FileItem LinkFiles(LinkEnvironment LinkEnvironment, bool bBuildImportLibraryOnly, List<Action> Actions);
-		public virtual FileItem[] LinkAllFiles(LinkEnvironment LinkEnvironment, bool bBuildImportLibraryOnly, List<Action> Actions)
+		public abstract FileItem LinkFiles(LinkEnvironment LinkEnvironment, bool bBuildImportLibraryOnly, IActionGraphBuilder Graph);
+		public virtual FileItem[] LinkAllFiles(LinkEnvironment LinkEnvironment, bool bBuildImportLibraryOnly, IActionGraphBuilder Graph)
 		{
-			return new FileItem[] { LinkFiles(LinkEnvironment, bBuildImportLibraryOnly, Actions) };
+			return new FileItem[] { LinkFiles(LinkEnvironment, bBuildImportLibraryOnly, Graph) };
 		}
 
 
@@ -66,7 +66,7 @@ namespace UnrealBuildTool
 			return FileReference.Combine(LinkEnvironment.IntermediateDirectory, OutputFile.Location.GetFileName() + ".response");
 		}
 
-		public virtual ICollection<FileItem> PostBuild(FileItem Executable, LinkEnvironment ExecutableLinkEnvironment, List<Action> Actions)
+		public virtual ICollection<FileItem> PostBuild(FileItem Executable, LinkEnvironment ExecutableLinkEnvironment, IActionGraphBuilder Graph)
 		{
 			return new List<FileItem>();
 		}
@@ -294,7 +294,7 @@ namespace UnrealBuildTool
 			return Suffix;
 		}
 
-		public override CPPOutput GenerateISPCHeaders(CppCompileEnvironment CompileEnvironment, List<FileItem> InputFiles, DirectoryReference OutputDir, List<Action> Actions)
+		public override CPPOutput GenerateISPCHeaders(CppCompileEnvironment CompileEnvironment, List<FileItem> InputFiles, DirectoryReference OutputDir, IActionGraphBuilder Graph)
 		{
 			CPPOutput Result = new CPPOutput();
 
@@ -307,7 +307,7 @@ namespace UnrealBuildTool
 
 			foreach (FileItem ISPCFile in InputFiles)
 			{
-				Action CompileAction = new Action(ActionType.Compile);
+				Action CompileAction = Graph.CreateAction(ActionType.Compile);
 				CompileAction.CommandDescription = "Compile";
 				CompileAction.WorkingDirectory = UnrealBuildTool.EngineSourceDirectory;
 				CompileAction.CommandPath = new FileReference(GetISPCHostCompilerPath(BuildHostPlatform.Current.Platform));
@@ -384,15 +384,13 @@ namespace UnrealBuildTool
 				CompileAction.ProducedItems.Add(ISPCIncludeHeaderFile);
 
 				FileReference ResponseFileName = new FileReference(ISPCIncludeHeaderFile.AbsolutePath + ".response");
-				FileItem ResponseFileItem = FileItem.CreateIntermediateTextFile(ResponseFileName, Arguments.Select(x => Utils.ExpandVariables(x)));
+				FileItem ResponseFileItem = Graph.CreateIntermediateTextFile(ResponseFileName, Arguments.Select(x => Utils.ExpandVariables(x)));
 				CompileAction.CommandArguments += String.Format("@\"{0}\"", ResponseFileName);
 				CompileAction.PrerequisiteItems.Add(ResponseFileItem);
 
 				// Add the source file and its included files to the prerequisite item list.
 				CompileAction.PrerequisiteItems.Add(ISPCFile);
 
-				Actions.Add(CompileAction);
-
 				FileItem ISPCFinalHeaderFile = FileItem.GetItemByFileReference(
 					FileReference.Combine(
 						OutputDir,
@@ -407,7 +405,7 @@ namespace UnrealBuildTool
 				FileItem SourceFileItem = FileItem.GetItemByFileReference(SourceFile);
 				FileItem TargetFileItem = FileItem.GetItemByFileReference(TargetFile);
 
-				Action CopyAction = new Action(ActionType.BuildProject);
+				Action CopyAction = Graph.CreateAction(ActionType.BuildProject);
 				CopyAction.CommandDescription = "Copy";
 				CopyAction.CommandPath = BuildHostPlatform.Current.Shell;
 				if (BuildHostPlatform.Current.ShellType == ShellType.Cmd)
@@ -424,7 +422,6 @@ namespace UnrealBuildTool
 				CopyAction.StatusDescription = TargetFileItem.Location.GetFileName();
 				CopyAction.bCanExecuteRemotely = false;
 				CopyAction.bShouldOutputStatusDescription = false;
-				Actions.Add(CopyAction);
 
 				Result.GeneratedHeaderFiles.Add(TargetFileItem);
 
@@ -434,7 +431,7 @@ namespace UnrealBuildTool
 			return Result;
 		}
 
-		public override CPPOutput CompileISPCFiles(CppCompileEnvironment CompileEnvironment, List<FileItem> InputFiles, DirectoryReference OutputDir, List<Action> Actions)
+		public override CPPOutput CompileISPCFiles(CppCompileEnvironment CompileEnvironment, List<FileItem> InputFiles, DirectoryReference OutputDir, IActionGraphBuilder Graph)
 		{
 			CPPOutput Result = new CPPOutput();
 
@@ -447,7 +444,7 @@ namespace UnrealBuildTool
 
 			foreach (FileItem ISPCFile in InputFiles)
 			{
-				Action CompileAction = new Action(ActionType.Compile);
+				Action CompileAction = Graph.CreateAction(ActionType.Compile);
 				CompileAction.CommandDescription = "Compile";
 				CompileAction.WorkingDirectory = UnrealBuildTool.EngineSourceDirectory;
 				CompileAction.CommandPath = new FileReference(GetISPCHostCompilerPath(BuildHostPlatform.Current.Platform));
@@ -574,15 +571,13 @@ namespace UnrealBuildTool
 				Result.ObjectFiles.AddRange(CompiledISPCObjFiles);
 
 				FileReference ResponseFileName = new FileReference(CompiledISPCObjFileNoISA.AbsolutePath + ".response");
-				FileItem ResponseFileItem = FileItem.CreateIntermediateTextFile(ResponseFileName, Arguments.Select(x => Utils.ExpandVariables(x)));
+				FileItem ResponseFileItem = Graph.CreateIntermediateTextFile(ResponseFileName, Arguments.Select(x => Utils.ExpandVariables(x)));
 				CompileAction.CommandArguments = " @\"" + ResponseFileName + "\"";
 				CompileAction.PrerequisiteItems.Add(ResponseFileItem);
 
 				// Add the source file and its included files to the prerequisite item list.
 				CompileAction.PrerequisiteItems.Add(ISPCFile);
 
-				Actions.Add(CompileAction);
-
 				Log.TraceVerbose("   ISPC Compiling " + CompileAction.StatusDescription + ": \"" + CompileAction.CommandPath + "\"" + CompileAction.CommandArguments);
 			}
 
diff --git a/Engine/Source/Programs/UnrealBuildTool/UnrealBuildTool.csproj b/Engine/Source/Programs/UnrealBuildTool/UnrealBuildTool.csproj
--- a/Engine/Source/Programs/UnrealBuildTool/UnrealBuildTool.csproj
+++ b/Engine/Source/Programs/UnrealBuildTool/UnrealBuildTool.csproj
@@ -1,4 +1,4 @@
-<?xml version="1.0" encoding="utf-8"?>
+﻿<?xml version="1.0" encoding="utf-8"?>
 <Project DefaultTargets="Build" xmlns="http://schemas.microsoft.com/developer/msbuild/2003" ToolsVersion="12.0">
   <PropertyGroup>
     <Configuration Condition=" '$(Configuration)' == '' ">Development</Configuration>
@@ -204,6 +204,7 @@
     <Compile Include="System\CppIncludeParser.cs" />
     <Compile Include="System\FileMetadataPrefetch.cs" />
     <Compile Include="Modes\ParseMsvcTimingInfo.cs" />
+    <Compile Include="System\ActionGraphBuilder.cs" />
     <Compile Include="System\RestrictedFolders.cs" />
     <Compile Include="System\RulesScope.cs" />
     <Compile Include="System\SourceFileMetadataCache.cs" />
@@ -307,7 +308,7 @@
   <Choose>
     <When Condition="Exists('Executors\NotForLicensees\FASTBuild.cs')">
       <ItemGroup>
-          <Compile Include="Executors\NotForLicensees\FASTBuild.cs" />
+        <Compile Include="Executors\NotForLicensees\FASTBuild.cs" />
       </ItemGroup>
       <PropertyGroup>
         <DefineConstants>$(DefineConstants);__FASTBUILD_AVAILABLE__</DefineConstants>
@@ -438,17 +439,14 @@
     <None Include="app.config" />
     <None Include="app.manifest" />
   </ItemGroup>
-
   <Import Project="$(MSBuildBinPath)\Microsoft.CSharp.targets" />
-
   <!-- The following will glob all platform extensions' UnrealBuildTool source
   files and include them in the build. They will also appear in VisualStudio -->
   <ItemGroup>
-		<Compile Include="../../../Platforms/*/Source/Programs/UnrealBuildTool/**/*.cs">
-			<Link>Platform/$([System.Text.RegularExpressions.Regex]::Split(`%(Identity)`, `[\\/]`)[4])/_</Link>
+    <Compile Include="../../../Platforms/*/Source/Programs/UnrealBuildTool/**/*.cs">
+      <Link>Platform/$([System.Text.RegularExpressions.Regex]::Split(`%(Identity)`, `[\\/]`)[4])/_</Link>
     </Compile>
   </ItemGroup>
-
   <!-- To modify your build process, add your task inside one of the targets below and uncomment it. 
        Other similar extension points exist, see Microsoft.Common.targets.
   <Target Name="BeforeBuild">
@@ -469,4 +467,4 @@
     </ProjectReference>
   </ItemGroup>
   <ItemGroup />
-</Project>
+</Project>
\ No newline at end of file
