UBT: Fix makefile being invalidated by response files updated during the build process. Also fix certain source files excluded from unity builds from triggering makefile invalidation.

Creation time of the makefile was being used to test whether dependencies have been modified, which is valid for external dependencies (since we want to catch any files modified after the build process starts), but not valid for files we create during the makefile building (such as response files).

Code attempting to detect response files was also incorrect, and would find any leaf files not included in a unity blob. Response files are now detected explicitly.

diff --git a/Engine/Source/Programs/UnrealBuildTool/Configuration/UEBuildTarget.cs b/Engine/Source/Programs/UnrealBuildTool/Configuration/UEBuildTarget.cs
--- a/Engine/Source/Programs/UnrealBuildTool/Configuration/UEBuildTarget.cs
+++ b/Engine/Source/Programs/UnrealBuildTool/Configuration/UEBuildTarget.cs
@@ -1948,35 +1948,25 @@ namespace UnrealBuildTool
 			}
 
 			// Add all the input files to the predicate store
-			Makefile.AdditionalDependencies.Add(FileItem.GetItemByFileReference(TargetRulesFile));
+			Makefile.ExternalDependencies.Add(FileItem.GetItemByFileReference(TargetRulesFile));
 			foreach(UEBuildModule Module in Modules.Values)
 			{
-				Makefile.AdditionalDependencies.Add(FileItem.GetItemByFileReference(Module.RulesFile));
+				Makefile.ExternalDependencies.Add(FileItem.GetItemByFileReference(Module.RulesFile));
 				foreach(string ExternalDependency in Module.Rules.ExternalDependencies)
 				{
 					FileReference Location = FileReference.Combine(Module.RulesFile.Directory, ExternalDependency);
-					Makefile.AdditionalDependencies.Add(FileItem.GetItemByFileReference(Location));
+					Makefile.ExternalDependencies.Add(FileItem.GetItemByFileReference(Location));
 				}
 				if (Module.Rules.SubclassRules != null)
 				{
 					foreach (string SubclassRule in Module.Rules.SubclassRules)
 					{
 						FileItem SubclassRuleFileItem = FileItem.GetItemByFileReference(new FileReference(SubclassRule));
-						Makefile.AdditionalDependencies.Add(SubclassRuleFileItem);
+						Makefile.ExternalDependencies.Add(SubclassRuleFileItem);
 					}
 				}
 			}
-			Makefile.AdditionalDependencies.UnionWith(Makefile.PluginFiles);
-
-			// Add any leaf dependencies (eg. response files) to the dependencies list
-			IEnumerable<FileItem> LeafPrerequisiteItems = Makefile.Actions.SelectMany(x => x.PrerequisiteItems).Except(Makefile.Actions.SelectMany(x => x.ProducedItems));
-			foreach (FileItem LeafPrerequisiteItem in LeafPrerequisiteItems)
-			{
-				if (LeafPrerequisiteItem.Exists)
-				{
-					Makefile.AdditionalDependencies.Add(LeafPrerequisiteItem);
-				}
-			}
+			Makefile.ExternalDependencies.UnionWith(Makefile.PluginFiles);
 
 			// Write a header containing public definitions for this target
 			if (Rules.ExportPublicHeader != null)
diff --git a/Engine/Source/Programs/UnrealBuildTool/System/TargetMakefile.cs b/Engine/Source/Programs/UnrealBuildTool/System/TargetMakefile.cs
--- a/Engine/Source/Programs/UnrealBuildTool/System/TargetMakefile.cs
+++ b/Engine/Source/Programs/UnrealBuildTool/System/TargetMakefile.cs
@@ -21,13 +21,18 @@ namespace UnrealBuildTool
 		/// <summary>
 		/// The version number to write
 		/// </summary>
-		public const int CurrentVersion = 18;
+		public const int CurrentVersion = 19;
 
 		/// <summary>
 		/// The time at which the makefile was created
 		/// </summary>
 		public DateTime CreateTimeUtc;
 
+		/// <summary>
+		/// The time at which the makefile was modified
+		/// </summary>
+		public DateTime ModifiedTimeUtc;
+
 		/// <summary>
 		/// Additional diagnostic output to print before building this target (toolchain version, etc...)
 		/// </summary>
@@ -150,9 +155,14 @@ namespace UnrealBuildTool
 		public HashSet<FileItem> PluginFiles;
 
 		/// <summary>
-		/// Additional files which are required 
+		/// Set of external (ie. user-owned) files which will cause the makefile to be invalidated if modified after
+		/// </summary>
+		public HashSet<FileItem> ExternalDependencies = new HashSet<FileItem>();
+
+		/// <summary>
+		/// Set of internal (eg. response files, unity files) which will cause the makefile to be invalidated if modified.
 		/// </summary>
-		public HashSet<FileItem> AdditionalDependencies = new HashSet<FileItem>();
+		public HashSet<FileItem> InternalDependencies = new HashSet<FileItem>();
 
 		/// <summary>
 		/// Constructor
@@ -167,7 +177,8 @@ namespace UnrealBuildTool
 		/// <param name="bHasProjectScriptPlugin">Whether the target has a project script plugin</param>
 		public TargetMakefile(string ExternalMetadata, FileReference ExecutableFile, FileReference ReceiptFile, DirectoryReference ProjectIntermediateDirectory, TargetType TargetType, ConfigValueTracker ConfigValueTracker, bool bDeployAfterCompile, bool bHasProjectScriptPlugin)
 		{
-			this.CreateTimeUtc = DateTime.UtcNow;
+			this.CreateTimeUtc = UnrealBuildTool.StartTimeUtc;
+			this.ModifiedTimeUtc = CreateTimeUtc;
 			this.Diagnostics = new List<string>();
 			this.ExternalMetadata = ExternalMetadata;
 			this.ExecutableFile = ExecutableFile;
@@ -188,16 +199,19 @@ namespace UnrealBuildTool
 			this.UObjectModules = new List<UHTModuleInfo>();
 			this.UObjectModuleHeaders = new List<UHTModuleHeaderInfo>();
 			this.PluginFiles = new HashSet<FileItem>();
-			this.AdditionalDependencies = new HashSet<FileItem>();
+			this.ExternalDependencies = new HashSet<FileItem>();
+			this.InternalDependencies = new HashSet<FileItem>();
 		}
 
 		/// <summary>
 		/// Constructor. Reads a makefile from disk.
 		/// </summary>
 		/// <param name="Reader">The archive to read from</param>
-		public TargetMakefile(BinaryArchiveReader Reader)
+		/// <param name="LastWriteTimeUtc">Last modified time for the makefile</param>
+		public TargetMakefile(BinaryArchiveReader Reader, DateTime LastWriteTimeUtc)
 		{
 			CreateTimeUtc = new DateTime(Reader.ReadLong(), DateTimeKind.Utc);
+			ModifiedTimeUtc = LastWriteTimeUtc;
 			Diagnostics = Reader.ReadList(() => Reader.ReadString());
 			ExternalMetadata = Reader.ReadString();
 			ExecutableFile = Reader.ReadFileReference();
@@ -221,7 +235,8 @@ namespace UnrealBuildTool
 			UObjectModules = Reader.ReadList(() => new UHTModuleInfo(Reader));
 			UObjectModuleHeaders = Reader.ReadList(() => new UHTModuleHeaderInfo(Reader));
 			PluginFiles = Reader.ReadHashSet(() => Reader.ReadFileItem());
-			AdditionalDependencies = Reader.ReadHashSet(() => Reader.ReadFileItem());
+			ExternalDependencies = Reader.ReadHashSet(() => Reader.ReadFileItem());
+			InternalDependencies = Reader.ReadHashSet(() => Reader.ReadFileItem());
 		}
 
 		/// <summary>
@@ -254,7 +269,8 @@ namespace UnrealBuildTool
 			Writer.WriteList(UObjectModules, e => e.Write(Writer));
 			Writer.WriteList(UObjectModuleHeaders, x => x.Write(Writer));
 			Writer.WriteHashSet(PluginFiles, x => Writer.WriteFileItem(x));
-			Writer.WriteHashSet(AdditionalDependencies, x => Writer.WriteFileItem(x));
+			Writer.WriteHashSet(ExternalDependencies, x => Writer.WriteFileItem(x));
+			Writer.WriteHashSet(InternalDependencies, x => Writer.WriteFileItem(x));
 		}
 
 		/// <summary>
@@ -282,10 +298,11 @@ namespace UnrealBuildTool
 		/// <returns>The loaded makefile, or null if it failed for some reason.  On failure, the 'ReasonNotLoaded' variable will contain information about why</returns>
 		public static TargetMakefile Load(FileReference MakefilePath, FileReference ProjectFile, UnrealTargetPlatform Platform, string[] Arguments, out string ReasonNotLoaded)
 		{
+			FileInfo MakefileInfo;
 			using(Timeline.ScopeEvent("Checking dependent timestamps"))
 			{
 				// Check the directory timestamp on the project files directory.  If the user has generated project files more recently than the makefile, then we need to consider the file to be out of date
-				FileInfo MakefileInfo = new FileInfo(MakefilePath.FullName);
+				MakefileInfo = new FileInfo(MakefilePath.FullName);
 				if (!MakefileInfo.Exists)
 				{
 					// Makefile doesn't even exist, so we won't bother loading it
@@ -387,7 +404,7 @@ namespace UnrealBuildTool
 							ReasonNotLoaded = "makefile version does not match";
 							return null;
 						}
-						Makefile = new TargetMakefile(Reader);
+						Makefile = new TargetMakefile(Reader, MakefileInfo.LastWriteTimeUtc);
 					}
 				}
 				catch (Exception Ex)
@@ -481,25 +498,42 @@ namespace UnrealBuildTool
 					}
 				}
 
-				// Check if any of the additional dependencies has changed
-				foreach(FileItem AdditionalDependency in Makefile.AdditionalDependencies)
+				// Check if any external dependencies have changed. These comparisons are done against the makefile creation time.
+				foreach(FileItem ExternalDependency in Makefile.ExternalDependencies)
+				{
+					if (!ExternalDependency.Exists)
+					{
+						Log.TraceLog("{0} has been deleted since makefile was built.", ExternalDependency.Location);
+						ReasonNotLoaded = string.Format("{0} deleted", ExternalDependency.Location.GetFileName());
+						return false;
+					}
+					if(ExternalDependency.LastWriteTimeUtc > Makefile.CreateTimeUtc)
+					{
+						Log.TraceLog("{0} has been modified since makefile was built.", ExternalDependency.Location);
+						ReasonNotLoaded = string.Format("{0} modified", ExternalDependency.Location.GetFileName());
+						return false;
+					}
+				}
+
+				// Check if any internal dependencies has changed. These comparisons are done against the makefile modified time.
+				foreach (FileItem InternalDependency in Makefile.InternalDependencies)
 				{
-					if (!AdditionalDependency.Exists)
+					if (!InternalDependency.Exists)
 					{
-						Log.TraceLog("{0} has been deleted since makefile was built.", AdditionalDependency.Location);
-						ReasonNotLoaded = string.Format("{0} deleted", AdditionalDependency.Location.GetFileName());
+						Log.TraceLog("{0} has been deleted since makefile was written.", InternalDependency.Location);
+						ReasonNotLoaded = string.Format("{0} deleted", InternalDependency.Location.GetFileName());
 						return false;
 					}
-					if(AdditionalDependency.LastWriteTimeUtc > Makefile.CreateTimeUtc)
+					if (InternalDependency.LastWriteTimeUtc > Makefile.ModifiedTimeUtc)
 					{
-						Log.TraceLog("{0} has been modified since makefile was built.", AdditionalDependency.Location);
-						ReasonNotLoaded = string.Format("{0} modified", AdditionalDependency.Location.GetFileName());
+						Log.TraceLog("{0} has been modified since makefile was written.", InternalDependency.Location);
+						ReasonNotLoaded = string.Format("{0} modified", InternalDependency.Location.GetFileName());
 						return false;
 					}
 				}
 
 				// Check that no new plugins have been added
-				foreach(FileReference PluginFile in Plugins.EnumeratePlugins(ProjectFile))
+				foreach (FileReference PluginFile in Plugins.EnumeratePlugins(ProjectFile))
 				{
 					FileItem PluginFileItem = FileItem.GetItemByFileReference(PluginFile);
 					if(!Makefile.PluginFiles.Contains(PluginFileItem))
@@ -654,6 +688,7 @@ namespace UnrealBuildTool
 
 			// Reset the file info, in case it already knows about the old file
 			FileItem Item = FileItem.GetItemByFileReference(Location);
+			InternalDependencies.Add(Item);
 			Item.ResetCachedInfo();
 			return Item;
 		}
diff --git a/Engine/Source/Programs/UnrealBuildTool/UnrealBuildTool.cs b/Engine/Source/Programs/UnrealBuildTool/UnrealBuildTool.cs
--- a/Engine/Source/Programs/UnrealBuildTool/UnrealBuildTool.cs
+++ b/Engine/Source/Programs/UnrealBuildTool/UnrealBuildTool.cs
@@ -14,6 +14,12 @@ namespace UnrealBuildTool
 {
 	static class UnrealBuildTool
 	{
+		/// <summary>
+		/// Save the application startup time. This can be used as the timestamp for build makefiles, to determine a base time after which any
+		/// modifications should invalidate it.
+		/// </summary>
+		static public DateTime StartTimeUtc { get; } = DateTime.UtcNow;
+
 		/// <summary>
 		/// The environment at boot time.
 		/// </summary>
