Fix recursive action when running PVS toolchain.

#rnx

diff --git a/Engine/Source/Programs/UnrealBuildTool/Platform/Windows/PVSToolChain.cs b/Engine/Source/Programs/UnrealBuildTool/Platform/Windows/PVSToolChain.cs
--- a/Engine/Source/Programs/UnrealBuildTool/Platform/Windows/PVSToolChain.cs
+++ b/Engine/Source/Programs/UnrealBuildTool/Platform/Windows/PVSToolChain.cs
@@ -354,7 +354,6 @@ namespace UnrealBuildTool
 			AnalyzeAction.PrerequisiteItems.AddRange(AnalyzeActionPrerequisiteItems);
 			AnalyzeAction.ProducedItems.Add(FileItem.GetItemByFileReference(OutputFile));
 			AnalyzeAction.DeleteItems.AddRange(AnalyzeAction.ProducedItems);
-			Makefile.Actions.Add(AnalyzeAction);
 
 			Makefile.OutputItems.AddRange(AnalyzeAction.ProducedItems);
 		}
