Fixed MSVC code analysis from giving compiler errors when run with Conformance Mode enabled.

When the '/permissive- is set we get a number of compile errors (C2337) due to microsoft attributes in \include\CodeAnalysis\sourceannotations.h

#rnx

diff --git a/Engine/Source/Programs/UnrealBuildTool/Platform/Windows/VCToolChain.cs b/Engine/Source/Programs/UnrealBuildTool/Platform/Windows/VCToolChain.cs
--- a/Engine/Source/Programs/UnrealBuildTool/Platform/Windows/VCToolChain.cs
+++ b/Engine/Source/Programs/UnrealBuildTool/Platform/Windows/VCToolChain.cs
@@ -230,8 +230,12 @@ namespace UnrealBuildTool
 			// Disable Microsoft extensions on VS2017+ for improved standards compliance.
 			if (Target.WindowsPlatform.Compiler >= WindowsCompiler.VisualStudio2017 && Target.WindowsPlatform.bStrictConformanceMode)
 			{
-				Arguments.Add("/permissive-");
-				Arguments.Add("/Zc:strictStrings-"); // Have to disable strict const char* semantics due to Windows headers not being compliant.
+				// Do not enable if we are running the MSVC code analyzer as it will fail to compile.
+				if (Target.WindowsPlatform.StaticAnalyzer != WindowsStaticAnalyzer.VisualCpp)
+				{
+					Arguments.Add("/permissive-");
+					Arguments.Add("/Zc:strictStrings-"); // Have to disable strict const char* semantics due to Windows headers not being compliant.
+				}
 			}
 
 			// @todo HoloLens: UE4 is non-compliant when it comes to use of %s and %S
