Compile fixes after recent stream merges
#rnx

diff --git a/Engine/Source/Programs/UnrealBuildTool/Configuration/UEBuildTarget.cs b/Engine/Source/Programs/UnrealBuildTool/Configuration/UEBuildTarget.cs
--- a/Engine/Source/Programs/UnrealBuildTool/Configuration/UEBuildTarget.cs
+++ b/Engine/Source/Programs/UnrealBuildTool/Configuration/UEBuildTarget.cs
@@ -1977,7 +1977,7 @@ namespace UnrealBuildTool
 			{
 				if (LeafPrerequisiteItem.Exists)
 				{
-					Makefile.AdditionalDependencies.Add(LeafPrerequisiteItem);
+					Makefile.InternalDependencies.Add(LeafPrerequisiteItem);
 				}
 			}
 
