Use a subdirectory for storing PVS studio artifacts, to avoid triggering a full rebuild with the next regular compile.

diff --git a/Engine/Source/Programs/UnrealBuildTool/Platform/Windows/PVSToolChain.cs b/Engine/Source/Programs/UnrealBuildTool/Platform/Windows/PVSToolChain.cs
--- a/Engine/Source/Programs/UnrealBuildTool/Platform/Windows/PVSToolChain.cs
+++ b/Engine/Source/Programs/UnrealBuildTool/Platform/Windows/PVSToolChain.cs
@@ -215,6 +215,9 @@ namespace UnrealBuildTool
 
 		public override CPPOutput CompileCPPFiles(CppCompileEnvironment CompileEnvironment, List<FileItem> InputFiles, DirectoryReference OutputDir, string ModuleName, IActionGraphBuilder Graph)
 		{
+			// Use a subdirectory for PVS output, to avoid clobbering regular build artifacts
+			OutputDir = DirectoryReference.Combine(OutputDir, "PVS");
+
 			// Preprocess the source files with the regular toolchain
 			CppCompileEnvironment PreprocessCompileEnvironment = new CppCompileEnvironment(CompileEnvironment);
 			PreprocessCompileEnvironment.bPreprocessOnly = true;
