Added more control over Gradle project
- ANDROID_TOOLS_BUILD_GRADLE_VERSION may be set in gradle.properties by <gradleProperties> in UPL
- settings.gradle may be added to with <settingsGradleAdditions> in UPL

#ue4
#android

diff --git a/Engine/Source/Programs/UnrealBuildTool/Platform/Android/UEDeployAndroid.cs b/Engine/Source/Programs/UnrealBuildTool/Platform/Android/UEDeployAndroid.cs
--- a/Engine/Source/Programs/UnrealBuildTool/Platform/Android/UEDeployAndroid.cs
+++ b/Engine/Source/Programs/UnrealBuildTool/Platform/Android/UEDeployAndroid.cs
@@ -21,6 +21,9 @@ namespace UnrealBuildTool
 		// filename of current BundleTool
 		private const string BUNDLETOOL_JAR = "bundletool-all-0.13.0.jar";
 
+		// classpath of default android build tools gradle plugin
+		private const string ANDROID_TOOLS_BUILD_GRADLE_VERSION = "com.android.tools.build:gradle:3.5.3";
+
 		// Minimum Android SDK that must be used for Java compiling
 		readonly int MinimumSDKLevel = 28;
 
@@ -2952,6 +2955,7 @@ namespace UnrealBuildTool
 				GradleProperties.AppendLine(string.Format("OBB_FILE{0}={1}", OBBFileIndex++, OBBFile.Replace("\\", "/")));
 			}
 
+			GradleProperties.AppendLine("ANDROID_TOOLS_BUILD_GRADLE_VERSION={0}", ANDROID_TOOLS_BUILD_GRADLE_VERSION);
 			GradleProperties.AppendLine("BUNDLETOOL_JAR=" + Path.GetFullPath(Path.Combine(UE4BuildFilesPath, "..", "Prebuilt", "bundletool", BUNDLETOOL_JAR)).Replace("\\", "/"));
 			GradleProperties.AppendLine("GENUNIVERSALAPK_JAR=" + Path.GetFullPath(Path.Combine(UE4BuildFilesPath, "..", "Prebuilt", "GenUniversalAPK", "bin", "GenUniversalAPK.jar")).Replace("\\", "/"));
 
@@ -3735,7 +3739,7 @@ namespace UnrealBuildTool
 					
 				// move JavaLibs into subprojects
 				string JavaLibsDir = Path.Combine(UE4BuildPath, "JavaLibs");
-				PrepareJavaLibsForGradle(JavaLibsDir, UE4BuildGradlePath, MinSDKVersion.ToString(), TargetSDKVersion.ToString(), CompileSDKVersion, BuildToolsVersion);
+				PrepareJavaLibsForGradle(JavaLibsDir, UE4BuildGradlePath, MinSDKVersion.ToString(), TargetSDKVersion.ToString(), CompileSDKVersion, BuildToolsVersion, NDKArch);
 
 				// Create local.properties
 				String LocalPropertiesFilename = Path.Combine(UE4BuildGradlePath, "local.properties");
@@ -4640,7 +4644,7 @@ namespace UnrealBuildTool
 			return AARHandler;
 		}
 
-		private void PrepareJavaLibsForGradle(string JavaLibsDir, string UE4BuildGradlePath, string InMinSdkVersion, string InTargetSdkVersion, string CompileSDKVersion, string BuildToolsVersion)
+		private void PrepareJavaLibsForGradle(string JavaLibsDir, string UE4BuildGradlePath, string InMinSdkVersion, string InTargetSdkVersion, string CompileSDKVersion, string BuildToolsVersion, string NDKArch)
 		{
 			StringBuilder SettingsGradleContent = new StringBuilder();
 			StringBuilder ProjectDependencyContent = new StringBuilder();
@@ -4796,6 +4800,9 @@ namespace UnrealBuildTool
 			}
 			ProjectDependencyContent.AppendLine("}");
 
+			// Add any UPL settingsGradleAdditions
+			SettingsGradleContent.Append(UPL.ProcessPluginNode(NDKArch, "settingsGradleAdditions", ""));
+
 			string SettingsGradleFilename = Path.Combine(UE4BuildGradlePath, "settings.gradle");
 			File.WriteAllText(SettingsGradleFilename, SettingsGradleContent.ToString());
 
