Make sure UBT and HoloLens platforms agree on the architecture being compiled. The architecture setting is now driven by the target being built, and cannot be overridden later. The HoloLens platform returns the default architecture value to UBT, ensuring that intermediate folders and the like are named consistently between builds where an architecture isn't specified (and it uses the default), and and builds where an architecture is explicitly specified.

#rnx

diff --git a/Engine/Source/Programs/UnrealBuildTool/Configuration/TargetRules.cs b/Engine/Source/Programs/UnrealBuildTool/Configuration/TargetRules.cs
--- a/Engine/Source/Programs/UnrealBuildTool/Configuration/TargetRules.cs
+++ b/Engine/Source/Programs/UnrealBuildTool/Configuration/TargetRules.cs
@@ -1467,7 +1467,7 @@ namespace UnrealBuildTool
 		/// HoloLens-specific target settings.
 		/// </summary>
 		[ConfigSubObject]
-		public HoloLensTargetRules HoloLensPlatform = new HoloLensTargetRules();
+		public HoloLensTargetRules HoloLensPlatform;
 
 		/// <summary>
 		/// Constructor.
@@ -1482,6 +1482,7 @@ namespace UnrealBuildTool
 			this.ProjectFile = Target.ProjectFile;
 			this.Version = Target.Version;
 			this.WindowsPlatform = new WindowsTargetRules(this);
+			this.HoloLensPlatform = new HoloLensTargetRules(Target);
 
 			// Read settings from config files
 			foreach (object ConfigurableObject in GetConfigurableObjects())
diff --git a/Engine/Source/Programs/UnrealBuildTool/Platform/HoloLens/UEBuildHoloLens.cs b/Engine/Source/Programs/UnrealBuildTool/Platform/HoloLens/UEBuildHoloLens.cs
--- a/Engine/Source/Programs/UnrealBuildTool/Platform/HoloLens/UEBuildHoloLens.cs
+++ b/Engine/Source/Programs/UnrealBuildTool/Platform/HoloLens/UEBuildHoloLens.cs
@@ -29,9 +29,10 @@ namespace UnrealBuildTool
 		/// <summary>
 		/// Architecture of Target.
 		/// </summary>
-		[CommandLine("-x64", Value = "x64")]
-		[CommandLine("-arm64", Value = "ARM64")]
-		public WindowsArchitecture Architecture = WindowsArchitecture.x64;
+		public WindowsArchitecture Architecture
+		{
+			get;
+		}
 
 		/// <summary>
 		/// Enable PIX debugging (automatically disabled in Shipping and Test configs)
@@ -58,6 +59,18 @@ namespace UnrealBuildTool
         /// </summary>
         [ConfigFile(ConfigHierarchyType.Engine, "/Script/HoloLensPlatformEditor.HoloLensTargetSettings", "bAutoIncrementVersion")]
         public bool bAutoIncrementVersion = false;
+
+		/// <summary>
+		/// Constructor
+		/// </summary>
+		/// <param name="Info">Target information</param>
+		public HoloLensTargetRules(TargetInfo Info)
+		{
+			if (Info.Platform == UnrealTargetPlatform.HoloLens)
+			{
+				Architecture = (WindowsArchitecture)Enum.Parse(typeof(WindowsArchitecture), Info.Architecture, true);
+			}
+		}
     }
 
 	/// <summary>
@@ -176,22 +189,9 @@ namespace UnrealBuildTool
 				}
 			}
 
-			Target.HoloLensPlatform.Architecture = WindowsArchitecture.ARM64;
-			if (Target.Architecture.ToLower() == "arm64")
-			{
-				Target.HoloLensPlatform.Architecture = WindowsArchitecture.ARM64;
-				if(!Target.bGenerateProjectFiles)
-				{
-					Log.TraceInformationOnce("Using ARM64 architecture for deploying to HoloLens device");
-				}
-			}
-			else
+			if(!Target.bGenerateProjectFiles)
 			{
-				Target.HoloLensPlatform.Architecture = WindowsArchitecture.x64;
-				if (!Target.bGenerateProjectFiles)
-				{
-					Log.TraceInformationOnce("Using x64 architecture for deploying to HoloLens emulator");
-				}
+				Log.TraceInformationOnce("Using {0} architecture for deploying to HoloLens device", Target.HoloLensPlatform.Architecture);
 			}
 
 			Target.WindowsPlatform.Compiler = Target.HoloLensPlatform.Compiler;
@@ -252,6 +252,16 @@ namespace UnrealBuildTool
 			HoloLensExports.InitWindowsSdkToolPath(Target.HoloLensPlatform.Win10SDKVersion.ToString());
 		}
 
+		/// <summary>
+		/// Gets the default HoloLens architecture
+		/// </summary>
+		/// <param name="ProjectFile">The project being built</param>
+		/// <returns>The default architecture</returns>
+		public override string GetDefaultArchitecture(FileReference ProjectFile)
+		{
+			return WindowsArchitecture.x64.ToString();
+		}
+
 		/// <summary>
 		/// Determines if the given name is a build product for a target.
 		/// </summary>
diff --git a/Engine/Source/Programs/UnrealBuildTool/Platform/Windows/UEBuildWindows.cs b/Engine/Source/Programs/UnrealBuildTool/Platform/Windows/UEBuildWindows.cs
--- a/Engine/Source/Programs/UnrealBuildTool/Platform/Windows/UEBuildWindows.cs
+++ b/Engine/Source/Programs/UnrealBuildTool/Platform/Windows/UEBuildWindows.cs
@@ -122,7 +122,12 @@ namespace UnrealBuildTool
 		/// <summary>
 		/// Architecture of Target.
 		/// </summary>
-		public WindowsArchitecture Architecture = WindowsArchitecture.x64;
+		public WindowsArchitecture Architecture
+		{
+			get;
+			internal set;
+		}
+		= WindowsArchitecture.x64;
 
 		/// <summary>
 		/// The specific toolchain version to use. This may be a specific version number (for example, "14.13.26128"), or the string "Latest", to select the newest available version. By default, and if it is available, we use the
