Fix exception generating project files.

#rnx

diff --git a/Engine/Source/Programs/UnrealBuildTool/Platform/HoloLens/UEBuildHoloLens.cs b/Engine/Source/Programs/UnrealBuildTool/Platform/HoloLens/UEBuildHoloLens.cs
--- a/Engine/Source/Programs/UnrealBuildTool/Platform/HoloLens/UEBuildHoloLens.cs
+++ b/Engine/Source/Programs/UnrealBuildTool/Platform/HoloLens/UEBuildHoloLens.cs
@@ -66,7 +66,7 @@ namespace UnrealBuildTool
 		/// <param name="Info">Target information</param>
 		public HoloLensTargetRules(TargetInfo Info)
 		{
-			if (Info.Platform == UnrealTargetPlatform.HoloLens)
+			if (Info.Platform == UnrealTargetPlatform.HoloLens && !String.IsNullOrEmpty(Info.Architecture))
 			{
 				Architecture = (WindowsArchitecture)Enum.Parse(typeof(WindowsArchitecture), Info.Architecture, true);
 			}
