UBT: Add a BuildConfiguration setting for forcing UBT to be built in debug configuration (VCProjectFileGenerator/bBuildUBTInDebug).

diff --git a/Engine/Source/Programs/UnrealBuildTool/ProjectFiles/VisualStudio/VCProjectFileGenerator.cs b/Engine/Source/Programs/UnrealBuildTool/ProjectFiles/VisualStudio/VCProjectFileGenerator.cs
--- a/Engine/Source/Programs/UnrealBuildTool/ProjectFiles/VisualStudio/VCProjectFileGenerator.cs
+++ b/Engine/Source/Programs/UnrealBuildTool/ProjectFiles/VisualStudio/VCProjectFileGenerator.cs
@@ -52,6 +52,12 @@ namespace UnrealBuildTool
 		[XmlConfigFile(Category = "BuildConfiguration")]
 		protected bool bWriteSolutionOptionFile = true;
 
+		/// <summary>
+		/// Forces UBT to be built in debug configuration, regardless of the solution configuration
+		/// </summary>
+		[XmlConfigFile]
+		protected bool bBuildUBTInDebug = false;
+
 		/// <summary>
 		/// Whether to add the -FastPDB option to build command lines by default.
 		/// </summary>
@@ -697,6 +703,12 @@ namespace UnrealBuildTool
 								// Get the context for the current solution context
 								MSBuildProjectContext ProjectContext = CurProject.GetMatchingProjectContext(SolutionConfigCombination.TargetConfigurationName, SolutionConfigCombination.Configuration, SolutionConfigCombination.Platform, PlatformProjectGenerators);
 
+								// Override the configuration to build for UBT
+								if (bBuildUBTInDebug && CurProject == UBTProject)
+								{
+									ProjectContext.ConfigurationName = "Debug";
+								}
+
 								// Write the solution mapping (e.g.  "{4232C52C-680F-4850-8855-DC39419B5E9B}.Debug|iOS.ActiveCfg = iOS_Debug|Win32")
 								string CurProjectGUID = CurProject.ProjectGUID.ToString("B").ToUpperInvariant();
 								VCSolutionFileContent.AppendLine("		{0}.{1}.ActiveCfg = {2}", CurProjectGUID, SolutionConfigCombination.VCSolutionConfigAndPlatformName, ProjectContext.Name);
