Fix installed files being deleted when cleaning a target from an installed engine build.

diff --git a/Engine/Source/Programs/UnrealBuildTool/System/RulesAssembly.cs b/Engine/Source/Programs/UnrealBuildTool/System/RulesAssembly.cs
--- a/Engine/Source/Programs/UnrealBuildTool/System/RulesAssembly.cs
+++ b/Engine/Source/Programs/UnrealBuildTool/System/RulesAssembly.cs
@@ -28,9 +28,9 @@ namespace UnrealBuildTool
 		private Assembly CompiledAssembly;
 
 		/// <summary>
-		/// The base directory for this assembly
+		/// The base directories for this assembly
 		/// </summary>
-		private DirectoryReference BaseDir;
+		private List<DirectoryReference> BaseDirs;
 
 		/// <summary>
 		/// All the plugins included in this assembly
@@ -82,7 +82,7 @@ namespace UnrealBuildTool
 		/// Constructor. Compiles a rules assembly from the given source files.
 		/// </summary>
 		/// <param name="Scope">The scope of items created by this assembly</param>
-		/// <param name="BaseDir">The base directory for this assembly</param>
+		/// <param name="BaseDirs">The base directories for this assembly</param>
 		/// <param name="Plugins">All the plugins included in this assembly</param>
 		/// <param name="ModuleFileToContext">List of module files to compile</param>
 		/// <param name="TargetFiles">List of target files to compile</param>
@@ -92,10 +92,10 @@ namespace UnrealBuildTool
 		/// <param name="bReadOnly">Whether the modules and targets in this assembly are installed, and should be created with the bUsePrecompiled flag set</param> 
 		/// <param name="bSkipCompile">Whether to skip compiling this assembly</param>
 		/// <param name="Parent">The parent rules assembly</param>
-		internal RulesAssembly(RulesScope Scope, DirectoryReference BaseDir, IReadOnlyList<PluginInfo> Plugins, Dictionary<FileReference, ModuleRulesContext> ModuleFileToContext, List<FileReference> TargetFiles, FileReference AssemblyFileName, bool bContainsEngineModules, BuildSettingsVersion? DefaultBuildSettings, bool bReadOnly, bool bSkipCompile, RulesAssembly Parent)
+		internal RulesAssembly(RulesScope Scope, List<DirectoryReference> BaseDirs, IReadOnlyList<PluginInfo> Plugins, Dictionary<FileReference, ModuleRulesContext> ModuleFileToContext, List<FileReference> TargetFiles, FileReference AssemblyFileName, bool bContainsEngineModules, BuildSettingsVersion? DefaultBuildSettings, bool bReadOnly, bool bSkipCompile, RulesAssembly Parent)
 		{
 			this.Scope = Scope;
-			this.BaseDir = BaseDir;
+			this.BaseDirs = BaseDirs;
 			this.Plugins = Plugins;
 			this.ModuleFileToContext = ModuleFileToContext;
 			this.bContainsEngineModules = bContainsEngineModules;
@@ -178,11 +178,11 @@ namespace UnrealBuildTool
 		/// <returns>True if the path is read-only, false otherwise</returns>
 		public bool IsReadOnly(FileSystemReference Location)
 		{
-			if(Location.IsUnderDirectory(BaseDir))
+			if (BaseDirs.Any(x => Location.IsUnderDirectory(x)))
 			{
 				return bReadOnly;
 			}
-			else if(Parent != null)
+			else if (Parent != null)
 			{
 				return Parent.IsReadOnly(Location);
 			}
diff --git a/Engine/Source/Programs/UnrealBuildTool/System/RulesCompiler.cs b/Engine/Source/Programs/UnrealBuildTool/System/RulesCompiler.cs
--- a/Engine/Source/Programs/UnrealBuildTool/System/RulesCompiler.cs
+++ b/Engine/Source/Programs/UnrealBuildTool/System/RulesCompiler.cs
@@ -464,6 +464,7 @@ namespace UnrealBuildTool
 					AddEngineModuleRulesWithContext(SourceDirectory, "ThirdParty", DefaultModuleContext, UHTModuleType.EngineThirdParty, ModuleFileToContext);
 				}
 			}
+
 			// Add all the plugin modules too (don't need to loop over RootDirectories since the plugins come in already found
 			using (Timeline.ScopeEvent("Finding plugin modules"))
 			{
@@ -485,7 +486,7 @@ namespace UnrealBuildTool
 
 			// Create the assembly
 			FileReference EngineAssemblyFileName = FileReference.Combine(AssemblyDir, "Intermediate", "Build", "BuildRules", AssemblyPrefix + "Rules" + FrameworkAssemblyExtension);
-			RulesAssembly EngineAssembly = new RulesAssembly(Scope, AssemblyDir, Plugins, ModuleFileToContext, new List<FileReference>(), EngineAssemblyFileName, bContainsEngineModules: true, DefaultBuildSettings: BuildSettingsVersion.Latest, bReadOnly: bReadOnly, bSkipCompile: bSkipCompile, Parent: Parent);
+			RulesAssembly EngineAssembly = new RulesAssembly(Scope, RootDirectories, Plugins, ModuleFileToContext, new List<FileReference>(), EngineAssemblyFileName, bContainsEngineModules: true, DefaultBuildSettings: BuildSettingsVersion.Latest, bReadOnly: bReadOnly, bSkipCompile: bSkipCompile, Parent: Parent);
 
 			List<FileReference> ProgramTargetFiles = new List<FileReference>();
 			Dictionary<FileReference, ModuleRulesContext> ProgramModuleFiles = new Dictionary<FileReference, ModuleRulesContext>();
@@ -512,7 +513,7 @@ namespace UnrealBuildTool
 
 			// Create a path to the assembly that we'll either load or compile
 			FileReference ProgramAssemblyFileName = FileReference.Combine(AssemblyDir, "Intermediate", "Build", "BuildRules", AssemblyPrefix + "ProgramRules" + FrameworkAssemblyExtension);
-			RulesAssembly ProgramAssembly = new RulesAssembly(ProgramsScope, AssemblyDir, new List<PluginInfo>().AsReadOnly(), ProgramModuleFiles, ProgramTargetFiles, ProgramAssemblyFileName, bContainsEngineModules: false, DefaultBuildSettings: BuildSettingsVersion.Latest, bReadOnly: bReadOnly, bSkipCompile: bSkipCompile, Parent: EngineAssembly);
+			RulesAssembly ProgramAssembly = new RulesAssembly(ProgramsScope, RootDirectories, new List<PluginInfo>().AsReadOnly(), ProgramModuleFiles, ProgramTargetFiles, ProgramAssemblyFileName, bContainsEngineModules: false, DefaultBuildSettings: BuildSettingsVersion.Latest, bReadOnly: bReadOnly, bSkipCompile: bSkipCompile, Parent: EngineAssembly);
 
 			// Return the combined assembly
 			return ProgramAssembly;
@@ -609,7 +610,7 @@ namespace UnrealBuildTool
 				}
 				else
 				{
-					ProjectRulesAssembly = new RulesAssembly(Scope, MainProjectDirectory, ProjectPlugins, ModuleFiles, TargetFiles, AssemblyFileName, bContainsEngineModules: false, DefaultBuildSettings: null, bReadOnly: UnrealBuildTool.IsProjectInstalled(), bSkipCompile: bSkipCompile, Parent: Parent);
+					ProjectRulesAssembly = new RulesAssembly(Scope, new List<DirectoryReference> { MainProjectDirectory }, ProjectPlugins, ModuleFiles, TargetFiles, AssemblyFileName, bContainsEngineModules: false, DefaultBuildSettings: null, bReadOnly: UnrealBuildTool.IsProjectInstalled(), bSkipCompile: bSkipCompile, Parent: Parent);
 				}
 				LoadedAssemblyMap.Add(ProjectFileName, ProjectRulesAssembly);
 			}
@@ -651,7 +652,7 @@ namespace UnrealBuildTool
 
 				// Compile the assembly
 				FileReference AssemblyFileName = FileReference.Combine(PluginFileName.Directory, "Intermediate", "Build", "BuildRules", Path.GetFileNameWithoutExtension(PluginFileName.FullName) + "ModuleRules" + FrameworkAssemblyExtension);
-				PluginRulesAssembly = new RulesAssembly(Scope, PluginFileName.Directory, ForeignPlugins, ModuleFiles, TargetFiles, AssemblyFileName, bContainsEngineModules, DefaultBuildSettings: null, bReadOnly: false, bSkipCompile: bSkipCompile, Parent: Parent);
+				PluginRulesAssembly = new RulesAssembly(Scope, new List<DirectoryReference> { PluginFileName.Directory }, ForeignPlugins, ModuleFiles, TargetFiles, AssemblyFileName, bContainsEngineModules, DefaultBuildSettings: null, bReadOnly: false, bSkipCompile: bSkipCompile, Parent: Parent);
 				LoadedAssemblyMap.Add(PluginFileName, PluginRulesAssembly);
 			}
 			return PluginRulesAssembly;
