Remove redundant null checks causing warnings on latest version of Mono.

#rnx

diff --git a/Engine/Source/Programs/UnrealBuildTool/System/ConfigCache.cs b/Engine/Source/Programs/UnrealBuildTool/System/ConfigCache.cs
--- a/Engine/Source/Programs/UnrealBuildTool/System/ConfigCache.cs
+++ b/Engine/Source/Programs/UnrealBuildTool/System/ConfigCache.cs
@@ -99,7 +99,7 @@ namespace UnrealBuildTool
 			/// <returns>Hash value for this object</returns>
 			public override int GetHashCode()
 			{
-				return ((ProjectDir != null)? ProjectDir.GetHashCode() : 0) + ((int)Type * 123) + ((Platform != null ? Platform.ToString() : "None").GetHashCode() * 345);
+				return ((ProjectDir != null) ? ProjectDir.GetHashCode() : 0) + ((int)Type * 123) + (Platform.GetHashCode() * 345);
 			}
 		}
 
diff --git a/Engine/Source/Programs/UnrealBuildTool/System/ConfigHierarchy.cs b/Engine/Source/Programs/UnrealBuildTool/System/ConfigHierarchy.cs
--- a/Engine/Source/Programs/UnrealBuildTool/System/ConfigHierarchy.cs
+++ b/Engine/Source/Programs/UnrealBuildTool/System/ConfigHierarchy.cs
@@ -1017,7 +1017,7 @@ namespace UnrealBuildTool
 			}
 			else
 			{
-				return TargetPlatform == null ? "None" : TargetPlatform.ToString();
+				return TargetPlatform.ToString();
 			}
 		}
 	}
