Remove unused variables causing warnings on Mono.

#rnx

diff --git a/Engine/Source/Programs/UnrealBuildTool/Configuration/UEBuildModule.cs b/Engine/Source/Programs/UnrealBuildTool/Configuration/UEBuildModule.cs
--- a/Engine/Source/Programs/UnrealBuildTool/Configuration/UEBuildModule.cs
+++ b/Engine/Source/Programs/UnrealBuildTool/Configuration/UEBuildModule.cs
@@ -56,11 +56,6 @@ namespace UnrealBuildTool
 			get { return Rules.File; }
 		}
 
-		/// <summary>
-		/// Is this module allowed to be redistributed.
-		/// </summary>
-		private readonly bool? IsRedistributableOverride;
-
 		/// <summary>
 		/// The binary the module will be linked into for the current target.  Only set after UEBuildBinary.BindModules is called.
 		/// </summary>
@@ -235,7 +230,6 @@ namespace UnrealBuildTool
 			{
 				PrivateIncludePaths = CreateDirectoryHashSet(Rules.PrivateIncludePaths);
 			}
-			IsRedistributableOverride = Rules.IsRedistributableOverride;
 
 			WhitelistRestrictedFolders = new HashSet<DirectoryReference>(Rules.WhitelistRestrictedFolders.Select(x => DirectoryReference.Combine(ModuleDirectory, x)));
 
diff --git a/Engine/Source/Programs/UnrealBuildTool/System/ConfigCache.cs b/Engine/Source/Programs/UnrealBuildTool/System/ConfigCache.cs
--- a/Engine/Source/Programs/UnrealBuildTool/System/ConfigCache.cs
+++ b/Engine/Source/Programs/UnrealBuildTool/System/ConfigCache.cs
@@ -103,11 +103,6 @@ namespace UnrealBuildTool
 			}
 		}
 
-		/// <summary>
-		/// Stores the number of config file types
-		/// </summary>
-		static readonly int NumConfigFileTypes = (int)((ConfigHierarchyType[])Enum.GetValues(typeof(ConfigHierarchyType))).Last() + 1;
-
 		/// <summary>
 		/// Cache of individual config files
 		/// </summary>
diff --git a/Engine/Source/Programs/UnrealBuildTool/System/RulesAssembly.cs b/Engine/Source/Programs/UnrealBuildTool/System/RulesAssembly.cs
--- a/Engine/Source/Programs/UnrealBuildTool/System/RulesAssembly.cs
+++ b/Engine/Source/Programs/UnrealBuildTool/System/RulesAssembly.cs
@@ -52,11 +52,6 @@ namespace UnrealBuildTool
 		/// </summary>
 		private Dictionary<FileReference, ModuleRulesContext> ModuleFileToContext;
 
-		/// <summary>
-		/// Cache for whether a module has source code
-		/// </summary>
-		private Dictionary<FileReference, bool> ModuleHasSource = new Dictionary<FileReference, bool>();
-
 		/// <summary>
 		/// Whether this assembly contains engine modules. Used to set default values for bTreatAsEngineModule.
 		/// </summary>
