Fix UnrealBuildTool unused variable warnings

diff --git a/Engine/Platforms/XXX/Source/Programs/UnrealBuildTool/UEBuildXXX.cs b/Engine/Platforms/XXX/Source/Programs/UnrealBuildTool/UEBuildXXX.cs
--- a/Engine/Platforms/XXX/Source/Programs/UnrealBuildTool/UEBuildXXX.cs
+++ b/Engine/Platforms/XXX/Source/Programs/UnrealBuildTool/UEBuildXXX.cs
@@ -57,7 +57,7 @@ namespace UnrealBuildTool
 		{
 			this.Inner = Inner;
 
-			UnrealPlatformGroup Group = UnrealPlatformGroup.Android;
+			//UnrealPlatformGroup Group = UnrealPlatformGroup.Android;
 		}
 	}
 
@@ -445,4 +445,4 @@ class XXXPlatformFactory : UEBuildPlatformFactory
 		UEBuildPlatform.RegisterBuildPlatform(new XXXPlatform(SDK));
 		UEBuildPlatform.RegisterPlatformWithGroup(UnrealTargetPlatform.XXX, UnrealPlatformGroup.Fake);
 	}
-}
\ No newline at end of file
+}
diff --git a/Engine/Source/Programs/UnrealBuildTool/Configuration/UEBuildModuleCPP.cs b/Engine/Source/Programs/UnrealBuildTool/Configuration/UEBuildModuleCPP.cs
--- a/Engine/Source/Programs/UnrealBuildTool/Configuration/UEBuildModuleCPP.cs
+++ b/Engine/Source/Programs/UnrealBuildTool/Configuration/UEBuildModuleCPP.cs
@@ -278,7 +278,7 @@ namespace UnrealBuildTool
 		// UEBuildModule interface.
 		public override List<FileItem> Compile(ReadOnlyTargetRules Target, UEToolChain ToolChain, CppCompileEnvironment BinaryCompileEnvironment, FileReference SingleFileToCompile, ISourceFileWorkingSet WorkingSet, IActionGraphBuilder Graph)
 		{
-			UEBuildPlatform BuildPlatform = UEBuildPlatform.GetBuildPlatform(BinaryCompileEnvironment.Platform);
+			//UEBuildPlatform BuildPlatform = UEBuildPlatform.GetBuildPlatform(BinaryCompileEnvironment.Platform);
 
 			List<FileItem> LinkInputFiles = base.Compile(Target, ToolChain, BinaryCompileEnvironment, SingleFileToCompile, WorkingSet, Graph);
 
diff --git a/Engine/Source/Programs/UnrealBuildTool/Configuration/UEBuildTarget.cs b/Engine/Source/Programs/UnrealBuildTool/Configuration/UEBuildTarget.cs
--- a/Engine/Source/Programs/UnrealBuildTool/Configuration/UEBuildTarget.cs
+++ b/Engine/Source/Programs/UnrealBuildTool/Configuration/UEBuildTarget.cs
@@ -3211,7 +3211,7 @@ namespace UnrealBuildTool
 				string PluginReferenceChain = String.Format("{0} -> {1}", ReferenceChain, Info.File.GetFileName());
 
 				// Create modules for this plugin
-				UEBuildBinaryType BinaryType = ShouldCompileMonolithic() ? UEBuildBinaryType.StaticLibrary : UEBuildBinaryType.DynamicLinkLibrary;
+				//UEBuildBinaryType BinaryType = ShouldCompileMonolithic() ? UEBuildBinaryType.StaticLibrary : UEBuildBinaryType.DynamicLinkLibrary;
 				if (Info.Descriptor.Modules != null)
 				{
 					foreach (ModuleDescriptor ModuleInfo in Info.Descriptor.Modules)
diff --git a/Engine/Source/Programs/UnrealBuildTool/Executors/SNDBS.cs b/Engine/Source/Programs/UnrealBuildTool/Executors/SNDBS.cs
--- a/Engine/Source/Programs/UnrealBuildTool/Executors/SNDBS.cs
+++ b/Engine/Source/Programs/UnrealBuildTool/Executors/SNDBS.cs
@@ -312,16 +312,16 @@ namespace UnrealBuildTool
 
 				NewProcess.WaitForExit();
 
-				TimeSpan Duration;
-				DateTimeOffset EndTime = DateTimeOffset.Now;
-				if (EndTime == DateTimeOffset.MinValue)
-				{
-					Duration = DateTimeOffset.Now - StartTime;
-				}
-				else
-				{
-					Duration = EndTime - StartTime;
-				}
+				//TimeSpan Duration;
+				//DateTimeOffset EndTime = DateTimeOffset.Now;
+				//if (EndTime == DateTimeOffset.MinValue)
+				//{
+				//	Duration = DateTimeOffset.Now - StartTime;
+				//}
+				//else
+				//{
+				//	Duration = EndTime - StartTime;
+				//}
 
 				DummyActionThread.bComplete = true;
 				int ExitCode = NewProcess.ExitCode;
diff --git a/Engine/Source/Programs/UnrealBuildTool/Executors/XGE.cs b/Engine/Source/Programs/UnrealBuildTool/Executors/XGE.cs
--- a/Engine/Source/Programs/UnrealBuildTool/Executors/XGE.cs
+++ b/Engine/Source/Programs/UnrealBuildTool/Executors/XGE.cs
@@ -139,7 +139,7 @@ namespace UnrealBuildTool
 		}
 
 		// precompile the Regex needed to parse the XGE output (the ones we want are of the form "File (Duration at +time)"
-		private static Regex XGEDurationRegex = new Regex(@"(?<Filename>.*) *\((?<Duration>[0-9:\.]+) at [0-9\+:\.]+\)", RegexOptions.ExplicitCapture);
+		//private static Regex XGEDurationRegex = new Regex(@"(?<Filename>.*) *\((?<Duration>[0-9:\.]+) at [0-9\+:\.]+\)", RegexOptions.ExplicitCapture);
 
 		public static void ExportActions(List<Action> ActionsToExecute)
 		{
diff --git a/Engine/Source/Programs/UnrealBuildTool/Platform/Android/AndroidAARHandler.cs b/Engine/Source/Programs/UnrealBuildTool/Platform/Android/AndroidAARHandler.cs
--- a/Engine/Source/Programs/UnrealBuildTool/Platform/Android/AndroidAARHandler.cs
+++ b/Engine/Source/Programs/UnrealBuildTool/Platform/Android/AndroidAARHandler.cs
@@ -514,7 +514,7 @@ namespace UnrealBuildTool
 				if (!File.Exists(TargetManifestFileName))
 				{
 					Log.TraceInformation("Extracting AAR {0}", BaseName);
-					IEnumerable<string> FileNames = UnzipFiles(Entry.Filename + ".aar", TargetPath);
+					/*IEnumerable<string> FileNames =*/ UnzipFiles(Entry.Filename + ".aar", TargetPath);
 
 					// Must have a src directory (even if empty)
 					string SrcDirectory = Path.Combine(TargetPath, "src");
diff --git a/Engine/Source/Programs/UnrealBuildTool/Platform/Android/AndroidToolChain.cs b/Engine/Source/Programs/UnrealBuildTool/Platform/Android/AndroidToolChain.cs
--- a/Engine/Source/Programs/UnrealBuildTool/Platform/Android/AndroidToolChain.cs
+++ b/Engine/Source/Programs/UnrealBuildTool/Platform/Android/AndroidToolChain.cs
@@ -134,7 +134,7 @@ namespace UnrealBuildTool
 			int RevisionNum = ToolchainInt / 10000;
 			int RevisionMinor = ToolchainInt - (RevisionNum * 10000);
 			int RevisionLetterNum = RevisionMinor / 100;
-			int RevisionBeta = RevisionMinor - (RevisionLetterNum * 100);
+			//int RevisionBeta = RevisionMinor - (RevisionLetterNum * 100);
 			char RevisionLetter = Convert.ToChar('a' + RevisionLetterNum - 1);
 
 			return "r" + RevisionNum + (RevisionLetterNum > 1 ? Char.ToString(RevisionLetter) : "");
@@ -314,12 +314,12 @@ namespace UnrealBuildTool
 			// NDK setup (use no less than 21 for 64-bit targets)
 			int NDKApiLevel32Int = GetNdkApiLevelInt();
 			int NDKApiLevel64Int = NDKApiLevel32Int;
-			string NDKApiLevel32Bit = GetNdkApiLevel();
-			string NDKApiLevel64Bit = NDKApiLevel32Bit;
+			//string NDKApiLevel32Bit = GetNdkApiLevel();
+			//string NDKApiLevel64Bit = NDKApiLevel32Bit;
 			if (NDKApiLevel64Int < 21)
 			{
 				NDKApiLevel64Int = 21;
-				NDKApiLevel64Bit = "android-21";
+				//NDKApiLevel64Bit = "android-21";
 			}
 
 			string GCCToolchainPath = Path.Combine(NDKPath, "toolchains", "llvm", ArchitecturePath);
@@ -1294,7 +1294,7 @@ namespace UnrealBuildTool
 			//string NDKRoot = Environment.GetEnvironmentVariable("NDKROOT").Replace("\\", "/");
 
 			string BasePCHName = "";
-			UEBuildPlatform BuildPlatform = UEBuildPlatform.GetBuildPlatform(CompileEnvironment.Platform);
+			//UEBuildPlatform BuildPlatform = UEBuildPlatform.GetBuildPlatform(CompileEnvironment.Platform);
 			string PCHExtension = ".gch";
 			if (CompileEnvironment.PrecompiledHeaderAction == PrecompiledHeaderAction.Include)
 			{
diff --git a/Engine/Source/Programs/UnrealBuildTool/Platform/Android/UEDeployAndroid.cs b/Engine/Source/Programs/UnrealBuildTool/Platform/Android/UEDeployAndroid.cs
--- a/Engine/Source/Programs/UnrealBuildTool/Platform/Android/UEDeployAndroid.cs
+++ b/Engine/Source/Programs/UnrealBuildTool/Platform/Android/UEDeployAndroid.cs
@@ -837,7 +837,7 @@ namespace UnrealBuildTool
 				if (bUseChangeListAsStoreVersion && IsBuildMachine)
 				{
 					// make sure changelist is cached
-					string EngineVersion = ReadEngineVersion();
+					//string EngineVersion = ReadEngineVersion();
 					
 					int Changelist = 0;
 					if (int.TryParse(EngineChangelist, out Changelist))
@@ -3140,6 +3140,7 @@ namespace UnrealBuildTool
 			{
 				if (!bSkipGradleBuild)
 				{
+					/*
 					IEnumerable<Tuple<string, string>> TargetList = null;
 
 					TargetList = from Arch in Arches
@@ -3149,7 +3150,6 @@ namespace UnrealBuildTool
 					string DestApkDirectory = Path.Combine(ProjectDirectory, "Binaries/Android");
 					string ApkFilename = Path.GetFileNameWithoutExtension(OutputPath).Replace("UE4Game", ProjectName);
 
-					/*
 					foreach (Tuple<string, string> target in TargetList)
 					{
 						string Arch = target.Item1;
@@ -3181,7 +3181,7 @@ namespace UnrealBuildTool
 			}
 
 			// make sure it is cached
-			string EngineVersion = ReadEngineVersion();
+			//string EngineVersion = ReadEngineVersion();
 
 			SetMinimumSDKLevelForGradle();
 
@@ -3204,8 +3204,8 @@ namespace UnrealBuildTool
 			string BuildToolsVersion = GetBuildToolsVersion();
 
 			// cache some tools paths
-			string NDKBuildPath = Environment.ExpandEnvironmentVariables("%NDKROOT%/ndk-build" + (Utils.IsRunningOnMono ? "" : ".cmd"));
-			bool HasNDKPath = File.Exists(NDKBuildPath);
+			//string NDKBuildPath = Environment.ExpandEnvironmentVariables("%NDKROOT%/ndk-build" + (Utils.IsRunningOnMono ? "" : ".cmd"));
+			//bool HasNDKPath = File.Exists(NDKBuildPath);
 
 			// set up some directory info
 			string IntermediateAndroidPath = Path.Combine(ProjectDirectory, "Intermediate", "Android");
@@ -3676,7 +3676,7 @@ namespace UnrealBuildTool
 					string[] CompileSDKLines = CompileSDKMin.Split(new[] { "\r\n", "\r", "\n" }, StringSplitOptions.None);
 					foreach (string CompileLine in CompileSDKLines)
 					{
-						string VersionString = CompileLine.Replace("android-", "");
+						//string VersionString = CompileLine.Replace("android-", "");
 						int VersionInt;
 						if (int.TryParse(CompileLine, out VersionInt))
 						{
diff --git a/Engine/Source/Programs/UnrealBuildTool/Platform/HoloLens/HoloLensDeploy.cs b/Engine/Source/Programs/UnrealBuildTool/Platform/HoloLens/HoloLensDeploy.cs
--- a/Engine/Source/Programs/UnrealBuildTool/Platform/HoloLens/HoloLensDeploy.cs
+++ b/Engine/Source/Programs/UnrealBuildTool/Platform/HoloLens/HoloLensDeploy.cs
@@ -15,8 +15,8 @@ namespace UnrealBuildTool
 	/// </summary>
 	class HoloLensDeploy : UEBuildDeploy
 	{
-		private FileReference MakeAppXPath;
-		private FileReference SignToolPath;
+		//private FileReference MakeAppXPath;
+		//private FileReference SignToolPath;
 
 		/// <summary>
 		/// Utility function to delete a file
@@ -102,7 +102,7 @@ namespace UnrealBuildTool
 					if (FilesInSource.Contains(DestFilename.ToUpperInvariant()) == false)
 					{
 						Log.TraceInformation("Destination file does not exist in Source - DELETING: {0}", DestFile);
-						FileAttributes attributes = File.GetAttributes(DestFile);
+						//FileAttributes attributes = File.GetAttributes(DestFile);
 						try
 						{
 							DeployHelper_DeleteFile(DestFile);
@@ -204,8 +204,8 @@ namespace UnrealBuildTool
 				ConfigDirRef = new DirectoryReference(UnrealBuildTool.GetRemoteIniPath());
 			}
 
-			MakeAppXPath = HoloLensExports.GetWindowsSdkToolPath("makeappx.exe");
-			SignToolPath = HoloLensExports.GetWindowsSdkToolPath("signtool.exe");
+			//MakeAppXPath = HoloLensExports.GetWindowsSdkToolPath("makeappx.exe");
+			//SignToolPath = HoloLensExports.GetWindowsSdkToolPath("signtool.exe");
 
 			return true;
 		}
diff --git a/Engine/Source/Programs/UnrealBuildTool/Platform/HoloLens/HoloLensProjectGenerator.cs b/Engine/Source/Programs/UnrealBuildTool/Platform/HoloLens/HoloLensProjectGenerator.cs
--- a/Engine/Source/Programs/UnrealBuildTool/Platform/HoloLens/HoloLensProjectGenerator.cs
+++ b/Engine/Source/Programs/UnrealBuildTool/Platform/HoloLens/HoloLensProjectGenerator.cs
@@ -101,14 +101,14 @@ namespace UnrealBuildTool
 			// Default to latest supported version.  Game projects can override this later.
 			// Because this property is only required for VS2017 we can safely say that's the compiler version (whether that's actually true
 			// or not)
-			string SDKFolder = "";
+			//string SDKFolder = "";
 			string SDKVersion = "";
 
 			DirectoryReference folder;
 			VersionNumber version;
 			if(WindowsPlatform.TryGetWindowsSdkDir("Latest", out version, out folder))
 			{
-				SDKFolder = folder.FullName;
+				//SDKFolder = folder.FullName;
 				SDKVersion = version.ToString();
 			}
 
diff --git a/Engine/Source/Programs/UnrealBuildTool/Platform/HoloLens/ManifestGenerator.cs b/Engine/Source/Programs/UnrealBuildTool/Platform/HoloLens/ManifestGenerator.cs
--- a/Engine/Source/Programs/UnrealBuildTool/Platform/HoloLens/ManifestGenerator.cs
+++ b/Engine/Source/Programs/UnrealBuildTool/Platform/HoloLens/ManifestGenerator.cs
@@ -31,7 +31,7 @@ namespace UnrealBuildTool
 
 		// Manifest configuration values/paths
 		private List<WinMDRegistrationInfo> WinMDReferences;
-		private UnrealTargetPlatform Platform;
+		//private UnrealTargetPlatform Platform;
 		private WindowsArchitecture Architecture;
 		private string TargetSettings;
 		private string BuildResourceProjectRelativePath;
@@ -653,7 +653,7 @@ namespace UnrealBuildTool
             bool iniUpdated = false;
 
             List<string> projectIniFilePaths = new List<string>();
-            List<ConfigFile> Files = new List<ConfigFile>();
+            //List<ConfigFile> Files = new List<ConfigFile>();
             foreach (FileReference IniFileName in ConfigHierarchy.EnumerateConfigFileLocations(ConfigHierarchyType.Game, IniDirRef, TargetPlatform))
             {
                 // If ini file does not exist, or is not in the project directory, move on to the next one.
@@ -761,7 +761,7 @@ namespace UnrealBuildTool
 			ManifestRelatedFilePaths = new List<string>();
 
 			WinMDReferences = ((InWinMDReferences == null) ? new List<WinMDRegistrationInfo>() : new List<WinMDRegistrationInfo>(InWinMDReferences));
-			Platform = TargetPlatform;
+			//Platform = TargetPlatform;
 			Architecture = TargetArchitecture;
 			TargetSettings = "/Script/HoloLensPlatformEditor.HoloLensTargetSettings";
 			BuildResourceProjectRelativePath = "Build\\HoloLens";
@@ -1132,7 +1132,7 @@ namespace UnrealBuildTool
 			ManifestRelatedFilePaths = new List<string>();
 
 			WinMDReferences = new List<WinMDRegistrationInfo>();
-			Platform = TargetPlatform;
+			//Platform = TargetPlatform;
 			TargetSettings = "/Script/HoloLensPlatformEditor.HoloLensTargetSettings";
 			BuildResourceProjectRelativePath = "Build\\HoloLens";
 			BuildResourceSubPath = "AssetsResources";
@@ -1682,7 +1682,7 @@ namespace UnrealBuildTool
 			{
 				for (int i = 0; i < CulturesToStage.Count; ++i)
 				{
-					string ValueToWrite = string.Empty;
+					//string ValueToWrite = string.Empty;
 					if (!IniLocalizedValues.ContainsKey(CulturesToStage[i]))
 					{
 						PerCultureResourceWriters[i].AddResource(ResourceEntryName, NeutralValue);
diff --git a/Engine/Source/Programs/UnrealBuildTool/Platform/HoloLens/UEBuildHoloLens.cs b/Engine/Source/Programs/UnrealBuildTool/Platform/HoloLens/UEBuildHoloLens.cs
--- a/Engine/Source/Programs/UnrealBuildTool/Platform/HoloLens/UEBuildHoloLens.cs
+++ b/Engine/Source/Programs/UnrealBuildTool/Platform/HoloLens/UEBuildHoloLens.cs
@@ -390,7 +390,7 @@ namespace UnrealBuildTool
 			if (DirectoryReference.Exists(ReferenceDir))
 			{
 				// Prefer a contract from a suitable SDK-versioned subdir of the references folder when available (starts with 15063 SDK)
-				Version WindowsSDKVersionMaxForToolchain = Compiler < WindowsCompiler.VisualStudio2017 ? HoloLens.MaximumSDKVersionForVS2015 : null;
+				//Version WindowsSDKVersionMaxForToolchain = Compiler < WindowsCompiler.VisualStudio2017 ? HoloLens.MaximumSDKVersionForVS2015 : null;
 				DirectoryReference SDKVersionedReferenceDir = DirectoryReference.Combine(ReferenceDir, SDKVersion.ToString());
 				DirectoryReference ContractDir = DirectoryReference.Combine(SDKVersionedReferenceDir, ApiContract);
 				Version ContractLatestVersion = null;
@@ -752,8 +752,8 @@ namespace UnrealBuildTool
 	class HoloLensPlatformSDK : UEBuildPlatformSDK
 	{
 		static bool bIsInstalled = false;
-		static string LatestVersionString = string.Empty;
-		static string InstallLocation = string.Empty;
+		//static string LatestVersionString = string.Empty;
+		//static string InstallLocation = string.Empty;
 
 		static HoloLensPlatformSDK()
 		{
@@ -776,8 +776,8 @@ namespace UnrealBuildTool
 				if (Result != null)
 				{
 					bIsInstalled = true;
-					InstallLocation = (string)Result;
-					LatestVersionString = Microsoft.Win32.Registry.GetValue(@"HKEY_LOCAL_MACHINE\SOFTWARE\Wow6432Node\Microsoft\Microsoft SDKs\Windows\" + Version, "ProductVersion", null) as string;
+					//InstallLocation = (string)Result;
+					//LatestVersionString = Microsoft.Win32.Registry.GetValue(@"HKEY_LOCAL_MACHINE\SOFTWARE\Wow6432Node\Microsoft\Microsoft SDKs\Windows\" + Version, "ProductVersion", null) as string;
 					break;
 				}
 			}
diff --git a/Engine/Source/Programs/UnrealBuildTool/Platform/IOS/IOSToolChain.cs b/Engine/Source/Programs/UnrealBuildTool/Platform/IOS/IOSToolChain.cs
--- a/Engine/Source/Programs/UnrealBuildTool/Platform/IOS/IOSToolChain.cs
+++ b/Engine/Source/Programs/UnrealBuildTool/Platform/IOS/IOSToolChain.cs
@@ -1593,7 +1593,7 @@ namespace UnrealBuildTool
 			string FabricPath = UnrealBuildTool.EngineDirectory + "/Intermediate/UnzippedFrameworks/Crashlytics/Fabric.embeddedframework";
             if (Directory.Exists(FabricPath) && Environment.GetEnvironmentVariable("IsBuildMachine") == "1")
             {
-				string PlistFile = ProjectDir + "/Intermediate/IOS/" + ProjectName + "-Info.plist";
+				//string PlistFile = ProjectDir + "/Intermediate/IOS/" + ProjectName + "-Info.plist";
 				Process FabricProcess = new Process();
 				FabricProcess.StartInfo.WorkingDirectory = Path.GetDirectoryName(DsymZip);
 				FabricProcess.StartInfo.FileName = "/bin/sh";
@@ -2083,7 +2083,7 @@ namespace UnrealBuildTool
 
 				foreach (KeyValuePair<string, DirectoryReference> Pair in Target.FrameworkNameToSourceDir)
 				{
-					string UnpackedZipPath = Pair.Value.FullName;
+					//string UnpackedZipPath = Pair.Value.FullName;
 
 					// For now, this is hard coded, but we need to loop over all modules, and copy bundled assets that need it
 					string LocalDest = LocalFrameworkAssets + "/" + Pair.Key;
diff --git a/Engine/Source/Programs/UnrealBuildTool/Platform/IOS/UEDeployIOS.cs b/Engine/Source/Programs/UnrealBuildTool/Platform/IOS/UEDeployIOS.cs
--- a/Engine/Source/Programs/UnrealBuildTool/Platform/IOS/UEDeployIOS.cs
+++ b/Engine/Source/Programs/UnrealBuildTool/Platform/IOS/UEDeployIOS.cs
@@ -316,7 +316,7 @@ namespace UnrealBuildTool
 				RequiredCaps += "\t\t<string>" + Arches[0] + "</string>\n";
 			}
 
-			ConfigHierarchy GameIni = ConfigCache.ReadHierarchy(ConfigHierarchyType.Game, DirRef, UnrealTargetPlatform.IOS);
+			//ConfigHierarchy GameIni = ConfigCache.ReadHierarchy(ConfigHierarchyType.Game, DirRef, UnrealTargetPlatform.IOS);
 
 			Ini.GetBool("/Script/IOSRuntimeSettings.IOSRuntimeSettings", "bSupportsMetal", out bSupported);
 			RequiredCaps += bSupported ? "\t\t<string>metal</string>\n" : "";
@@ -840,7 +840,7 @@ namespace UnrealBuildTool
 		public virtual bool GeneratePList(FileReference ProjectFile, UnrealTargetConfiguration Config, string ProjectDirectory, bool bIsUE4Game, string GameName, bool bIsClient, string ProjectName, string InEngineDir, string AppDirectory, List<string> UPLScripts, VersionNumber SdkVersion, string BundleID, bool bBuildAsFramework, out bool bSupportsPortrait, out bool bSupportsLandscape, out bool bSkipIcons)
 		{
 			// remember name with -IOS-Shipping, etc
-			string ExeName = GameName;
+			// string ExeName = GameName;
 
 			// strip out the markup
 			GameName = GameName.Split("-".ToCharArray())[0];
diff --git a/Engine/Source/Programs/UnrealBuildTool/Platform/Linux/LinuxToolChain.cs b/Engine/Source/Programs/UnrealBuildTool/Platform/Linux/LinuxToolChain.cs
--- a/Engine/Source/Programs/UnrealBuildTool/Platform/Linux/LinuxToolChain.cs
+++ b/Engine/Source/Programs/UnrealBuildTool/Platform/Linux/LinuxToolChain.cs
@@ -1278,7 +1278,7 @@ namespace UnrealBuildTool
 			string Arguments = GetCLArguments_Global(CompileEnvironment);
 			string PCHArguments = "";
 
-			var BuildPlatform = UEBuildPlatform.GetBuildPlatform(CompileEnvironment.Platform);
+			//var BuildPlatform = UEBuildPlatform.GetBuildPlatform(CompileEnvironment.Platform);
 
 			if (!bHasPrintedBuildDetails)
 			{
diff --git a/Engine/Source/Programs/UnrealBuildTool/Platform/Lumin/LuminToolChain.cs b/Engine/Source/Programs/UnrealBuildTool/Platform/Lumin/LuminToolChain.cs
--- a/Engine/Source/Programs/UnrealBuildTool/Platform/Lumin/LuminToolChain.cs
+++ b/Engine/Source/Programs/UnrealBuildTool/Platform/Lumin/LuminToolChain.cs
@@ -13,12 +13,12 @@ namespace UnrealBuildTool
 {
 	class LuminToolChain : AndroidToolChain
 	{
-		static private Dictionary<string, string[]> AllArchNames = new Dictionary<string, string[]> {
-			{ "-armv7", new string[] { "armv7", "armeabi-v7a", "android-kk-egl-t124-a32" } },
-			{ "-arm64", new string[] { "arm64", "arm64-v8a", "android-L-egl-t132-a64" } },
-		};
+		//static private Dictionary<string, string[]> AllArchNames = new Dictionary<string, string[]> {
+		//	{ "-armv7", new string[] { "armv7", "armeabi-v7a", "android-kk-egl-t124-a32" } },
+		//	{ "-arm64", new string[] { "arm64", "arm64-v8a", "android-L-egl-t132-a64" } },
+		//};
 
-		static private string[] LibrariesToSkip = new string[] { "nvToolsExt", "nvToolsExtStub", "oculus", "vrapi", "ovrkernel", "systemutils", "openglloader", "gpg", };
+		//static private string[] LibrariesToSkip = new string[] { "nvToolsExt", "nvToolsExtStub", "oculus", "vrapi", "ovrkernel", "systemutils", "openglloader", "gpg", };
 
 		private List<string> AdditionalGPUArches;
 
@@ -91,7 +91,7 @@ namespace UnrealBuildTool
 				throw new BuildException("Could not parse clang version from mabu spec '{0}'", MabuSpec);
 			}
 
-			string ClangVersion = string.Format("{0}.{1}", SpecMatch.Groups[1].Value, SpecMatch.Groups[2].Value);
+			//string ClangVersion = string.Format("{0}.{1}", SpecMatch.Groups[1].Value, SpecMatch.Groups[2].Value);
 
 			SetClangVersion(int.Parse(SpecMatch.Groups[1].Value), int.Parse(SpecMatch.Groups[2].Value), 0);
 
diff --git a/Engine/Source/Programs/UnrealBuildTool/Platform/TVOS/UEDeployTVOS.cs b/Engine/Source/Programs/UnrealBuildTool/Platform/TVOS/UEDeployTVOS.cs
--- a/Engine/Source/Programs/UnrealBuildTool/Platform/TVOS/UEDeployTVOS.cs
+++ b/Engine/Source/Programs/UnrealBuildTool/Platform/TVOS/UEDeployTVOS.cs
@@ -42,11 +42,11 @@ namespace UnrealBuildTool
 			string OldPListData = File.Exists(PListFile) ? File.ReadAllText(PListFile) : "";
 
 			// determine if there is a launch.xib
-			string LaunchXib = InEngineDir + "/Build/IOS/Resources/Interface/LaunchScreen.xib";
-			if (File.Exists(BuildDirectory + "/Resources/Interface/LaunchScreen.xib"))
-			{
-				LaunchXib = BuildDirectory + "/Resources/Interface/LaunchScreen.xib";
-			}
+			//string LaunchXib = InEngineDir + "/Build/IOS/Resources/Interface/LaunchScreen.xib";
+			//if (File.Exists(BuildDirectory + "/Resources/Interface/LaunchScreen.xib"))
+			//{
+			//	LaunchXib = BuildDirectory + "/Resources/Interface/LaunchScreen.xib";
+			//}
 
 			// get the settings from the ini file
 			// plist replacements
diff --git a/Engine/Source/Programs/UnrealBuildTool/Platform/Windows/PVSToolChain.cs b/Engine/Source/Programs/UnrealBuildTool/Platform/Windows/PVSToolChain.cs
--- a/Engine/Source/Programs/UnrealBuildTool/Platform/Windows/PVSToolChain.cs
+++ b/Engine/Source/Programs/UnrealBuildTool/Platform/Windows/PVSToolChain.cs
@@ -86,7 +86,7 @@ namespace UnrealBuildTool
 							string[] Tokens = Line.Split(new string[] { "<#~>" }, StringSplitOptions.None);
 							if(Tokens.Length >= 9)
 							{
-								string Trial = Tokens[1];
+								//string Trial = Tokens[1];
 								string LineNumberStr = Tokens[2];
 								string FileName = Tokens[3];
 								string WarningCode = Tokens[5];
diff --git a/Engine/Source/Programs/UnrealBuildTool/Platform/Windows/UEBuildWindows.cs b/Engine/Source/Programs/UnrealBuildTool/Platform/Windows/UEBuildWindows.cs
--- a/Engine/Source/Programs/UnrealBuildTool/Platform/Windows/UEBuildWindows.cs
+++ b/Engine/Source/Programs/UnrealBuildTool/Platform/Windows/UEBuildWindows.cs
@@ -1710,7 +1710,7 @@ namespace UnrealBuildTool
 		public static bool TryGetWindowsSdkDir(string DesiredVersion, out VersionNumber OutSdkVersion, out DirectoryReference OutSdkDir)
 		{
 			// Get a map of Windows SDK versions to their root directories
-			IReadOnlyDictionary<VersionNumber, DirectoryReference> WindowsSdkDirs = FindWindowsSdkDirs();
+			//IReadOnlyDictionary<VersionNumber, DirectoryReference> WindowsSdkDirs = FindWindowsSdkDirs();
 
 			// Figure out which version number to look for
 			VersionNumber WindowsSdkVersion = null;
diff --git a/Engine/Source/Programs/UnrealBuildTool/Platform/Windows/VCToolChain.cs b/Engine/Source/Programs/UnrealBuildTool/Platform/Windows/VCToolChain.cs
--- a/Engine/Source/Programs/UnrealBuildTool/Platform/Windows/VCToolChain.cs
+++ b/Engine/Source/Programs/UnrealBuildTool/Platform/Windows/VCToolChain.cs
@@ -1253,7 +1253,7 @@ namespace UnrealBuildTool
 			{
 				CompileAction.TimingFile = FileItem.GetItemByFileReference(FileReference.Combine(OutputDir, String.Format("{0}.timing.txt", SourceFile.Location.GetFileName())));
 				CompileAction.ProducedItems.Add(CompileAction.TimingFile);
-				Action ParseTimingInfoAction = GenerateParseTimingInfoAction(SourceFile, CompileAction, Result, Graph);
+				/*Action ParseTimingInfoAction = */GenerateParseTimingInfoAction(SourceFile, CompileAction, Result, Graph);
 				CommandArguments.Add(string.Format("-timing={0}", Utils.MakePathSafeToUseWithCommandLine(CompileAction.TimingFile.Location)));
 			}
 
diff --git a/Engine/Source/Programs/UnrealBuildTool/ProjectFiles/CMake/CMakefileGenerator.cs b/Engine/Source/Programs/UnrealBuildTool/ProjectFiles/CMake/CMakefileGenerator.cs
--- a/Engine/Source/Programs/UnrealBuildTool/ProjectFiles/CMake/CMakefileGenerator.cs
+++ b/Engine/Source/Programs/UnrealBuildTool/ProjectFiles/CMake/CMakefileGenerator.cs
@@ -325,7 +325,7 @@ namespace UnrealBuildTool
 			string ProjectHeadersFilePath = FileReference.Combine(IntermediateProjectFilesPath, CMakeProjectHeadersFileName).ToNormalizedPath();
 			string EngineSourcesFilePath = FileReference.Combine(IntermediateProjectFilesPath, CMakeEngineSourcesFileName).ToNormalizedPath();
 			string ProjectSourcesFilePath = FileReference.Combine(IntermediateProjectFilesPath, CMakeProjectSourcesFileName).ToNormalizedPath();
-			string ProjectFilePath = FileReference.Combine(IntermediateProjectFilesPath, CMakeProjectSourcesFileName).ToNormalizedPath();
+			//string ProjectFilePath = FileReference.Combine(IntermediateProjectFilesPath, CMakeProjectSourcesFileName).ToNormalizedPath();
 			string IncludeFilePath = FileReference.Combine(IntermediateProjectFilesPath, CMakeIncludesFileName).ToNormalizedPath();
 			string EngineConfigsFilePath = FileReference.Combine(IntermediateProjectFilesPath, CMakeEngineConfigsFileName).ToNormalizedPath();
 			string ProjectConfigsFilePath = FileReference.Combine(IntermediateProjectFilesPath, CMakeProjectConfigsFileName).ToNormalizedPath();
diff --git a/Engine/Source/Programs/UnrealBuildTool/ProjectFiles/CodeLite/CodeLiteProject.cs b/Engine/Source/Programs/UnrealBuildTool/ProjectFiles/CodeLite/CodeLiteProject.cs
--- a/Engine/Source/Programs/UnrealBuildTool/ProjectFiles/CodeLite/CodeLiteProject.cs
+++ b/Engine/Source/Programs/UnrealBuildTool/ProjectFiles/CodeLite/CodeLiteProject.cs
@@ -38,7 +38,7 @@ namespace UnrealBuildTool
 		{
 			bool bSuccess = false;
 			string ProjectNameRaw = ProjectFilePath.GetFileNameWithoutExtension();
-			string ProjectPath = ProjectFilePath.FullName;
+			//string ProjectPath = ProjectFilePath.FullName;
 			string ProjectExtension = ProjectFilePath.GetExtension();
 			string ProjectPlatformName = BuildHostPlatform.Current.Platform.ToString();
 
diff --git a/Engine/Source/Programs/UnrealBuildTool/ProjectFiles/Eddie/EddieProject.cs b/Engine/Source/Programs/UnrealBuildTool/ProjectFiles/Eddie/EddieProject.cs
--- a/Engine/Source/Programs/UnrealBuildTool/ProjectFiles/Eddie/EddieProject.cs
+++ b/Engine/Source/Programs/UnrealBuildTool/ProjectFiles/Eddie/EddieProject.cs
@@ -45,14 +45,14 @@ namespace UnrealBuildTool
 	
 	class EddieProjectFile : ProjectFile
 	{
-		FileReference OnlyGameProject;
+		//FileReference OnlyGameProject;
 		
 		Dictionary<string, EddieFolder> Folders = new Dictionary<string, EddieFolder>();
 		
 		public EddieProjectFile(FileReference InitFilePath, FileReference InOnlyGameProject)
 			: base(InitFilePath)
 		{
-			OnlyGameProject = InOnlyGameProject;
+			//OnlyGameProject = InOnlyGameProject;
 		}
 
 		public override string ToString()
@@ -123,9 +123,9 @@ namespace UnrealBuildTool
 			{
 				EddieSourceFile SourceFile = CurSourceFile as EddieSourceFile;
 				string FileName = SourceFile.Reference.GetFileName();
-				string FileExtension = Path.GetExtension(FileName);
-				string FilePath = SourceFile.Reference.MakeRelativeTo(ProjectFilePath.Directory);
-				string FilePathMac = Utils.CleanDirectorySeparators(FilePath, '/');
+				//string FileExtension = Path.GetExtension(FileName);
+				//string FilePath = SourceFile.Reference.MakeRelativeTo(ProjectFilePath.Directory);
+				//string FilePathMac = Utils.CleanDirectorySeparators(FilePath, '/');
 				
 				string ProjectRelativeSourceFile = CurSourceFile.Reference.MakeRelativeTo(ProjectFilePath.Directory);
 				string RelativeSourceDirectory = Path.GetDirectoryName(ProjectRelativeSourceFile);
@@ -199,17 +199,17 @@ namespace UnrealBuildTool
 		{
 			bool bSuccess = false;
 			
-			string TargetName = ProjectFilePath.GetFileNameWithoutExtension();
+			//string TargetName = ProjectFilePath.GetFileNameWithoutExtension();
 			
-			FileReference GameProjectPath = null;
-			foreach(ProjectTarget Target in ProjectTargets)
-			{
-				if(Target.UnrealProjectFilePath != null)
-				{
-					GameProjectPath = Target.UnrealProjectFilePath;
-					break;
-				}
-			}
+			//FileReference GameProjectPath = null;
+			//foreach(ProjectTarget Target in ProjectTargets)
+			//{
+			//	if(Target.UnrealProjectFilePath != null)
+			//	{
+			//		GameProjectPath = Target.UnrealProjectFilePath;
+			//		break;
+			//	}
+			//}
 			
 			StringBuilder ProjectFileContent = new StringBuilder();
 			
@@ -225,4 +225,4 @@ namespace UnrealBuildTool
 		}
 	}
 
-}
\ No newline at end of file
+}
diff --git a/Engine/Source/Programs/UnrealBuildTool/ProjectFiles/ProjectFileGenerator.cs b/Engine/Source/Programs/UnrealBuildTool/ProjectFiles/ProjectFileGenerator.cs
--- a/Engine/Source/Programs/UnrealBuildTool/ProjectFiles/ProjectFileGenerator.cs
+++ b/Engine/Source/Programs/UnrealBuildTool/ProjectFiles/ProjectFileGenerator.cs
@@ -2220,7 +2220,7 @@ namespace UnrealBuildTool
 			ProgramProjects = new Dictionary<FileReference, ProjectFile>();
 
 			// Get some standard directories
-			DirectoryReference EngineSourceProgramsDirectory = DirectoryReference.Combine(UnrealBuildTool.EngineSourceDirectory, "Programs");
+			//DirectoryReference EngineSourceProgramsDirectory = DirectoryReference.Combine(UnrealBuildTool.EngineSourceDirectory, "Programs");
 			DirectoryReference EnterpriseSourceProgramsDirectory = DirectoryReference.Combine(UnrealBuildTool.EnterpriseSourceDirectory, "Programs");
 
 			foreach( FileReference TargetFilePath in AllTargetFiles )
diff --git a/Engine/Source/Programs/UnrealBuildTool/ProjectFiles/VisualStudio/VCProject.cs b/Engine/Source/Programs/UnrealBuildTool/ProjectFiles/VisualStudio/VCProject.cs
--- a/Engine/Source/Programs/UnrealBuildTool/ProjectFiles/VisualStudio/VCProject.cs
+++ b/Engine/Source/Programs/UnrealBuildTool/ProjectFiles/VisualStudio/VCProject.cs
@@ -253,7 +253,7 @@ namespace UnrealBuildTool
 
 	class VCProjectFile : MSBuildProjectFile
 	{
-		FileReference OnlyGameProject;
+		//FileReference OnlyGameProject;
 		VCProjectFileFormat ProjectFileFormat;
 		bool bUseFastPDB;
 		bool bUsePerFileIntellisense;
@@ -289,7 +289,7 @@ namespace UnrealBuildTool
 		public VCProjectFile(FileReference InFilePath, FileReference InOnlyGameProject, VCProjectFileFormat InProjectFileFormat, bool bUseFastPDB, bool bUsePerFileIntellisense, bool bUsePrecompiled, bool bEditorDependsOnShaderCompileWorker, bool bBuildLiveCodingConsole, string BuildToolOverride)
 			: base(InFilePath)
 		{
-			OnlyGameProject = InOnlyGameProject;
+			//OnlyGameProject = InOnlyGameProject;
 			ProjectFileFormat = InProjectFileFormat;
 			this.bUseFastPDB = bUseFastPDB;
 			this.bUsePerFileIntellisense = bUsePerFileIntellisense;
@@ -675,7 +675,7 @@ namespace UnrealBuildTool
 		/// <returns>True on success</returns>
 		public override List<Tuple<ProjectFile, string>> WriteDebugProjectFiles(List<UnrealTargetPlatform> InPlatforms, List<UnrealTargetConfiguration> InConfigurations, PlatformProjectGeneratorCollection PlatformProjectGenerators)
 		{
-			string ProjectName = ProjectFilePath.GetFileNameWithoutExtension();
+			//string ProjectName = ProjectFilePath.GetFileNameWithoutExtension();
 
 			List<UnrealTargetPlatform> ProjectPlatforms = new List<UnrealTargetPlatform>();
 			List<Tuple<ProjectFile, string>> ProjectFiles = new List<Tuple<ProjectFile, string>>();
@@ -1005,7 +1005,7 @@ namespace UnrealBuildTool
 
 				// Add all file directories to the filters file as solution filters
 				HashSet<string> FilterDirectories = new HashSet<string>();
-				UEBuildPlatform BuildPlatform = UEBuildPlatform.GetBuildPlatform(BuildHostPlatform.Current.Platform);
+				//UEBuildPlatform BuildPlatform = UEBuildPlatform.GetBuildPlatform(BuildHostPlatform.Current.Platform);
 
 				Dictionary<DirectoryReference, string> DirectoryToForceIncludePaths = new Dictionary<DirectoryReference, string>();
 				foreach (AliasedFile AliasedFile in LocalAliasedFiles)
diff --git a/Engine/Source/Programs/UnrealBuildTool/ProjectFiles/VisualStudio/VCProjectFileGenerator.cs b/Engine/Source/Programs/UnrealBuildTool/ProjectFiles/VisualStudio/VCProjectFileGenerator.cs
--- a/Engine/Source/Programs/UnrealBuildTool/ProjectFiles/VisualStudio/VCProjectFileGenerator.cs
+++ b/Engine/Source/Programs/UnrealBuildTool/ProjectFiles/VisualStudio/VCProjectFileGenerator.cs
@@ -532,7 +532,7 @@ namespace UnrealBuildTool
 
 
 				// Project files
-				List<MSBuildProjectFile> AllProjectFilesSorted = AllProjectFiles.OrderBy((ProjFile) => ProjFile.ProjectFilePath.GetFileNameWithoutExtension()).Cast<MSBuildProjectFile>().ToList();
+				//List<MSBuildProjectFile> AllProjectFilesSorted = AllProjectFiles.OrderBy((ProjFile) => ProjFile.ProjectFilePath.GetFileNameWithoutExtension()).Cast<MSBuildProjectFile>().ToList();
 				foreach (MSBuildProjectFile CurProject in AllProjectFiles)
 				{
 					// Visual Studio uses different GUID types depending on the project type
@@ -874,7 +874,7 @@ namespace UnrealBuildTool
 				UnrealVSContent.Append("</UnrealVS>" + ProjectFileGenerator.NewLine );
 
 				string ConfigFilePath = FileReference.Combine(IntermediateProjectFilesPath, "UnrealVS.xml").FullName;
-				bool bSuccess = ProjectFileGenerator.WriteFileIfChanged(ConfigFilePath, UnrealVSContent.ToString());
+				/* bool bSuccess = */ ProjectFileGenerator.WriteFileIfChanged(ConfigFilePath, UnrealVSContent.ToString());
 			}
 		}
 
diff --git a/Engine/Source/Programs/UnrealBuildTool/ProjectFiles/VisualStudio/VCSolutionOptions.cs b/Engine/Source/Programs/UnrealBuildTool/ProjectFiles/VisualStudio/VCSolutionOptions.cs
--- a/Engine/Source/Programs/UnrealBuildTool/ProjectFiles/VisualStudio/VCSolutionOptions.cs
+++ b/Engine/Source/Programs/UnrealBuildTool/ProjectFiles/VisualStudio/VCSolutionOptions.cs
@@ -368,8 +368,8 @@ namespace UnrealBuildTool
 			if(Format >= VCProjectFileFormat.VisualStudio2017)
 			{
 				int Header1 = Reader.ReadInt32();
-				int Header2 = Reader.ReadInt32();
-				int Header3 = Reader.ReadByte();
+				/*int Header2 =*/ Reader.ReadInt32();
+				/*int Header3 =*/ Reader.ReadByte();
 				if (Header1 != 15)
 				{
 					throw new BuildException("Unexpected data in open projects section");
diff --git a/Engine/Source/Programs/UnrealBuildTool/ProjectFiles/VisualStudioCode/VSCodeProjectFileGenerator.cs b/Engine/Source/Programs/UnrealBuildTool/ProjectFiles/VisualStudioCode/VSCodeProjectFileGenerator.cs
--- a/Engine/Source/Programs/UnrealBuildTool/ProjectFiles/VisualStudioCode/VSCodeProjectFileGenerator.cs
+++ b/Engine/Source/Programs/UnrealBuildTool/ProjectFiles/VisualStudioCode/VSCodeProjectFileGenerator.cs
@@ -678,7 +678,7 @@ namespace UnrealBuildTool
 						string Command = BaseCommand == "Rebuild" ? "Build" : BaseCommand;
 						string TaskName = String.Format("{0} {1} {2} {3}", Target.Name, BuildProduct.Platform.ToString(), BuildProduct.Config, BaseCommand);
 						string CleanTaskName = String.Format("{0} {1} {2} {3}", Target.Name, BuildProduct.Platform.ToString(), BuildProduct.Config, "Clean");
-						List<string> ExtraParams = new List<string>();
+						//List<string> ExtraParams = new List<string>();
 
 						OutFile.BeginObject();
 						{
@@ -901,9 +901,9 @@ namespace UnrealBuildTool
 			FileReference TargetFilePath = Target.TargetFilePath;
 			string TargetName = TargetFilePath == null ? Project.ProjectFilePath.GetFileNameWithoutExtension() : TargetFilePath.GetFileNameWithoutAnyExtensions();
 			string UBTPlatformName = Platform.ToString();
-			string UBTConfigurationName = Configuration.ToString();
+			//string UBTConfigurationName = Configuration.ToString();
 
-			string ProjectName = Project.ProjectFilePath.GetFileNameWithoutExtension();
+			//string ProjectName = Project.ProjectFilePath.GetFileNameWithoutExtension();
 
 			// Setup output path
 			UEBuildPlatform BuildPlatform = UEBuildPlatform.GetBuildPlatform(Platform);
@@ -1358,4 +1358,4 @@ namespace UnrealBuildTool
 			}
 		}
 	}
-}
\ No newline at end of file
+}
diff --git a/Engine/Source/Programs/UnrealBuildTool/ProjectFiles/Xcode/XcodeProject.cs b/Engine/Source/Programs/UnrealBuildTool/ProjectFiles/Xcode/XcodeProject.cs
--- a/Engine/Source/Programs/UnrealBuildTool/ProjectFiles/Xcode/XcodeProject.cs
+++ b/Engine/Source/Programs/UnrealBuildTool/ProjectFiles/Xcode/XcodeProject.cs
@@ -124,7 +124,7 @@ namespace UnrealBuildTool
 
 	class XcodeProjectFile : ProjectFile
 	{
-		FileReference OnlyGameProject;
+		//FileReference OnlyGameProject;
 
 		Dictionary<string, XcodeFileGroup> Groups = new Dictionary<string, XcodeFileGroup>();
 
@@ -138,7 +138,7 @@ namespace UnrealBuildTool
 		public XcodeProjectFile(FileReference InitFilePath, FileReference InOnlyGameProject, bool IsForDistribution, string BundleID)
 			: base(InitFilePath)
 		{
-			OnlyGameProject = InOnlyGameProject;
+			//OnlyGameProject = InOnlyGameProject;
 			bForDistribution = IsForDistribution;
 			BundleIdentifier = BundleID;
 		}
@@ -355,7 +355,7 @@ namespace UnrealBuildTool
 			if (UProjectPath != null)
 			{
 				string ProjectExtensionsDir = Path.Combine(Path.GetDirectoryName(UProjectPath.FullName), "Build/IOS/Extensions");
-				string ProjectIntermediateDir = Path.Combine(Path.GetDirectoryName(UProjectPath.FullName), "Intermediate/IOS/Extensions");
+				//string ProjectIntermediateDir = Path.Combine(Path.GetDirectoryName(UProjectPath.FullName), "Intermediate/IOS/Extensions");
 
 				if (Directory.Exists(ProjectExtensionsDir))
 				{
@@ -764,11 +764,11 @@ namespace UnrealBuildTool
 					// look at each zip
 					foreach (FileInfo FI in new System.IO.DirectoryInfo(FrameworkDir.FullName).EnumerateFiles("*.framework.zip"))
 					{
-						string Guid = XcodeProjectFileGenerator.MakeXcodeGuid();
-						string RefGuid = XcodeProjectFileGenerator.MakeXcodeGuid();
+						//string Guid = XcodeProjectFileGenerator.MakeXcodeGuid();
+						//string RefGuid = XcodeProjectFileGenerator.MakeXcodeGuid();
 
 						// for FI of foo.framework.zip, this will give us foo.framework
-						string Framework = Path.GetFileNameWithoutExtension(FI.FullName);
+						//string Framework = Path.GetFileNameWithoutExtension(FI.FullName);
 
 						// unzip the framework right into the .app
 						FrameworkScript.AppendFormat("\\techo Unzipping {0}...\\n", FI.FullName);
@@ -842,7 +842,7 @@ namespace UnrealBuildTool
 			{
 				TVOSPlatform TVOSPlatform = ((TVOSPlatform)UEBuildPlatform.GetBuildPlatform(UnrealTargetPlatform.TVOS));
 				TVOSProjectSettings ProjectSettings = TVOSPlatform.ReadProjectSettings(ProjectFile);
-				TVOSProvisioningData ProvisioningData = TVOSPlatform.ReadProvisioningData(ProjectSettings, bForDistribution);
+				//TVOSProvisioningData ProvisioningData = TVOSPlatform.ReadProvisioningData(ProjectSettings, bForDistribution);
 				bAutomaticSigning = ProjectSettings.bAutomaticSigning;
 			}
 
@@ -1136,7 +1136,7 @@ namespace UnrealBuildTool
 			Content.Append("\t\t\tbuildSettings = {" + ProjectFileGenerator.NewLine);
 
 			string UE4Dir = ConvertPath(Path.GetFullPath(Directory.GetCurrentDirectory() + "../../.."));
-			string MacExecutableDir = ConvertPath(Config.MacExecutablePath.Directory.FullName);
+			//string MacExecutableDir = ConvertPath(Config.MacExecutablePath.Directory.FullName);
 			string MacExecutableFileName = Config.MacExecutablePath.GetFileName();
 
 			string IOSRunTimeVersion, TVOSRunTimeVersion;
@@ -1424,7 +1424,7 @@ namespace UnrealBuildTool
 		{
 			List<XcodeBuildConfig> BuildConfigs = new List<XcodeBuildConfig>();
 
-			string ProjectName = ProjectFilePath.GetFileNameWithoutExtension();
+			//string ProjectName = ProjectFilePath.GetFileNameWithoutExtension();
 
 			foreach (UnrealTargetConfiguration Configuration in Configurations)
 			{
@@ -1500,7 +1500,7 @@ namespace UnrealBuildTool
                                             {
                                                 string IOSExecutableName = MakeExecutableFileName(ExeName, UnrealTargetPlatform.IOS, Configuration, ProjectTarget.TargetRules.Architecture, ProjectTarget.TargetRules.UndecoratedConfiguration);
                                                 string TVOSExecutableName = IOSExecutableName.Replace("-IOS-", "-TVOS-");
-                                                string MacExecutableName = IOSExecutableName.Replace("-IOS-", "-Mac-");
+                                                //string MacExecutableName = IOSExecutableName.Replace("-IOS-", "-Mac-");
                                                 BuildConfigs.Add(new XcodeBuildConfig(ConfigName, TargetName, FileReference.Combine(OutputDirectory, "Mac", IOSExecutableName), FileReference.Combine(OutputDirectory, "IOS", IOSExecutableName), FileReference.Combine(OutputDirectory, "TVOS", TVOSExecutableName), ProjectTarget, Configuration));
                                             }
                                         }
diff --git a/Engine/Source/Programs/UnrealBuildTool/System/Plugins.cs b/Engine/Source/Programs/UnrealBuildTool/System/Plugins.cs
--- a/Engine/Source/Programs/UnrealBuildTool/System/Plugins.cs
+++ b/Engine/Source/Programs/UnrealBuildTool/System/Plugins.cs
@@ -369,7 +369,7 @@ namespace UnrealBuildTool
 			Parent.ChildFiles.Add(Child.File);
 
 			// this should cause an error if it's invalid platform name
-			UnrealTargetPlatform Platform = UnrealTargetPlatform.Parse(PlatformName);
+			//UnrealTargetPlatform Platform = UnrealTargetPlatform.Parse(PlatformName);
 
 			// merge the supported platforms
 			if (Child.Descriptor.SupportedTargetPlatforms != null)
diff --git a/Engine/Source/Programs/UnrealBuildTool/System/RulesCompiler.cs b/Engine/Source/Programs/UnrealBuildTool/System/RulesCompiler.cs
--- a/Engine/Source/Programs/UnrealBuildTool/System/RulesCompiler.cs
+++ b/Engine/Source/Programs/UnrealBuildTool/System/RulesCompiler.cs
@@ -414,7 +414,7 @@ namespace UnrealBuildTool
 				{
 					RulesScope EnterpriseScope = new RulesScope("Enterprise", EngineAssembly.Scope);
 
-					List<DirectoryReference> EnterpriseDirectories = new List<DirectoryReference>() { UnrealBuildTool.EnterpriseDirectory };
+					//List<DirectoryReference> EnterpriseDirectories = new List<DirectoryReference>() { UnrealBuildTool.EnterpriseDirectory };
 
 					IReadOnlyList<PluginInfo> IncludedPlugins = Plugins.ReadEnterprisePlugins(UnrealBuildTool.EnterpriseDirectory);
 					EnterpriseRulesAssembly = CreateEngineOrEnterpriseRulesAssembly(EnterpriseScope, new List<DirectoryReference>() { UnrealBuildTool.EnterpriseDirectory }, 
@@ -546,7 +546,7 @@ namespace UnrealBuildTool
 				}
 
 				DirectoryReference MainProjectDirectory = ProjectFileName.Directory;
-				DirectoryReference MainProjectSourceDirectory = DirectoryReference.Combine(MainProjectDirectory, "Source");
+				//DirectoryReference MainProjectSourceDirectory = DirectoryReference.Combine(MainProjectDirectory, "Source");
 
 				// Create a scope for things in this assembly
 				RulesScope Scope = new RulesScope("Project", Parent.Scope);
diff --git a/Engine/Source/Programs/UnrealBuildTool/System/Unity.cs b/Engine/Source/Programs/UnrealBuildTool/System/Unity.cs
--- a/Engine/Source/Programs/UnrealBuildTool/System/Unity.cs
+++ b/Engine/Source/Programs/UnrealBuildTool/System/Unity.cs
@@ -168,7 +168,7 @@ namespace UnrealBuildTool
 		{
 			List<FileItem> NewCPPFiles = new List<FileItem>();
 
-			UEBuildPlatform BuildPlatform = UEBuildPlatform.GetBuildPlatform(CompileEnvironment.Platform);
+			//UEBuildPlatform BuildPlatform = UEBuildPlatform.GetBuildPlatform(CompileEnvironment.Platform);
 
 			// Figure out size of all input files combined. We use this to determine whether to use larger unity threshold or not.
 			long TotalBytesInCPPFiles = CPPFiles.Sum(F => F.Length);
