Fix Windows UBT build break caused by 12224860

diff --git a/Engine/Source/Programs/UnrealBuildTool/Platform/Windows/UEBuildWindows.cs b/Engine/Source/Programs/UnrealBuildTool/Platform/Windows/UEBuildWindows.cs
--- a/Engine/Source/Programs/UnrealBuildTool/Platform/Windows/UEBuildWindows.cs
+++ b/Engine/Source/Programs/UnrealBuildTool/Platform/Windows/UEBuildWindows.cs
@@ -1710,7 +1710,7 @@ namespace UnrealBuildTool
 		public static bool TryGetWindowsSdkDir(string DesiredVersion, out VersionNumber OutSdkVersion, out DirectoryReference OutSdkDir)
 		{
 			// Get a map of Windows SDK versions to their root directories
-			//IReadOnlyDictionary<VersionNumber, DirectoryReference> WindowsSdkDirs = FindWindowsSdkDirs();
+			/*IReadOnlyDictionary<VersionNumber, DirectoryReference> WindowsSdkDirs =*/ FindWindowsSdkDirs();
 
 			// Figure out which version number to look for
 			VersionNumber WindowsSdkVersion = null;
