UBT: Ignore any targets that can't be built on the current host platform when generating intellisense data for project files.

diff --git a/Engine/Source/Programs/UnrealBuildTool/ProjectFiles/ProjectFileGenerator.cs b/Engine/Source/Programs/UnrealBuildTool/ProjectFiles/ProjectFileGenerator.cs
--- a/Engine/Source/Programs/UnrealBuildTool/ProjectFiles/ProjectFileGenerator.cs
+++ b/Engine/Source/Programs/UnrealBuildTool/ProjectFiles/ProjectFileGenerator.cs
@@ -1823,6 +1823,13 @@ namespace UnrealBuildTool
 						ProjectFile TargetProjectFile = Targets[ TargetIndex ].Item1;
 						ProjectTarget CurTarget = Targets[ TargetIndex ].Item2;
 
+						// Ignore projects for platforms we can't build on this host
+						UnrealTargetPlatform IntellisensePlatform = BuildHostPlatform.Current.Platform;
+						if (!CurTarget.SupportedPlatforms.Any(x => x == IntellisensePlatform))
+						{
+							continue;
+						}
+
 						Log.TraceVerbose("Found target: " + CurTarget.Name);
 
 						List<string> NewArguments = new List<string>(Arguments.Length + 4);
@@ -1835,10 +1842,10 @@ namespace UnrealBuildTool
 						try
 						{
 							// Get the architecture from the target platform
-							string DefaultArchitecture = UEBuildPlatform.GetBuildPlatform(BuildHostPlatform.Current.Platform).GetDefaultArchitecture(CurTarget.UnrealProjectFilePath);
+							string DefaultArchitecture = UEBuildPlatform.GetBuildPlatform(IntellisensePlatform).GetDefaultArchitecture(CurTarget.UnrealProjectFilePath);
 
 							// Create the target descriptor
-							TargetDescriptor TargetDesc = new TargetDescriptor(CurTarget.UnrealProjectFilePath, CurTarget.Name, BuildHostPlatform.Current.Platform, UnrealTargetConfiguration.Development, DefaultArchitecture, new CommandLineArguments(NewArguments.ToArray()));
+							TargetDescriptor TargetDesc = new TargetDescriptor(CurTarget.UnrealProjectFilePath, CurTarget.Name, IntellisensePlatform, UnrealTargetConfiguration.Development, DefaultArchitecture, new CommandLineArguments(NewArguments.ToArray()));
 
 							// Create the target
 							UEBuildTarget Target = UEBuildTarget.Create(TargetDesc, false, bUsePrecompiled);
