Fix broken changelist as version

#ue4
#android

diff --git a/Engine/Source/Programs/UnrealBuildTool/Platform/Android/UEDeployAndroid.cs b/Engine/Source/Programs/UnrealBuildTool/Platform/Android/UEDeployAndroid.cs
--- a/Engine/Source/Programs/UnrealBuildTool/Platform/Android/UEDeployAndroid.cs
+++ b/Engine/Source/Programs/UnrealBuildTool/Platform/Android/UEDeployAndroid.cs
@@ -836,9 +836,13 @@ namespace UnrealBuildTool
 				// override store version with changelist if enabled and is build machine
 				if (bUseChangeListAsStoreVersion && IsBuildMachine)
 				{
-					// make sure changelist is cached
-					//string EngineVersion = ReadEngineVersion();
-					
+					// make sure changelist is cached (clear unused warning)
+					string EngineVersion = ReadEngineVersion();
+					if (EngineVersion == null)
+					{
+						throw new BuildException("No engine version!");
+					}
+
 					int Changelist = 0;
 					if (int.TryParse(EngineChangelist, out Changelist))
 					{
@@ -3180,8 +3184,12 @@ namespace UnrealBuildTool
 				throw new BuildException("Cannot make APK with UPL errors");
 			}
 
-			// make sure it is cached
-			//string EngineVersion = ReadEngineVersion();
+			// make sure it is cached (clear unused warning)
+			string EngineVersion = ReadEngineVersion();
+			if (EngineVersion == null)
+			{
+				throw new BuildException("No engine version!");
+			}
 
 			SetMinimumSDKLevelForGradle();
 
