Fix arm64 non-debug flags for -fdata-sections

#ue4
#android

diff --git a/Engine/Source/Programs/UnrealBuildTool/Platform/Android/AndroidToolChain.cs b/Engine/Source/Programs/UnrealBuildTool/Platform/Android/AndroidToolChain.cs
--- a/Engine/Source/Programs/UnrealBuildTool/Platform/Android/AndroidToolChain.cs
+++ b/Engine/Source/Programs/UnrealBuildTool/Platform/Android/AndroidToolChain.cs
@@ -704,7 +704,7 @@ namespace UnrealBuildTool
 				if (CompileEnvironment.Configuration != CppConfiguration.Debug)
 				{
 					Result += " -ffunction-sections";   // Places each function in its own section of the output file, linker may be able to perform opts to improve locality of reference
-					Result += " -fdata-sections"; // Places each data item in its own section of the output file, linker may be able to perform opts to improve locality of reference
+					Result += " -fdata-sections";		// Places each data item in its own section of the output file, linker may be able to perform opts to improve locality of reference
 				}
 
 				Result += " -fsigned-char";             // Treat chars as signed //@todo android: any concerns about ABI compatibility with libs here?
@@ -722,10 +722,21 @@ namespace UnrealBuildTool
 				//Result += " -mfloat-abi=softfp";
 				//Result += " -mfpu=vfpv3-d16";			//@todo android: UE3 was just vfp. arm7a should all support v3 with 16 registers
 
+				// Add flags for on-device debugging
+				if (CompileEnvironment.Configuration == CppConfiguration.Debug)
+				{
+					Result += " -fno-omit-frame-pointer";   // Disable removing the save/restore frame pointer for better debugging
+					if (CompilerVersionGreaterOrEqual(3, 6, 0))
+					{
+						Result += " -fno-function-sections";    // Improve breakpoint location
+					}
+				}
+
 				// Some switches interfere with on-device debugging
 				if (CompileEnvironment.Configuration != CppConfiguration.Debug)
 				{
 					Result += " -ffunction-sections";   // Places each function in its own section of the output file, linker may be able to perform opts to improve locality of reference
+					Result += " -fdata-sections";		// Places each data item in its own section of the output file, linker may be able to perform opts to improve locality of reference
 				}
 
 				Result += " -fsigned-char";             // Treat chars as signed //@todo android: any concerns about ABI compatibility with libs here?
