- Another android link attempt, this one linked in all 3 configs on EC
#rnx

diff --git a/Engine/Source/Programs/UnrealBuildTool/Platform/Android/AndroidToolChain.cs b/Engine/Source/Programs/UnrealBuildTool/Platform/Android/AndroidToolChain.cs
--- a/Engine/Source/Programs/UnrealBuildTool/Platform/Android/AndroidToolChain.cs
+++ b/Engine/Source/Programs/UnrealBuildTool/Platform/Android/AndroidToolChain.cs
@@ -408,6 +408,13 @@ namespace UnrealBuildTool
 			return Ini.GetBool("/Script/AndroidRuntimeSettings.AndroidRuntimeSettings", "bForceLDLinker", out bForceLDLinker) && bForceLDLinker;
 		}
 
+		private bool DisableFunctionDataSections()
+		{
+			ConfigHierarchy Ini = ConfigCache.ReadHierarchy(ConfigHierarchyType.Engine, DirectoryReference.FromFile(ProjectFile), UnrealTargetPlatform.Android);
+			bool bDisableFunctionDataSections = false;
+			return Ini.GetBool("/Script/AndroidRuntimeSettings.AndroidRuntimeSettings", "bDisableFunctionDataSections", out bDisableFunctionDataSections) && bDisableFunctionDataSections;
+		}
+
 		public override void SetUpGlobalEnvironment(ReadOnlyTargetRules Target)
 		{
 			base.SetUpGlobalEnvironment(Target);
@@ -701,7 +708,7 @@ namespace UnrealBuildTool
 				}
 
 				// Some switches interfere with on-device debugging
-				if (CompileEnvironment.Configuration != CppConfiguration.Debug)
+				if (CompileEnvironment.Configuration != CppConfiguration.Debug && !DisableFunctionDataSections())
 				{
 					Result += " -ffunction-sections";   // Places each function in its own section of the output file, linker may be able to perform opts to improve locality of reference
 					Result += " -fdata-sections";		// Places each data item in its own section of the output file, linker may be able to perform opts to improve locality of reference
@@ -733,7 +740,7 @@ namespace UnrealBuildTool
 				}
 
 				// Some switches interfere with on-device debugging
-				if (CompileEnvironment.Configuration != CppConfiguration.Debug)
+				if (CompileEnvironment.Configuration != CppConfiguration.Debug && !DisableFunctionDataSections())
 				{
 					Result += " -ffunction-sections";   // Places each function in its own section of the output file, linker may be able to perform opts to improve locality of reference
 					Result += " -fdata-sections";		// Places each data item in its own section of the output file, linker may be able to perform opts to improve locality of reference
