add xml comment to IsZipFile

diff --git a/Engine/Source/Programs/UnrealBuildTool/Configuration/ModuleRules.cs b/Engine/Source/Programs/UnrealBuildTool/Configuration/ModuleRules.cs
--- a/Engine/Source/Programs/UnrealBuildTool/Configuration/ModuleRules.cs
+++ b/Engine/Source/Programs/UnrealBuildTool/Configuration/ModuleRules.cs
@@ -345,6 +345,9 @@ namespace UnrealBuildTool
 				this.bCopyFramework = bCopyFramework;
 			}
 
+			/// <summary>
+			/// Specifies if the file is a zip file
+			/// </summary>
 			public bool IsZipFile()
 			{
 				return Path.EndsWith(".zip");
