UE4: Compile fix.
#rnx

diff --git a/Engine/Source/Programs/UnrealBuildTool/Platform/IOS/IOSToolChain.cs b/Engine/Source/Programs/UnrealBuildTool/Platform/IOS/IOSToolChain.cs
--- a/Engine/Source/Programs/UnrealBuildTool/Platform/IOS/IOSToolChain.cs
+++ b/Engine/Source/Programs/UnrealBuildTool/Platform/IOS/IOSToolChain.cs
@@ -1235,7 +1235,7 @@ namespace UnrealBuildTool
 				ExtractScript.AppendLine(String.Format("unzip -q -o {0} -d {1}", Utils.MakePathSafeToUseWithCommandLine(Framework.ZipFile.FullName), Utils.MakePathSafeToUseWithCommandLine(Framework.FrameworkDirectory.ParentDirectory.FullName))); // Zip contains folder with the same name, hence ParentDirectory
 				ExtractScript.AppendLine(String.Format("touch {0}", Utils.MakePathSafeToUseWithCommandLine(Framework.ExtractedTokenFile.AbsolutePath)));
 
-				FileItem ExtractScriptFileItem = FileItem.CreateIntermediateTextFile(new FileReference(Framework.OutputDirectory.FullName + ".sh"), ExtractScript.ToString());
+				FileItem ExtractScriptFileItem = Graph.CreateIntermediateTextFile(new FileReference(Framework.OutputDirectory.FullName + ".sh"), ExtractScript.ToString());
 
 				Action UnzipAction = Graph.CreateAction(ActionType.BuildProject);
 				UnzipAction.CommandPath = new FileReference("/bin/sh");
