UBT: Add a warning to discourage generating project files by passing arguments to GenerateProjectFiles.bat.

diff --git a/Engine/Source/Programs/UnrealBuildTool/Modes/GenerateProjectFilesMode.cs b/Engine/Source/Programs/UnrealBuildTool/Modes/GenerateProjectFilesMode.cs
--- a/Engine/Source/Programs/UnrealBuildTool/Modes/GenerateProjectFilesMode.cs
+++ b/Engine/Source/Programs/UnrealBuildTool/Modes/GenerateProjectFilesMode.cs
@@ -62,6 +62,28 @@ namespace UnrealBuildTool
 			FileReference ProjectFile;
 			TryParseProjectFileArgument(Arguments, out ProjectFile);
 
+			// Warn if there are explicit project file formats specified
+			if (ProjectFileFormats.Count > 0)
+			{
+				StringBuilder Configuration = new StringBuilder();
+				Configuration.Append("Project file formats specified via the command line will be ignored when generating\n");
+				Configuration.Append("project files from the editor and other engine tools.\n");
+				Configuration.Append("\n");
+				Configuration.Append("Consider setting your desired IDE from the editor preferences window, or modify your\n");
+				Configuration.Append("BuildConfiguration.xml file with:\n");
+				Configuration.Append("\n");
+				Configuration.Append("<?xml version=\"1.0\" encoding=\"utf-8\" ?>\n");
+				Configuration.Append("<Configuration xmlns=\"https://www.unrealengine.com/BuildConfiguration\">\n");
+				Configuration.Append("  <ProjectFileGenerator>\n");
+				foreach(ProjectFileFormat ProjectFileFormat in ProjectFileFormats)
+				{
+					Configuration.AppendFormat("    <Format>{0}</Format>\n", ProjectFileFormat);
+				}
+				Configuration.Append("  </ProjectFileGenerator>\n");
+				Configuration.Append("</Configuration>\n");
+				Log.TraceWarning("{0}", Configuration.ToString());
+			}
+
 			// If there aren't any formats set, read the default project file format from the config file
 			if (ProjectFileFormats.Count == 0)
 			{
