Storyboard Full Support (Including compilation and Editor Support)

diff --git a/Engine/Source/Programs/AutomationTool/IOS/IOSPlatform.Automation.cs b/Engine/Source/Programs/AutomationTool/IOS/IOSPlatform.Automation.cs
--- a/Engine/Source/Programs/AutomationTool/IOS/IOSPlatform.Automation.cs
+++ b/Engine/Source/Programs/AutomationTool/IOS/IOSPlatform.Automation.cs
@@ -952,6 +952,7 @@ public class IOSPlatform : Platform
 
 		bool bLaunchscreenStoryboard = false;
 		ConfigHierarchy PlatformGameConfig;
+		bool bXCArchive = false;
 		if (Params.EngineConfigs.TryGetValue(SC.StageTargetPlatform.PlatformType, out PlatformGameConfig))
 		{
 			PlatformGameConfig.GetBool("/Script/IOSRuntimeSettings.IOSRuntimeSettings", "bLaunchscreenStoryboard", out bLaunchscreenStoryboard);
diff --git a/Engine/Source/Programs/UnrealBuildTool/Platform/IOS/UEDeployIOS.cs b/Engine/Source/Programs/UnrealBuildTool/Platform/IOS/UEDeployIOS.cs
--- a/Engine/Source/Programs/UnrealBuildTool/Platform/IOS/UEDeployIOS.cs
+++ b/Engine/Source/Programs/UnrealBuildTool/Platform/IOS/UEDeployIOS.cs
@@ -28,7 +28,6 @@ namespace UnrealBuildTool
 			set { bForDistribution = value; }
 		}
 		bool bForDistribution = false;
-		public static string LaunchStoryboardPath;
 
 		protected class VersionUtilities
 		{
@@ -211,11 +210,12 @@ namespace UnrealBuildTool
 			string OldPListData = File.Exists(PListFile) ? File.ReadAllText(PListFile) : "";
 
 			// determine if there is a Lanchscreen Storyboard THEN a launch.xib
-			LaunchStoryboardPath = InEngineDir + "/Build/IOS/Resources/Interface";
+			string LaunchStoryboard = InEngineDir + "/Build/IOS/Resources/Interface/LaunchScreen.storyboard";
 			string LaunchXib = InEngineDir + "/Build/IOS/Resources/Interface/LaunchScreen.xib";
-			if (Directory.Exists(BuildDirectory + "/Resources/Interface/LaunchScreen.storyboardc"))
+			if (File.Exists(BuildDirectory + "/Resources/Interface/LaunchScreen.storyboard"))
+
 			{
-				LaunchStoryboardPath = BuildDirectory + "/Resources/Interface";
+				LaunchStoryboard = BuildDirectory + "/Resources/Interface/LaunchScreen.storyboard";
 			}
 			else if (File.Exists(BuildDirectory + "/Resources/Interface/LaunchScreen.xib"))
 			{
@@ -580,8 +580,16 @@ namespace UnrealBuildTool
 				Text.AppendLine("\t\t</dict>");
 				Text.AppendLine("\t</dict>");
 			}
-			if (Directory.Exists(LaunchStoryboardPath + "/LaunchScreen.storyboardc") || File.Exists(LaunchXib) && VersionUtilities.bLaunchscreenStoryboard)
+			// Launchscreen storyboard only compatible on Mac for now.
+			if ((BuildHostPlatform.Current.Platform == UnrealTargetPlatform.Mac) && (File.Exists(LaunchStoryboard) || File.Exists(LaunchXib)) && VersionUtilities.bLaunchscreenStoryboard)
 			{
+				string outputStoryboard = LaunchStoryboard + "c";
+				string argsStoryboard = "--compile " + outputStoryboard + " " + LaunchStoryboard;
+				string stdOutLaunchScreen;
+				stdOutLaunchScreen = Utils.RunLocalProcessAndReturnStdOut("ibtool", argsStoryboard);
+
+				Log.TraceInformation("LaunchScreen Storyboard compilation results : " + stdOutLaunchScreen);
+
 				Text.AppendLine("\t<key>UILaunchStoryboardName</key>");
 				Text.AppendLine("\t<string>LaunchScreen</string>");
 				bSkipDefaultPNGs = true;
@@ -883,51 +891,50 @@ namespace UnrealBuildTool
 			{
 				if (VersionUtilities.bLaunchscreenStoryboard)
 				{
-					if(!Directory.Exists(AppDirectory + "/LaunchScreen.storyboardc"))
-					{
-						Directory.CreateDirectory(AppDirectory + "/LaunchScreen.storyboardc");
-					}
-					
-					CopyFolder(LaunchStoryboardPath + "/LaunchScreen.storyboardc", AppDirectory + "/LaunchScreen.storyboardc", true);
+					CopyFolder(InEngineDir + "/Build/IOS/Resources/Interface/LaunchScreen.storyboardc", AppDirectory + "/LaunchScreen.storyboardc", true);
+					CopyFiles(InEngineDir + "/Build/IOS/Resources/Interface/Assets", AppDirectory, "*", true);
 				}
 				// we still want default icons
 				if (!bSkipIcons)
 				{
 					CopyFiles(InEngineDir + "/Build/IOS/Resources/Graphics", AppDirectory, "Icon*.png", true);
 				}
-			}
-			else
-			{
-				if (!bSkipIcons)
-				{
-					CopyFiles(InEngineDir + "/Build/IOS/Resources/Graphics", AppDirectory, "Icon*.png", true);
-				}
-				if (bSupportsPortrait)
-				{
-					CopyFiles(InEngineDir + "/Build/IOS/Resources/Graphics", AppDirectory, "Default-IPhone6.png", true);
-					CopyFiles(InEngineDir + "/Build/IOS/Resources/Graphics", AppDirectory, "Default-IPhone6Plus-Portrait.png", true);
-					CopyFiles(InEngineDir + "/Build/IOS/Resources/Graphics", AppDirectory, "[REDACTED_EMAIL]", true);
-					CopyFiles(InEngineDir + "/Build/IOS/Resources/Graphics", AppDirectory, "[REDACTED_EMAIL]", true);
-					CopyFiles(InEngineDir + "/Build/IOS/Resources/Graphics", AppDirectory, "[REDACTED_EMAIL]", true);
-					CopyFiles(InEngineDir + "/Build/IOS/Resources/Graphics", AppDirectory, "[REDACTED_EMAIL]", true);
-					CopyFiles(InEngineDir + "/Build/IOS/Resources/Graphics", AppDirectory, "Default-IPhoneXS-Portrait.png", true);
-					CopyFiles(InEngineDir + "/Build/IOS/Resources/Graphics", AppDirectory, "Default-IPhoneXSMax-Portrait.png", true);
-					CopyFiles(InEngineDir + "/Build/IOS/Resources/Graphics", AppDirectory, "Default-IPhoneXR-Portrait.png", true);
-				}
-				if (bSupportsLandscape)
+				else
 				{
-					CopyFiles(InEngineDir + "/Build/IOS/Resources/Graphics", AppDirectory, "Default-IPhone6-Landscape.png", true);
-					CopyFiles(InEngineDir + "/Build/IOS/Resources/Graphics", AppDirectory, "Default-IPhone6Plus-Landscape.png", true);
-					CopyFiles(InEngineDir + "/Build/IOS/Resources/Graphics", AppDirectory, "[REDACTED_EMAIL]", true);
-					CopyFiles(InEngineDir + "/Build/IOS/Resources/Graphics", AppDirectory, "[REDACTED_EMAIL]", true);
-					CopyFiles(InEngineDir + "/Build/IOS/Resources/Graphics", AppDirectory, "[REDACTED_EMAIL]", true);
-					CopyFiles(InEngineDir + "/Build/IOS/Resources/Graphics", AppDirectory, "[REDACTED_EMAIL]", true);
-					CopyFiles(InEngineDir + "/Build/IOS/Resources/Graphics", AppDirectory, "Default-IPhoneXS-Landscape.png", true);
-					CopyFiles(InEngineDir + "/Build/IOS/Resources/Graphics", AppDirectory, "Default-IPhoneXSMax-Landscape.png", true);
-					CopyFiles(InEngineDir + "/Build/IOS/Resources/Graphics", AppDirectory, "Default-IPhoneXR-Landscape.png", true);
+					if (!bSkipIcons)
+					{
+						CopyFiles(InEngineDir + "/Build/IOS/Resources/Graphics", AppDirectory, "Icon*.png", true);
+					}
+					if (bSupportsPortrait && !VersionUtilities.bLaunchscreenStoryboard)
+					{
+						CopyFiles(InEngineDir + "/Build/IOS/Resources/Graphics", AppDirectory, "Default-IPhone6.png", true);
+						CopyFiles(InEngineDir + "/Build/IOS/Resources/Graphics", AppDirectory, "Default-IPhone6Plus-Portrait.png", true);
+						CopyFiles(InEngineDir + "/Build/IOS/Resources/Graphics", AppDirectory, "[REDACTED_EMAIL]", true);
+						CopyFiles(InEngineDir + "/Build/IOS/Resources/Graphics", AppDirectory, "[REDACTED_EMAIL]", true);
+						CopyFiles(InEngineDir + "/Build/IOS/Resources/Graphics", AppDirectory, "[REDACTED_EMAIL]", true);
+						CopyFiles(InEngineDir + "/Build/IOS/Resources/Graphics", AppDirectory, "[REDACTED_EMAIL]", true);
+						CopyFiles(InEngineDir + "/Build/IOS/Resources/Graphics", AppDirectory, "Default-IPhoneXS-Portrait.png", true);
+						CopyFiles(InEngineDir + "/Build/IOS/Resources/Graphics", AppDirectory, "Default-IPhoneXSMax-Portrait.png", true);
+						CopyFiles(InEngineDir + "/Build/IOS/Resources/Graphics", AppDirectory, "Default-IPhoneXR-Portrait.png", true);
+					}
+					if (bSupportsLandscape && !VersionUtilities.bLaunchscreenStoryboard)
+					{
+						CopyFiles(InEngineDir + "/Build/IOS/Resources/Graphics", AppDirectory, "Default-IPhone6-Landscape.png", true);
+						CopyFiles(InEngineDir + "/Build/IOS/Resources/Graphics", AppDirectory, "Default-IPhone6Plus-Landscape.png", true);
+						CopyFiles(InEngineDir + "/Build/IOS/Resources/Graphics", AppDirectory, "[REDACTED_EMAIL]", true);
+						CopyFiles(InEngineDir + "/Build/IOS/Resources/Graphics", AppDirectory, "[REDACTED_EMAIL]", true);
+						CopyFiles(InEngineDir + "/Build/IOS/Resources/Graphics", AppDirectory, "[REDACTED_EMAIL]", true);
+						CopyFiles(InEngineDir + "/Build/IOS/Resources/Graphics", AppDirectory, "[REDACTED_EMAIL]", true);
+						CopyFiles(InEngineDir + "/Build/IOS/Resources/Graphics", AppDirectory, "Default-IPhoneXS-Landscape.png", true);
+						CopyFiles(InEngineDir + "/Build/IOS/Resources/Graphics", AppDirectory, "Default-IPhoneXSMax-Landscape.png", true);
+						CopyFiles(InEngineDir + "/Build/IOS/Resources/Graphics", AppDirectory, "Default-IPhoneXR-Landscape.png", true);
+					}
+					if (!VersionUtilities.bLaunchscreenStoryboard)
+					{
+						CopyFiles(BuildDirectory + "/Resources/Graphics", AppDirectory, "[REDACTED_EMAIL]", true);
+						CopyFiles(BuildDirectory + "/Resources/Graphics", AppDirectory, "[REDACTED_EMAIL]", true);
+					}
 				}
-				CopyFiles(InEngineDir + "/Build/IOS/Resources/Graphics", AppDirectory, "[REDACTED_EMAIL]", true);
-				CopyFiles(InEngineDir + "/Build/IOS/Resources/Graphics", AppDirectory, "[REDACTED_EMAIL]", true);
 			}
 			// merge game assets on top
 			// @todo tvos: Do we want to copy IOS and TVOS both in? (Engine/IOS -> Game/IOS -> Game/TVOS)?
@@ -937,7 +944,7 @@ namespace UnrealBuildTool
 				{
 					CopyFiles(BuildDirectory + "/Resources/Graphics", AppDirectory, "Icon*.png", true);
 				}
-				if (bSupportsPortrait && !VersionUtilities.bLaunchscreenStoryboard)
+				if (bSupportsPortrait)
 				{
 					CopyFiles(BuildDirectory + "/Resources/Graphics", AppDirectory, "Default-IPhone6.png", true);
 					CopyFiles(BuildDirectory + "/Resources/Graphics", AppDirectory, "Default-IPhone6Plus-Portrait.png", true);
@@ -949,7 +956,7 @@ namespace UnrealBuildTool
 					CopyFiles(BuildDirectory + "/Resources/Graphics", AppDirectory, "Default-IPhoneXSMax-Portrait.png", true);
 					CopyFiles(BuildDirectory + "/Resources/Graphics", AppDirectory, "Default-IPhoneXR-Portrait.png", true);
 				}
-				if (bSupportsLandscape && !VersionUtilities.bLaunchscreenStoryboard)
+				if (bSupportsLandscape)
 				{
 					CopyFiles(BuildDirectory + "/Resources/Graphics", AppDirectory, "Default-IPhone6-Landscape.png", true);
 					CopyFiles(BuildDirectory + "/Resources/Graphics", AppDirectory, "Default-IPhone6Plus-Landscape.png", true);
@@ -961,13 +968,11 @@ namespace UnrealBuildTool
 					CopyFiles(BuildDirectory + "/Resources/Graphics", AppDirectory, "Default-IPhoneXSMax-Landscape.png", true);
 					CopyFiles(BuildDirectory + "/Resources/Graphics", AppDirectory, "Default-IPhoneXR-Landscape.png", true);
 				}
-				if (!VersionUtilities.bLaunchscreenStoryboard)
-				{
-					CopyFiles(BuildDirectory + "/Resources/Graphics", AppDirectory, "[REDACTED_EMAIL]", true);
-					CopyFiles(BuildDirectory + "/Resources/Graphics", AppDirectory, "[REDACTED_EMAIL]", true);
-				}
+				CopyFiles(BuildDirectory + "/Resources/Graphics", AppDirectory, "[REDACTED_EMAIL]", true);
+				CopyFiles(BuildDirectory + "/Resources/Graphics", AppDirectory, "[REDACTED_EMAIL]", true);
 			}
 		}
+		
 		protected virtual void CopyLocalizationsResources(string InEngineDir, string AppDirectory, string BuildDirectory, string IntermediateDir)
 		{
 			string LocalizationsPath = BuildDirectory + "/Resources/Localizations";
diff --git a/Engine/Source/Runtime/IOS/IOSRuntimeSettings/Classes/IOSRuntimeSettings.h b/Engine/Source/Runtime/IOS/IOSRuntimeSettings/Classes/IOSRuntimeSettings.h
--- a/Engine/Source/Runtime/IOS/IOSRuntimeSettings/Classes/IOSRuntimeSettings.h
+++ b/Engine/Source/Runtime/IOS/IOSRuntimeSettings/Classes/IOSRuntimeSettings.h
@@ -432,6 +432,14 @@ public:
 	UPROPERTY(GlobalConfig, EditAnywhere, Category = LaunchScreen)
 	bool 	bLaunchscreenStoryboard;
 
+	/**
+	 * Choose whether to use a storyboard as a Launchscreen or not. If True, LaunchScreen.storyboard in Build/IOS/Resources/Interface
+	 * in the Engine folder will be copied to the bundle app and the Launchscreen images WILL NOT be included.
+	 * When using assets in your LaunchScreen.Storyboard, add them in Build/IOS/Resources/Interface/Assets and they will be included.
+	 */
+	UPROPERTY(GlobalConfig, EditAnywhere, Category = LaunchScreen)
+	bool 	bLaunchscreenStoryboard;
+
 	// Whether the app supports Facebook
 	UPROPERTY(GlobalConfig, EditAnywhere, Category = Online)
 	bool bEnableFacebookSupport;
