Check if the XGE service is running as part of the availability check.

UE4Build / BuildCookRun doesn't check BuildConfiguration.xml* before initiating builds so people who have the IB service installed but disabled will see builds sent to XGE that fail with "Failed to connect to Build Service (on local machine)".

#tests ran BuildCookRun -build with and without the XGE agent service disabled.Verified that builds fell back to local execution with the service disabled and still attempted XGE builds with it enabled.

diff --git a/Engine/Source/Programs/UnrealBuildTool/Executors/XGE.cs b/Engine/Source/Programs/UnrealBuildTool/Executors/XGE.cs
--- a/Engine/Source/Programs/UnrealBuildTool/Executors/XGE.cs
+++ b/Engine/Source/Programs/UnrealBuildTool/Executors/XGE.cs
@@ -135,7 +135,25 @@ namespace UnrealBuildTool
 		public static bool IsAvailable()
 		{
 			string XgConsoleExe;
-			return TryGetXgConsoleExecutable(out XgConsoleExe);
+
+			bool Success = TryGetXgConsoleExecutable(out XgConsoleExe);
+
+			// on windows check the service is actually running
+			if (Success && BuildHostPlatform.Current.Platform == UnrealTargetPlatform.Win64)
+			{
+				try
+				{
+					// will throw if the service doesn't exist, which it should if IB is present but just incase...
+					System.ServiceProcess.ServiceController SC = new System.ServiceProcess.ServiceController("Incredibuild Agent");
+					Success = SC.Status == System.ServiceProcess.ServiceControllerStatus.Running;
+				}
+				catch
+				{
+					Success = false;
+				}
+			}
+
+			return Success;
 		}
 
 		// precompile the Regex needed to parse the XGE output (the ones we want are of the form "File (Duration at +time)"
diff --git a/Engine/Source/Programs/UnrealBuildTool/System/PlatformExports.cs b/Engine/Source/Programs/UnrealBuildTool/System/PlatformExports.cs
--- a/Engine/Source/Programs/UnrealBuildTool/System/PlatformExports.cs
+++ b/Engine/Source/Programs/UnrealBuildTool/System/PlatformExports.cs
@@ -106,7 +106,7 @@ namespace UnrealBuildTool
 		/// <returns>True if the platform supports XGE</returns>
 		public static bool CanUseXGE(UnrealTargetPlatform Platform)
 		{
-			return UEBuildPlatform.IsPlatformAvailable(Platform) && UEBuildPlatform.GetBuildPlatform(Platform).CanUseXGE();
+			return UEBuildPlatform.IsPlatformAvailable(Platform) && UEBuildPlatform.GetBuildPlatform(Platform).CanUseXGE() && XGE.IsAvailable();
 		}
 
 		/// <summary>
