- Fixed up the whitelist/blacklist logic of child plugins. It was incredibly broken before, where in non-staged builds (UnrealPak was one), the host platform would be blacklisted by a random child plugin

diff --git a/Engine/Source/Programs/UnrealBuildTool/System/Plugins.cs b/Engine/Source/Programs/UnrealBuildTool/System/Plugins.cs
--- a/Engine/Source/Programs/UnrealBuildTool/System/Plugins.cs
+++ b/Engine/Source/Programs/UnrealBuildTool/System/Plugins.cs
@@ -444,18 +444,16 @@ namespace UnrealBuildTool
 						PluginReferenceDescriptor ParentPluginReference = Parent.Descriptor.Plugins.FirstOrDefault(x => x.Name.Equals(ChildPluginReference.Name));
 						if (ParentPluginReference != null)
 						{
-							// merge white/blacklists (if the parent had a list, and child didn't specify a list, just add the child platform to the parent list - for white and black!)
-							if (ChildPluginReference.WhitelistPlatforms != null)
+							// we only need to whitelist the platform if the parent had a whitelist (otherwise, we could mistakenly remove all other platforms)
+							if (ParentPluginReference.WhitelistPlatforms != null)
 							{
-								if (ParentPluginReference.WhitelistPlatforms == null)
-								{
-									ParentPluginReference.WhitelistPlatforms = ChildPluginReference.WhitelistPlatforms;
-								}
-								else 
+								if (ChildPluginReference.WhitelistPlatforms != null)
 								{
 									ParentPluginReference.WhitelistPlatforms = ParentPluginReference.WhitelistPlatforms.Union(ChildPluginReference.WhitelistPlatforms).ToList();
 								}
 							}
+
+							// if we want to blacklist a platform, add it even if the parent didn't have a blacklist. this won't cause problems with other platforms
 							if (ChildPluginReference.BlacklistPlatforms != null)
 							{
 								if (ParentPluginReference.BlacklistPlatforms == null)
diff --git a/Engine/Source/Runtime/Projects/Private/PluginManager.cpp b/Engine/Source/Runtime/Projects/Private/PluginManager.cpp
--- a/Engine/Source/Runtime/Projects/Private/PluginManager.cpp
+++ b/Engine/Source/Runtime/Projects/Private/PluginManager.cpp
@@ -410,28 +410,14 @@ void FPluginManager::ReadAllPlugins(TMap<FString, TSharedRef<FPlugin>>& Plugins,
 				{
 					if (ParentModule.Name == ChildModule.Name && ParentModule.Type == ChildModule.Type)
 					{
+						// we only need to whitelist the platform if the parent had a whitelist (otherwise, we could mistakenly remove all other platforms)
 						if (ParentModule.WhitelistPlatforms.Num() > 0)
 						{
-							if (ChildModule.WhitelistPlatforms.Num() > 0)
-							{
-								ParentModule.WhitelistPlatforms.Append(ChildModule.WhitelistPlatforms);
-							}
-							else
-							{
-								ParentModule.WhitelistPlatforms.Add(FPlatformProperties::PlatformName());
-							}
-						}
-						if (ParentModule.BlacklistPlatforms.Num() > 0)
-						{
-							if (ChildModule.BlacklistPlatforms.Num() > 0)
-							{
-								ParentModule.BlacklistPlatforms.Append(ChildModule.BlacklistPlatforms);
-							}
-							else
-							{
-								ParentModule.BlacklistPlatforms.Add(FPlatformProperties::PlatformName());
-							}
+							ParentModule.WhitelistPlatforms.Append(ChildModule.WhitelistPlatforms);
 						}
+
+						// if we want to blacklist a platform, add it even if the parent didn't have a blacklist. this won't cause problems with other platforms
+						ParentModule.BlacklistPlatforms.Append(ChildModule.BlacklistPlatforms);
 					}
 				}
 			}
