Alternative way to set the metadata that VS uses to organise globbed platform UBT source. The Split() approach has one further step of indirection that assumes the returned array's size. Somehow this assumption has been observed failing. Replace() returns a string directly and gracefully handles cases where a regex doesn't match so is a much more tolerant solution.

diff --git a/Engine/Source/Programs/UnrealBuildTool/UnrealBuildTool.csproj b/Engine/Source/Programs/UnrealBuildTool/UnrealBuildTool.csproj
--- a/Engine/Source/Programs/UnrealBuildTool/UnrealBuildTool.csproj
+++ b/Engine/Source/Programs/UnrealBuildTool/UnrealBuildTool.csproj
@@ -444,7 +444,10 @@
   files and include them in the build. They will also appear in VisualStudio -->
   <ItemGroup>
     <Compile Include="../../../Platforms/*/Source/Programs/UnrealBuildTool/**/*.cs">
-      <Link>Platform/$([System.Text.RegularExpressions.Regex]::Split(`%(Identity)`, `[\\/]`)[4])/_</Link>
+      <!-- RecursiveDir metadata expands to [PlatformName]/Source/Programs/[etc]. The Replace() will
+      replace everything from the first slash to the end leaving just the platform name. Should the
+      Replace() fail then RecursiveDir is still an agreeable value -->
+      <Link>Platform/$([System.Text.RegularExpressions.Regex]::Replace(%(Compile.RecursiveDir), [\\/].+$, ``))/_</Link>
     </Compile>
   </ItemGroup>
   <!-- To modify your build process, add your task inside one of the targets below and uncomment it. 
