Merged from Dev-Build:
Fixed issue with converting UnrealTargetPlatform to string when parsing command line options

Added a TypeConverter for UnrealTargetPlatform
Added check if a typeconverter exists when parsing command line arguments

diff --git a/Engine/Source/Programs/DotNETCommon/DotNETUtilities/CommandLineArguments.cs b/Engine/Source/Programs/DotNETCommon/DotNETUtilities/CommandLineArguments.cs
--- a/Engine/Source/Programs/DotNETCommon/DotNETUtilities/CommandLineArguments.cs
+++ b/Engine/Source/Programs/DotNETCommon/DotNETUtilities/CommandLineArguments.cs
@@ -1191,6 +1191,14 @@ namespace Tools.DotNETCommon
 			}
 			else
 			{
+				// First check for a TypeConverter
+				TypeConverter TypeConverter = TypeDescriptor.GetConverter(FieldType);
+				if (TypeConverter.CanConvertFrom(typeof(string)))
+				{
+					Value = TypeConverter.ConvertFrom(Text);
+					return true;
+				}
+
 				// Otherwise let the framework convert between types
 				try
 				{
diff --git a/Engine/Source/Programs/UnrealBuildTool/Configuration/UEBuildTarget.cs b/Engine/Source/Programs/UnrealBuildTool/Configuration/UEBuildTarget.cs
--- a/Engine/Source/Programs/UnrealBuildTool/Configuration/UEBuildTarget.cs
+++ b/Engine/Source/Programs/UnrealBuildTool/Configuration/UEBuildTarget.cs
@@ -2,9 +2,11 @@
 
 using System;
 using System.Collections.Generic;
+using System.ComponentModel;
 using System.Linq;
 using System.Text;
 using System.Diagnostics;
+using System.Globalization;
 using System.IO;
 using System.Xml;
 using System.Runtime.Serialization;
@@ -17,7 +19,7 @@ namespace UnrealBuildTool
 	/// <summary>
 	/// The platform we're building for
 	/// </summary>
-	[Serializable]
+	[Serializable, TypeConverter(typeof(UnrealTargetPlatformTypeConverter))]
 	public partial struct UnrealTargetPlatform : ISerializable
 	{
 		#region Private/boilerplate
@@ -276,6 +278,44 @@ namespace UnrealBuildTool
 		public static UnrealTargetPlatform Lumin = FindOrAddByName("Lumin");
 	}
 
+	internal class UnrealTargetPlatformTypeConverter : TypeConverter
+	{
+		public override bool CanConvertFrom(ITypeDescriptorContext context, Type sourceType)
+		{
+			if (sourceType == typeof(string))
+				return true;
+
+			return base.CanConvertFrom(context, sourceType);
+		}
+
+		public override bool CanConvertTo(ITypeDescriptorContext context, Type destinationType)
+		{
+			if (destinationType == typeof(string))
+				return true;
+
+			return base.CanConvertTo(context, destinationType);
+		}
+
+		public override object ConvertFrom(ITypeDescriptorContext context, CultureInfo culture, object value)
+		{
+			if (value.GetType() == typeof(string))
+			{
+				return UnrealTargetPlatform.Parse((string) value);
+			}
+			return base.ConvertFrom(context, culture, value);
+		}
+
+		public override object ConvertTo(ITypeDescriptorContext context, CultureInfo culture, object value, Type destinationType)
+		{
+			if (destinationType == typeof(string))
+			{ 
+				UnrealTargetPlatform Platform = (UnrealTargetPlatform) value;
+				return Platform.ToString();
+			}
+			return base.ConvertTo(context, culture, value, destinationType);
+		}
+	}
+
 	/// <summary>
 	/// Extension methods used for serializing UnrealTargetPlatform instances
 	/// </summary>
