Fix issue with launch screen storyboards in the Project folder

diff --git a/Engine/Source/Programs/UnrealBuildTool/Platform/IOS/UEDeployIOS.cs b/Engine/Source/Programs/UnrealBuildTool/Platform/IOS/UEDeployIOS.cs
--- a/Engine/Source/Programs/UnrealBuildTool/Platform/IOS/UEDeployIOS.cs
+++ b/Engine/Source/Programs/UnrealBuildTool/Platform/IOS/UEDeployIOS.cs
@@ -212,11 +212,7 @@
 			// determine if there is a Lanchscreen Storyboard THEN a launch.xib
 			string LaunchStoryboard = InEngineDir + "/Build/IOS/Resources/Interface/LaunchScreen.storyboard";
 			string LaunchXib = InEngineDir + "/Build/IOS/Resources/Interface/LaunchScreen.xib";
-			if (File.Exists(BuildDirectory + "/Resources/Interface/LaunchScreen.storyboard"))
-
-			{
-				LaunchStoryboard = BuildDirectory + "/Resources/Interface/LaunchScreen.storyboard";
-			}
+			if (File.Exists(BuildDirectory + "/Resources/Interface/LaunchScreen.storyboard"))			{				LaunchStoryboard = BuildDirectory + "/Resources/Interface/LaunchScreen.storyboard";			}
 			else if (File.Exists(BuildDirectory + "/Resources/Interface/LaunchScreen.xib"))
 			{
 				LaunchXib = BuildDirectory + "/Resources/Interface/LaunchScreen.xib";
@@ -891,8 +887,16 @@
 			{
 				if (VersionUtilities.bLaunchscreenStoryboard)
 				{
-					CopyFolder(InEngineDir + "/Build/IOS/Resources/Interface/LaunchScreen.storyboardc", AppDirectory + "/LaunchScreen.storyboardc", true);
-					CopyFiles(InEngineDir + "/Build/IOS/Resources/Interface/Assets", AppDirectory, "*", true);
+					if (Directory.Exists(BuildDirectory + "/Resources/Interface/LaunchScreen.storyboardc"))
+					{
+						CopyFolder(BuildDirectory + "/Resources/Interface/LaunchScreen.storyboardc", AppDirectory + "/LaunchScreen.storyboardc", true);
+						CopyFiles(BuildDirectory + "/Resources/Interface/Assets", AppDirectory, "*", true);
+					}
+					else
+					{
+						CopyFolder(InEngineDir + "/Build/IOS/Resources/Interface/LaunchScreen.storyboardc", AppDirectory + "/LaunchScreen.storyboardc", true);
+						CopyFiles(InEngineDir + "/Build/IOS/Resources/Interface/Assets", AppDirectory, "*", true);
+					}
 				}
 				// we still want default icons
 				if (!bSkipIcons)
