Storyboard Full Support (Including compilation and Editor Support)

diff --git a/Engine/Source/Programs/AutomationTool/IOS/IOSPlatform.Automation.cs b/Engine/Source/Programs/AutomationTool/IOS/IOSPlatform.Automation.cs
--- a/Engine/Source/Programs/AutomationTool/IOS/IOSPlatform.Automation.cs
+++ b/Engine/Source/Programs/AutomationTool/IOS/IOSPlatform.Automation.cs
@@ -880,14 +880,14 @@ public class IOSPlatform : Platform
 
 					// copy the icons/launch screens from the engine
 					{
-						DirectoryReference DataPath = DirectoryReference.Combine(SC.EngineRoot, "Build", "IOS", "Resources", "Graphics");
-						StageImageAndIconFiles(DataPath, bSupportsPortrait, bSupportsLandscape, SC, bSkipIcons);
+						DirectoryReference GraphicsDataPath = DirectoryReference.Combine(SC.EngineRoot, "Build", "IOS", "Resources", "Graphics");
+						StageImageAndIconFiles(Params, GraphicsDataPath, bSupportsPortrait, bSupportsLandscape, SC, bSkipIcons);
 					}
 
 					// copy the icons/launch screens from the game (may stomp the engine copies)
 					{
-						DirectoryReference DataPath = DirectoryReference.Combine(SC.ProjectRoot, "Build", "IOS", "Resources", "Graphics");
-						StageImageAndIconFiles(DataPath, bSupportsPortrait, bSupportsLandscape, SC, bSkipIcons);
+						DirectoryReference GraphicsDataPath = DirectoryReference.Combine(SC.ProjectRoot, "Build", "IOS", "Resources", "Graphics");
+						StageImageAndIconFiles(Params, GraphicsDataPath, bSupportsPortrait, bSupportsLandscape, SC, bSkipIcons);
 					}
 				}
 
@@ -947,9 +947,18 @@ public class IOSPlatform : Platform
 		}
 	}
 
-	private void StageImageAndIconFiles(DirectoryReference DataPath, bool bSupportsPortrait, bool bSupportsLandscape, DeploymentContext SC, bool bSkipIcons)
+	private void StageImageAndIconFiles(ProjectParams Params, DirectoryReference GraphicsDataPath, bool bSupportsPortrait, bool bSupportsLandscape, DeploymentContext SC, bool bSkipIcons)
 	{
-		if (DirectoryReference.Exists(DataPath))
+
+		bool bLaunchscreenStoryboard = false;
+		ConfigHierarchy PlatformGameConfig;
+		bool bXCArchive = false;
+		if (Params.EngineConfigs.TryGetValue(SC.StageTargetPlatform.PlatformType, out PlatformGameConfig))
+		{
+			PlatformGameConfig.GetBool("/Script/IOSRuntimeSettings.IOSRuntimeSettings", "bLaunchscreenStoryboard", out bLaunchscreenStoryboard);
+		}
+
+		if (DirectoryReference.Exists(GraphicsDataPath) && !bLaunchscreenStoryboard)
 		{
 			List<string> ImageFileNames = new List<string>();
 			if (bSupportsPortrait)
@@ -981,7 +990,7 @@ public class IOSPlatform : Platform
 
 			foreach (string ImageFileName in ImageFileNames)
 			{
-				FileReference ImageFile = FileReference.Combine(DataPath, ImageFileName);
+				FileReference ImageFile = FileReference.Combine(GraphicsDataPath, ImageFileName);
 				if (FileReference.Exists(ImageFile))
 				{
 					SC.StageFile(StagedFileType.SystemNonUFS, ImageFile, new StagedFileReference(ImageFileName));
@@ -990,7 +999,7 @@ public class IOSPlatform : Platform
 
 			if (!bSkipIcons)
 			{
-				SC.StageFiles(StagedFileType.SystemNonUFS, DataPath, "Icon*.png", StageFilesSearch.TopDirectoryOnly, StagedDirectoryReference.Root);
+				SC.StageFiles(StagedFileType.SystemNonUFS, GraphicsDataPath, "Icon*.png", StageFilesSearch.TopDirectoryOnly, StagedDirectoryReference.Root);
 			}
 		}
 	}
diff --git a/Engine/Source/Programs/UnrealBuildTool/Platform/IOS/UEDeployIOS.cs b/Engine/Source/Programs/UnrealBuildTool/Platform/IOS/UEDeployIOS.cs
--- a/Engine/Source/Programs/UnrealBuildTool/Platform/IOS/UEDeployIOS.cs
+++ b/Engine/Source/Programs/UnrealBuildTool/Platform/IOS/UEDeployIOS.cs
@@ -42,7 +42,7 @@ namespace UnrealBuildTool
 				set;
 			}
 
-			
+			public static bool bLaunchscreenStoryboard = false;			
 
 			static string RunningVersionFilename
 			{
@@ -401,6 +401,8 @@ namespace UnrealBuildTool
 			string ExtraData = "";
 			Ini.GetString("/Script/IOSRuntimeSettings.IOSRuntimeSettings", "AdditionalPlistData", out ExtraData);
 
+			Ini.GetBool("/Script/IOSRuntimeSettings.IOSRuntimeSettings", "bLaunchscreenStoryboard", out VersionUtilities.bLaunchscreenStoryboard);
+
 			// generate the plist file
 			StringBuilder Text = new StringBuilder();
 			Text.AppendLine("<?xml version=\"1.0\" encoding=\"UTF-8\"?>");
@@ -577,10 +579,22 @@ namespace UnrealBuildTool
 				Text.AppendLine("\t\t</dict>");
 				Text.AppendLine("\t</dict>");
 			}
-			if (File.Exists(LaunchStoryboard) || File.Exists(LaunchXib))
+			// Launchscreen storyboard only compatible on Mac for now.
+			if ((BuildHostPlatform.Current.Platform == UnrealTargetPlatform.Mac) && (File.Exists(LaunchStoryboard) || File.Exists(LaunchXib)) && VersionUtilities.bLaunchscreenStoryboard)
 			{
+				string outputStoryboard = LaunchStoryboard + "c";
+
+				string argsStoryboard = "--compile " + outputStoryboard + " " + LaunchStoryboard;
+
+				string stdOutLaunchScreen;
+
+				stdOutLaunchScreen = Utils.RunLocalProcessAndReturnStdOut("ibtool", argsStoryboard);
+
+				Log.TraceInformation("LaunchScreen Storyboard compilation results : " + stdOutLaunchScreen);
+
 				Text.AppendLine("\t<key>UILaunchStoryboardName</key>");
 				Text.AppendLine("\t<string>LaunchScreen</string>");
+
 				bSkipDefaultPNGs = true;
 			}
 			else
@@ -878,46 +892,52 @@ namespace UnrealBuildTool
 			// copy engine assets in (IOS and TVOS shared in IOS)
 			if (bSkipDefaultPNGs)
 			{
-				//Copy storyboard or xib
-				CopyFiles(InEngineDir + "/Build/IOS/Resources/Interface", AppDirectory, "*", true);
-				// we still want default icons
-				if (!bSkipIcons)
+				if (VersionUtilities.bLaunchscreenStoryboard)
 				{
-					CopyFiles(InEngineDir + "/Build/IOS/Resources/Graphics", AppDirectory, "Icon*.png", true);
+					CopyFolder(InEngineDir + "/Build/IOS/Resources/Interface/LaunchScreen.storyboardc", AppDirectory + "/LaunchScreen.storyboardc", true);
+					CopyFiles(InEngineDir + "/Build/IOS/Resources/Interface/Assets", AppDirectory, "*", true);
 				}
-			}
-			else
-			{
+				// we still want default icons
 				if (!bSkipIcons)
 				{
 					CopyFiles(InEngineDir + "/Build/IOS/Resources/Graphics", AppDirectory, "Icon*.png", true);
 				}
-				if (bSupportsPortrait)
-				{
-					CopyFiles(InEngineDir + "/Build/IOS/Resources/Graphics", AppDirectory, "Default-IPhone6.png", true);
-					CopyFiles(InEngineDir + "/Build/IOS/Resources/Graphics", AppDirectory, "Default-IPhone6Plus-Portrait.png", true);
-					CopyFiles(InEngineDir + "/Build/IOS/Resources/Graphics", AppDirectory, "[REDACTED_EMAIL]", true);
-					CopyFiles(InEngineDir + "/Build/IOS/Resources/Graphics", AppDirectory, "[REDACTED_EMAIL]", true);
-					CopyFiles(InEngineDir + "/Build/IOS/Resources/Graphics", AppDirectory, "[REDACTED_EMAIL]", true);
-					CopyFiles(InEngineDir + "/Build/IOS/Resources/Graphics", AppDirectory, "[REDACTED_EMAIL]", true);
-					CopyFiles(InEngineDir + "/Build/IOS/Resources/Graphics", AppDirectory, "Default-IPhoneXS-Portrait.png", true);
-					CopyFiles(InEngineDir + "/Build/IOS/Resources/Graphics", AppDirectory, "Default-IPhoneXSMax-Portrait.png", true);
-					CopyFiles(InEngineDir + "/Build/IOS/Resources/Graphics", AppDirectory, "Default-IPhoneXR-Portrait.png", true);
-				}
-				if (bSupportsLandscape)
+				else
 				{
-					CopyFiles(InEngineDir + "/Build/IOS/Resources/Graphics", AppDirectory, "Default-IPhone6-Landscape.png", true);
-					CopyFiles(InEngineDir + "/Build/IOS/Resources/Graphics", AppDirectory, "Default-IPhone6Plus-Landscape.png", true);
-					CopyFiles(InEngineDir + "/Build/IOS/Resources/Graphics", AppDirectory, "[REDACTED_EMAIL]", true);
-					CopyFiles(InEngineDir + "/Build/IOS/Resources/Graphics", AppDirectory, "[REDACTED_EMAIL]", true);
-					CopyFiles(InEngineDir + "/Build/IOS/Resources/Graphics", AppDirectory, "[REDACTED_EMAIL]", true);
-					CopyFiles(InEngineDir + "/Build/IOS/Resources/Graphics", AppDirectory, "[REDACTED_EMAIL]", true);
-					CopyFiles(InEngineDir + "/Build/IOS/Resources/Graphics", AppDirectory, "Default-IPhoneXS-Landscape.png", true);
-					CopyFiles(InEngineDir + "/Build/IOS/Resources/Graphics", AppDirectory, "Default-IPhoneXSMax-Landscape.png", true);
-					CopyFiles(InEngineDir + "/Build/IOS/Resources/Graphics", AppDirectory, "Default-IPhoneXR-Landscape.png", true);
+					if (!bSkipIcons)
+					{
+						CopyFiles(InEngineDir + "/Build/IOS/Resources/Graphics", AppDirectory, "Icon*.png", true);
+					}
+					if (bSupportsPortrait && !VersionUtilities.bLaunchscreenStoryboard)
+					{
+						CopyFiles(InEngineDir + "/Build/IOS/Resources/Graphics", AppDirectory, "Default-IPhone6.png", true);
+						CopyFiles(InEngineDir + "/Build/IOS/Resources/Graphics", AppDirectory, "Default-IPhone6Plus-Portrait.png", true);
+						CopyFiles(InEngineDir + "/Build/IOS/Resources/Graphics", AppDirectory, "[REDACTED_EMAIL]", true);
+						CopyFiles(InEngineDir + "/Build/IOS/Resources/Graphics", AppDirectory, "[REDACTED_EMAIL]", true);
+						CopyFiles(InEngineDir + "/Build/IOS/Resources/Graphics", AppDirectory, "[REDACTED_EMAIL]", true);
+						CopyFiles(InEngineDir + "/Build/IOS/Resources/Graphics", AppDirectory, "[REDACTED_EMAIL]", true);
+						CopyFiles(InEngineDir + "/Build/IOS/Resources/Graphics", AppDirectory, "Default-IPhoneXS-Portrait.png", true);
+						CopyFiles(InEngineDir + "/Build/IOS/Resources/Graphics", AppDirectory, "Default-IPhoneXSMax-Portrait.png", true);
+						CopyFiles(InEngineDir + "/Build/IOS/Resources/Graphics", AppDirectory, "Default-IPhoneXR-Portrait.png", true);
+					}
+					if (bSupportsLandscape && !VersionUtilities.bLaunchscreenStoryboard)
+					{
+						CopyFiles(InEngineDir + "/Build/IOS/Resources/Graphics", AppDirectory, "Default-IPhone6-Landscape.png", true);
+						CopyFiles(InEngineDir + "/Build/IOS/Resources/Graphics", AppDirectory, "Default-IPhone6Plus-Landscape.png", true);
+						CopyFiles(InEngineDir + "/Build/IOS/Resources/Graphics", AppDirectory, "[REDACTED_EMAIL]", true);
+						CopyFiles(InEngineDir + "/Build/IOS/Resources/Graphics", AppDirectory, "[REDACTED_EMAIL]", true);
+						CopyFiles(InEngineDir + "/Build/IOS/Resources/Graphics", AppDirectory, "[REDACTED_EMAIL]", true);
+						CopyFiles(InEngineDir + "/Build/IOS/Resources/Graphics", AppDirectory, "[REDACTED_EMAIL]", true);
+						CopyFiles(InEngineDir + "/Build/IOS/Resources/Graphics", AppDirectory, "Default-IPhoneXS-Landscape.png", true);
+						CopyFiles(InEngineDir + "/Build/IOS/Resources/Graphics", AppDirectory, "Default-IPhoneXSMax-Landscape.png", true);
+						CopyFiles(InEngineDir + "/Build/IOS/Resources/Graphics", AppDirectory, "Default-IPhoneXR-Landscape.png", true);
+					}
+					if (!VersionUtilities.bLaunchscreenStoryboard)
+					{
+						CopyFiles(BuildDirectory + "/Resources/Graphics", AppDirectory, "[REDACTED_EMAIL]", true);
+						CopyFiles(BuildDirectory + "/Resources/Graphics", AppDirectory, "[REDACTED_EMAIL]", true);
+					}
 				}
-				CopyFiles(InEngineDir + "/Build/IOS/Resources/Graphics", AppDirectory, "[REDACTED_EMAIL]", true);
-				CopyFiles(InEngineDir + "/Build/IOS/Resources/Graphics", AppDirectory, "[REDACTED_EMAIL]", true);
 			}
 			// merge game assets on top
 			// @todo tvos: Do we want to copy IOS and TVOS both in? (Engine/IOS -> Game/IOS -> Game/TVOS)?
@@ -955,6 +975,7 @@ namespace UnrealBuildTool
 				CopyFiles(BuildDirectory + "/Resources/Graphics", AppDirectory, "[REDACTED_EMAIL]", true);
 			}
 		}
+		
 		protected virtual void CopyLocalizationsResources(string InEngineDir, string AppDirectory, string BuildDirectory, string IntermediateDir)
 		{
 			string LocalizationsPath = BuildDirectory + "/Resources/Localizations";
diff --git a/Engine/Source/Runtime/IOS/IOSRuntimeSettings/Classes/IOSRuntimeSettings.h b/Engine/Source/Runtime/IOS/IOSRuntimeSettings/Classes/IOSRuntimeSettings.h
--- a/Engine/Source/Runtime/IOS/IOSRuntimeSettings/Classes/IOSRuntimeSettings.h
+++ b/Engine/Source/Runtime/IOS/IOSRuntimeSettings/Classes/IOSRuntimeSettings.h
@@ -424,6 +424,14 @@ public:
 	UPROPERTY(GlobalConfig, EditAnywhere, Category = ExtraData)
 	FString AdditionalPlistData;
 
+	/**
+	 * Choose whether to use a storyboard as a Launchscreen or not. If True, LaunchScreen.storyboard in Build/IOS/Resources/Interface
+	 * in the Engine folder will be copied to the bundle app and the Launchscreen images WILL NOT be included.
+	 * When using assets in your LaunchScreen.Storyboard, add them in Build/IOS/Resources/Interface/Assets and they will be included.
+	 */
+	UPROPERTY(GlobalConfig, EditAnywhere, Category = LaunchScreen)
+	bool 	bLaunchscreenStoryboard;
+
 	// Whether the app supports Facebook
 	UPROPERTY(GlobalConfig, EditAnywhere, Category = Online)
 	bool bEnableFacebookSupport;
diff --git a/Engine/Source/Runtime/IOS/IOSRuntimeSettings/Private/IOSRuntimeSettings.cpp b/Engine/Source/Runtime/IOS/IOSRuntimeSettings/Private/IOSRuntimeSettings.cpp
--- a/Engine/Source/Runtime/IOS/IOSRuntimeSettings/Private/IOSRuntimeSettings.cpp
+++ b/Engine/Source/Runtime/IOS/IOSRuntimeSettings/Private/IOSRuntimeSettings.cpp
@@ -37,6 +37,7 @@ UIOSRuntimeSettings::UIOSRuntimeSettings(const FObjectInitializer& ObjectInitial
 	bShipForArmV7S = false;
 	bShipForBitcode = true;
 	bUseRSync = true;
+	bLaunchscreenStoryboard = false;
 	AdditionalPlistData = TEXT("");
 	AdditionalLinkerFlags = TEXT("");
 	AdditionalShippingLinkerFlags = TEXT("");
