Disable XGE if the user is connected via VPN (determined as whether they have a network adapter with an IP address within a list of subnets)

diff --git a/Engine/Source/Programs/UnrealBuildTool/Executors/XGE.cs b/Engine/Source/Programs/UnrealBuildTool/Executors/XGE.cs
--- a/Engine/Source/Programs/UnrealBuildTool/Executors/XGE.cs
+++ b/Engine/Source/Programs/UnrealBuildTool/Executors/XGE.cs
@@ -12,6 +12,8 @@ using System.Reflection;
 using Microsoft.Win32;
 using System.Text;
 using Tools.DotNETCommon;
+using System.Net;
+using System.Net.NetworkInformation;
 
 namespace UnrealBuildTool
 {
@@ -35,6 +37,18 @@ namespace UnrealBuildTool
 		[XmlConfigFile(Category = "BuildConfiguration")]
 		bool bStopXGECompilationAfterErrors = false;
 
+		/// <summary>
+		/// When set to false, XGE will not be When enabled, XGE will stop compiling targets after a compile error occurs.  Recommended, as it saves computing resources for others.
+		/// </summary>
+		[XmlConfigFile(Category = "XGE")]
+		static bool bAllowOverVpn = true;
+
+		/// <summary>
+		/// List of subnets containing IP addresses assigned by VPN
+		/// </summary>
+		[XmlConfigFile(Category = "XGE")]
+		static string[] VpnSubnets = null;
+
 		private const string ProgressMarkupPrefix = "@action";
 
 		public XGE()
@@ -135,25 +149,68 @@ namespace UnrealBuildTool
 		public static bool IsAvailable()
 		{
 			string XgConsoleExe;
-
-			bool Success = TryGetXgConsoleExecutable(out XgConsoleExe);
+			if (!TryGetXgConsoleExecutable(out XgConsoleExe))
+			{
+				return false;
+			}
 
 			// on windows check the service is actually running
-			if (Success && BuildHostPlatform.Current.Platform == UnrealTargetPlatform.Win64)
+			if (BuildHostPlatform.Current.Platform == UnrealTargetPlatform.Win64)
 			{
 				try
 				{
 					// will throw if the service doesn't exist, which it should if IB is present but just incase...
 					System.ServiceProcess.ServiceController SC = new System.ServiceProcess.ServiceController("Incredibuild Agent");
-					Success = SC.Status == System.ServiceProcess.ServiceControllerStatus.Running;
+					if (SC.Status != System.ServiceProcess.ServiceControllerStatus.Running)
+					{
+						return false;
+					}
+				}
+				catch(Exception Ex)
+				{
+					Log.TraceLog("Unable to query for status of Incredibuild service: {0}", ExceptionUtils.FormatExceptionDetails(Ex));
+					return false;
+				}
+			}
+
+			// Check if we're connected over VPN
+			if (!bAllowOverVpn && VpnSubnets != null && VpnSubnets.Length > 0)
+			{
+				// Parse all the subnets from the config file
+				List<Subnet> ParsedVpnSubnets = new List<Subnet>();
+				foreach (string VpnSubnet in VpnSubnets)
+				{
+					ParsedVpnSubnets.Add(Subnet.Parse(VpnSubnet));
+				}
+
+				// Check if any network adapters have an IP within one of these subnets
+				try
+				{
+					NetworkInterface[] Interfaces = NetworkInterface.GetAllNetworkInterfaces();
+					foreach (NetworkInterface Interface in Interfaces)
+					{
+						IPInterfaceProperties Properties = Interface.GetIPProperties();
+						foreach (UnicastIPAddressInformation UnicastAddressInfo in Properties.UnicastAddresses)
+						{
+							byte[] AddressBytes = UnicastAddressInfo.Address.GetAddressBytes();
+							foreach (Subnet Subnet in ParsedVpnSubnets)
+							{
+								if (Subnet.Contains(AddressBytes))
+								{
+									Log.TraceInformationOnce("XGE will be not be used over VPN (adapter '{0}' with IP {1} is in subnet {2}). Set <XGE><bAllowOverVpn>true</bAllowOverVpn></XGE> in BuildConfiguration.xml to override.", Interface.Name, UnicastAddressInfo.Address, Subnet);
+									return false;
+								}
+							}
+						}
+					}
 				}
-				catch
+				catch (Exception Ex)
 				{
-					Success = false;
+					Log.TraceWarning("Unable to check whether machine is connected to VPN:\n{0}", ExceptionUtils.FormatExceptionDetails(Ex));
 				}
 			}
 
-			return Success;
+			return true;
 		}
 
 		// precompile the Regex needed to parse the XGE output (the ones we want are of the form "File (Duration at +time)"
diff --git a/Engine/Source/Programs/UnrealBuildTool/System/Subnet.cs b/Engine/Source/Programs/UnrealBuildTool/System/Subnet.cs
new file mode 100644
--- /dev/null
+++ b/Engine/Source/Programs/UnrealBuildTool/System/Subnet.cs
@@ -0,0 +1,120 @@
+using System;
+using System.Collections.Generic;
+using System.Diagnostics;
+using System.Linq;
+using System.Net;
+using System.Text;
+using System.Threading.Tasks;
+
+namespace UnrealBuildTool
+{
+	/// <summary>
+	/// Represents an IP subnet
+	/// </summary>
+	class Subnet
+	{
+		/// <summary>
+		/// The prefix address
+		/// </summary>
+		public IPAddress Prefix
+		{
+			get;
+			private set;
+		}
+
+		/// <summary>
+		/// Number of bits that need to match in an IP address for this subnet
+		/// </summary>
+		public int MaskBits
+		{
+			get;
+			private set;
+		}
+
+		/// <summary>
+		/// Bytes corresponding to the prefix address
+		/// </summary>
+		byte[] PrefixBytes;
+
+		/// <summary>
+		/// Constructor
+		/// </summary>
+		/// <param name="Prefix">The prefix IP address</param>
+		/// <param name="MaskBits">Number of bits to match for this subnet</param>
+		public Subnet(IPAddress Prefix, int MaskBits)
+		{
+			this.Prefix = Prefix;
+			this.PrefixBytes = Prefix.GetAddressBytes();
+			this.MaskBits = MaskBits;
+		}
+
+		/// <summary>
+		/// Parses a subnet from a string
+		/// </summary>
+		/// <param name="Text">The string to parse</param>
+		/// <returns>New subnet that was parsed</returns>
+		public static Subnet Parse(string Text)
+		{
+			int SlashIdx = Text.IndexOf('/');
+			IPAddress Address = IPAddress.Parse(Text.Substring(0, SlashIdx));
+			return new Subnet(Address, int.Parse(Text.Substring(SlashIdx + 1)));
+		}
+
+		/// <summary>
+		/// Checks if this subnet contains the given address
+		/// </summary>
+		/// <param name="Address">IP address to test</param>
+		/// <returns>True if the subnet contains this address</returns>
+		public bool Contains(IPAddress Address)
+		{
+			return Contains(Address.GetAddressBytes());
+		}
+
+		/// <summary>
+		/// Checks if this subnet contains the given address bytes
+		/// </summary>
+		/// <param name="Bytes">Bytes of the IP address to text</param>
+		/// <returns>True if the subnet contains this address</returns>
+		public bool Contains(byte[] Bytes)
+		{
+			if (Bytes.Length != PrefixBytes.Length)
+			{
+				return false;
+			}
+
+			int Index = 0;
+			int RemainingBits = MaskBits;
+
+			// Check all the full bytes first
+			for (; RemainingBits >= 8; RemainingBits -= 8)
+			{
+				if (Bytes[Index] != PrefixBytes[Index])
+				{
+					return false;
+				}
+				Index++;
+			}
+
+			// Check the remaining bits
+			if (RemainingBits > 0)
+			{
+				int LastMaskByte = ((1 << RemainingBits) - 1) << (8 - RemainingBits);
+				if ((Bytes[Index] & LastMaskByte) != (PrefixBytes[Index] & LastMaskByte))
+				{
+					return false;
+				}
+			}
+
+			return true;
+		}
+
+		/// <summary>
+		/// Formats this subnet as text
+		/// </summary>
+		/// <returns>String representation of the subnet</returns>
+		public override string ToString()
+		{
+			return String.Format("{0}/{1}", Prefix, MaskBits);
+		}
+	}
+}
diff --git a/Engine/Source/Programs/UnrealBuildTool/UnrealBuildTool.csproj b/Engine/Source/Programs/UnrealBuildTool/UnrealBuildTool.csproj
--- a/Engine/Source/Programs/UnrealBuildTool/UnrealBuildTool.csproj
+++ b/Engine/Source/Programs/UnrealBuildTool/UnrealBuildTool.csproj
@@ -241,6 +241,7 @@
     <Compile Include="System\SingleInstanceMutex.cs" />
     <Compile Include="System\SourceFileWorkingSet.cs" />
     <Compile Include="System\StartupTraceListener.cs" />
+    <Compile Include="System\Subnet.cs" />
     <Compile Include="System\SupportedConfigurationsAttribute.cs" />
     <Compile Include="System\SupportedPlatformsAttribute.cs" />
     <Compile Include="System\TargetMakefile.cs" />
