#UE4 added copyright

diff --git a/Engine/Source/Programs/UnrealBuildTool/System/Subnet.cs b/Engine/Source/Programs/UnrealBuildTool/System/Subnet.cs
--- a/Engine/Source/Programs/UnrealBuildTool/System/Subnet.cs
+++ b/Engine/Source/Programs/UnrealBuildTool/System/Subnet.cs
@@ -1,3 +1,5 @@
+// Copyright Epic Games, Inc. All Rights Reserved.
+
 using System;
 using System.Collections.Generic;
 using System.Diagnostics;
