Add additional checks to disabling XGE over VPN.

diff --git a/Engine/Source/Programs/UnrealBuildTool/Executors/XGE.cs b/Engine/Source/Programs/UnrealBuildTool/Executors/XGE.cs
--- a/Engine/Source/Programs/UnrealBuildTool/Executors/XGE.cs
+++ b/Engine/Source/Programs/UnrealBuildTool/Executors/XGE.cs
@@ -189,16 +189,19 @@ namespace UnrealBuildTool
 					NetworkInterface[] Interfaces = NetworkInterface.GetAllNetworkInterfaces();
 					foreach (NetworkInterface Interface in Interfaces)
 					{
-						IPInterfaceProperties Properties = Interface.GetIPProperties();
-						foreach (UnicastIPAddressInformation UnicastAddressInfo in Properties.UnicastAddresses)
+						if (Interface.OperationalStatus == OperationalStatus.Up)
 						{
-							byte[] AddressBytes = UnicastAddressInfo.Address.GetAddressBytes();
-							foreach (Subnet Subnet in ParsedVpnSubnets)
+							IPInterfaceProperties Properties = Interface.GetIPProperties();
+							foreach (UnicastIPAddressInformation UnicastAddressInfo in Properties.UnicastAddresses)
 							{
-								if (Subnet.Contains(AddressBytes))
+								byte[] AddressBytes = UnicastAddressInfo.Address.GetAddressBytes();
+								foreach (Subnet Subnet in ParsedVpnSubnets)
 								{
-									Log.TraceInformationOnce("XGE will be not be used over VPN (adapter '{0}' with IP {1} is in subnet {2}). Set <XGE><bAllowOverVpn>true</bAllowOverVpn></XGE> in BuildConfiguration.xml to override.", Interface.Name, UnicastAddressInfo.Address, Subnet);
-									return false;
+									if (Subnet.Contains(AddressBytes))
+									{
+										Log.TraceInformationOnce("XGE will be not be used over VPN (adapter '{0}' with IP {1} is in subnet {2}). Set <XGE><bAllowOverVpn>true</bAllowOverVpn></XGE> in BuildConfiguration.xml to override.", Interface.Description, UnicastAddressInfo.Address, Subnet);
+										return false;
+									}
 								}
 							}
 						}
