Refine logic for disabling XGE over VPN. Now reads the coordinator hostname from the registry, and figures out which network adapter it routes through. XGE is only disabled if the IP address of this adapter is on a VPN subnet.

diff --git a/Engine/Source/Programs/UnrealBuildTool/Executors/XGE.cs b/Engine/Source/Programs/UnrealBuildTool/Executors/XGE.cs
--- a/Engine/Source/Programs/UnrealBuildTool/Executors/XGE.cs
+++ b/Engine/Source/Programs/UnrealBuildTool/Executors/XGE.cs
@@ -14,6 +14,7 @@ using System.Text;
 using Tools.DotNETCommon;
 using System.Net;
 using System.Net.NetworkInformation;
+using System.Runtime.InteropServices;
 
 namespace UnrealBuildTool
 {
@@ -146,6 +147,117 @@ namespace UnrealBuildTool
 			return false;
 		}
 
+		static bool TryReadRegistryValue(RegistryHive Hive, RegistryView View, string KeyName, string ValueName, out string OutCoordinator)
+		{
+			using (RegistryKey BaseKey = RegistryKey.OpenBaseKey(Hive, View))
+			{
+				using (RegistryKey SubKey = BaseKey.OpenSubKey(KeyName))
+				{
+					if (SubKey != null)
+					{
+						string Coordinator = SubKey.GetValue(ValueName) as string;
+						if (!String.IsNullOrEmpty(Coordinator))
+						{
+							OutCoordinator = Coordinator;
+							return true;
+						}
+					}
+				}
+			}
+
+			OutCoordinator = null;
+			return false;
+		}
+
+		static bool TryGetCoordinatorHost(out string OutCoordinator)
+		{
+			if (BuildHostPlatform.Current.Platform == UnrealTargetPlatform.Win64)
+			{
+				const string KeyName = @"SOFTWARE\Xoreax\IncrediBuild\BuildService";
+				const string ValueName = "CoordHost";
+
+				return TryReadRegistryValue(RegistryHive.CurrentUser, RegistryView.Registry64, KeyName, ValueName, out OutCoordinator) ||
+					TryReadRegistryValue(RegistryHive.CurrentUser, RegistryView.Registry32, KeyName, ValueName, out OutCoordinator) ||
+					TryReadRegistryValue(RegistryHive.LocalMachine, RegistryView.Registry64, KeyName, ValueName, out OutCoordinator) ||
+					TryReadRegistryValue(RegistryHive.LocalMachine, RegistryView.Registry32, KeyName, ValueName, out OutCoordinator);
+			}
+			else
+			{
+				OutCoordinator = null;
+				return false;
+			}
+		}
+
+		[DllImport("iphlpapi")]
+		static extern int GetBestInterface(uint dwDestAddr, ref int pdwBestIfIndex);
+
+		static NetworkInterface GetInterfaceForHost(string Host)
+		{
+			if (BuildHostPlatform.Current.Platform == UnrealTargetPlatform.Win64)
+			{
+				IPHostEntry HostEntry = Dns.GetHostEntry(Host);
+				foreach (IPAddress HostAddress in HostEntry.AddressList)
+				{
+					int InterfaceIdx = 0;
+					if (GetBestInterface(BitConverter.ToUInt32(HostAddress.GetAddressBytes(), 0), ref InterfaceIdx) == 0)
+					{
+						foreach (NetworkInterface Interface in NetworkInterface.GetAllNetworkInterfaces())
+						{
+							IPv4InterfaceProperties Properties = Interface.GetIPProperties().GetIPv4Properties();
+							if (Properties.Index == InterfaceIdx)
+							{
+								return Interface;
+							}
+						}
+					}
+				}
+			}
+			return null;
+		}
+
+		public static bool IsHostOnVpn(string HostName)
+		{
+			// If there aren't any defined subnets, just early out
+			if (VpnSubnets == null || VpnSubnets.Length == 0)
+			{
+				return false;
+			}
+
+			// Parse all the subnets from the config file
+			List<Subnet> ParsedVpnSubnets = new List<Subnet>();
+			foreach (string VpnSubnet in VpnSubnets)
+			{
+				ParsedVpnSubnets.Add(Subnet.Parse(VpnSubnet));
+			}
+
+			// Check if any network adapters have an IP within one of these subnets
+			try
+			{
+				NetworkInterface Interface = GetInterfaceForHost(HostName);
+				if (Interface != null && Interface.OperationalStatus == OperationalStatus.Up)
+				{
+					IPInterfaceProperties Properties = Interface.GetIPProperties();
+					foreach (UnicastIPAddressInformation UnicastAddressInfo in Properties.UnicastAddresses)
+					{
+						byte[] AddressBytes = UnicastAddressInfo.Address.GetAddressBytes();
+						foreach (Subnet Subnet in ParsedVpnSubnets)
+						{
+							if (Subnet.Contains(AddressBytes))
+							{
+								Log.TraceInformationOnce("XGE coordinator {0} will be not be used over VPN (adapter '{1}' with IP {2} is in subnet {3}). Set <XGE><bAllowOverVpn>true</bAllowOverVpn></XGE> in BuildConfiguration.xml to override.", HostName, Interface.Description, UnicastAddressInfo.Address, Subnet);
+								return false;
+							}
+						}
+					}
+				}
+			}
+			catch (Exception Ex)
+			{
+				Log.TraceWarning("Unable to check whether host {0} is connected to VPN:\n{1}", HostName, ExceptionUtils.FormatExceptionDetails(Ex));
+			}
+			return false;
+		}
+
 		public static bool IsAvailable()
 		{
 			string XgConsoleExe;
@@ -176,40 +288,10 @@ namespace UnrealBuildTool
 			// Check if we're connected over VPN
 			if (!bAllowOverVpn && VpnSubnets != null && VpnSubnets.Length > 0)
 			{
-				// Parse all the subnets from the config file
-				List<Subnet> ParsedVpnSubnets = new List<Subnet>();
-				foreach (string VpnSubnet in VpnSubnets)
-				{
-					ParsedVpnSubnets.Add(Subnet.Parse(VpnSubnet));
-				}
-
-				// Check if any network adapters have an IP within one of these subnets
-				try
+				string CoordinatorHost;
+				if (TryGetCoordinatorHost(out CoordinatorHost) && IsHostOnVpn(CoordinatorHost))
 				{
-					NetworkInterface[] Interfaces = NetworkInterface.GetAllNetworkInterfaces();
-					foreach (NetworkInterface Interface in Interfaces)
-					{
-						if (Interface.OperationalStatus == OperationalStatus.Up)
-						{
-							IPInterfaceProperties Properties = Interface.GetIPProperties();
-							foreach (UnicastIPAddressInformation UnicastAddressInfo in Properties.UnicastAddresses)
-							{
-								byte[] AddressBytes = UnicastAddressInfo.Address.GetAddressBytes();
-								foreach (Subnet Subnet in ParsedVpnSubnets)
-								{
-									if (Subnet.Contains(AddressBytes))
-									{
-										Log.TraceInformationOnce("XGE will be not be used over VPN (adapter '{0}' with IP {1} is in subnet {2}). Set <XGE><bAllowOverVpn>true</bAllowOverVpn></XGE> in BuildConfiguration.xml to override.", Interface.Description, UnicastAddressInfo.Address, Subnet);
-										return false;
-									}
-								}
-							}
-						}
-					}
-				}
-				catch (Exception Ex)
-				{
-					Log.TraceWarning("Unable to check whether machine is connected to VPN:\n{0}", ExceptionUtils.FormatExceptionDetails(Ex));
+					return false;
 				}
 			}
 
