Change Lumin UBT to the LLVM toolchain
Prevents mixing of LLVM and GCC toolchains

diff --git a/Engine/Source/Programs/UnrealBuildTool/Platform/Lumin/LuminToolChain.cs b/Engine/Source/Programs/UnrealBuildTool/Platform/Lumin/LuminToolChain.cs
--- a/Engine/Source/Programs/UnrealBuildTool/Platform/Lumin/LuminToolChain.cs
+++ b/Engine/Source/Programs/UnrealBuildTool/Platform/Lumin/LuminToolChain.cs
@@ -79,7 +79,7 @@ namespace UnrealBuildTool
 
 			MLSDKPath = MLSDKPath.Replace("\"", "");
 
-			string MabuSpec = RunMabuAndReadOutput("-t device --print-spec -q");
+			string MabuSpec = RunMabuAndReadOutput("-t lumin_clang-8.0 --print-spec -q");
 
 			// parse clange version
 			Regex SpecRegex = new Regex("\\s*(?:[a-z]+_lumin_clang-)(\\d)[.](\\d)\\s*");
@@ -95,7 +95,7 @@ namespace UnrealBuildTool
 
 			SetClangVersion(int.Parse(SpecMatch.Groups[1].Value), int.Parse(SpecMatch.Groups[2].Value), 0);
 
-			string MabuTools = RunMabuAndReadOutput("-t device --print-tools -q");
+			string MabuTools = RunMabuAndReadOutput("-t lumin_clang-8.0 --print-tools -q");
 
 			Dictionary<string, string> ToolsDict = new Dictionary<string, string>();
 			using (StringReader Reader = new StringReader(MabuTools))
