Merging support for -Preview toolchains to Main.

#rnx

diff --git a/Engine/Source/Programs/UnrealBuildTool/Platform/Windows/UEBuildWindows.cs b/Engine/Source/Programs/UnrealBuildTool/Platform/Windows/UEBuildWindows.cs
--- a/Engine/Source/Programs/UnrealBuildTool/Platform/Windows/UEBuildWindows.cs
+++ b/Engine/Source/Programs/UnrealBuildTool/Platform/Windows/UEBuildWindows.cs
@@ -1121,15 +1121,17 @@ namespace UnrealBuildTool
 
 						// Enumerate all the AutoSDK toolchains
 						DirectoryReference PlatformDir;
-						if(UEBuildPlatformSDK.TryGetHostPlatformAutoSDKDir(out PlatformDir))
+						if (UEBuildPlatformSDK.TryGetHostPlatformAutoSDKDir(out PlatformDir))
 						{
-							DirectoryReference ToolChainBaseDir = DirectoryReference.Combine(PlatformDir, "Win64", (Compiler == WindowsCompiler.VisualStudio2019)? "VS2019" : "VS2017");
-							if(DirectoryReference.Exists(ToolChainBaseDir))
+							DirectoryReference ToolChainBaseDir = DirectoryReference.Combine(PlatformDir, "Win64", (Compiler == WindowsCompiler.VisualStudio2019) ? "VS2019" : "VS2017");
+							DirectoryReference PreviewToolChainBaseDir = DirectoryReference.Combine(ToolChainBaseDir.ParentDirectory, string.Format("{0}-Preview", ToolChainBaseDir.GetDirectoryName()));
+							IEnumerable<DirectoryReference> ToolChainDirectories = DirectoryReference.EnumerateDirectories(PreviewToolChainBaseDir).Concat(DirectoryReference.EnumerateDirectories(ToolChainBaseDir));
+							if (DirectoryReference.Exists(ToolChainBaseDir))
 							{
-								foreach(DirectoryReference ToolChainDir in DirectoryReference.EnumerateDirectories(ToolChainBaseDir))
+								foreach (DirectoryReference ToolChainDir in ToolChainDirectories)
 								{
 									VersionNumber Version;
-									if(IsValidToolChainDir2017or2019(ToolChainDir, out Version))
+									if (IsValidToolChainDir2017or2019(ToolChainDir, out Version))
 									{
 										Log.TraceLog("Found Visual Studio toolchain: {0} (Version={1})", ToolChainDir, Version);
 										if (!ToolChainVersionToDir.ContainsKey(Version))
