Merging fix for crash if VS2019-Preview directory does not exist.

#rnx

diff --git a/Engine/Source/Programs/UnrealBuildTool/Platform/Windows/UEBuildWindows.cs b/Engine/Source/Programs/UnrealBuildTool/Platform/Windows/UEBuildWindows.cs
--- a/Engine/Source/Programs/UnrealBuildTool/Platform/Windows/UEBuildWindows.cs
+++ b/Engine/Source/Programs/UnrealBuildTool/Platform/Windows/UEBuildWindows.cs
@@ -1101,46 +1101,19 @@ namespace UnrealBuildTool
 						List<DirectoryReference> InstallDirs = FindVSInstallDirs(Compiler);
 					    foreach(DirectoryReference InstallDir in InstallDirs)
 					    {
-						    DirectoryReference ToolChainBaseDir = DirectoryReference.Combine(InstallDir, "VC", "Tools", "MSVC");
-						    if(DirectoryReference.Exists(ToolChainBaseDir))
-						    {
-							    foreach(DirectoryReference ToolChainDir in DirectoryReference.EnumerateDirectories(ToolChainBaseDir))
-							    {
-								    VersionNumber Version;
-								    if(IsValidToolChainDir2017or2019(ToolChainDir, out Version))
-									{
-										Log.TraceLog("Found Visual Studio toolchain: {0} (Version={1})", ToolChainDir, Version);
-										if(!ToolChainVersionToDir.ContainsKey(Version))
-										{
-											ToolChainVersionToDir[Version] = ToolChainDir;
-										}
-								    }
-							    }
-						    }
+						    DirectoryReference BaseDir = DirectoryReference.Combine(InstallDir, "VC", "Tools", "MSVC");
+							FindVisualStudioToolChains(BaseDir, ToolChainVersionToDir);
 					    }
 
 						// Enumerate all the AutoSDK toolchains
 						DirectoryReference PlatformDir;
 						if (UEBuildPlatformSDK.TryGetHostPlatformAutoSDKDir(out PlatformDir))
 						{
-							DirectoryReference ToolChainBaseDir = DirectoryReference.Combine(PlatformDir, "Win64", (Compiler == WindowsCompiler.VisualStudio2019) ? "VS2019" : "VS2017");
-							DirectoryReference PreviewToolChainBaseDir = DirectoryReference.Combine(ToolChainBaseDir.ParentDirectory, string.Format("{0}-Preview", ToolChainBaseDir.GetDirectoryName()));
-							IEnumerable<DirectoryReference> ToolChainDirectories = DirectoryReference.EnumerateDirectories(PreviewToolChainBaseDir).Concat(DirectoryReference.EnumerateDirectories(ToolChainBaseDir));
-							if (DirectoryReference.Exists(ToolChainBaseDir))
-							{
-								foreach (DirectoryReference ToolChainDir in ToolChainDirectories)
-								{
-									VersionNumber Version;
-									if (IsValidToolChainDir2017or2019(ToolChainDir, out Version))
-									{
-										Log.TraceLog("Found Visual Studio toolchain: {0} (Version={1})", ToolChainDir, Version);
-										if (!ToolChainVersionToDir.ContainsKey(Version))
-										{
-											ToolChainVersionToDir[Version] = ToolChainDir;
-										}
-									}
-								}
-							}
+							DirectoryReference ReleaseBaseDir = DirectoryReference.Combine(PlatformDir, "Win64", (Compiler == WindowsCompiler.VisualStudio2019) ? "VS2019" : "VS2017");
+							FindVisualStudioToolChains(ReleaseBaseDir, ToolChainVersionToDir);
+
+							DirectoryReference PreviewBaseDir = DirectoryReference.Combine(PlatformDir, "Win64", (Compiler == WindowsCompiler.VisualStudio2019) ? "VS2019-Preview" : "VS2017-Preview");
+							FindVisualStudioToolChains(PreviewBaseDir, ToolChainVersionToDir);
 						}
 					}
 					else
@@ -1153,6 +1126,30 @@ namespace UnrealBuildTool
 			return ToolChainVersionToDir;
 		}
 
+		/// <summary>
+		/// Finds all the valid Visual Studio toolchains under the given base directory
+		/// </summary>
+		/// <param name="BaseDir">Base directory to search</param>
+		/// <param name="ToolChainVersionToDir">Map of tool chain version to directory</param>
+		static void FindVisualStudioToolChains(DirectoryReference BaseDir, Dictionary<VersionNumber, DirectoryReference> ToolChainVersionToDir)
+		{
+			if (DirectoryReference.Exists(BaseDir))
+			{
+				foreach (DirectoryReference ToolChainDir in DirectoryReference.EnumerateDirectories(BaseDir))
+				{
+					VersionNumber Version;
+					if (IsValidToolChainDir2017or2019(ToolChainDir, out Version))
+					{
+						Log.TraceLog("Found Visual Studio toolchain: {0} (Version={1})", ToolChainDir, Version);
+						if (!ToolChainVersionToDir.ContainsKey(Version))
+						{
+							ToolChainVersionToDir[Version] = ToolChainDir;
+						}
+					}
+				}
+			}
+		}
+
 		/// <summary>
 		/// Checks if the given directory contains a valid Clang toolchain
 		/// </summary>
