Add support for different analysis modes in the PVS toolchain.

Defaults to just using general analysis (which UE4 builds cleanly with), but other options can be configured by setting TargetRules.WindowsPlatform.PVS.ModeFlags or by setting TargetRules.WindowsPlatform.PVS.UseApplicationSettings = true to read from the PVS config file.

Based on PR #6730.

diff --git a/Engine/Source/Programs/UnrealBuildTool/Platform/Windows/PVSToolChain.cs b/Engine/Source/Programs/UnrealBuildTool/Platform/Windows/PVSToolChain.cs
--- a/Engine/Source/Programs/UnrealBuildTool/Platform/Windows/PVSToolChain.cs
+++ b/Engine/Source/Programs/UnrealBuildTool/Platform/Windows/PVSToolChain.cs
@@ -13,6 +13,37 @@ using Tools.DotNETCommon;
 
 namespace UnrealBuildTool
 {
+	/// <summary>
+	/// Flags for the PVS analyzer mode
+	/// </summary>
+	public enum PVSAnalysisModeFlags : uint
+	{
+		/// <summary>
+		/// Check for 64-bit portability issues
+		/// </summary>
+		Check64BitPortability = 1,
+
+		/// <summary>
+		/// Enable general analysis
+		/// </summary>
+		GeneralAnalysis = 4,
+
+		/// <summary>
+		/// Check for optimizations
+		/// </summary>
+		Optimizations = 8,
+
+		/// <summary>
+		/// Enable customer-specific rules
+		/// </summary>
+		CustomerSpecific = 16,
+
+		/// <summary>
+		/// Enable MISRA analysis
+		/// </summary>
+		MISRA = 32,
+	}
+
 	/// <summary>
 	/// Partial representation of PVS-Studio main settings file
 	/// </summary>
@@ -33,6 +64,177 @@ namespace UnrealBuildTool
 		/// Registered serial number
 		/// </summary>
 		public string SerialNumber;
+
+		/// <summary>
+		/// Disable the 64-bit Analysis
+		/// </summary>
+		public bool Disable64BitAnalysis;
+
+		/// <summary>
+		/// Disable the General Analysis
+		/// </summary>
+		public bool DisableGAAnalysis;
+
+		/// <summary>
+		/// Disable the Optimization Analysis
+		/// </summary>
+		public bool DisableOPAnalysis;
+
+		/// <summary>
+		/// Disable the Customer's Specific diagnostic rules
+		/// </summary>
+		public bool DisableCSAnalysis;
+
+		/// <summary>
+		/// Disable the MISRA Analysis
+		/// </summary>
+		public bool DisableMISRAAnalysis;
+
+		/// <summary>
+		/// Gets the analysis mode flags from the settings
+		/// </summary>
+		/// <returns>Mode flags</returns>
+		public PVSAnalysisModeFlags GetModeFlags()
+		{
+			PVSAnalysisModeFlags Flags = 0;
+			if (!Disable64BitAnalysis)
+			{
+				Flags |= PVSAnalysisModeFlags.Check64BitPortability;
+			}
+			if (!DisableGAAnalysis)
+			{
+				Flags |= PVSAnalysisModeFlags.GeneralAnalysis;
+			}
+			if (!DisableOPAnalysis)
+			{
+				Flags |= PVSAnalysisModeFlags.Optimizations;
+			}
+			if (!DisableCSAnalysis)
+			{
+				Flags |= PVSAnalysisModeFlags.CustomerSpecific;
+			}
+			if (!DisableMISRAAnalysis)
+			{
+				Flags |= PVSAnalysisModeFlags.MISRA;
+			}
+			return Flags;
+		}
+
+		/// <summary>
+		/// Attempts to read the application settings from the default location
+		/// </summary>
+		/// <returns>Application settings instance, or null if no file was present</returns>
+		internal static PVSApplicationSettings Read()
+		{
+			FileReference SettingsPath = FileReference.Combine(new DirectoryReference(Environment.GetFolderPath(Environment.SpecialFolder.ApplicationData)), "PVS-Studio", "Settings.xml");
+			if (FileReference.Exists(SettingsPath))
+			{
+				try
+				{
+					XmlSerializer Serializer = new XmlSerializer(typeof(PVSApplicationSettings));
+					using (FileStream Stream = new FileStream(SettingsPath.FullName, FileMode.Open, FileAccess.Read, FileShare.Read))
+					{
+						return (PVSApplicationSettings)Serializer.Deserialize(Stream);
+					}
+				}
+				catch (Exception Ex)
+				{
+					throw new BuildException(Ex, "Unable to read PVS-Studio settings file from {0}", SettingsPath);
+				}
+			}
+			return null;
+		}
+	}
+
+	/// <summary>
+	/// Settings for the PVS Studio analyzer
+	/// </summary>
+	public class PVSTargetSettings
+	{
+		/// <summary>
+		/// Returns the application settings
+		/// </summary>
+		internal Lazy<PVSApplicationSettings> ApplicationSettings { get; } = new Lazy<PVSApplicationSettings>(() => PVSApplicationSettings.Read());
+
+		/// <summary>
+		/// Whether to use application settings to determine the analysis mode
+		/// </summary>
+		public bool UseApplicationSettings { get; set; }
+
+		/// <summary>
+		/// Override for the analysis mode to use
+		/// </summary>
+		public PVSAnalysisModeFlags ModeFlags
+		{
+			get
+ 			{
+				if (ModePrivate.HasValue)
+				{
+					return ModePrivate.Value;
+				}
+				else if (UseApplicationSettings && ApplicationSettings.Value != null)
+				{
+					return ApplicationSettings.Value.GetModeFlags();
+				}
+				else
+				{
+					return PVSAnalysisModeFlags.GeneralAnalysis;
+				}
+			}
+			set
+			{
+				ModePrivate = value;
+			}
+		}
+
+		/// <summary>
+		/// Private storage for the mode flags
+		/// </summary>
+		PVSAnalysisModeFlags? ModePrivate;
+	}
+
+	/// <summary>
+	/// Read-only version of the PVS toolchain settings
+	/// </summary>
+	public class ReadOnlyPVSTargetSettings
+	{
+		/// <summary>
+		/// Inner settings
+		/// </summary>
+		PVSTargetSettings Inner;
+
+		/// <summary>
+		/// Constructor
+		/// </summary>
+		/// <param name="Inner">The inner object</param>
+		public ReadOnlyPVSTargetSettings(PVSTargetSettings Inner)
+		{
+			this.Inner = Inner;
+		}
+
+		/// <summary>
+		/// Accessor for the Application settings
+		/// </summary>
+		internal PVSApplicationSettings ApplicationSettings
+		{
+			get { return Inner.ApplicationSettings.Value; }
+		}
+
+		/// <summary>
+		/// Whether to use the application settings for the mode
+		/// </summary>
+		public bool UseApplicationSettings
+		{
+			get { return Inner.UseApplicationSettings; }
+		}
+
+		/// <summary>
+		/// Override for the analysis mode to use
+		/// </summary>
+		public PVSAnalysisModeFlags ModeFlags
+		{
+			get { return Inner.ModeFlags; }
+		}
 	}
 
 	/// <summary>
@@ -132,10 +334,11 @@ namespace UnrealBuildTool
 	class PVSToolChain : UEToolChain
 	{
 		ReadOnlyTargetRules Target;
+		ReadOnlyPVSTargetSettings Settings;
+		PVSApplicationSettings ApplicationSettings;
 		VCToolChain InnerToolChain;
 		FileReference AnalyzerFile;
 		FileReference LicenseFile;
-		PVSApplicationSettings ApplicationSettings;
 		UnrealTargetPlatform Platform;
 
 		public PVSToolChain(ReadOnlyTargetRules Target)
@@ -158,28 +361,22 @@ namespace UnrealBuildTool
 				}
 			}
 
-			FileReference SettingsPath = FileReference.Combine(new DirectoryReference(Environment.GetFolderPath(Environment.SpecialFolder.ApplicationData)), "PVS-Studio", "Settings.xml");
-			if (FileReference.Exists(SettingsPath))
+			Settings = Target.WindowsPlatform.PVS;
+			ApplicationSettings = Settings.ApplicationSettings;
+
+			if(ApplicationSettings != null)
 			{
-				try
+				if (Settings.ModeFlags == 0)
 				{
-					XmlSerializer Serializer = new XmlSerializer(typeof(PVSApplicationSettings));
-					using(FileStream Stream = new FileStream(SettingsPath.FullName, FileMode.Open, FileAccess.Read, FileShare.Read))
-					{
-						ApplicationSettings = (PVSApplicationSettings)Serializer.Deserialize(Stream);
-					}
+					throw new BuildException("All PVS-Studio analysis modes are disabled.");
 				}
-				catch(Exception Ex)
+
+				if (!String.IsNullOrEmpty(ApplicationSettings.UserName) && !String.IsNullOrEmpty(ApplicationSettings.SerialNumber))
 				{
-					throw new BuildException(Ex, "Unable to read PVS-Studio settings file from {0}", SettingsPath);
+					LicenseFile = FileReference.Combine(UnrealBuildTool.EngineDirectory, "Intermediate", "PVS", "PVS-Studio.lic");
+					Utils.WriteFileIfChanged(LicenseFile, String.Format("{0}\n{1}\n", ApplicationSettings.UserName, ApplicationSettings.SerialNumber), StringComparison.Ordinal);
 				}
 			}
-
-			if(ApplicationSettings != null && !String.IsNullOrEmpty(ApplicationSettings.UserName) && !String.IsNullOrEmpty(ApplicationSettings.SerialNumber))
-			{
-				LicenseFile = FileReference.Combine(UnrealBuildTool.EngineDirectory, "Intermediate", "PVS", "PVS-Studio.lic");
-				Utils.WriteFileIfChanged(LicenseFile, String.Format("{0}\n{1}\n", ApplicationSettings.UserName, ApplicationSettings.SerialNumber), StringComparison.Ordinal);
-			}
 			else
 			{
 				FileReference DefaultLicenseFile = AnalyzerFile.ChangeExtension(".lic");
@@ -190,9 +387,12 @@ namespace UnrealBuildTool
 			}
 		}
 
-		static string GetFullIncludePath(string IncludePath)
+		public override void GetVersionInfo(List<string> Lines)
 		{
-			return Path.GetFullPath(Utils.ExpandVariables(IncludePath));
+			base.GetVersionInfo(Lines);
+
+			ReadOnlyPVSTargetSettings Settings = Target.WindowsPlatform.PVS;
+			Lines.Add(String.Format("Using PVS-Studio installation at {0} with analysis mode {1} ({2})", AnalyzerFile, (uint)Settings.ModeFlags, Settings.ModeFlags.ToString()));
 		}
 
 		class ActionGraphCapture : ForwardingActionGraphBuilder
@@ -307,7 +507,7 @@ namespace UnrealBuildTool
 				AnalyzeAction.StatusDescription = BaseFileName;
 				AnalyzeAction.WorkingDirectory = UnrealBuildTool.EngineSourceDirectory;
 				AnalyzeAction.CommandPath = AnalyzerFile;
-				AnalyzeAction.CommandArguments = String.Format("--cl-params \"{0}\" --source-file \"{1}\" --output-file \"{2}\" --cfg \"{3}\" --analysis-mode 4", PreprocessAction.CommandArguments, SourceFileItem.AbsolutePath, OutputFileLocation, ConfigFileItem.AbsolutePath);
+				AnalyzeAction.CommandArguments = String.Format("--cl-params \"{0}\" --source-file \"{1}\" --output-file \"{2}\" --cfg \"{3}\" --analysis-mode {4}", PreprocessAction.CommandArguments, SourceFileItem.AbsolutePath, OutputFileLocation, ConfigFileItem.AbsolutePath, (uint)Settings.ModeFlags);
 				if (LicenseFile != null)
 				{
 					AnalyzeAction.CommandArguments += String.Format(" --lic-file \"{0}\"", LicenseFile);
diff --git a/Engine/Source/Programs/UnrealBuildTool/Platform/Windows/UEBuildWindows.cs b/Engine/Source/Programs/UnrealBuildTool/Platform/Windows/UEBuildWindows.cs
--- a/Engine/Source/Programs/UnrealBuildTool/Platform/Windows/UEBuildWindows.cs
+++ b/Engine/Source/Programs/UnrealBuildTool/Platform/Windows/UEBuildWindows.cs
@@ -310,6 +310,11 @@ namespace UnrealBuildTool
 		/// </summary>
 		public bool bUseBundledDbgHelp = true;
 
+		/// <summary>
+		/// Settings for PVS studio
+		/// </summary>
+		public PVSTargetSettings PVS = new PVSTargetSettings();
+
 		/// <summary>
 		/// The Visual C++ environment to use for this target. Only initialized after all the target settings are finalized, in ValidateTarget().
 		/// </summary>
@@ -396,6 +401,7 @@ namespace UnrealBuildTool
 		public ReadOnlyWindowsTargetRules(WindowsTargetRules Inner)
 		{
 			this.Inner = Inner;
+			this.PVS = new ReadOnlyPVSTargetSettings(Inner.PVS);
 		}
 
 		/// <summary>
@@ -410,6 +416,7 @@ namespace UnrealBuildTool
 		{
 			get { return Inner.Compiler; }
 		}
+		
 		public WindowsArchitecture Architecture
 		{
 			get { return Inner.Architecture; }
@@ -555,6 +562,11 @@ namespace UnrealBuildTool
 			get { return Inner.bUseBundledDbgHelp; }
 		}
 
+		public ReadOnlyPVSTargetSettings PVS
+		{
+			get; private set;
+		}
+
 		internal VCEnvironment Environment
 		{
 			get { return Inner.Environment; }
