Fixed AdditionalRestrictedFolders not resulting in an array.
#rnx

diff --git a/Engine/Source/Programs/UnrealBuildTool/Configuration/ModuleRules.cs b/Engine/Source/Programs/UnrealBuildTool/Configuration/ModuleRules.cs
--- a/Engine/Source/Programs/UnrealBuildTool/Configuration/ModuleRules.cs
+++ b/Engine/Source/Programs/UnrealBuildTool/Configuration/ModuleRules.cs
@@ -752,6 +752,11 @@ namespace UnrealBuildTool
 		/// </summary>
 		public List<string> WhitelistRestrictedFolders = new List<string>();
 
+		/// <summary>
+		/// Set of aliased restricted folder references
+		/// </summary>
+		public Dictionary<string, string> AliasRestrictedFolders = new Dictionary<string, string>();
+
 		/// <summary>
 		/// Enforce "include what you use" rules when PCHUsage is set to ExplicitOrSharedPCH; warns when monolithic headers (Engine.h, UnrealEd.h, etc...) 
 		/// are used, and checks that source files include their matching header first.
diff --git a/Engine/Source/Programs/UnrealBuildTool/Configuration/UEBuildBinary.cs b/Engine/Source/Programs/UnrealBuildTool/Configuration/UEBuildBinary.cs
--- a/Engine/Source/Programs/UnrealBuildTool/Configuration/UEBuildBinary.cs
+++ b/Engine/Source/Programs/UnrealBuildTool/Configuration/UEBuildBinary.cs
@@ -607,6 +607,23 @@ namespace UnrealBuildTool
 				// Find the permitted restricted folder references under the base directory
 				List<RestrictedFolder> BinaryFolders = RestrictedFolders.FindPermittedRestrictedFolderReferences(BaseDir, OutputFilePath.Directory);
 
+				List<RestrictedFolder> AliasedBinaryFolders = new List<RestrictedFolder>();
+				foreach (RestrictedFolder BinaryFolder in BinaryFolders)
+				{
+					string Alias;
+					if (PrimaryModule.AliasRestrictedFolders.TryGetValue(BinaryFolder.ToString(), out Alias))
+					{
+						foreach(RestrictedFolder Folder in RestrictedFolder.GetValues())
+						{
+							if (Folder.ToString().Equals(Alias))
+							{
+								AliasedBinaryFolders.Add(Folder);
+							}
+						}
+					}
+				}
+				BinaryFolders.AddRange(AliasedBinaryFolders);
+
 				// Check all the dependent modules
 				foreach(UEBuildModule Module in ModuleReferencedBy.Keys)
 				{
diff --git a/Engine/Source/Programs/UnrealBuildTool/Configuration/UEBuildModule.cs b/Engine/Source/Programs/UnrealBuildTool/Configuration/UEBuildModule.cs
--- a/Engine/Source/Programs/UnrealBuildTool/Configuration/UEBuildModule.cs
+++ b/Engine/Source/Programs/UnrealBuildTool/Configuration/UEBuildModule.cs
@@ -161,6 +161,11 @@ namespace UnrealBuildTool
 		/// </summary>
 		private readonly HashSet<DirectoryReference> WhitelistRestrictedFolders;
 
+		/// <summary>
+		/// Set of aliased restricted folder references
+		/// </summary>
+		public readonly Dictionary<string, string> AliasRestrictedFolders;
+
 		/// <summary>
 		/// Constructor
 		/// </summary>
@@ -234,6 +239,7 @@ namespace UnrealBuildTool
 			}
 
 			WhitelistRestrictedFolders = new HashSet<DirectoryReference>(Rules.WhitelistRestrictedFolders.Select(x => DirectoryReference.Combine(ModuleDirectory, x)));
+			AliasRestrictedFolders = new Dictionary<string, string>(Rules.AliasRestrictedFolders);
 
 			// merge the main directory and any others set in the Rules
 			List<DirectoryReference> MergedDirectories = new List<DirectoryReference> { ModuleDirectory };
