Fix PVS Studio not working in installed engine builds.

diff --git a/Engine/Source/Programs/UnrealBuildTool/Configuration/UEBuildBinary.cs b/Engine/Source/Programs/UnrealBuildTool/Configuration/UEBuildBinary.cs
--- a/Engine/Source/Programs/UnrealBuildTool/Configuration/UEBuildBinary.cs
+++ b/Engine/Source/Programs/UnrealBuildTool/Configuration/UEBuildBinary.cs
@@ -177,7 +177,7 @@ namespace UnrealBuildTool
 		public List<FileItem> Build(ReadOnlyTargetRules Target, UEToolChain ToolChain, CppCompileEnvironment CompileEnvironment, LinkEnvironment LinkEnvironment, FileReference SingleFileToCompile, ISourceFileWorkingSet WorkingSet, DirectoryReference ExeDir, IActionGraphBuilder Graph)
 		{
 			// Return nothing if we're using precompiled binaries. If we're not linking, we might want just one module to be compiled (eg. a foreign plugin), so allow any actions to run.
-			if (bUsePrecompiled && !Target.bDisableLinking)
+			if (bUsePrecompiled && !(Target.LinkType == TargetLinkType.Monolithic && Target.bDisableLinking))
 			{
 				return new List<FileItem>();
 			}
