Fixed VsCode error failing to find omnisharp extension.
It has been renamed by microsoft so updated to the new name.

diff --git a/Engine/Source/Programs/UnrealBuildTool/ProjectFiles/VisualStudioCode/VSCodeProjectFileGenerator.cs b/Engine/Source/Programs/UnrealBuildTool/ProjectFiles/VisualStudioCode/VSCodeProjectFileGenerator.cs
--- a/Engine/Source/Programs/UnrealBuildTool/ProjectFiles/VisualStudioCode/VSCodeProjectFileGenerator.cs
+++ b/Engine/Source/Programs/UnrealBuildTool/ProjectFiles/VisualStudioCode/VSCodeProjectFileGenerator.cs
@@ -1305,7 +1305,7 @@ namespace UnrealBuildTool
 				WorkspaceFile.BeginArray("recommendations");
 				{
 					WorkspaceFile.AddUnnamedField("ms-vscode.cpptools");
-					WorkspaceFile.AddUnnamedField("ms-vscode.csharp");
+					WorkspaceFile.AddUnnamedField("ms-dotnettools.csharp");
 
 					// If the platform we run the generator on uses mono, there are additional debugging extensions to add.
 					if (Utils.IsRunningOnMono)
