Full Storyboard support for iOS
Removal of old Static Launch Screen Logic

diff --git a/Engine/Source/Developer/IOS/IOSPlatformEditor/Private/IOSTargetSettingsCustomization.cpp b/Engine/Source/Developer/IOS/IOSPlatformEditor/Private/IOSTargetSettingsCustomization.cpp
--- a/Engine/Source/Developer/IOS/IOSPlatformEditor/Private/IOSTargetSettingsCustomization.cpp
+++ b/Engine/Source/Developer/IOS/IOSPlatformEditor/Private/IOSTargetSettingsCustomization.cpp
@@ -86,24 +86,7 @@ FIOSTargetSettingsCustomization::FIOSTargetSettingsCustomization()
 	new (IconNames)FPlatformIconInfo(TEXT("[REDACTED_EMAIL]"), LOCTEXT("AppIcon_iPadProRetina_iOS9", "iPad Pro Retina iOS9 App Icon"), FText::GetEmpty(), 167, 167, FPlatformIconInfo::Required);
 	new (IconNames)FPlatformIconInfo(TEXT("Icon1024.png"), LOCTEXT("AppIcon_Marketing", "Marketing Icon"), FText::GetEmpty(), 1024, 1024, FPlatformIconInfo::Required);
 
-	new (LaunchImageNames) FPlatformIconInfo(TEXT("Default-IPhone5-Landscape.png"), LOCTEXT("LaunchImage_iPhone5_Landscape", "iPhone 5/5S/SE in Landscape"), FText::GetEmpty(), 1136, 640, FPlatformIconInfo::Required);
-    new (LaunchImageNames) FPlatformIconInfo(TEXT("[REDACTED_EMAIL]"), LOCTEXT("LaunchImage_iPhone5_Portrait", "iPhone 5/5S/SE in Portrait"), FText::GetEmpty(), 640, 1136, FPlatformIconInfo::Required);
-    new (LaunchImageNames) FPlatformIconInfo(TEXT("Default-IPhone6-Landscape.png"), LOCTEXT("LaunchImage_iPhone6_Landscape", "iPhone 6/7/8 in Landscape"), FText::GetEmpty(), 1334, 750, FPlatformIconInfo::Required);
-    new (LaunchImageNames) FPlatformIconInfo(TEXT("Default-IPhone6.png"), LOCTEXT("LaunchImage_iPhone6", "iPhone 6/7/8 in Portrait"), FText::GetEmpty(), 750, 1334, FPlatformIconInfo::Required);
-    new (LaunchImageNames) FPlatformIconInfo(TEXT("Default-IPhone6Plus-Landscape.png"), LOCTEXT("LaunchImage_iPhone6Plus_Landscape", "iPhone 6+/7+/8+ in Landscape"), FText::GetEmpty(), 2208, 1242, FPlatformIconInfo::Required);
-    new (LaunchImageNames) FPlatformIconInfo(TEXT("Default-IPhone6Plus-Portrait.png"), LOCTEXT("LaunchImage_iPhone6Plus_Portrait", "iPhone 6+/7+/8+ in Portrait"), FText::GetEmpty(), 1242, 2208, FPlatformIconInfo::Required);
-    new (LaunchImageNames) FPlatformIconInfo(TEXT("Default-IPhoneXS-Landscape.png"), LOCTEXT("LaunchImage_iPhoneX_Landscape", "iPhone X/XS in Landscape"), FText::GetEmpty(), 2436, 1125, FPlatformIconInfo::Required);
-    new (LaunchImageNames) FPlatformIconInfo(TEXT("Default-IPhoneXS-Portrait.png"), LOCTEXT("LaunchImage_iPhoneX_Portrait", "iPhone X/XS in Portrait"), FText::GetEmpty(), 1125, 2436, FPlatformIconInfo::Required);
-	new (LaunchImageNames) FPlatformIconInfo(TEXT("Default-IPhoneXSMax-Landscape.png"), LOCTEXT("LaunchImage_iPhoneXSMax_Landscape", "iPhone XS Max in Landscape"), FText::GetEmpty(), 2688, 1242, FPlatformIconInfo::Required);
-	new (LaunchImageNames) FPlatformIconInfo(TEXT("Default-IPhoneXSMax-Portrait.png"), LOCTEXT("LaunchImage_iPhoneXSMax_Portrait", "iPhone XS Max in Portrait"), FText::GetEmpty(), 1242, 2688, FPlatformIconInfo::Required);
-	new (LaunchImageNames) FPlatformIconInfo(TEXT("[REDACTED_EMAIL]"), LOCTEXT("LaunchImage_iPadRetina_Landscape", "iPad in Landscape"), FText::GetEmpty(), 2048, 1536, FPlatformIconInfo::Required);
-	new (LaunchImageNames) FPlatformIconInfo(TEXT("[REDACTED_EMAIL]"), LOCTEXT("LaunchImage_iPadRetina_Portrait", "iPad in Portrait"), FText::GetEmpty(), 1536, 2048, FPlatformIconInfo::Required);
-	new (LaunchImageNames)FPlatformIconInfo(TEXT("[REDACTED_EMAIL]"), LOCTEXT("LaunchImage_iPadPro105Retina_Landscape", "iPad Pro 10.5 in Landscape"), FText::GetEmpty(), 2224, 1668, FPlatformIconInfo::Required);
-	new (LaunchImageNames)FPlatformIconInfo(TEXT("[REDACTED_EMAIL]"), LOCTEXT("LaunchImage_iPadPro105Retina_Portrait", "iPad Pro 10.5 in Portrait"), FText::GetEmpty(), 1668, 2224, FPlatformIconInfo::Required);
-	new (LaunchImageNames)FPlatformIconInfo(TEXT("[REDACTED_EMAIL]"), LOCTEXT("LaunchImage_iPadPro11Retina_Landscape", "iPad Pro 11 in Landscape"), FText::GetEmpty(), 2388, 1668, FPlatformIconInfo::Required);
-	new (LaunchImageNames)FPlatformIconInfo(TEXT("[REDACTED_EMAIL]"), LOCTEXT("LaunchImage_iPadPro11Retina_Portrait", "iPad Pro 11 in Portrait"), FText::GetEmpty(), 1668, 2388, FPlatformIconInfo::Required);
-	new (LaunchImageNames)FPlatformIconInfo(TEXT("[REDACTED_EMAIL]"), LOCTEXT("LaunchImage_iPadProRetina_Landscape", "iPad Pro 12.9 in Landscape"), FText::GetEmpty(), 2732, 2048, FPlatformIconInfo::Required);
-	new (LaunchImageNames)FPlatformIconInfo(TEXT("[REDACTED_EMAIL]"), LOCTEXT("LaunchImage_iPadProRetina_Portrait", "iPad Pro 12.9 in Portrait"), FText::GetEmpty(), 2048, 2732, FPlatformIconInfo::Required);
+	new (LaunchImageNames)FPlatformIconInfo(TEXT("LaunchScreenIOS.png"), LOCTEXT("LaunchImageIOS", "Image to be used in as Launch Screen"), FText::GetEmpty(), -1, -1, FPlatformIconInfo::Required);
 
 	bShowAllProvisions = false;
 	bShowAllCertificates = false;
@@ -1087,7 +1070,7 @@ void FIOSTargetSettingsCustomization::BuildIconSection(IDetailLayoutBuilder& Det
 	}
 
 	// Add the launch images
-	IDetailCategoryBuilder& LaunchImageCategory = DetailLayout.EditCategory(TEXT("Launch Images"));
+	IDetailCategoryBuilder& LaunchImageCategory = DetailLayout.EditCategory(TEXT("Launch Image"));
 	const FVector2D LaunchImageMaxSize(150.0f, 150.0f);
 	for (const FPlatformIconInfo& Info : LaunchImageNames)
 	{
diff --git a/Engine/Source/Programs/AutomationTool/IOS/IOSPlatform.Automation.cs b/Engine/Source/Programs/AutomationTool/IOS/IOSPlatform.Automation.cs
--- a/Engine/Source/Programs/AutomationTool/IOS/IOSPlatform.Automation.cs
+++ b/Engine/Source/Programs/AutomationTool/IOS/IOSPlatform.Automation.cs
@@ -950,52 +950,15 @@ public class IOSPlatform : Platform
 	private void StageImageAndIconFiles(ProjectParams Params, DirectoryReference GraphicsDataPath, bool bSupportsPortrait, bool bSupportsLandscape, DeploymentContext SC, bool bSkipIcons)
 	{
 
-		bool bLaunchscreenStoryboard = false;
+		bool bCustomLaunchscreenStoryboard = false;
 		ConfigHierarchy PlatformGameConfig;
 		if (Params.EngineConfigs.TryGetValue(SC.StageTargetPlatform.PlatformType, out PlatformGameConfig))
 		{
-			PlatformGameConfig.GetBool("/Script/IOSRuntimeSettings.IOSRuntimeSettings", "bLaunchscreenStoryboard", out bLaunchscreenStoryboard);
+			PlatformGameConfig.GetBool("/Script/IOSRuntimeSettings.IOSRuntimeSettings", "bCustomLaunchscreenStoryboard", out bCustomLaunchscreenStoryboard);
 		}
 
-		if (DirectoryReference.Exists(GraphicsDataPath) && !bLaunchscreenStoryboard)
+		if (DirectoryReference.Exists(GraphicsDataPath))
 		{
-			List<string> ImageFileNames = new List<string>();
-			if (bSupportsPortrait)
-			{
-				ImageFileNames.Add("Default-IPhone6.png");
-				ImageFileNames.Add("Default-IPhone6Plus-Portrait.png");
-				ImageFileNames.Add("[REDACTED_EMAIL]");
-				ImageFileNames.Add("[REDACTED_EMAIL]");
-				ImageFileNames.Add("[REDACTED_EMAIL]");
-				ImageFileNames.Add("[REDACTED_EMAIL]");
-				ImageFileNames.Add("Default-IPhoneXS-Portrait.png");
-				ImageFileNames.Add("Default-IPhoneXSMax-Portrait.png");
-				ImageFileNames.Add("Default-IPhoneXR-Portrait.png");
-			}
-			if (bSupportsLandscape)
-			{
-				ImageFileNames.Add("Default-IPhone6-Landscape.png");
-				ImageFileNames.Add("Default-IPhone6Plus-Landscape.png");
-				ImageFileNames.Add("[REDACTED_EMAIL]");
-				ImageFileNames.Add("[REDACTED_EMAIL]");
-				ImageFileNames.Add("[REDACTED_EMAIL]");
-				ImageFileNames.Add("[REDACTED_EMAIL]");
-				ImageFileNames.Add("Default-IPhoneXS-Landscape.png");
-				ImageFileNames.Add("Default-IPhoneXSMax-Landscape.png");
-				ImageFileNames.Add("Default-IPhoneXR-Landscape.png");
-			}
-			ImageFileNames.Add("[REDACTED_EMAIL]");
-			ImageFileNames.Add("[REDACTED_EMAIL]");
-
-			foreach (string ImageFileName in ImageFileNames)
-			{
-				FileReference ImageFile = FileReference.Combine(GraphicsDataPath, ImageFileName);
-				if (FileReference.Exists(ImageFile))
-				{
-					SC.StageFile(StagedFileType.SystemNonUFS, ImageFile, new StagedFileReference(ImageFileName));
-				}
-			}
-
 			if (!bSkipIcons)
 			{
 				SC.StageFiles(StagedFileType.SystemNonUFS, GraphicsDataPath, "Icon*.png", StageFilesSearch.TopDirectoryOnly, StagedDirectoryReference.Root);
diff --git a/Engine/Source/Programs/UnrealBuildTool/Platform/IOS/UEDeployIOS.cs b/Engine/Source/Programs/UnrealBuildTool/Platform/IOS/UEDeployIOS.cs
--- a/Engine/Source/Programs/UnrealBuildTool/Platform/IOS/UEDeployIOS.cs
+++ b/Engine/Source/Programs/UnrealBuildTool/Platform/IOS/UEDeployIOS.cs
@@ -42,7 +42,7 @@ namespace UnrealBuildTool
 				set;
 			}
 
-			public static bool bLaunchscreenStoryboard = false;
+			public static bool bCustomLaunchscreenStoryboard = false;			
 
 			static string RunningVersionFilename
 			{
@@ -199,7 +199,6 @@ namespace UnrealBuildTool
 		{
 			// generate the Info.plist for future use
 			string BuildDirectory = ProjectDirectory + "/Build/IOS";
-			bool bSkipDefaultPNGs = false;
 			string IntermediateDirectory = (bIsUE4Game ? InEngineDir : ProjectDirectory) + "/Intermediate/IOS";
 			string PListFile = IntermediateDirectory + "/" + GameName + "-Info.plist";
 			ProjectName = !String.IsNullOrEmpty(ProjectName) ? ProjectName : GameName;
@@ -209,14 +208,8 @@ namespace UnrealBuildTool
 			// read the old file
 			string OldPListData = File.Exists(PListFile) ? File.ReadAllText(PListFile) : "";
 
-			// determine if there is a Lanchscreen Storyboard THEN a launch.xib
 			string LaunchStoryboard = InEngineDir + "/Build/IOS/Resources/Interface/LaunchScreen.storyboard";
-			string LaunchXib = InEngineDir + "/Build/IOS/Resources/Interface/LaunchScreen.xib";
 			if (File.Exists(BuildDirectory + "/Resources/Interface/LaunchScreen.storyboard"))			{				LaunchStoryboard = BuildDirectory + "/Resources/Interface/LaunchScreen.storyboard";			}
-			else if (File.Exists(BuildDirectory + "/Resources/Interface/LaunchScreen.xib"))
-			{
-				LaunchXib = BuildDirectory + "/Resources/Interface/LaunchScreen.xib";
-			}
 
 			// get the settings from the ini file
 			// plist replacements
@@ -398,7 +391,7 @@ namespace UnrealBuildTool
 			string ExtraData = "";
 			Ini.GetString("/Script/IOSRuntimeSettings.IOSRuntimeSettings", "AdditionalPlistData", out ExtraData);
 
-			Ini.GetBool("/Script/IOSRuntimeSettings.IOSRuntimeSettings", "bLaunchscreenStoryboard", out VersionUtilities.bLaunchscreenStoryboard);
+			Ini.GetBool("/Script/IOSRuntimeSettings.IOSRuntimeSettings", "bCustomLaunchscreenStoryboard", out VersionUtilities.bCustomLaunchscreenStoryboard);
 
 			// generate the plist file
 			StringBuilder Text = new StringBuilder();
@@ -576,94 +569,20 @@ namespace UnrealBuildTool
 				Text.AppendLine("\t\t</dict>");
 				Text.AppendLine("\t</dict>");
 			}
-			// Launchscreen storyboard only compatible on Mac for now.
-			if ((BuildHostPlatform.Current.Platform == UnrealTargetPlatform.Mac) && (File.Exists(LaunchStoryboard) || File.Exists(LaunchXib)) && VersionUtilities.bLaunchscreenStoryboard)
+
+			Text.AppendLine("\t<key>UILaunchStoryboardName</key>");
+			Text.AppendLine("\t<string>LaunchScreen</string>");
+
+			// Custom Launchscreen storyboard only compatible on Mac for now.
+			if ((BuildHostPlatform.Current.Platform == UnrealTargetPlatform.Mac) && File.Exists(LaunchStoryboard) && VersionUtilities.bCustomLaunchscreenStoryboard)
 			{
 				string outputStoryboard = LaunchStoryboard + "c";
 				string argsStoryboard = "--compile " + outputStoryboard + " " + LaunchStoryboard;
-				string stdOutLaunchScreen;
-				stdOutLaunchScreen = Utils.RunLocalProcessAndReturnStdOut("ibtool", argsStoryboard);
+				string stdOutLaunchScreen = Utils.RunLocalProcessAndReturnStdOut("ibtool", argsStoryboard);
 
 				Log.TraceInformation("LaunchScreen Storyboard compilation results : " + stdOutLaunchScreen);
-
-				Text.AppendLine("\t<key>UILaunchStoryboardName</key>");
-				Text.AppendLine("\t<string>LaunchScreen</string>");
-				bSkipDefaultPNGs = true;
 			}
-			else
-			{
-				// this is a temp way to inject the iphone 6 images without needing to upgrade everyone's plist
-				// eventually we want to generate this based on what the user has set in the project settings
-				string[] IPhoneConfigs =
-				{
-					"Default-IPhone6-Landscape.png", "Landscape", "{375, 667}", "8.0",
-					"Default-IPhone6.png", "Portrait", "{375, 667}",  "8.0",
-					"Default-IPhone6Plus-Landscape.png", "Landscape", "{414, 736}",  "8.0",
-					"Default-IPhone6Plus-Portrait.png", "Portrait", "{414, 736}",  "8.0",
-					"Default.png", "Portrait", "{320, 480}", "7.0",
-					"Default-568h.png", "Portrait", "{320, 568}", "7.0",
-					"Default-IPhoneXS-Landscape.png", "Landscape", "{375, 812}",  "11.0",
-					"Default-IPhoneXS-Portrait.png", "Portrait", "{375, 812}",  "11.0",
-					"Default-IPhoneXSMax-Landscape.png", "Landscape", "{414, 896}",  "12.0",
-					"Default-IPhoneXSMax-Portrait.png", "Portrait", "{414, 896}",  "12.0",
-					"Default-IPhoneXR-Landscape.png", "Landscape", "{276, 597}",  "12.0",
-					"Default-IPhoneXR-Portrait.png", "Portrait", "{276, 597}",  "12.0",
-				};
-
-				Text.AppendLine("\t<key>UILaunchImages~iphone</key>");
-				Text.AppendLine("\t<array>");
-				for (int ConfigIndex = 0; ConfigIndex < IPhoneConfigs.Length; ConfigIndex += 4)
-				{
-					{
-						Text.AppendLine("\t\t<dict>");
-						Text.AppendLine("\t\t\t<key>UILaunchImageMinimumOSVersion</key>");
-						Text.AppendLine(string.Format("\t\t\t<string>{0}</string>", IPhoneConfigs[ConfigIndex + 3]));
-						Text.AppendLine("\t\t\t<key>UILaunchImageName</key>");
-						Text.AppendLine(string.Format("\t\t\t<string>{0}</string>", IPhoneConfigs[ConfigIndex + 0]));
-						Text.AppendLine("\t\t\t<key>UILaunchImageOrientation</key>");
-						Text.AppendLine(string.Format("\t\t\t<string>{0}</string>", IPhoneConfigs[ConfigIndex + 1]));
-						Text.AppendLine("\t\t\t<key>UILaunchImageSize</key>");
-						Text.AppendLine(string.Format("\t\t\t<string>{0}</string>", IPhoneConfigs[ConfigIndex + 2]));
-						Text.AppendLine("\t\t</dict>");
-					}
-				}
-
-				// close it out
-				Text.AppendLine("\t</array>");
 
-				// this is a temp way to inject the iPad Pro without needing to upgrade everyone's plist
-				// eventually we want to generate this based on what the user has set in the project settings
-				string[] IPadConfigs =
-				{
-					"Default-Landscape.png", "Landscape", "{768, 1024}", "7.0",
-					"Default-Portrait.png", "Portrait", "{768, 1024}",  "7.0",
-					"Default-Landscape-1112.png", "Landscape", "{834, 1112}",  "9.0",
-					"Default-Portrait-1112.png", "Portrait", "{834, 1112}",  "9.0",
-					"Default-Landscape-1194.png", "Landscape", "{834, 1194}",  "9.0",
-					"Default-Portrait-1194.png", "Portrait", "{834, 1194}",  "9.0",
-					"Default-Landscape-1336.png", "Landscape", "{1024, 1366}",  "9.0",
-					"Default-Portrait-1336.png", "Portrait", "{1024, 1366}",  "9.0",
-				};
-
-				Text.AppendLine("\t<key>UILaunchImages~ipad</key>");
-				Text.AppendLine("\t<array>");
-				for (int ConfigIndex = 0; ConfigIndex < IPadConfigs.Length; ConfigIndex += 4)
-				{
-					{
-						Text.AppendLine("\t\t<dict>");
-						Text.AppendLine("\t\t\t<key>UILaunchImageMinimumOSVersion</key>");
-						Text.AppendLine(string.Format("\t\t\t<string>{0}</string>", IPadConfigs[ConfigIndex + 3]));
-						Text.AppendLine("\t\t\t<key>UILaunchImageName</key>");
-						Text.AppendLine(string.Format("\t\t\t<string>{0}</string>", IPadConfigs[ConfigIndex + 0]));
-						Text.AppendLine("\t\t\t<key>UILaunchImageOrientation</key>");
-						Text.AppendLine(string.Format("\t\t\t<string>{0}</string>", IPadConfigs[ConfigIndex + 1]));
-						Text.AppendLine("\t\t\t<key>UILaunchImageSize</key>");
-						Text.AppendLine(string.Format("\t\t\t<string>{0}</string>", IPadConfigs[ConfigIndex + 2]));
-						Text.AppendLine("\t\t</dict>");
-					}
-				}
-				Text.AppendLine("\t</array>");
-			}
 			Text.AppendLine("\t<key>CFBundleSupportedPlatforms</key>");
 			Text.AppendLine("\t<array>");
 			Text.AppendLine("\t\t<string>iPhoneOS</string>");
@@ -798,7 +717,7 @@ namespace UnrealBuildTool
 				File.WriteAllText(AppDirectory + "/Info.plist", Text.ToString());
 			}
 
-			return bSkipDefaultPNGs;
+			return true;
 		}
 
         static bool SupportsIconCatalog(VersionNumber SdkVersion)
@@ -880,12 +799,13 @@ namespace UnrealBuildTool
 		{
 			CopyFiles(InEngineDir + "/Build/IOS/Cloud", AppDirectory, "*.*", true);
 		}
+
 		protected virtual void CopyGraphicsResources(bool bSkipDefaultPNGs, bool bSkipIcons, string InEngineDir, string AppDirectory, string BuildDirectory, string IntermediateDir, bool bSupportsPortrait, bool bSupportsLandscape)
 		{
 			// copy engine assets in (IOS and TVOS shared in IOS)
 			if (bSkipDefaultPNGs)
 			{
-				if (VersionUtilities.bLaunchscreenStoryboard)
+				if (VersionUtilities.bCustomLaunchscreenStoryboard)
 				{
 					if (Directory.Exists(BuildDirectory + "/Resources/Interface/LaunchScreen.storyboardc"))
 					{
@@ -898,6 +818,20 @@ namespace UnrealBuildTool
 						CopyFiles(InEngineDir + "/Build/IOS/Resources/Interface/Assets", AppDirectory, "*", true);
 					}
 				}
+				else
+				{
+					if (File.Exists(BuildDirectory + "/Resources/Graphics/LaunchScreenIOS.png"))
+					{
+						CopyFolder(InEngineDir + "/Build/IOS/Resources/Interface/LaunchScreen.storyboardc", AppDirectory + "/LaunchScreen.storyboardc", true);
+						CopyFiles(BuildDirectory + "/Resources/Graphics", AppDirectory, "LaunchScreenIOS.png", true);
+					}
+					else
+					{
+						CopyFolder(InEngineDir + "/Build/IOS/Resources/Interface/LaunchScreen.storyboardc", AppDirectory + "/LaunchScreen.storyboardc", true);
+						CopyFiles(InEngineDir + "/Build/IOS/Resources/Graphics", AppDirectory, "LaunchScreenIOS.png", true);
+					}
+
+				}
 				// we still want default icons
 				if (!bSkipIcons)
 				{
@@ -909,35 +843,6 @@ namespace UnrealBuildTool
 					{
 						CopyFiles(InEngineDir + "/Build/IOS/Resources/Graphics", AppDirectory, "Icon*.png", true);
 					}
-					if (bSupportsPortrait && !VersionUtilities.bLaunchscreenStoryboard)
-					{
-						CopyFiles(InEngineDir + "/Build/IOS/Resources/Graphics", AppDirectory, "Default-IPhone6.png", true);
-						CopyFiles(InEngineDir + "/Build/IOS/Resources/Graphics", AppDirectory, "Default-IPhone6Plus-Portrait.png", true);
-						CopyFiles(InEngineDir + "/Build/IOS/Resources/Graphics", AppDirectory, "[REDACTED_EMAIL]", true);
-						CopyFiles(InEngineDir + "/Build/IOS/Resources/Graphics", AppDirectory, "[REDACTED_EMAIL]", true);
-						CopyFiles(InEngineDir + "/Build/IOS/Resources/Graphics", AppDirectory, "[REDACTED_EMAIL]", true);
-						CopyFiles(InEngineDir + "/Build/IOS/Resources/Graphics", AppDirectory, "[REDACTED_EMAIL]", true);
-						CopyFiles(InEngineDir + "/Build/IOS/Resources/Graphics", AppDirectory, "Default-IPhoneXS-Portrait.png", true);
-						CopyFiles(InEngineDir + "/Build/IOS/Resources/Graphics", AppDirectory, "Default-IPhoneXSMax-Portrait.png", true);
-						CopyFiles(InEngineDir + "/Build/IOS/Resources/Graphics", AppDirectory, "Default-IPhoneXR-Portrait.png", true);
-					}
-					if (bSupportsLandscape && !VersionUtilities.bLaunchscreenStoryboard)
-					{
-						CopyFiles(InEngineDir + "/Build/IOS/Resources/Graphics", AppDirectory, "Default-IPhone6-Landscape.png", true);
-						CopyFiles(InEngineDir + "/Build/IOS/Resources/Graphics", AppDirectory, "Default-IPhone6Plus-Landscape.png", true);
-						CopyFiles(InEngineDir + "/Build/IOS/Resources/Graphics", AppDirectory, "[REDACTED_EMAIL]", true);
-						CopyFiles(InEngineDir + "/Build/IOS/Resources/Graphics", AppDirectory, "[REDACTED_EMAIL]", true);
-						CopyFiles(InEngineDir + "/Build/IOS/Resources/Graphics", AppDirectory, "[REDACTED_EMAIL]", true);
-						CopyFiles(InEngineDir + "/Build/IOS/Resources/Graphics", AppDirectory, "[REDACTED_EMAIL]", true);
-						CopyFiles(InEngineDir + "/Build/IOS/Resources/Graphics", AppDirectory, "Default-IPhoneXS-Landscape.png", true);
-						CopyFiles(InEngineDir + "/Build/IOS/Resources/Graphics", AppDirectory, "Default-IPhoneXSMax-Landscape.png", true);
-						CopyFiles(InEngineDir + "/Build/IOS/Resources/Graphics", AppDirectory, "Default-IPhoneXR-Landscape.png", true);
-					}
-					if (!VersionUtilities.bLaunchscreenStoryboard)
-					{
-						CopyFiles(BuildDirectory + "/Resources/Graphics", AppDirectory, "[REDACTED_EMAIL]", true);
-						CopyFiles(BuildDirectory + "/Resources/Graphics", AppDirectory, "[REDACTED_EMAIL]", true);
-					}
 				}
 			}
 			// merge game assets on top
@@ -948,32 +853,6 @@ namespace UnrealBuildTool
 				{
 					CopyFiles(BuildDirectory + "/Resources/Graphics", AppDirectory, "Icon*.png", true);
 				}
-				if (bSupportsPortrait)
-				{
-					CopyFiles(BuildDirectory + "/Resources/Graphics", AppDirectory, "Default-IPhone6.png", true);
-					CopyFiles(BuildDirectory + "/Resources/Graphics", AppDirectory, "Default-IPhone6Plus-Portrait.png", true);
-					CopyFiles(BuildDirectory + "/Resources/Graphics", AppDirectory, "[REDACTED_EMAIL]", true);
-					CopyFiles(BuildDirectory + "/Resources/Graphics", AppDirectory, "[REDACTED_EMAIL]", true);
-					CopyFiles(BuildDirectory + "/Resources/Graphics", AppDirectory, "[REDACTED_EMAIL]", true);
-					CopyFiles(BuildDirectory + "/Resources/Graphics", AppDirectory, "[REDACTED_EMAIL]", true);
-					CopyFiles(BuildDirectory + "/Resources/Graphics", AppDirectory, "Default-IPhoneXS-Portrait.png", true);
-					CopyFiles(BuildDirectory + "/Resources/Graphics", AppDirectory, "Default-IPhoneXSMax-Portrait.png", true);
-					CopyFiles(BuildDirectory + "/Resources/Graphics", AppDirectory, "Default-IPhoneXR-Portrait.png", true);
-				}
-				if (bSupportsLandscape)
-				{
-					CopyFiles(BuildDirectory + "/Resources/Graphics", AppDirectory, "Default-IPhone6-Landscape.png", true);
-					CopyFiles(BuildDirectory + "/Resources/Graphics", AppDirectory, "Default-IPhone6Plus-Landscape.png", true);
-					CopyFiles(BuildDirectory + "/Resources/Graphics", AppDirectory, "[REDACTED_EMAIL]", true);
-					CopyFiles(BuildDirectory + "/Resources/Graphics", AppDirectory, "[REDACTED_EMAIL]", true);
-					CopyFiles(BuildDirectory + "/Resources/Graphics", AppDirectory, "[REDACTED_EMAIL]", true);
-					CopyFiles(BuildDirectory + "/Resources/Graphics", AppDirectory, "[REDACTED_EMAIL]", true);
-					CopyFiles(BuildDirectory + "/Resources/Graphics", AppDirectory, "Default-IPhoneXS-Landscape.png", true);
-					CopyFiles(BuildDirectory + "/Resources/Graphics", AppDirectory, "Default-IPhoneXSMax-Landscape.png", true);
-					CopyFiles(BuildDirectory + "/Resources/Graphics", AppDirectory, "Default-IPhoneXR-Landscape.png", true);
-				}
-				CopyFiles(BuildDirectory + "/Resources/Graphics", AppDirectory, "[REDACTED_EMAIL]", true);
-				CopyFiles(BuildDirectory + "/Resources/Graphics", AppDirectory, "[REDACTED_EMAIL]", true);
 			}
 		}
 		
@@ -1159,7 +1038,7 @@ namespace UnrealBuildTool
 			bool bSupportsPortrait = true;
 			bool bSupportsLandscape = false;
 			bool bSkipIcons = false;
-			bool bSkipDefaultPNGs = GeneratePList(ProjectFile, Config, InProjectDirectory, bIsUE4Game, GameExeName, false, InProjectName, InEngineDir, AppDirectory, UPLScripts, SdkVersion, BundleID, bBuildAsFramework, out bSupportsPortrait, out bSupportsLandscape, out bSkipIcons);
+			GeneratePList(ProjectFile, Config, InProjectDirectory, bIsUE4Game, GameExeName, false, InProjectName, InEngineDir, AppDirectory, UPLScripts, SdkVersion, BundleID, bBuildAsFramework, out bSupportsPortrait, out bSupportsLandscape, out bSkipIcons);
 
 			// ensure the destination is writable
 			if (File.Exists(AppDirectory + "/" + GameName))
@@ -1173,7 +1052,7 @@ namespace UnrealBuildTool
 
 			if (!bCreateStubIPA)
 			{
-				CopyGraphicsResources(bSkipDefaultPNGs, bSkipIcons, InEngineDir, AppDirectory, BuildDirectory, IntermediateDirectory, bSupportsPortrait, bSupportsLandscape);
+				CopyGraphicsResources(true, bSkipIcons, InEngineDir, AppDirectory, BuildDirectory, IntermediateDirectory, bSupportsPortrait, bSupportsLandscape);
 				CopyLocalizationsResources(InEngineDir, AppDirectory, BuildDirectory, IntermediateDirectory);
 
 				CopyCloudResources(InProjectDirectory, AppDirectory);
diff --git a/Engine/Source/Runtime/IOS/IOSRuntimeSettings/Classes/IOSRuntimeSettings.h b/Engine/Source/Runtime/IOS/IOSRuntimeSettings/Classes/IOSRuntimeSettings.h
--- a/Engine/Source/Runtime/IOS/IOSRuntimeSettings/Classes/IOSRuntimeSettings.h
+++ b/Engine/Source/Runtime/IOS/IOSRuntimeSettings/Classes/IOSRuntimeSettings.h
@@ -423,14 +423,14 @@ public:
 	// Any additional plist key/value data utilizing \n for a new line
 	UPROPERTY(GlobalConfig, EditAnywhere, Category = ExtraData)
 	FString AdditionalPlistData;
-	
+
 	/**
-	 * Choose whether to use a storyboard as a Launchscreen or not. If True, LaunchScreen.storyboard in Build/IOS/Resources/Interface
-	 * in the Engine folder will be copied to the bundle app and the Launchscreen images WILL NOT be included.
-	 * When using assets in your LaunchScreen.Storyboard, add them in Build/IOS/Resources/Interface/Assets and they will be included.
+	 * Choose whether to use a custom LaunchScreen.Storyboard as a Launchscreen. If True, the LaunchScreen.storyboard in Build/IOS/Resources/Interface
+	 * in the Engine folder will be copied to the bundle app and the Launchscreen defined in iOS Project Settings WILL NOT be included.
+	 * When using assets in your custom LaunchScreen.storyboard, add them in Build/IOS/Resources/Interface/Assets and they will be included.
 	 */
 	UPROPERTY(GlobalConfig, EditAnywhere, Category = LaunchScreen)
-	bool bLaunchscreenStoryboard;
+	bool bCustomLaunchscreenStoryboard;
 
 	// Whether the app supports Facebook
 	UPROPERTY(GlobalConfig, EditAnywhere, Category = Online)
diff --git a/Engine/Source/Runtime/IOS/IOSRuntimeSettings/Private/IOSRuntimeSettings.cpp b/Engine/Source/Runtime/IOS/IOSRuntimeSettings/Private/IOSRuntimeSettings.cpp
--- a/Engine/Source/Runtime/IOS/IOSRuntimeSettings/Private/IOSRuntimeSettings.cpp
+++ b/Engine/Source/Runtime/IOS/IOSRuntimeSettings/Private/IOSRuntimeSettings.cpp
@@ -37,7 +37,7 @@ UIOSRuntimeSettings::UIOSRuntimeSettings(const FObjectInitializer& ObjectInitial
 	bShipForArmV7S = false;
 	bShipForBitcode = true;
 	bUseRSync = true;
-	bLaunchscreenStoryboard = false;
+	bCustomLaunchscreenStoryboard = false;
 	AdditionalPlistData = TEXT("");
 	AdditionalLinkerFlags = TEXT("");
 	AdditionalShippingLinkerFlags = TEXT("");
