Only append sanatizer args for monolithic builds

diff --git a/Engine/Source/Programs/UnrealBuildTool/Platform/Linux/UEBuildLinux.cs b/Engine/Source/Programs/UnrealBuildTool/Platform/Linux/UEBuildLinux.cs
--- a/Engine/Source/Programs/UnrealBuildTool/Platform/Linux/UEBuildLinux.cs
+++ b/Engine/Source/Programs/UnrealBuildTool/Platform/Linux/UEBuildLinux.cs
@@ -536,21 +536,26 @@ namespace UnrealBuildTool
 			List<FileReference> FinalBinaryPath = new List<FileReference>();
 
 			string SanitizerSuffix = null;
-			if(Target.LinuxPlatform.bEnableAddressSanitizer)
-			{
-				SanitizerSuffix = "ASan";
-			}
-			else if(Target.LinuxPlatform.bEnableThreadSanitizer)
-			{
-				SanitizerSuffix = "TSan";
-			}
-			else if(Target.LinuxPlatform.bEnableUndefinedBehaviorSanitizer)
-			{
-				SanitizerSuffix = "UBSan";
-			}
-			else if(Target.LinuxPlatform.bEnableMemorySanitizer)
+
+			// Only append these for monolithic builds. non-monolithic runs into issues dealing with target/modules files
+			if (Target.LinkType == TargetLinkType.Monolithic)
 			{
-				SanitizerSuffix = "MSan";
+				if(Target.LinuxPlatform.bEnableAddressSanitizer)
+				{
+					SanitizerSuffix = "ASan";
+				}
+				else if(Target.LinuxPlatform.bEnableThreadSanitizer)
+				{
+					SanitizerSuffix = "TSan";
+				}
+				else if(Target.LinuxPlatform.bEnableUndefinedBehaviorSanitizer)
+				{
+					SanitizerSuffix = "UBSan";
+				}
+				else if(Target.LinuxPlatform.bEnableMemorySanitizer)
+				{
+					SanitizerSuffix = "MSan";
+				}
 			}
 
 			if (String.IsNullOrEmpty(SanitizerSuffix))
