Improved SN-DBS builds
 - Switched to the json interface, and removed the local task execution. All task scheduling is handled by SN-DBS now.
 - Local-only tasks are annotated as such in the json, meaning we don't have to flush the pipe between batches of remote jobs.

Measured full engine rebuild improvement of ~40% vs the previous implementation on a network with 4 worker agents.

diff --git a/Engine/Source/Programs/UnrealBuildTool/Executors/SNDBS.cs b/Engine/Source/Programs/UnrealBuildTool/Executors/SNDBS.cs
--- a/Engine/Source/Programs/UnrealBuildTool/Executors/SNDBS.cs
+++ b/Engine/Source/Programs/UnrealBuildTool/Executors/SNDBS.cs
@@ -7,504 +7,122 @@ using System.Diagnostics;
 using System.IO;
 using System.Linq;
 using System.ServiceProcess;
-using System.Text;
 using Tools.DotNETCommon;
 
 namespace UnrealBuildTool
 {
 	class SNDBS : ActionExecutor
 	{
-		private FileReference IncludeRewriteRulesFile = FileReference.Combine(UnrealBuildTool.EngineDirectory, "Build", "SNDBS", "include-rewrite-rules.ini");
+		private static readonly string SCERoot = Environment.GetEnvironmentVariable("SCE_ROOT_DIR");
+		private static readonly string SNDBSExecutable = Path.Combine(SCERoot ?? string.Empty, "Common", "SN-DBS", "bin", "dbsbuild.exe");
 
-		/// <summary>
-		/// Processor count multiplier for local execution. Can be below 1 to reserve CPU for other tasks.
-		/// </summary>
-		[XmlConfigFile]
-		public double ProcessorCountMultiplier = 1.0;
+		private static readonly string IncludeRewriteRulesDirectory = Path.Combine(UnrealBuildTool.EngineDirectory.FullName, "Build", "SNDBS");
+		private static readonly string IncludeRewriteRulesFilePath = Path.Combine(IncludeRewriteRulesDirectory, "include-rewrite-rules.ini");
 
-		/// <summary>
-		/// Maximum processor count for local execution. 
-		/// </summary>
-		[XmlConfigFile]
-		public int MaxProcessorCount = int.MaxValue;
+		private static readonly string ScriptFilePath = Path.Combine(UnrealBuildTool.EngineDirectory.FullName, "Intermediate", "Build", "sndbs.json");
 
-		/// <summary>
-		/// The number of actions to execute in parallel is trying to keep the CPU busy enough in presence of I/O stalls.
-		/// </summary>
-		int MaxActionsToExecuteInParallel;
+		private static readonly string TemplatesInputDir = Path.Combine(UnrealBuildTool.EngineDirectory.FullName, "Build", "SNDBSTemplates");
+		private static readonly string TemplatesOutputDir = Path.Combine(UnrealBuildTool.EngineDirectory.FullName, "Programs", "UnrealBuildTool", "SNDBSTemplates");
 
-		/// <summary>
-		/// Unique id for new jobs
-		/// </summary>
-		int JobNumber;
+		public override string Name => "SNDBS";
 
-		public SNDBS()
-		{
-			XmlConfig.ApplyTo(this);
-		}
-
-		public override string Name
-		{
-			get { return "SNDBS"; }
-		}
-
-		/// <summary>
-		/// Used when debugging Actions outputs all action return values to debug out
-		/// </summary>
-		/// <param name="sender"> Sending object</param>
-		/// <param name="e">  Event arguments (In this case, the line of string output)</param>
-		static protected void ActionDebugOutput(object sender, DataReceivedEventArgs e)
+		public static bool IsAvailable()
 		{
-			string Output = e.Data;
-			if (Output == null)
+			if(SCERoot == null || !File.Exists(SNDBSExecutable))
 			{
-				return;
+				return false;
 			}
 
-			Log.TraceInformation(Output);
-		}
-
-		internal bool ExecuteLocalActions(List<Action> InLocalActions, Dictionary<Action, ActionThread> InActionThreadDictionary, int TotalNumJobs)
-		{
-			// Time to sleep after each iteration of the loop in order to not busy wait.
-			const float LoopSleepTime = 0.1f;
-
-			bool LocalActionsResult = true;
-
-			while (true)
+			ServiceController[] services = ServiceController.GetServices();
+			foreach (ServiceController service in services)
 			{
-				// Count the number of pending and still executing actions.
-				int NumUnexecutedActions = 0;
-				int NumExecutingActions = 0;
-				foreach (Action Action in InLocalActions)
-				{
-					ActionThread ActionThread = null;
-					bool bFoundActionProcess = InActionThreadDictionary.TryGetValue(Action, out ActionThread);
-					if (bFoundActionProcess == false)
-					{
-						NumUnexecutedActions++;
-					}
-					else if (ActionThread != null)
-					{
-						if (ActionThread.bComplete == false)
-						{
-							NumUnexecutedActions++;
-							NumExecutingActions++;
-						}
-					}
-				}
-
-				// If there aren't any pending actions left, we're done executing.
-				if (NumUnexecutedActions == 0)
-				{
-					break;
-				}
-
-				// If there are fewer actions executing than the maximum, look for pending actions that don't have any outdated
-				// prerequisites.
-				foreach (Action Action in InLocalActions)
+				if (service.ServiceName.StartsWith("SNDBS") && service.Status == ServiceControllerStatus.Running)
 				{
-					ActionThread ActionProcess = null;
-					bool bFoundActionProcess = InActionThreadDictionary.TryGetValue(Action, out ActionProcess);
-					if (bFoundActionProcess == false)
-					{
-						if (NumExecutingActions < Math.Max(1, MaxActionsToExecuteInParallel))
-						{
-							// Determine whether there are any prerequisites of the action that are outdated.
-							bool bHasOutdatedPrerequisites = false;
-							bool bHasFailedPrerequisites = false;
-							foreach (Action PrerequisiteAction in Action.PrerequisiteActions)
-							{
-								if (InLocalActions.Contains(PrerequisiteAction))
-								{
-									ActionThread PrerequisiteProcess = null;
-									bool bFoundPrerequisiteProcess = InActionThreadDictionary.TryGetValue(PrerequisiteAction, out PrerequisiteProcess);
-									if (bFoundPrerequisiteProcess == true)
-									{
-										if (PrerequisiteProcess == null)
-										{
-											bHasFailedPrerequisites = true;
-										}
-										else if (PrerequisiteProcess.bComplete == false)
-										{
-											bHasOutdatedPrerequisites = true;
-										}
-										else if (PrerequisiteProcess.ExitCode != 0)
-										{
-											bHasFailedPrerequisites = true;
-										}
-									}
-									else
-									{
-										bHasOutdatedPrerequisites = true;
-									}
-								}
-							}
-
-							// If there are any failed prerequisites of this action, don't execute it.
-							if (bHasFailedPrerequisites)
-							{
-								// Add a null entry in the dictionary for this action.
-								InActionThreadDictionary.Add(Action, null);
-							}
-							// If there aren't any outdated prerequisites of this action, execute it.
-							else if (!bHasOutdatedPrerequisites)
-							{
-								ActionThread ActionThread = new ActionThread(Action, JobNumber, TotalNumJobs);
-								ActionThread.Run();
-
-								InActionThreadDictionary.Add(Action, ActionThread);
-
-								NumExecutingActions++;
-								JobNumber++;
-							}
-						}
-					}
+					return true;
 				}
-
-				System.Threading.Thread.Sleep(TimeSpan.FromSeconds(LoopSleepTime));
 			}
 
-			return LocalActionsResult;
+			return false;
 		}
 
-		internal bool ExecuteActions(List<Action> InActions, Dictionary<Action, ActionThread> InActionThreadDictionary)
+		private sealed class SNDBSJob
 		{
-			// Build the script file that will be executed by SN-DBS
-			StreamWriter ScriptFile;
-			string ScriptFilename = Path.Combine(UnrealBuildTool.EngineDirectory.FullName, "Intermediate", "Build", "SNDBS.bat");
-
-			FileStream ScriptFileStream = new FileStream(ScriptFilename, FileMode.Create, FileAccess.ReadWrite, FileShare.Read);
-			ScriptFile = new StreamWriter(ScriptFileStream);
-			ScriptFile.AutoFlush = true;
-
-			int NumScriptedActions = 0;
-			List<Action> LocalActions = new List<Action>();
-			ActionThread DummyActionThread = new ActionThread(null, 1, 1);
-			bool PrintDebugInfo = true;
-			foreach (Action Action in InActions)
-			{
-				ActionThread ActionProcess = null;
-				bool bFoundActionProcess = InActionThreadDictionary.TryGetValue(Action, out ActionProcess);
-				if (bFoundActionProcess == false)
-				{
-					// Determine whether there are any prerequisites of the action that are outdated.
-					bool bHasOutdatedPrerequisites = false;
-					bool bHasFailedPrerequisites = false;
-					foreach (Action PrerequisiteAction in Action.PrerequisiteActions)
-					{
-						if (InActions.Contains(PrerequisiteAction))
-						{
-							ActionThread PrerequisiteProcess = null;
-							bool bFoundPrerequisiteProcess = InActionThreadDictionary.TryGetValue(PrerequisiteAction, out PrerequisiteProcess);
-							if (bFoundPrerequisiteProcess == true)
-							{
-								if (PrerequisiteProcess == null)
-								{
-									bHasFailedPrerequisites = true;
-								}
-								else if (PrerequisiteProcess.bComplete == false)
-								{
-									bHasOutdatedPrerequisites = true;
-								}
-								else if (PrerequisiteProcess.ExitCode != 0)
-								{
-									bHasFailedPrerequisites = true;
-								}
-							}
-							else
-							{
-								bHasOutdatedPrerequisites = true;
-							}
-						}
-					}
-
-					// If there are any failed prerequisites of this action, don't execute it.
-					if (bHasFailedPrerequisites)
-					{
-						// Add a null entry in the dictionary for this action.
-						InActionThreadDictionary.Add(Action, null);
-					}
-					// If there aren't any outdated prerequisites of this action, execute it.
-					else if (!bHasOutdatedPrerequisites)
-					{
-						if (Action.bCanExecuteRemotely == false || Action.bCanExecuteRemotelyWithSNDBS == false)
-						{
-							// Execute locally
-							LocalActions.Add(Action);
-						}
-						else
-						{
-							// Create a dummy force-included file which references PCH files, so that SN-DBS knows they are dependencies.
-							string AdditionalStubIncludes = "";
-							if (Action.CommandPath.GetFileName().Equals("cl.exe", StringComparison.OrdinalIgnoreCase) || Action.CommandPath.GetFileName().Equals("cl-filter.exe", StringComparison.OrdinalIgnoreCase))
-							{
-								string DummyPCHIncludeFile = Action.DependencyListFile.AbsolutePath.Replace("\"", "").Replace("@", "").Trim();
-								DummyPCHIncludeFile = Path.ChangeExtension(DummyPCHIncludeFile, null);
-
-								StringBuilder WrapperContents = new StringBuilder();
-								using (StringWriter Writer = new StringWriter(WrapperContents))
-								{
-									Writer.WriteLine("// PCH dependencies for {0}", DummyPCHIncludeFile);
-									Writer.WriteLine("#if 0");
-									foreach (FileItem Preqrequisite in Action.PrerequisiteItems)
-									{
-										if (Preqrequisite.AbsolutePath.EndsWith(".pch"))
-										{
-											Writer.WriteLine("#include \"{0}\"", Preqrequisite.AbsolutePath.Replace(".pch", ".obj"));
-										}
-									}
-									Writer.WriteLine("#endif");
-								}
-
-								FileReference DummyPCHIncludeFileDependency = new FileReference(DummyPCHIncludeFile + ".dummy.h");
-								Utils.WriteFileIfChanged(DummyPCHIncludeFileDependency, WrapperContents.ToString(), StringComparison.OrdinalIgnoreCase);
-								AdditionalStubIncludes = string.Format("/FI\"{0}\"", DummyPCHIncludeFileDependency);
-							}
-
-							// Add to script for execution by SN-DBS
-							string NewCommandArguments = "\"" + Action.CommandPath + "\"" + " " + Action.CommandArguments + " " + AdditionalStubIncludes;
-							ScriptFile.WriteLine(NewCommandArguments);
-							InActionThreadDictionary.Add(Action, DummyActionThread);
-							Action.StartTime = Action.EndTime = DateTimeOffset.Now;
-							Log.TraceInformation("[{0}/{1}] {2} {3}", JobNumber, InActions.Count, Action.CommandDescription, Action.StatusDescription);
-							JobNumber++;
-							NumScriptedActions++;
-							PrintDebugInfo |= Action.bPrintDebugInfo;
-
-							if (Action.DependencyListFile != null && File.Exists(Action.DependencyListFile.AbsolutePath))
-							{
-								Log.TraceVerbose("Deleting dependency list file {0}", Action.DependencyListFile.AbsolutePath);
-								File.Delete(Action.DependencyListFile.AbsolutePath);
-							}
-						}
-					}
-				}
-			}
-
-			ScriptFile.Flush();
-			ScriptFile.Close();
-			ScriptFile.Dispose();
-			ScriptFile = null;
-
-			if (NumScriptedActions > 0)
-			{
-                // Create the process
-                string SCERoot = Environment.GetEnvironmentVariable("SCE_ROOT_DIR");
-                string SNDBSExecutable = Path.Combine(SCERoot, "Common/SN-DBS/bin/dbsbuild.exe");
-				DirectoryReference TemplatesDir = DirectoryReference.Combine(UnrealBuildTool.EngineDirectory, "Programs", "UnrealBuildTool", "SndbsTemplates");
-				string IncludeRewriteRulesArg = String.Format("--include-rewrite-rules \"{0}\"", IncludeRewriteRulesFile.FullName);
-				string VerbosityLevel = PrintDebugInfo ? "-v" : "-q";
-                ProcessStartInfo PSI = new ProcessStartInfo(SNDBSExecutable, String.Format("{0} -p UE4Code -s \"{1}\" -templates \"{2}\" {3}", VerbosityLevel, FileReference.Combine(UnrealBuildTool.EngineDirectory, "Intermediate", "Build", "sndbs.bat").FullName, TemplatesDir.FullName, IncludeRewriteRulesArg));
-				PSI.RedirectStandardOutput = true;
-				PSI.RedirectStandardError = true;
-				PSI.UseShellExecute = false;
-				PSI.CreateNoWindow = true;
-				PSI.WorkingDirectory = Path.GetFullPath(".");
-				Process NewProcess = new Process();
-				NewProcess.StartInfo = PSI;
-				NewProcess.OutputDataReceived += new DataReceivedEventHandler(ActionDebugOutput);
-				NewProcess.ErrorDataReceived += new DataReceivedEventHandler(ActionDebugOutput);
-				DateTimeOffset StartTime = DateTimeOffset.Now;
-				NewProcess.Start();
-				NewProcess.BeginOutputReadLine();
-				NewProcess.BeginErrorReadLine();
-
-				NewProcess.WaitForExit();
-
-				//TimeSpan Duration;
-				//DateTimeOffset EndTime = DateTimeOffset.Now;
-				//if (EndTime == DateTimeOffset.MinValue)
-				//{
-				//	Duration = DateTimeOffset.Now - StartTime;
-				//}
-				//else
-				//{
-				//	Duration = EndTime - StartTime;
-				//}
-
-				DummyActionThread.bComplete = true;
-				int ExitCode = NewProcess.ExitCode;
-				if (ExitCode != 0)
-				{
-					return false;
-				}
-			}
-
-			// Execute local tasks
-			if (LocalActions.Count > 0)
-			{
-				return ExecuteLocalActions(LocalActions, InActionThreadDictionary, InActions.Count);
-			}
+			public string title { get; set; }
+			public string command { get; set; }
+			public string working_directory { get; set; }
+			public string[] dependencies { get; set; }
+			public bool run_locally { get; set; }
+		}
 
-			return true;
+		private sealed class SNDBSScript
+		{
+			public IDictionary<string, SNDBSJob> jobs { get; set; }
 		}
 
-		public static bool IsAvailable()
+		public override bool ExecuteActions(List<Action> Actions, bool bLogDetailedActionStats)
 		{
-			string SCERoot = Environment.GetEnvironmentVariable("SCE_ROOT_DIR");
-			if(SCERoot == null)
-			{
-				return false;
-			}
-			if (!File.Exists(Path.Combine(SCERoot, "Common/SN-DBS/bin/dbsbuild.exe")))
-			{
-				return false;
-			}
+			if (Actions.Count == 0)
+				return true;
 
-			ServiceController[] services = ServiceController.GetServices();
-			foreach (ServiceController service in services)
+			// Build the json script file to describe all the actions and their dependencies
+			var ActionIds = Actions.ToDictionary(a => a, a => Guid.NewGuid().ToString());
+			var Script = new SNDBSScript()
 			{
-				if (service.ServiceName.StartsWith("SNDBS") && service.Status == ServiceControllerStatus.Running)
+				jobs = Actions.ToDictionary(a => ActionIds[a], a => new SNDBSJob
 				{
-					return true;
-				}
-			}
-			return false;
+					title = a.StatusDescription,
+					command = $"\"{a.CommandPath}\" {a.CommandArguments}",
+					working_directory = a.WorkingDirectory.FullName,
+					dependencies = a.PrerequisiteActions.Select(p => ActionIds[p]).ToArray(),
+					run_locally = !(a.bCanExecuteRemotely && a.bCanExecuteRemotelyWithSNDBS)
+				})
+			};
+
+			var ScriptJson = fastJSON.JSON.Instance.ToJSON(Script, new fastJSON.JSONParameters() { UseExtensions = false });
+			File.WriteAllText(ScriptFilePath, ScriptJson);
+
+			GenerateSNDBSIncludeRewriteRules();
+			PrepareToolTemplates(Actions);
+
+			var LocalProcess = new Process();
+			LocalProcess.StartInfo = new ProcessStartInfo(SNDBSExecutable, $"-q -p UE4Code -s \"{ScriptFilePath}\" -templates \"{TemplatesOutputDir}\" --include-rewrite-rules \"{IncludeRewriteRulesFilePath}\"");
+			LocalProcess.OutputDataReceived += (Sender, Args) => Log.TraceInformation("{0}", Args.Data);
+			LocalProcess.ErrorDataReceived += (Sender, Args) => Log.TraceInformation("{0}", Args.Data);
+			return Utils.RunLocalProcess(LocalProcess) == 0;
 		}
 
-		public override bool ExecuteActions(List<Action> Actions, bool bLogDetailedActionStats)
+		private void PrepareToolTemplates(IEnumerable<Action> Actions)
 		{
-			bool SNDBSResult = true;
-			if (Actions.Count > 0)
+			if (!Directory.Exists(TemplatesOutputDir))
 			{
-				// Generate any needed templates. Can only generate one per executable, so just use the first Action for reference
-				IEnumerable<Action> CommandPaths = Actions.GroupBy(a => a.CommandPath).Select(g => g.FirstOrDefault()).ToList();
-				foreach (Action CommandPath in CommandPaths)
-				{
-					PrepareToolTemplate(CommandPath);
-				}
-
-				// Generate include-rewrite-rules.ini.
-				GenerateSNDBSIncludeRewriteRules();
+				Directory.CreateDirectory(TemplatesOutputDir);
+			}
 
-				// Use WMI to figure out physical cores, excluding hyper threading.
-				int NumCores = Utils.GetPhysicalProcessorCount();
-				
-				// If we failed to detect the number of cores, default to the logical processor count
-				if (NumCores == -1)
-				{
-					NumCores = System.Environment.ProcessorCount;
-				}
-				// The number of actions to execute in parallel is trying to keep the CPU busy enough in presence of I/O stalls.
-				MaxActionsToExecuteInParallel = 0;
-				// The CPU has more logical cores than physical ones, aka uses hyper-threading. 
-				if (NumCores < System.Environment.ProcessorCount)
-				{
-					MaxActionsToExecuteInParallel = (int)(NumCores * ProcessorCountMultiplier);
-				}
-				// No hyper-threading. Only kicking off a task per CPU to keep machine responsive.
-				else
-				{
-					MaxActionsToExecuteInParallel = NumCores;
-				}
-				MaxActionsToExecuteInParallel = Math.Min(MaxActionsToExecuteInParallel, MaxProcessorCount);
+			var UniqueTools = new HashSet<string>(Actions.Select(a => a.CommandPath.FullName));
+			foreach (var Tool in UniqueTools)
+			{
+				var TemplateFileName = $"{Path.GetFileName(Tool)}.sn-dbs-tool.ini";
 
-				JobNumber = 1;
-				Dictionary<Action, ActionThread> ActionThreadDictionary = new Dictionary<Action, ActionThread>();
+				var TemplateInputFilePath = Path.Combine(TemplatesInputDir, TemplateFileName);
+				var TemplateOutputFilePath = Path.Combine(TemplatesOutputDir, TemplateFileName);
 
-				while (true)
+				// If no base template exists, don't try to generate one.
+				if (!File.Exists(TemplateInputFilePath))
 				{
-					bool bUnexecutedActions = false;
-					foreach (Action Action in Actions)
-					{
-						ActionThread ActionThread = null;
-						bool bFoundActionProcess = ActionThreadDictionary.TryGetValue(Action, out ActionThread);
-						if (bFoundActionProcess == false)
-						{
-							bUnexecutedActions = true;
-							if(!ExecuteActions(Actions, ActionThreadDictionary))
-							{
-								return false;
-							}
-							break;
-						}
-					}
-
-					if (bUnexecutedActions == false)
-					{
-						break;
-					}
+					continue;
 				}
 
-				Log.WriteLineIf(bLogDetailedActionStats, LogEventType.Console, "-------- Begin Detailed Action Stats ----------------------------------------------------------");
-				Log.WriteLineIf(bLogDetailedActionStats, LogEventType.Console, "^Action Type^Duration (seconds)^Tool^Task^Using PCH");
-
-				double TotalThreadSeconds = 0;
+				string TemplateText = File.ReadAllText(TemplateInputFilePath);
 
-				// Check whether any of the tasks failed and log action stats if wanted.
-				foreach (KeyValuePair<Action, ActionThread> ActionProcess in ActionThreadDictionary)
+				TemplateText = TemplateText.Replace("{COMMAND_PATH}", Path.GetDirectoryName(Tool));
+				foreach (DictionaryEntry Variable in Environment.GetEnvironmentVariables(EnvironmentVariableTarget.Process))
 				{
-					Action Action = ActionProcess.Key;
-					ActionThread ActionThread = ActionProcess.Value;
-
-					// Check for pending actions, preemptive failure
-					if (ActionThread == null)
-					{
-						SNDBSResult = false;
-						continue;
-					}
-					// Check for executed action but general failure
-					if (ActionThread.ExitCode != 0)
-					{
-						SNDBSResult = false;
-					}
-					// Log CPU time, tool and task.
-					double ThreadSeconds = Action.Duration.TotalSeconds;
-
-					Log.WriteLineIf(bLogDetailedActionStats,
-						LogEventType.Console,
-						"^{0}^{1:0.00}^{2}^{3}",
-						Action.ActionType.ToString(),
-						ThreadSeconds,
-						Action.CommandPath.GetFileName(),
-						Action.StatusDescription);
-
-					// Keep track of total thread seconds spent on tasks.
-					TotalThreadSeconds += ThreadSeconds;
+					string VariableName = string.Format("{{{0}}}", Variable.Key);
+					TemplateText = TemplateText.Replace(VariableName, Variable.Value.ToString());
 				}
 
-				Log.TraceInformation("-------- End Detailed Actions Stats -----------------------------------------------------------");
-
-				// Log total CPU seconds and numbers of processors involved in tasks.
-				Log.WriteLineIf(bLogDetailedActionStats,
-					LogEventType.Console, "Cumulative thread seconds ({0} processors): {1:0.00}", System.Environment.ProcessorCount, TotalThreadSeconds);
-			}
-
-			// Delete the include-rewrite-rules.ini if it was generated.
-			if (File.Exists(IncludeRewriteRulesFile.FullName))
-			{
-				File.Delete(IncludeRewriteRulesFile.FullName);
-			}
-
-			return SNDBSResult;
-		}
-
-		private void PrepareToolTemplate(Action Action)
-		{
-			string TemplateFileName = String.Format("{0}.sn-dbs-tool.ini", Action.CommandPath.GetFileName());
-			FileReference TemplateInput = FileReference.Combine(UnrealBuildTool.EngineDirectory, "Build", "SNDBSTemplates", TemplateFileName);
-			
-			// If no base template exists, don't try to generate one.
-			if (!File.Exists(TemplateInput.FullName))
-			{
-				return;
-			}
-
-			FileReference TemplateOutput = FileReference.Combine(UnrealBuildTool.EngineDirectory, "Programs", "UnrealBuildTool", "SNDBSTemplates", TemplateFileName);
-			if (!Directory.Exists(TemplateOutput.Directory.FullName))
-			{
-				Directory.CreateDirectory(TemplateOutput.Directory.FullName);
-			}
-
-			string TemplateText = File.ReadAllText(TemplateInput.FullName);
-			TemplateText = TemplateText.Replace("{COMMAND_PATH}", Action.CommandPath.Directory.FullName);
-			foreach (DictionaryEntry Variable in Environment.GetEnvironmentVariables(EnvironmentVariableTarget.Process))
-			{
-				string VariableName = String.Format("{{{0}}}", Variable.Key);
-				TemplateText = TemplateText.Replace(VariableName, Variable.Value.ToString());
+				File.WriteAllText(TemplateOutputFilePath, TemplateText);
 			}
-			File.WriteAllText(TemplateOutput.FullName, TemplateText);
 		}
 
 		private void GenerateSNDBSIncludeRewriteRules()
@@ -522,35 +140,35 @@ namespace UnrealBuildTool
 				}
 			}
 
-			if (!Directory.Exists(IncludeRewriteRulesFile.Directory.FullName))
+			if (!Directory.Exists(IncludeRewriteRulesDirectory))
 			{
-				Directory.CreateDirectory(IncludeRewriteRulesFile.Directory.FullName);
+				Directory.CreateDirectory(IncludeRewriteRulesDirectory);
 			}
 
 			List<string> IncludeRewriteRulesText = new List<string>();
 			IncludeRewriteRulesText.Add("[computed-include-rules]");
 			{
 				IncludeRewriteRulesText.Add(@"pattern1=^COMPILED_PLATFORM_HEADER\(\s*([^ ,]+)\)");
-				IEnumerable<string> PlatformExpansions = PlatformNames.Select(p => String.Format("{0}/{0}$1|{0}$1", p));
-				IncludeRewriteRulesText.Add(String.Format("expansions1={0}", String.Join("|", PlatformExpansions)));
+				IEnumerable<string> PlatformExpansions = PlatformNames.Select(p => string.Format("{0}/{0}$1|{0}$1", p));
+				IncludeRewriteRulesText.Add(string.Format("expansions1={0}", string.Join("|", PlatformExpansions)));
 			}
 			{
 				IncludeRewriteRulesText.Add(@"pattern2=^COMPILED_PLATFORM_HEADER_WITH_PREFIX\(\s*([^ ,]+)\s*,\s*([^ ,]+)\)");
-				IEnumerable<string> PlatformExpansions = PlatformNames.Select(p => String.Format("$1/{0}/{0}$2|$1/{0}$2", p));
-				IncludeRewriteRulesText.Add(String.Format("expansions2={0}", String.Join("|", PlatformExpansions)));
+				IEnumerable<string> PlatformExpansions = PlatformNames.Select(p => string.Format("$1/{0}/{0}$2|$1/{0}$2", p));
+				IncludeRewriteRulesText.Add(string.Format("expansions2={0}", string.Join("|", PlatformExpansions)));
 			}
 			{
 				IncludeRewriteRulesText.Add(@"pattern3=^[A-Z]{5}_PLATFORM_HEADER_NAME\(\s*([^ ,]+)\)");
-				IEnumerable<string> PlatformExpansions = PlatformNames.Select(p => String.Format("{0}/{0}$1|{0}$1", p));
-				IncludeRewriteRulesText.Add(String.Format("expansions3={0}", String.Join("|", PlatformExpansions)));
+				IEnumerable<string> PlatformExpansions = PlatformNames.Select(p => string.Format("{0}/{0}$1|{0}$1", p));
+				IncludeRewriteRulesText.Add(string.Format("expansions3={0}", string.Join("|", PlatformExpansions)));
 			}
 			{
 				IncludeRewriteRulesText.Add(@"pattern4=^[A-Z]{5}_PLATFORM_HEADER_NAME_WITH_PREFIX\(\s*([^ ,]+)\s*,\s*([^ ,]+)\)");
-				IEnumerable<string> PlatformExpansions = PlatformNames.Select(p => String.Format("$1/{0}/{0}$2|$1/{0}$2", p));
-				IncludeRewriteRulesText.Add(String.Format("expansions4={0}", String.Join("|", PlatformExpansions)));
+				IEnumerable<string> PlatformExpansions = PlatformNames.Select(p => string.Format("$1/{0}/{0}$2|$1/{0}$2", p));
+				IncludeRewriteRulesText.Add(string.Format("expansions4={0}", string.Join("|", PlatformExpansions)));
 			}
 
-			File.WriteAllText(IncludeRewriteRulesFile.FullName, String.Join(Environment.NewLine, IncludeRewriteRulesText));
+			File.WriteAllText(IncludeRewriteRulesFilePath, string.Join(Environment.NewLine, IncludeRewriteRulesText));
 		}
 	}
 }
