Dont hard code c++14 in LinuxToolchain

diff --git a/Engine/Source/Programs/UnrealBuildTool/Platform/Linux/LinuxToolChain.cs b/Engine/Source/Programs/UnrealBuildTool/Platform/Linux/LinuxToolChain.cs
--- a/Engine/Source/Programs/UnrealBuildTool/Platform/Linux/LinuxToolChain.cs
+++ b/Engine/Source/Programs/UnrealBuildTool/Platform/Linux/LinuxToolChain.cs
@@ -846,22 +846,30 @@ namespace UnrealBuildTool
 				: string.Format("{0}={1}", myKey, myValue);
 		}
 
-		static string GetCompileArguments_CPP(CppCompileEnvironment CompilerEnvironment)
+		static string GetCompilerStandardVersion_CPP(CppCompileEnvironment CompileEnvironment)
 		{
-			string Result = "";
-			Result += " -x c++";
-			if (CompilerEnvironment.CppStandard == CppStandardVersion.Cpp14 || CompilerEnvironment.CppStandard == CppStandardVersion.Default)
+			if (CompileEnvironment.CppStandard == CppStandardVersion.Cpp14 || CompileEnvironment.CppStandard == CppStandardVersion.Default)
 			{
-				Result += " -std=c++14";
+				return " -std=c++14";
 			}
-			else if (CompilerEnvironment.CppStandard == CppStandardVersion.Cpp17)
+			else if (CompileEnvironment.CppStandard == CppStandardVersion.Cpp17)
 			{
-				Result += " -std=c++17";
+				return " -std=c++17";
 			}
-			else if (CompilerEnvironment.CppStandard == CppStandardVersion.Latest)
+			else if (CompileEnvironment.CppStandard == CppStandardVersion.Latest)
 			{
-				Result += " -std=c++17";
+				return " -std=c++17";
 			}
+
+			throw new BuildException(
+			string.Format("Unknown C++ standard type set: {0}", CompileEnvironment.CppStandard));
+		}
+
+		static string GetCompileArguments_CPP(CppCompileEnvironment CompileEnvironment)
+		{
+			string Result = "";
+			Result += " -x c++";
+			Result += GetCompilerStandardVersion_CPP(CompileEnvironment);
 			return Result;
 		}
 
@@ -872,13 +880,13 @@ namespace UnrealBuildTool
 			return Result;
 		}
 
-		static string GetCompileArguments_MM()
+		static string GetCompileArguments_MM(CppCompileEnvironment CompileEnvironment)
 		{
 			string Result = "";
 			Result += " -x objective-c++";
 			Result += " -fobjc-abi-version=2";
 			Result += " -fobjc-legacy-dispatch";
-			Result += " -std=c++14";
+			Result += GetCompilerStandardVersion_CPP(CompileEnvironment);
 			return Result;
 		}
 
@@ -901,21 +909,21 @@ namespace UnrealBuildTool
 			return Result;
 		}
 
-		static string GetCompileArguments_M()
+		static string GetCompileArguments_M(CppCompileEnvironment CompileEnvironment)
 		{
 			string Result = "";
 			Result += " -x objective-c";
 			Result += " -fobjc-abi-version=2";
 			Result += " -fobjc-legacy-dispatch";
-			Result += " -std=c++14";
+			Result += GetCompilerStandardVersion_CPP(CompileEnvironment);
 			return Result;
 		}
 
-		static string GetCompileArguments_PCH()
+		static string GetCompileArguments_PCH(CppCompileEnvironment CompileEnvironment)
 		{
 			string Result = "";
 			Result += " -x c++-header";
-			Result += " -std=c++14";
+			Result += GetCompilerStandardVersion_CPP(CompileEnvironment);
 			return Result;
 		}
 
@@ -1351,7 +1359,7 @@ namespace UnrealBuildTool
 				// Add C or C++ specific compiler arguments.
 				if (CompileEnvironment.PrecompiledHeaderAction == PrecompiledHeaderAction.Create)
 				{
-					FileArguments += GetCompileArguments_PCH();
+					FileArguments += GetCompileArguments_PCH(CompileEnvironment);
 				}
 				else if (Extension == ".C")
 				{
@@ -1361,13 +1369,13 @@ namespace UnrealBuildTool
 				else if (Extension == ".MM")
 				{
 					// Compile the file as Objective-C++ code.
-					FileArguments += GetCompileArguments_MM();
+					FileArguments += GetCompileArguments_MM(CompileEnvironment);
 					FileArguments += GetRTTIFlag(CompileEnvironment);
 				}
 				else if (Extension == ".M")
 				{
 					// Compile the file as Objective-C code.
-					FileArguments += GetCompileArguments_M();
+					FileArguments += GetCompileArguments_M(CompileEnvironment);
 				}
 				else
 				{
