Unshelved from pending changelist '12670438':

Fix up Android NDK 21 merge differences and native issues

#android

diff --git a/Engine/Source/Programs/UnrealBuildTool/Platform/Android/UEBuildAndroid.cs b/Engine/Source/Programs/UnrealBuildTool/Platform/Android/UEBuildAndroid.cs
--- a/Engine/Source/Programs/UnrealBuildTool/Platform/Android/UEBuildAndroid.cs
+++ b/Engine/Source/Programs/UnrealBuildTool/Platform/Android/UEBuildAndroid.cs
@@ -401,12 +401,11 @@ namespace UnrealBuildTool
 
 			SetUpSpecificEnvironment(Target, CompileEnvironment, LinkEnvironment);
 
-//			LinkEnvironment.AdditionalLibraries.Add("libc++_shared");
+			// deliberately not linking stl or stdc++ here (c++_shared is default)
 			LinkEnvironment.AdditionalLibraries.Add("c");
 			LinkEnvironment.AdditionalLibraries.Add("dl");
 			LinkEnvironment.AdditionalLibraries.Add("log");
 			LinkEnvironment.AdditionalLibraries.Add("m");
-//			LinkEnvironment.AdditionalLibraries.Add("stdc++");
 			LinkEnvironment.AdditionalLibraries.Add("z");
 			LinkEnvironment.AdditionalLibraries.Add("atomic");
 		}
diff --git a/Engine/Source/Runtime/Launch/Private/Android/LaunchAndroid.cpp b/Engine/Source/Runtime/Launch/Private/Android/LaunchAndroid.cpp
--- a/Engine/Source/Runtime/Launch/Private/Android/LaunchAndroid.cpp
+++ b/Engine/Source/Runtime/Launch/Private/Android/LaunchAndroid.cpp
@@ -117,7 +117,8 @@ static const uint32 ValidGamepadKeyCodesList[] =
 static TSet<uint32> ValidGamepadKeyCodes;
 
 // -nostdlib means no crtbegin_so.o, so we have to provide our own __dso_handle and atexit()
-/*
+// this is not needed now we are using stdlib (later NDK has more functionality we should keep)
+#if 0
 extern "C"
 {
 	int atexit(void (*func)(void)) { return 0; }
@@ -125,7 +126,7 @@ extern "C"
 	extern void *__dso_handle __attribute__((__visibility__ ("hidden")));
 	void *__dso_handle;
 }
-*/
+#endif
 
 int32 GAndroidEnableNativeResizeEvent = 0;
 static FAutoConsoleVariableRef CVarEnableResizeNativeEvent(
