- UBT/C++ changes needed for new Restricted locations
- Does not have the changes for new .ini file locations

diff --git a/Engine/Source/Programs/UnrealBuildTool/Configuration/ModuleRules.cs b/Engine/Source/Programs/UnrealBuildTool/Configuration/ModuleRules.cs
--- a/Engine/Source/Programs/UnrealBuildTool/Configuration/ModuleRules.cs
+++ b/Engine/Source/Programs/UnrealBuildTool/Configuration/ModuleRules.cs
@@ -1317,7 +1317,7 @@ namespace UnrealBuildTool
 					case ModuleRules.PrecompileTargetsType.None:
 						return false;
 					case ModuleRules.PrecompileTargetsType.Default:
-						return (Target.Type == TargetType.Editor || !UnrealBuildTool.GetAllEngineDirectories("Source/Developer").Any(Dir => RulesFile.IsUnderDirectory(Dir)) || Plugin != null);
+						return (Target.Type == TargetType.Editor || !UnrealBuildTool.GetExtensionDirs(UnrealBuildTool.EngineDirectory, "Source/Developer").Any(Dir => RulesFile.IsUnderDirectory(Dir)) || Plugin != null);
 					case ModuleRules.PrecompileTargetsType.Game:
 						return (Target.Type == TargetType.Client || Target.Type == TargetType.Server || Target.Type == TargetType.Game);
 					case ModuleRules.PrecompileTargetsType.Editor:
diff --git a/Engine/Source/Programs/UnrealBuildTool/Configuration/UEBuildTarget.cs b/Engine/Source/Programs/UnrealBuildTool/Configuration/UEBuildTarget.cs
--- a/Engine/Source/Programs/UnrealBuildTool/Configuration/UEBuildTarget.cs
+++ b/Engine/Source/Programs/UnrealBuildTool/Configuration/UEBuildTarget.cs
@@ -1492,7 +1492,7 @@ namespace UnrealBuildTool
 
 					foreach (FileReference ChildFile in EnabledPlugin.ChildFiles)
 					{
-							if (ValidFileNames.Contains(ChildFile.GetFileNameWithoutExtension(), StringComparer.InvariantCultureIgnoreCase))
+						if (ValidFileNames.Contains(ChildFile.GetFileNameWithoutExtension(), StringComparer.InvariantCultureIgnoreCase))
 						{
 							Receipt.RuntimeDependencies.Add(ChildFile, StagedFileType.UFS);
 						}
@@ -2338,7 +2338,20 @@ namespace UnrealBuildTool
 
 			if(Module.RulesFile != null)
 			{
-				return !Module.RulesFile.IsUnderDirectory(UnrealBuildTool.EngineSourceDeveloperDirectory) && !Module.RulesFile.IsUnderDirectory(UnrealBuildTool.EngineSourceEditorDirectory);
+				foreach (DirectoryReference ExtensionDir in UnrealBuildTool.GetExtensionDirs(UnrealBuildTool.EngineDirectory))
+				{
+					DirectoryReference SourceDeveloperDir = DirectoryReference.Combine(ExtensionDir, "Source/Developer");
+					if(Module.RulesFile.IsUnderDirectory(SourceDeveloperDir))
+					{
+						return false;
+					}
+
+					DirectoryReference SourceEditorDir = DirectoryReference.Combine(ExtensionDir, "Source/Editor");
+					if(Module.RulesFile.IsUnderDirectory(SourceEditorDir))
+					{
+						return false;
+					}
+				}
 			}
 
 			return true;
@@ -2712,15 +2725,15 @@ namespace UnrealBuildTool
 				List<DirectoryReference> Directories = new List<DirectoryReference>();
 				if (TargetType == TargetType.Editor)
 				{
-					Directories.AddRange(UnrealBuildTool.GetAllEngineDirectories("Source/Editor"));
+					Directories.AddRange(UnrealBuildTool.GetExtensionDirs(UnrealBuildTool.EngineDirectory, "Source/Editor"));
 				}
-				Directories.AddRange(UnrealBuildTool.GetAllEngineDirectories("Source/Runtime"));
+				Directories.AddRange(UnrealBuildTool.GetExtensionDirs(UnrealBuildTool.EngineDirectory, "Source/Runtime"));
 
 				// Also allow anything in the developer directory in non-shipping configurations (though we blacklist by default unless the PrecompileForTargets
 				// setting indicates that it's actually useful at runtime).
 				if(Rules.bBuildDeveloperTools)
 				{
-					Directories.AddRange(UnrealBuildTool.GetAllEngineDirectories("Source/Developer"));
+					Directories.AddRange(UnrealBuildTool.GetExtensionDirs(UnrealBuildTool.EngineDirectory, "Source/Developer"));
 					Directories.Add(DirectoryReference.Combine(UnrealBuildTool.EnterpriseSourceDirectory, "Developer"));
 				}
 
@@ -2729,8 +2742,8 @@ namespace UnrealBuildTool
 				{
 					FileReference ModuleFileName = RulesAssembly.GetModuleFileName(ModuleName);
 					if (Directories.Any(BaseDir => ModuleFileName.IsUnderDirectory(BaseDir)))
-						{
-							Type RulesType = RulesAssembly.GetModuleRulesType(ModuleName);
+					{
+						Type RulesType = RulesAssembly.GetModuleRulesType(ModuleName);
 
 						// Skip platform extension modules. We only care about the base modules, not the platform overrides.
 						// The platform overrides get applied at a later stage when we actually come to build the module.
@@ -3367,24 +3380,15 @@ namespace UnrealBuildTool
 			}
 
 			// Construct the output paths for this target's executable
-			DirectoryReference OutputDirectory = UnrealBuildTool.EngineDirectory;
-
-			bool bOutputToPlatformExtensionDirectory = Rules.File.IsUnderDirectory(UnrealBuildTool.EnginePlatformExtensionsDirectory) || Rules.File.IsUnderDirectory(UnrealBuildTool.ProjectPlatformExtensionsDirectory(ProjectDirectory));
-			bool bOutputToProjectDirectory = (bCompileMonolithic || !bUseSharedBuildEnvironment) && (ProjectDirectory != UnrealBuildTool.EngineDirectory);
-			
-			if (bOutputToPlatformExtensionDirectory && bOutputToProjectDirectory)
+			List<DirectoryReference> PossibleOutputDirectories = new List<DirectoryReference>();
+			PossibleOutputDirectories.AddRange(UnrealBuildTool.GetExtensionDirs(UnrealBuildTool.EngineDirectory));
+			if (bCompileMonolithic || !bUseSharedBuildEnvironment)
 			{
-				OutputDirectory = UnrealBuildTool.GetAllProjectDirectories(ProjectDirectory).First(x => x != ProjectDirectory && Rules.File.IsUnderDirectory(x));
-			}
-			else if(bOutputToPlatformExtensionDirectory)
-			{
-				OutputDirectory = UnrealBuildTool.GetAllEngineDirectories().First(x => x != UnrealBuildTool.EngineDirectory && Rules.File.IsUnderDirectory(x));
-			} 
-			else if (bOutputToProjectDirectory)
-			{
-				OutputDirectory = ProjectDirectory;
+				PossibleOutputDirectories.AddRange(UnrealBuildTool.GetExtensionDirs(UnrealBuildTool.EngineDirectory));
 			}
 
+			DirectoryReference OutputDirectory = PossibleOutputDirectories.Where(x => Rules.File.IsUnderDirectory(x)).OrderByDescending(x => x.FullName.Length).FirstOrDefault() ?? UnrealBuildTool.EngineDirectory;
+
 			bool bCompileAsDLL = Rules.bShouldCompileAsDLL && bCompileMonolithic;
 			List<FileReference> OutputPaths = MakeBinaryPaths(OutputDirectory, bCompileMonolithic ? TargetName : AppName, Platform, Configuration, bCompileAsDLL ? UEBuildBinaryType.DynamicLinkLibrary : UEBuildBinaryType.Executable, Rules.Architecture, Rules.UndecoratedConfiguration, bCompileMonolithic && ProjectFile != null, Rules.ExeBinariesSubFolder, ProjectFile, Rules);
 
diff --git a/Engine/Source/Programs/UnrealBuildTool/Platform/Linux/LinuxToolChain.cs b/Engine/Source/Programs/UnrealBuildTool/Platform/Linux/LinuxToolChain.cs
--- a/Engine/Source/Programs/UnrealBuildTool/Platform/Linux/LinuxToolChain.cs
+++ b/Engine/Source/Programs/UnrealBuildTool/Platform/Linux/LinuxToolChain.cs
@@ -1292,7 +1292,7 @@ namespace UnrealBuildTool
 			{
 				PrintBuildDetails(CompileEnvironment);
 
-				string LinuxDependenciesPath = Path.Combine(UnrealBuildTool.EngineSourceThirdPartyDirectory.FullName, "Linux", PlatformSDK.HaveLinuxDependenciesFile());
+				string LinuxDependenciesPath = Path.Combine(UnrealBuildTool.EngineDirectory.FullName, "Source/ThirdParty/Linux", PlatformSDK.HaveLinuxDependenciesFile());
 				if (!File.Exists(LinuxDependenciesPath))
 				{
 					throw new BuildException("Please make sure that Engine/Source/ThirdParty/Linux is complete (re - run Setup script if using a github build)");
diff --git a/Engine/Source/Programs/UnrealBuildTool/Platform/Windows/VCManifestGenerator.cs b/Engine/Source/Programs/UnrealBuildTool/Platform/Windows/VCManifestGenerator.cs
--- a/Engine/Source/Programs/UnrealBuildTool/Platform/Windows/VCManifestGenerator.cs
+++ b/Engine/Source/Programs/UnrealBuildTool/Platform/Windows/VCManifestGenerator.cs
@@ -292,7 +292,7 @@ namespace UnrealBuildTool
 					// look in Platform extensions too
 					if (!bFileExists)
 					{
-						SourcePath = Path.Combine(UnrealBuildTool.EnginePlatformExtensionsDirectory.FullName, Platform.ToString(), "Build", EngineResourceSubPath);
+						SourcePath = Path.Combine(UnrealBuildTool.EngineDirectory.FullName, "Platforms", Platform.ToString(), "Build", EngineResourceSubPath);
 						bFileExists = File.Exists(Path.Combine(SourcePath, ResourceFileName));
 					}
 				}
diff --git a/Engine/Source/Programs/UnrealBuildTool/ProjectFiles/ProjectFileGenerator.cs b/Engine/Source/Programs/UnrealBuildTool/ProjectFiles/ProjectFileGenerator.cs
--- a/Engine/Source/Programs/UnrealBuildTool/ProjectFiles/ProjectFileGenerator.cs
+++ b/Engine/Source/Programs/UnrealBuildTool/ProjectFiles/ProjectFileGenerator.cs
@@ -477,7 +477,7 @@ namespace UnrealBuildTool
 			DirectoryReference EngineExtras = DirectoryReference.Combine(UnrealBuildTool.EngineDirectory, "Extras");
 			DiscoverCSharpProgramProjectsRecursively(EngineExtras, FoundProjects);
 
-			DirectoryReference[] AllEngineDirectories = UnrealBuildTool.GetAllEngineDirectories("Source/Programs");
+			List<DirectoryReference> AllEngineDirectories = UnrealBuildTool.GetExtensionDirs(UnrealBuildTool.EngineDirectory, "Source/Programs");
 			foreach (DirectoryReference EngineDir in AllEngineDirectories)
 			{
 				DiscoverCSharpProgramProjectsRecursively(EngineDir, FoundProjects);
@@ -1261,12 +1261,9 @@ namespace UnrealBuildTool
 				{
 					// @todo projectfiles: We have engine localization files, but should we also add GAME localization files?
 
-					GameProject.AddFilesToProject( SourceFileSearch.FindFiles( GameProjectDirectory, SearchSubdirectories: false ), GameProjectDirectory );
-
-					DirectoryReference GamePlatformsDirectory = UnrealBuildTool.ProjectPlatformExtensionsDirectory(GameProjectDirectory);
-					if (DirectoryReference.Exists(GamePlatformsDirectory))
+					foreach (DirectoryReference ExtensionDir in UnrealBuildTool.GetExtensionDirs(GameProjectDirectory))
 					{
-						GameProject.AddFilesToProject(SourceFileSearch.FindFiles(GamePlatformsDirectory), GameProjectDirectory);
+						GameProject.AddFilesToProject(SourceFileSearch.FindFiles(ExtensionDir, SearchSubdirectories: false), ExtensionDir);
 					}
 
 					// Game config files
@@ -1347,15 +1344,17 @@ namespace UnrealBuildTool
 		{
 			// @todo: this will add the same files to the solution (like the UBT source files that also get added to UnrealBuildTool project).
 			// not sure of a good filtering method here
-			DirectoryReference PlatformExtensionsDirectory = UnrealBuildTool.EnginePlatformExtensionsDirectory;
-			if (DirectoryReference.Exists(PlatformExtensionsDirectory))
+			foreach(DirectoryReference ExtensionDir in UnrealBuildTool.GetExtensionDirs(UnrealBuildTool.EngineDirectory))
 			{
-				List<string> SubdirectoryNamesToExclude = new List<string>();
-				SubdirectoryNamesToExclude.Add("AutomationTool"); //automation files are added separately to the AutomationTool project
-				SubdirectoryNamesToExclude.Add("Binaries");
-				SubdirectoryNamesToExclude.Add("Content");
+				if (ExtensionDir != UnrealBuildTool.EngineDirectory)
+				{
+					List<string> SubdirectoryNamesToExclude = new List<string>();
+					SubdirectoryNamesToExclude.Add("AutomationTool"); //automation files are added separately to the AutomationTool project
+					SubdirectoryNamesToExclude.Add("Binaries");
+					SubdirectoryNamesToExclude.Add("Content");
 
-				EngineProject.AddFilesToProject(SourceFileSearch.FindFiles(PlatformExtensionsDirectory, SubdirectoryNamesToExclude), UnrealBuildTool.EngineDirectory);
+					EngineProject.AddFilesToProject(SourceFileSearch.FindFiles(ExtensionDir, SubdirectoryNamesToExclude), UnrealBuildTool.EngineDirectory);
+				}
 			}
 		}
 
@@ -2031,14 +2030,11 @@ namespace UnrealBuildTool
 				// check for engine, or platform extension engine folders
 				if( !bIncludeEngineSource )
 				{
-					foreach (DirectoryReference EngineDirectory in UnrealBuildTool.GetAllEngineDirectories())
+					if (CurModuleFile.IsUnderDirectory(UnrealBuildTool.EngineDirectory))
 					{
-						if (CurModuleFile.IsUnderDirectory(EngineDirectory))
-						{
-							// We were asked to exclude engine modules from the generated projects
-							WantProjectFileForModule = false;
-							break;
-						}
+						// We were asked to exclude engine modules from the generated projects
+						WantProjectFileForModule = false;
+						break;
 					}
 				}
 
@@ -2076,8 +2072,7 @@ namespace UnrealBuildTool
 					// Check if there's a plugin directory here
 					if(!ProjectsWithPlugins.Contains(ProjectFile))
 					{
-						DirectoryReference PluginFolder = DirectoryReference.Combine(BaseFolder, "Plugins");
-						if(DirectoryReference.Exists(PluginFolder))
+						foreach (DirectoryReference PluginFolder in UnrealBuildTool.GetExtensionDirs(BaseFolder, "Plugins"))
 						{
 							// Add all the plugin files for this project
 							foreach(FileReference PluginFileName in Plugins.EnumeratePlugins(PluginFolder))
@@ -2176,15 +2171,13 @@ namespace UnrealBuildTool
 				}
 
 				// check for engine, or platform extension engine folders
-				if (Path == UnrealBuildTool.EngineDirectory)
-				{
-					BaseFolder = UnrealBuildTool.EngineDirectory;
-					return FindOrAddProjectHelper(EngineProjectFileNameBase, BaseFolder);
-				}
-				if (Path == UnrealBuildTool.EnginePlatformExtensionsDirectory)
+				foreach (DirectoryReference ExtensionDir in UnrealBuildTool.GetExtensionDirs(UnrealBuildTool.EngineDirectory))
 				{
-					BaseFolder = UnrealBuildTool.EngineDirectory;
-					return FindOrAddProjectHelper(EngineProjectFileNameBase, BaseFolder);
+					if (Path == ExtensionDir)
+					{
+						BaseFolder = UnrealBuildTool.EngineDirectory;
+						return FindOrAddProjectHelper(EngineProjectFileNameBase, BaseFolder);
+					}
 				}
 				if (Path == UnrealBuildTool.EnterpriseDirectory)
 				{
@@ -2238,16 +2231,16 @@ namespace UnrealBuildTool
 				bool IsEngineTarget = false;
 				bool IsEnterpriseTarget = false;
 				bool WantProjectFileForTarget = true;
-				if(UnrealBuildTool.GetAllEngineDirectories().Any(x => TargetFilePath.IsUnderDirectory(x)))
+				if(TargetFilePath.IsUnderDirectory(UnrealBuildTool.EngineDirectory))
 				{
 					// This is an engine target
 					IsEngineTarget = true;
 
-					if(UnrealBuildTool.GetAllEngineDirectories("Source/Programs").Any(x => TargetFilePath.IsUnderDirectory(x)))
+					if(UnrealBuildTool.GetExtensionDirs(UnrealBuildTool.EngineDirectory, "Source/Programs").Any(x => TargetFilePath.IsUnderDirectory(x)))
 					{
 						WantProjectFileForTarget = IncludeEnginePrograms;
 					}
-					else if(UnrealBuildTool.GetAllEngineDirectories("Source").Any(x => TargetFilePath.IsUnderDirectory(x)))
+					else if(UnrealBuildTool.GetExtensionDirs(UnrealBuildTool.EngineDirectory, "Source").Any(x => TargetFilePath.IsUnderDirectory(x)))
 					{
 						WantProjectFileForTarget = bIncludeEngineSource;
 					}
diff --git a/Engine/Source/Programs/UnrealBuildTool/System/ConfigHierarchy.cs b/Engine/Source/Programs/UnrealBuildTool/System/ConfigHierarchy.cs
--- a/Engine/Source/Programs/UnrealBuildTool/System/ConfigHierarchy.cs
+++ b/Engine/Source/Programs/UnrealBuildTool/System/ConfigHierarchy.cs
@@ -826,8 +826,8 @@ namespace UnrealBuildTool
 			out bool bHasPlatformTag, out bool bHasProjectTag, out bool bHasExpansionTag)
 		{
 			// cache some platform extension information that can be used inside the loops
-			string PlatformExtensionEngineConfigDir = DirectoryReference.Combine(UnrealBuildTool.EnginePlatformExtensionsDirectory, PlatformExtensionName, "Config").FullName;
-			string PlatformExtensionProjectConfigDir = ProjectDir != null ? DirectoryReference.Combine(UnrealBuildTool.ProjectPlatformExtensionsDirectory(ProjectDir), PlatformExtensionName, "Config").FullName : null;
+			string PlatformExtensionEngineConfigDir = DirectoryReference.Combine(UnrealBuildTool.EngineDirectory, "Platforms", PlatformExtensionName, "Config").FullName;
+			string PlatformExtensionProjectConfigDir = ProjectDir != null ? DirectoryReference.Combine(ProjectDir, "Platforms", PlatformExtensionName, "Config").FullName : null;
 			bool bHasPlatformExtensionEngineConfigDir = Directory.Exists(PlatformExtensionEngineConfigDir);
 			bool bHasPlatformExtensionProjectConfigDir = PlatformExtensionProjectConfigDir != null && Directory.Exists(PlatformExtensionProjectConfigDir);
 
diff --git a/Engine/Source/Programs/UnrealBuildTool/System/DataDrivenPlatformInfo.cs b/Engine/Source/Programs/UnrealBuildTool/System/DataDrivenPlatformInfo.cs
--- a/Engine/Source/Programs/UnrealBuildTool/System/DataDrivenPlatformInfo.cs
+++ b/Engine/Source/Programs/UnrealBuildTool/System/DataDrivenPlatformInfo.cs
@@ -49,7 +49,8 @@ namespace UnrealBuildTool
 				PlatformInfos = new Dictionary<string, ConfigDataDrivenPlatformInfo>();
 				Dictionary<string, string> IniParents = new Dictionary<string, string>();
 
-				foreach (DirectoryReference EngineConfigDir in UnrealBuildTool.GetAllEngineDirectories("Config"))
+				// find all platform directories (skipping NFL/NoRedist)
+				foreach (DirectoryReference EngineConfigDir in UnrealBuildTool.GetExtensionDirs(UnrealBuildTool.EngineDirectory, "Config", bIncludeRestrictedDirectories:false))
 				{
 					// look through all config dirs looking for the data driven ini file
 					foreach (string FilePath in Directory.EnumerateFiles(EngineConfigDir.FullName, "DataDrivenPlatformInfo.ini", SearchOption.AllDirectories))
diff --git a/Engine/Source/Programs/UnrealBuildTool/System/FileMetadataPrefetch.cs b/Engine/Source/Programs/UnrealBuildTool/System/FileMetadataPrefetch.cs
--- a/Engine/Source/Programs/UnrealBuildTool/System/FileMetadataPrefetch.cs
+++ b/Engine/Source/Programs/UnrealBuildTool/System/FileMetadataPrefetch.cs
@@ -110,20 +110,26 @@ namespace UnrealBuildTool
 		/// </summary>
 		static void ScanEngineDirectory()
 		{
-			DirectoryItem EngineDirectory = DirectoryItem.GetItemByDirectoryReference(UnrealBuildTool.EngineDirectory);
-			EngineDirectory.CacheDirectories();
+			foreach (DirectoryReference ExtensionDir in UnrealBuildTool.GetExtensionDirs(UnrealBuildTool.EngineDirectory))
+			{
+				DirectoryItem BaseDirectory = DirectoryItem.GetItemByDirectoryReference(ExtensionDir);
+				BaseDirectory.CacheDirectories();
+
+				DirectoryItem BasePluginsDirectory = DirectoryItem.Combine(BaseDirectory, "Plugins");
+				Enqueue(() => ScanPluginFolder(BasePluginsDirectory));
 
-			DirectoryItem EnginePluginsDirectory = DirectoryItem.Combine(EngineDirectory, "Plugins");
-			Enqueue(() => ScanPluginFolder(EnginePluginsDirectory));
+				DirectoryItem BaseSourceDirectory = DirectoryItem.Combine(BaseDirectory, "Source");
+				BaseSourceDirectory.CacheDirectories();
 
-			DirectoryItem EngineRuntimeDirectory = DirectoryItem.GetItemByDirectoryReference(UnrealBuildTool.EngineSourceRuntimeDirectory);
-			Enqueue(() => ScanDirectoryTree(EngineRuntimeDirectory));
+				DirectoryItem BaseSourceRuntimeDirectory = DirectoryItem.Combine(BaseSourceDirectory, "Runtime");
+				Enqueue(() => ScanDirectoryTree(BaseSourceRuntimeDirectory));
 
-			DirectoryItem EngineDeveloperDirectory = DirectoryItem.GetItemByDirectoryReference(UnrealBuildTool.EngineSourceDeveloperDirectory);
-			Enqueue(() => ScanDirectoryTree(EngineDeveloperDirectory));
+				DirectoryItem BaseSourceDeveloperDirectory = DirectoryItem.Combine(BaseSourceDirectory, "Developer");
+				Enqueue(() => ScanDirectoryTree(BaseSourceDeveloperDirectory));
 
-			DirectoryItem EngineEditorDirectory = DirectoryItem.GetItemByDirectoryReference(UnrealBuildTool.EngineSourceEditorDirectory);
-			Enqueue(() => ScanDirectoryTree(EngineEditorDirectory));
+				DirectoryItem BaseSourceEditorDirectory = DirectoryItem.Combine(BaseSourceDirectory, "Editor");
+				Enqueue(() => ScanDirectoryTree(BaseSourceEditorDirectory));
+			}
 		}
 
 		/// <summary>
@@ -132,11 +138,17 @@ namespace UnrealBuildTool
 		/// <param name="ProjectDirectory">The project directory to search</param>
 		static void ScanProjectDirectory(DirectoryItem ProjectDirectory)
 		{
-			DirectoryItem ProjectPluginsDirectory = DirectoryItem.Combine(ProjectDirectory, "Plugins");
-			Enqueue(() => ScanPluginFolder(ProjectPluginsDirectory));
+			foreach (DirectoryReference ExtensionDir in UnrealBuildTool.GetExtensionDirs(UnrealBuildTool.EngineDirectory))
+			{
+				DirectoryItem BaseDirectory = DirectoryItem.GetItemByDirectoryReference(ExtensionDir);
+				BaseDirectory.CacheDirectories();
+
+				DirectoryItem BasePluginsDirectory = DirectoryItem.Combine(ProjectDirectory, "Plugins");
+				Enqueue(() => ScanPluginFolder(BasePluginsDirectory));
 
-			DirectoryItem ProjectSourceDirectory = DirectoryItem.Combine(ProjectDirectory, "Source");
-			Enqueue(() => ScanDirectoryTree(ProjectSourceDirectory));
+				DirectoryItem BaseSourceDirectory = DirectoryItem.Combine(ProjectDirectory, "Source");
+				Enqueue(() => ScanDirectoryTree(BaseSourceDirectory));
+			}
 		}
 
 		/// <summary>
diff --git a/Engine/Source/Programs/UnrealBuildTool/System/Plugins.cs b/Engine/Source/Programs/UnrealBuildTool/System/Plugins.cs
--- a/Engine/Source/Programs/UnrealBuildTool/System/Plugins.cs
+++ b/Engine/Source/Programs/UnrealBuildTool/System/Plugins.cs
@@ -486,18 +486,8 @@ namespace UnrealBuildTool
 		/// <returns>Sequence of the found PluginInfo object.</returns>
 		public static IReadOnlyList<PluginInfo> ReadPluginsFromDirectory(DirectoryReference RootDirectory, string Subdirectory, PluginType Type)
 		{
-			// look for directories in RootDirectory and and Platform directories under RootDirectory
-			List<DirectoryReference> RootDirectories = new List<DirectoryReference>() { DirectoryReference.Combine(RootDirectory, Subdirectory) };
-
-			// now look for platform subdirectories with the Subdirectory
-			DirectoryReference PlatformDirectory = DirectoryReference.Combine(RootDirectory, "Platforms");
-			if (DirectoryReference.Exists(PlatformDirectory))
-			{
-				foreach (DirectoryReference Dir in DirectoryReference.EnumerateDirectories(PlatformDirectory))
-				{
-					RootDirectories.Add(DirectoryReference.Combine(Dir, Subdirectory));
-				}
-			}
+			// look for directories in RootDirectory and and extension directories under RootDirectory
+			List<DirectoryReference> RootDirectories = UnrealBuildTool.GetExtensionDirs(RootDirectory, Subdirectory);
 
 			Dictionary<PluginInfo, FileReference> ChildPlugins = new Dictionary<PluginInfo, FileReference>();
 			List<PluginInfo> AllParentPlugins = new List<PluginInfo>();
diff --git a/Engine/Source/Programs/UnrealBuildTool/System/RulesCompiler.cs b/Engine/Source/Programs/UnrealBuildTool/System/RulesCompiler.cs
--- a/Engine/Source/Programs/UnrealBuildTool/System/RulesCompiler.cs
+++ b/Engine/Source/Programs/UnrealBuildTool/System/RulesCompiler.cs
@@ -67,25 +67,20 @@ namespace UnrealBuildTool
 		/// <param name="bIncludeEngine"></param>
 		/// <param name="bIncludeEnterprise"></param>
 		/// <param name="bIncludeTempTargets">Whether to include targets generated by UAT to accomodate content-only projects that need to be compiled to include plugins</param>
-		/// <param name="bIncludePlatformExtensions"></param>
 		/// <returns></returns>
-		public static List<FileReference> FindAllRulesSourceFiles(RulesFileType RulesFileType, List<DirectoryReference> GameFolders, List<FileReference> ForeignPlugins, List<DirectoryReference> AdditionalSearchPaths, bool bIncludeEngine = true, bool bIncludeEnterprise = true, bool bIncludeTempTargets = true, bool bIncludePlatformExtensions = true)
+		public static List<FileReference> FindAllRulesSourceFiles(RulesFileType RulesFileType, List<DirectoryReference> GameFolders, List<FileReference> ForeignPlugins, List<DirectoryReference> AdditionalSearchPaths, bool bIncludeEngine = true, bool bIncludeEnterprise = true, bool bIncludeTempTargets = true)
 		{
 			List<DirectoryReference> Folders = new List<DirectoryReference>();
 
 			// Add all engine source (including third party source)
 			if (bIncludeEngine)
 			{
-				Folders.Add(UnrealBuildTool.EngineSourceDirectory);
+				Folders.AddRange(UnrealBuildTool.GetExtensionDirs(UnrealBuildTool.EngineDirectory, "Source"));
 			}
 			if(bIncludeEnterprise)
 			{
 				Folders.Add(UnrealBuildTool.EnterpriseSourceDirectory);
 			}
-			if (bIncludePlatformExtensions)
-			{
-				Folders.Add(UnrealBuildTool.EnginePlatformExtensionsDirectory);
-			}
 
 			// @todo plugin: Disallow modules from including plugin modules as dependency modules? (except when the module is part of that plugin)
 
@@ -93,7 +88,7 @@ namespace UnrealBuildTool
 			List<DirectoryReference> RootFolders = new List<DirectoryReference>();
 			if (bIncludeEngine)
 			{
-				RootFolders.AddRange(UnrealBuildTool.GetAllEngineDirectories());
+				RootFolders.AddRange(UnrealBuildTool.GetExtensionDirs(UnrealBuildTool.EngineDirectory));
 			}
 			if(bIncludeEnterprise)
 			{
@@ -101,17 +96,7 @@ namespace UnrealBuildTool
 			}
 			if (GameFolders != null)
 			{
-				if (bIncludePlatformExtensions)
-				{
-					foreach (DirectoryReference GameFolder in GameFolders)
-					{
-						RootFolders.AddRange(UnrealBuildTool.GetAllProjectDirectories(GameFolder));
-					}
-				}
-				else
-				{
-					RootFolders.AddRange(GameFolders);
-				}
+				RootFolders.AddRange(GameFolders.SelectMany(x => UnrealBuildTool.GetExtensionDirs(x)));
 			}
 
 			// Find all the plugin source directories
@@ -138,14 +123,7 @@ namespace UnrealBuildTool
 			{
 				foreach (DirectoryReference GameFolder in GameFolders)
 				{
-					if (bIncludePlatformExtensions)
-					{
-						Folders.AddRange(UnrealBuildTool.GetAllProjectDirectories(GameFolder, "Source"));
-					}
-					else
-					{
-						Folders.Add(DirectoryReference.Combine(GameFolder, "Source"));
-					}
+					Folders.AddRange(UnrealBuildTool.GetExtensionDirs(GameFolder, "Source"));
 
 					if (bIncludeTempTargets)
 					{
@@ -394,7 +372,7 @@ namespace UnrealBuildTool
 
 				RulesScope EngineScope = new RulesScope("Engine", null);
 
-				EngineRulesAssembly = CreateEngineOrEnterpriseRulesAssembly(EngineScope, UnrealBuildTool.GetAllEngineDirectories().ToList(), ProjectFileGenerator.EngineProjectFileNameBase, IncludedPlugins, UnrealBuildTool.IsEngineInstalled() || bUsePrecompiled, bSkipCompile, null);
+				EngineRulesAssembly = CreateEngineOrEnterpriseRulesAssembly(EngineScope, UnrealBuildTool.GetExtensionDirs(UnrealBuildTool.EngineDirectory), ProjectFileGenerator.EngineProjectFileNameBase, IncludedPlugins, UnrealBuildTool.IsEngineInstalled() || bUsePrecompiled, bSkipCompile, null);
 			}
 			return EngineRulesAssembly;
 		}
@@ -563,7 +541,7 @@ namespace UnrealBuildTool
 				List<FileReference> TargetFiles = new List<FileReference>();
 
 				// Find all the project directories
-				List<DirectoryReference> ProjectDirectories = new List<DirectoryReference>(UnrealBuildTool.GetAllProjectDirectories(ProjectFileName));
+				List<DirectoryReference> ProjectDirectories = UnrealBuildTool.GetExtensionDirs(ProjectFileName.Directory);
 				if (Project.AdditionalRootDirectories != null)
 				{
 					ProjectDirectories.AddRange(Project.AdditionalRootDirectories);
diff --git a/Engine/Source/Programs/UnrealBuildTool/ToolChain/UEToolChain.cs b/Engine/Source/Programs/UnrealBuildTool/ToolChain/UEToolChain.cs
--- a/Engine/Source/Programs/UnrealBuildTool/ToolChain/UEToolChain.cs
+++ b/Engine/Source/Programs/UnrealBuildTool/ToolChain/UEToolChain.cs
@@ -240,7 +240,7 @@ namespace UnrealBuildTool
 		/// <returns>Path to ISPC compiler</returns>
 		public virtual string GetISPCHostCompilerPath(UnrealTargetPlatform Platform)
 		{
-			string ISPCCompilerPathCommon = Path.Combine(UnrealBuildTool.EngineSourceThirdPartyDirectory.FullName, "IntelISPC", "bin");
+			string ISPCCompilerPathCommon = Path.Combine(UnrealBuildTool.EngineDirectory.FullName, "Source", "ThirdParty", "IntelISPC", "bin");
 			string ISPCArchitecturePath = "";
 			string ExeExtension = ".exe";
 
diff --git a/Engine/Source/Programs/UnrealBuildTool/UnrealBuildTool.cs b/Engine/Source/Programs/UnrealBuildTool/UnrealBuildTool.cs
--- a/Engine/Source/Programs/UnrealBuildTool/UnrealBuildTool.cs
+++ b/Engine/Source/Programs/UnrealBuildTool/UnrealBuildTool.cs
@@ -73,33 +73,33 @@ namespace UnrealBuildTool
 		/// <summary>
 		/// Full path to the Engine/Source/Runtime directory
 		/// </summary>
+		[Obsolete("Please use UnrealBuildTool.GetExtensionDirs(UnrealBuildTool.EngineDirectory, \"Source/Runtime\") instead.")]
 		public static readonly DirectoryReference EngineSourceRuntimeDirectory = DirectoryReference.Combine(EngineSourceDirectory, "Runtime");
 
 		/// <summary>
 		/// Full path to the Engine/Source/Developer directory
 		/// </summary>
+		[Obsolete("Please use UnrealBuildTool.GetExtensionDirs(UnrealBuildTool.EngineDirectory, \"Source/Developer\") instead.")]
 		public static readonly DirectoryReference EngineSourceDeveloperDirectory = DirectoryReference.Combine(EngineSourceDirectory, "Developer");
 
 		/// <summary>
 		/// Full path to the Engine/Source/Editor directory
 		/// </summary>
+		[Obsolete("Please use UnrealBuildTool.GetExtensionDirs(UnrealBuildTool.EngineDirectory, \"Source/Editor\") instead.")]
 		public static readonly DirectoryReference EngineSourceEditorDirectory = DirectoryReference.Combine(EngineSourceDirectory, "Editor");
 
 		/// <summary>
 		/// Full path to the Engine/Source/Programs directory
 		/// </summary>
+		[Obsolete("Please use UnrealBuildTool.GetExtensionDirs(UnrealBuildTool.EngineDirectory, \"Source/Programs\") instead.")]
 		public static readonly DirectoryReference EngineSourceProgramsDirectory = DirectoryReference.Combine(EngineSourceDirectory, "Programs");
 
 		/// <summary>
 		/// Full path to the Engine/Source/ThirdParty directory
 		/// </summary>
+		[Obsolete("Please use UnrealBuildTool.GetExtensionDirs(UnrealBuildTool.EngineDirectory, \"Source/ThirdParty\") instead.")]
 		public static readonly DirectoryReference EngineSourceThirdPartyDirectory = DirectoryReference.Combine(EngineSourceDirectory, "ThirdParty");
 
-		/// <summary>
-		/// The full name of the Engine's PlatformExtensions directory
-		/// </summary>
-		public static readonly DirectoryReference EnginePlatformExtensionsDirectory = DirectoryReference.Combine(EngineDirectory, "Platforms");
-
 		/// <summary>
 		/// The full name of the Enterprise directory
 		/// </summary>
@@ -143,79 +143,63 @@ namespace UnrealBuildTool
 		}
 
 		/// <summary>
-		/// Returns the root location of platform extensions within the given project
-		/// </summary>
-		/// <param name="ProjectDirectory">Location of the project on disk</param>
-		/// <returns>the root location of platform extensions within the given project</returns>
-		public static DirectoryReference ProjectPlatformExtensionsDirectory(DirectoryReference ProjectDirectory)
-		{
-			return DirectoryReference.Combine(ProjectDirectory, "Platforms");
-		}
-
-		/// <summary>
-		/// Returns the root location of platform extensions within the given project
+		/// Finds all the extension directories for the given base directory. This includes platform extensions and restricted folders.
 		/// </summary>
-		/// <param name="ProjectFile">Location of the .uproject file on disk</param>
-		/// <returns>the root location of platform extensions within the given project</returns>
-		public static DirectoryReference ProjectPlatformExtensionsDirectory(FileReference ProjectFile)
+		/// <param name="BaseDir">Location of the base directory</param>
+		/// <param name="bIncludePlatformDirectories">If true, platform subdirectories are included (will return platform directories under Restricted dirs, even if bIncludeRestrictedDirectories is false)</param>
+		/// <param name="bIncludeRestrictedDirectories">If true, restricted (NotForLicensees, NoRedist) subdirectories are included</param>
+		/// <returns>List of extension directories, including the given base directory</returns>
+		public static List<DirectoryReference> GetExtensionDirs(DirectoryReference BaseDir, bool bIncludePlatformDirectories=true, bool bIncludeRestrictedDirectories=true)
 		{
-			return ProjectPlatformExtensionsDirectory(ProjectFile.Directory);
-		}
+			List<DirectoryReference> ExtensionDirs = new List<DirectoryReference>();
+			ExtensionDirs.Add(BaseDir);
 
-		/// <summary>
-		/// The main engine directory and all found platform extension engine directories
-		/// </summary>
-		public static DirectoryReference[] GetAllEngineDirectories(string Suffix="")
-		{
-			List<DirectoryReference> EngineDirectories = new List<DirectoryReference>() { DirectoryReference.Combine(EngineDirectory, Suffix) };
-			if (DirectoryReference.Exists(EnginePlatformExtensionsDirectory))
+			if (bIncludePlatformDirectories)
 			{
-				foreach (DirectoryReference PlatformDirectory in DirectoryReference.EnumerateDirectories(EnginePlatformExtensionsDirectory))
+				DirectoryReference PlatformExtensionBaseDir = DirectoryReference.Combine(BaseDir, "Platforms");
+				if (DirectoryReference.Exists(PlatformExtensionBaseDir))
 				{
-					DirectoryReference PlatformEngineDirectory = DirectoryReference.Combine(PlatformDirectory, Suffix);
-					if (DirectoryReference.Exists(PlatformEngineDirectory))
-					{
-						EngineDirectories.Add(PlatformEngineDirectory);
-					}
+					ExtensionDirs.AddRange(DirectoryReference.EnumerateDirectories(PlatformExtensionBaseDir));
 				}
 			}
-			return EngineDirectories.ToArray();
-		}
 
-		/// <summary>
-		/// Returns the main project directory and all found platform extension project directories, with
-		/// an optional subdirectory to look for within each location (ie, "Config" or "Source/Runtime")
-		/// </summary>
-		/// <param name="ProjectDirectory">Location of the project on disk</param>
-		/// <param name="Suffix">Optional subdirectory to look in for each location</param>
-		/// <returns>the main project directory and all found platform extension project directories</returns>
-		public static DirectoryReference[] GetAllProjectDirectories(DirectoryReference ProjectDirectory, string Suffix = "")
-		{
-			List<DirectoryReference> ProjectDirectories = new List<DirectoryReference>() { DirectoryReference.Combine(ProjectDirectory, Suffix) };
-			if (DirectoryReference.Exists(ProjectPlatformExtensionsDirectory(ProjectDirectory)))
+			DirectoryReference RestrictedBaseDir = DirectoryReference.Combine(BaseDir, "Restricted");
+			if (DirectoryReference.Exists(RestrictedBaseDir))
 			{
-				foreach (DirectoryReference PlatformDirectory in DirectoryReference.EnumerateDirectories(ProjectPlatformExtensionsDirectory(ProjectDirectory), "*", SearchOption.TopDirectoryOnly))
+				IEnumerable<DirectoryReference> RestrictedDirs = DirectoryReference.EnumerateDirectories(RestrictedBaseDir);
+				if (bIncludeRestrictedDirectories)
 				{
-					DirectoryReference PlatformEngineDirectory = DirectoryReference.Combine(PlatformDirectory, Suffix);
-					if (DirectoryReference.Exists(PlatformEngineDirectory))
+					ExtensionDirs.AddRange(RestrictedDirs);
+				}
+
+				if (bIncludePlatformDirectories)
+				{
+					// also look for nested platforms in the restricted 
+					foreach (DirectoryReference RestrictedDir in RestrictedDirs)
 					{
-						ProjectDirectories.Add(PlatformEngineDirectory);
+						DirectoryReference RestrictedPlatformExtensionBaseDir = DirectoryReference.Combine(RestrictedBaseDir, "Platforms");
+						if (DirectoryReference.Exists(RestrictedPlatformExtensionBaseDir))
+						{
+							ExtensionDirs.AddRange(DirectoryReference.EnumerateDirectories(RestrictedPlatformExtensionBaseDir));
+						}
 					}
 				}
 			}
-			return ProjectDirectories.ToArray();
+
+			return ExtensionDirs;
 		}
 
 		/// <summary>
-		/// Returns the main project directory and all found platform extension project directories, with
-		/// an optional subdirectory to look for within each location (ie, "Config" or "Source/Runtime")
+		/// Finds all the extension directories for the given base directory. This includes platform extensions and restricted folders.
 		/// </summary>
-		/// <param name="ProjectFile">Location of the .uproject file on disk</param>
-		/// <param name="Suffix">Optional subdirectory to look in for each location</param>
-		/// <returns>the main project directory and all found platform extension project directories</returns>
-		public static DirectoryReference[] GetAllProjectDirectories(FileReference ProjectFile, string Suffix = "")
+		/// <param name="BaseDir">Location of the base directory</param>
+		/// <param name="SubDir">The subdirectory to find</param>
+		/// <param name="bIncludePlatformDirectories">If true, platform subdirectories are included (will return platform directories under Restricted dirs, even if bIncludeRestrictedDirectories is false)</param>
+		/// <param name="bIncludeRestrictedDirectories">If true, restricted (NotForLicensees, NoRedist) subdirectories are included</param>
+		/// <returns>List of extension directories, including the given base directory</returns>
+		public static List<DirectoryReference> GetExtensionDirs(DirectoryReference BaseDir, string SubDir, bool bIncludePlatformDirectories = true, bool bIncludeRestrictedDirectories = true)
 		{
-			return GetAllProjectDirectories(ProjectFile.Directory, Suffix);
+			return GetExtensionDirs(BaseDir, bIncludePlatformDirectories, bIncludeRestrictedDirectories).Select(x => DirectoryReference.Combine(x, SubDir)).Where(x => DirectoryReference.Exists(x)).ToList();
 		}
 
 		/// <summary>
diff --git a/Engine/Source/Runtime/Core/Private/Misc/Paths.cpp b/Engine/Source/Runtime/Core/Private/Misc/Paths.cpp
--- a/Engine/Source/Runtime/Core/Private/Misc/Paths.cpp
+++ b/Engine/Source/Runtime/Core/Private/Misc/Paths.cpp
@@ -248,6 +248,46 @@ FString FPaths::ProjectPlatformExtensionsDir()
 	return FPaths::ProjectDir() + TEXT("Platforms/");
 }
 
+
+static void AddIfDirectoryExists(TArray<FString>& ExtensionDirs, FString&& Dir)
+{
+	if (FPaths::DirectoryExists(Dir))
+	{
+		ExtensionDirs.Add(Dir);
+	}
+}
+
+TArray<FString> FPaths::GetExtensionDirs(const FString& BaseDir, const FString& SubDir)
+{
+	TArray<FString> ExtensionDirs;
+	AddIfDirectoryExists(ExtensionDirs, FPaths::Combine(BaseDir, SubDir));
+
+	FString PlatformExtensionBaseDir = FPaths::Combine(BaseDir, TEXT("Platforms"));
+	for (auto const& Platform : FDataDrivenPlatformInfoRegistry::GetAllPlatformInfos())
+	{
+		for (auto const& AdditionalFolder : Platform.Value.AdditionalRestrictedFolders)
+		{
+			AddIfDirectoryExists(ExtensionDirs, FPaths::Combine(PlatformExtensionBaseDir, AdditionalFolder, SubDir));
+		}
+		AddIfDirectoryExists(ExtensionDirs, FPaths::Combine(PlatformExtensionBaseDir, Platform.Key, SubDir));
+	}
+
+	FString RestrictedBaseDir = FPaths::Combine(BaseDir, TEXT("Restricted"));
+	IFileManager::Get().IterateDirectory(*RestrictedBaseDir, [&ExtensionDirs, SubDir](const TCHAR* FilenameOrDirectory, bool bIsDirectory)  -> bool
+	{
+		if (bIsDirectory)
+		{
+			AddIfDirectoryExists(ExtensionDirs, FPaths::Combine(FilenameOrDirectory, SubDir));
+
+			// now look for platforms under it
+			ExtensionDirs.Append(GetExtensionDirs(FilenameOrDirectory, SubDir));
+		}
+		return true;
+	});
+
+	return ExtensionDirs;
+}
+
 FString FPaths::RootDir()
 {
 	return FString(FPlatformMisc::RootDir());
diff --git a/Engine/Source/Runtime/Core/Public/Misc/Paths.h b/Engine/Source/Runtime/Core/Public/Misc/Paths.h
--- a/Engine/Source/Runtime/Core/Public/Misc/Paths.h
+++ b/Engine/Source/Runtime/Core/Public/Misc/Paths.h
@@ -138,6 +138,13 @@ public:
 	 */
 	static FString ProjectPlatformExtensionsDir();
 
+	/**
+	 * Returns platform and restricted extensions that are present and valid (for platforms, it uses FDataDrivePlatformInfo to determine valid platforms, it doesn't just use what's present)
+	 *
+	 * @return BaseDir and usable extension directories under BaseDir (either Engine or Project)
+	 */
+	static TArray<FString> GetExtensionDirs(const FString& BaseDir, const FString& SubDir=FString());
+
 	/**
 	 * Returns the root directory of the engine directory tree
 	 *
diff --git a/Engine/Source/Runtime/Projects/Private/PluginManager.cpp b/Engine/Source/Runtime/Projects/Private/PluginManager.cpp
--- a/Engine/Source/Runtime/Projects/Private/PluginManager.cpp
+++ b/Engine/Source/Runtime/Projects/Private/PluginManager.cpp
@@ -223,7 +223,7 @@ bool FPluginManager::AddToPluginsList(const FString& PluginFilename)
 	{
 		// Determine the plugin type
 		EPluginType PluginType = EPluginType::External;
-		if (PluginFilename.StartsWith(FPaths::EngineDir()) || PluginFilename.StartsWith(FPaths::EnginePlatformExtensionsDir()))
+		if (PluginFilename.StartsWith(FPaths::EngineDir()))
 		{
 			PluginType = EPluginType::Engine;
 		}
@@ -294,29 +294,18 @@ void FPluginManager::ReadAllPlugins(TMap<FString, TSharedRef<FPlugin>>& Plugins,
 	if (ManifestFileNames.Num() == 0)
 	{
 		// Find "built-in" plugins.  That is, plugins situated right within the Engine directory.
-		ReadPluginsInDirectory(FPaths::EnginePluginsDir(), EPluginType::Engine, Plugins, ChildPlugins);
-		for (auto const& Platform : FDataDrivenPlatformInfoRegistry::GetAllPlatformInfos())
+		for (const FString& EnginePluginDir : FPaths::GetExtensionDirs(FPaths::EngineDir(), TEXT("Plugins")))
 		{
-			for (auto const& AdditionalFolder : Platform.Value.AdditionalRestrictedFolders)
-			{
-				ReadPluginsInDirectory(FPaths::EnginePlatformExtensionsDir() / AdditionalFolder / TEXT("Plugins"), EPluginType::Engine, Plugins, ChildPlugins);
-			}
-			ReadPluginsInDirectory(FPaths::EnginePlatformExtensionsDir() / Platform.Key / TEXT("Plugins"), EPluginType::Engine, Plugins, ChildPlugins);
+			ReadPluginsInDirectory(EnginePluginDir, EPluginType::Engine, Plugins, ChildPlugins);
 		}
 
 		// Find plugins in the game project directory (<MyGameProject>/Plugins). If there are any engine plugins matching the name of a game plugin,
 		// assume that the game plugin version is preferred.
 		if (Project != nullptr)
 		{
-			FString Root = FPaths::GetPath(FPaths::GetProjectFilePath());
-			ReadPluginsInDirectory(Root / TEXT("Plugins"), EPluginType::Project, Plugins, ChildPlugins);
-			for (auto const& Platform : FDataDrivenPlatformInfoRegistry::GetAllPlatformInfos())
+			for (const FString& ProjectPluginDir : FPaths::GetExtensionDirs(FPaths::GetPath(FPaths::GetProjectFilePath()), TEXT("Plugins")))
 			{
-				for (auto const& AdditionalFolder : Platform.Value.AdditionalRestrictedFolders)
-				{
-					ReadPluginsInDirectory(Root / TEXT("Platforms") / AdditionalFolder / TEXT("Plugins"), EPluginType::Project, Plugins, ChildPlugins);
-				}
-				ReadPluginsInDirectory(Root / TEXT("Platforms") / Platform.Key / TEXT("Plugins"), EPluginType::Project, Plugins, ChildPlugins);
+				ReadPluginsInDirectory(ProjectPluginDir, EPluginType::Project, Plugins, ChildPlugins);
 			}
 		}
 	}
