- Sped up GetExtensionDirs() in UBT

diff --git a/Engine/Source/Programs/UnrealBuildTool/UnrealBuildTool.cs b/Engine/Source/Programs/UnrealBuildTool/UnrealBuildTool.cs
--- a/Engine/Source/Programs/UnrealBuildTool/UnrealBuildTool.cs
+++ b/Engine/Source/Programs/UnrealBuildTool/UnrealBuildTool.cs
@@ -142,6 +142,9 @@ namespace UnrealBuildTool
 			}
 		}
 
+		// cached dictionary of BaseDir to extension directories
+		private static Dictionary<DirectoryReference, Tuple<List<DirectoryReference>, List<DirectoryReference>>> CachedExtensionDirectories = new Dictionary<DirectoryReference, Tuple<List<DirectoryReference>, List<DirectoryReference>>>();
+		
 		/// <summary>
 		/// Finds all the extension directories for the given base directory. This includes platform extensions and restricted folders.
 		/// </summary>
@@ -151,41 +154,47 @@ namespace UnrealBuildTool
 		/// <returns>List of extension directories, including the given base directory</returns>
 		public static List<DirectoryReference> GetExtensionDirs(DirectoryReference BaseDir, bool bIncludePlatformDirectories=true, bool bIncludeRestrictedDirectories=true)
 		{
-			List<DirectoryReference> ExtensionDirs = new List<DirectoryReference>();
-			ExtensionDirs.Add(BaseDir);
-
-			if (bIncludePlatformDirectories)
+			Tuple<List<DirectoryReference>, List<DirectoryReference>> CachedDirs;
+			if (!CachedExtensionDirectories.TryGetValue(BaseDir, out CachedDirs))
 			{
+				CachedDirs = Tuple.Create(new List<DirectoryReference>(), new List<DirectoryReference>());
+
+				CachedExtensionDirectories[BaseDir] = CachedDirs;
+
 				DirectoryReference PlatformExtensionBaseDir = DirectoryReference.Combine(BaseDir, "Platforms");
 				if (DirectoryReference.Exists(PlatformExtensionBaseDir))
 				{
-					ExtensionDirs.AddRange(DirectoryReference.EnumerateDirectories(PlatformExtensionBaseDir));
+					CachedDirs.Item1.AddRange(DirectoryReference.EnumerateDirectories(PlatformExtensionBaseDir));
 				}
-			}
 
-			DirectoryReference RestrictedBaseDir = DirectoryReference.Combine(BaseDir, "Restricted");
-			if (DirectoryReference.Exists(RestrictedBaseDir))
-			{
-				IEnumerable<DirectoryReference> RestrictedDirs = DirectoryReference.EnumerateDirectories(RestrictedBaseDir);
-				if (bIncludeRestrictedDirectories)
+				DirectoryReference RestrictedBaseDir = DirectoryReference.Combine(BaseDir, "Restricted");
+				if (DirectoryReference.Exists(RestrictedBaseDir))
 				{
-					ExtensionDirs.AddRange(RestrictedDirs);
-				}
+					IEnumerable<DirectoryReference> RestrictedDirs = DirectoryReference.EnumerateDirectories(RestrictedBaseDir);
+					CachedDirs.Item2.AddRange(RestrictedDirs);
 
-				if (bIncludePlatformDirectories)
-				{
 					// also look for nested platforms in the restricted 
 					foreach (DirectoryReference RestrictedDir in RestrictedDirs)
 					{
 						DirectoryReference RestrictedPlatformExtensionBaseDir = DirectoryReference.Combine(RestrictedBaseDir, "Platforms");
 						if (DirectoryReference.Exists(RestrictedPlatformExtensionBaseDir))
 						{
-							ExtensionDirs.AddRange(DirectoryReference.EnumerateDirectories(RestrictedPlatformExtensionBaseDir));
+							CachedDirs.Item1.AddRange(DirectoryReference.EnumerateDirectories(RestrictedPlatformExtensionBaseDir));
 						}
 					}
 				}
 			}
 
+			// now return what the caller wanted (always include BaseDir)
+			List<DirectoryReference> ExtensionDirs = new List<DirectoryReference> { BaseDir };
+			if (bIncludePlatformDirectories)
+			{
+				ExtensionDirs.AddRange(CachedDirs.Item1);
+			}
+			if (bIncludeRestrictedDirectories)
+			{
+				ExtensionDirs.AddRange(CachedDirs.Item2);
+			}
 			return ExtensionDirs;
 		}
 
