Fix output paths for programs under project directories.

diff --git a/Engine/Source/Programs/UnrealBuildTool/Configuration/UEBuildTarget.cs b/Engine/Source/Programs/UnrealBuildTool/Configuration/UEBuildTarget.cs
--- a/Engine/Source/Programs/UnrealBuildTool/Configuration/UEBuildTarget.cs
+++ b/Engine/Source/Programs/UnrealBuildTool/Configuration/UEBuildTarget.cs
@@ -3382,9 +3382,9 @@ namespace UnrealBuildTool
 			// Construct the output paths for this target's executable
 			List<DirectoryReference> PossibleOutputDirectories = new List<DirectoryReference>();
 			PossibleOutputDirectories.AddRange(UnrealBuildTool.GetExtensionDirs(UnrealBuildTool.EngineDirectory));
-			if (bCompileMonolithic || !bUseSharedBuildEnvironment)
+			if (ProjectFile != null && (bCompileMonolithic || !bUseSharedBuildEnvironment))
 			{
-				PossibleOutputDirectories.AddRange(UnrealBuildTool.GetExtensionDirs(UnrealBuildTool.EngineDirectory));
+				PossibleOutputDirectories.AddRange(UnrealBuildTool.GetExtensionDirs(ProjectDirectory));
 			}
 
 			DirectoryReference OutputDirectory = PossibleOutputDirectories.Where(x => Rules.File.IsUnderDirectory(x)).OrderByDescending(x => x.FullName.Length).FirstOrDefault() ?? UnrealBuildTool.EngineDirectory;
