UBT: Allow building targets from project platform/restricted directories.

diff --git a/Engine/Source/Programs/UnrealBuildTool/System/NativeProjects.cs b/Engine/Source/Programs/UnrealBuildTool/System/NativeProjects.cs
--- a/Engine/Source/Programs/UnrealBuildTool/System/NativeProjects.cs
+++ b/Engine/Source/Programs/UnrealBuildTool/System/NativeProjects.cs
@@ -139,16 +139,19 @@ namespace UnrealBuildTool
 						Dictionary<string, FileReference> TargetNameToProjectFile = new Dictionary<string, FileReference>();
 						foreach(FileReference ProjectFile in EnumerateProjectFiles())
 						{
-							DirectoryReference SourceDirectory = DirectoryReference.Combine(ProjectFile.Directory, "Source");
-							if(DirectoryLookupCache.DirectoryExists(SourceDirectory))
+							foreach (DirectoryReference ExtensionDir in UnrealBuildTool.GetExtensionDirs(ProjectFile.Directory))
 							{
-								FindTargetFiles(SourceDirectory, TargetNameToProjectFile, ProjectFile);
-							}
+								DirectoryReference SourceDirectory = DirectoryReference.Combine(ExtensionDir, "Source");
+								if (DirectoryLookupCache.DirectoryExists(SourceDirectory))
+								{
+									FindTargetFiles(SourceDirectory, TargetNameToProjectFile, ProjectFile);
+								}
 
-							DirectoryReference IntermediateSourceDirectory = DirectoryReference.Combine(ProjectFile.Directory, "Intermediate", "Source");
-							if(DirectoryLookupCache.DirectoryExists(IntermediateSourceDirectory))
-							{
-								FindTargetFiles(IntermediateSourceDirectory, TargetNameToProjectFile, ProjectFile);
+								DirectoryReference IntermediateSourceDirectory = DirectoryReference.Combine(ExtensionDir, "Intermediate", "Source");
+								if (DirectoryLookupCache.DirectoryExists(IntermediateSourceDirectory))
+								{
+									FindTargetFiles(IntermediateSourceDirectory, TargetNameToProjectFile, ProjectFile);
+								}
 							}
 						}
 						CachedTargetNameToProjectFile = TargetNameToProjectFile;
