UBT: Change preference for binary output directories (choose most generic output dir)

diff --git a/Engine/Source/Programs/UnrealBuildTool/Configuration/UEBuildTarget.cs b/Engine/Source/Programs/UnrealBuildTool/Configuration/UEBuildTarget.cs
--- a/Engine/Source/Programs/UnrealBuildTool/Configuration/UEBuildTarget.cs
+++ b/Engine/Source/Programs/UnrealBuildTool/Configuration/UEBuildTarget.cs
@@ -3387,7 +3387,7 @@ namespace UnrealBuildTool
 				PossibleOutputDirectories.AddRange(UnrealBuildTool.GetExtensionDirs(ProjectDirectory));
 			}
 
-			DirectoryReference OutputDirectory = PossibleOutputDirectories.Where(x => Rules.File.IsUnderDirectory(x)).OrderByDescending(x => x.FullName.Length).FirstOrDefault() ?? UnrealBuildTool.EngineDirectory;
+			DirectoryReference OutputDirectory = PossibleOutputDirectories.Where(x => Rules.File.IsUnderDirectory(x)).OrderBy(x => x.FullName.Length).FirstOrDefault() ?? UnrealBuildTool.EngineDirectory;
 
 			bool bCompileAsDLL = Rules.bShouldCompileAsDLL && bCompileMonolithic;
 			List<FileReference> OutputPaths = MakeBinaryPaths(OutputDirectory, bCompileMonolithic ? TargetName : AppName, Platform, Configuration, bCompileAsDLL ? UEBuildBinaryType.DynamicLinkLibrary : UEBuildBinaryType.Executable, Rules.Architecture, Rules.UndecoratedConfiguration, bCompileMonolithic && ProjectFile != null, Rules.ExeBinariesSubFolder, ProjectFile, Rules);
