Move PVS-Studio into Restricted/NoRedist folder.

diff --git a/Engine/Source/Programs/UnrealBuildTool/Platform/Windows/PVSToolChain.cs b/Engine/Source/Programs/UnrealBuildTool/Platform/Windows/PVSToolChain.cs
--- a/Engine/Source/Programs/UnrealBuildTool/Platform/Windows/PVSToolChain.cs
+++ b/Engine/Source/Programs/UnrealBuildTool/Platform/Windows/PVSToolChain.cs
@@ -350,7 +350,7 @@ namespace UnrealBuildTool
 			AnalyzerFile = FileReference.Combine(new DirectoryReference(Environment.GetFolderPath(Environment.SpecialFolder.ProgramFilesX86)), "PVS-Studio", "x64", "PVS-Studio.exe");
 			if(!FileReference.Exists(AnalyzerFile))
 			{
-				FileReference EngineAnalyzerFile = FileReference.Combine(UnrealBuildTool.RootDirectory, "Engine", "Extras", "ThirdPartyNotUE", "NoRedist", "PVS-Studio", "PVS-Studio.exe");
+				FileReference EngineAnalyzerFile = FileReference.Combine(UnrealBuildTool.RootDirectory, "Engine", "Restricted", "NoRedist", "Extras", "ThirdPartyNotUE", "PVS-Studio", "PVS-Studio.exe");
 				if (FileReference.Exists(EngineAnalyzerFile))
 				{
 					AnalyzerFile = EngineAnalyzerFile;
