Fix for universal apk appbundle crashing in WriteAllLines because file is in use

diff --git a/Engine/Source/Programs/UnrealBuildTool/Platform/Android/UEDeployAndroid.cs b/Engine/Source/Programs/UnrealBuildTool/Platform/Android/UEDeployAndroid.cs
--- a/Engine/Source/Programs/UnrealBuildTool/Platform/Android/UEDeployAndroid.cs
+++ b/Engine/Source/Programs/UnrealBuildTool/Platform/Android/UEDeployAndroid.cs
@@ -2702,29 +2702,32 @@ namespace UnrealBuildTool
 				return false;
 			}
 
-			FileStream SourceStream = new FileStream(SourceFilename, FileMode.Open, FileAccess.Read, FileShare.Read);
-			BinaryReader SourceReader = new BinaryReader(SourceStream);
-
-			FileStream DestStream = new FileStream(DestFilename, FileMode.Open, FileAccess.Read, FileShare.Read);
-			BinaryReader DestReader = new BinaryReader(DestStream);
-
-			bool bEOF = false;
-			while (!bEOF)
+			using (FileStream SourceStream = new FileStream(SourceFilename, FileMode.Open, FileAccess.Read, FileShare.Read))
+			using (FileStream DestStream = new FileStream(DestFilename, FileMode.Open, FileAccess.Read, FileShare.Read))
 			{
-				byte[] SourceData = SourceReader.ReadBytes(32768);
-				if (SourceData.Length == 0)
+				using(BinaryReader SourceReader = new BinaryReader(SourceStream))
+				using(BinaryReader DestReader = new BinaryReader(DestStream))
 				{
-					bEOF = true;
-					break;
-				}
+					bool bEOF = false;
+					while (!bEOF)
+					{
+						byte[] SourceData = SourceReader.ReadBytes(32768);
+						if (SourceData.Length == 0)
+						{
+							bEOF = true;
+							break;
+						}
 
-				byte[] DestData = DestReader.ReadBytes(32768);
-				if (!SourceData.SequenceEqual(DestData))
-				{
-					return false;
+						byte[] DestData = DestReader.ReadBytes(32768);
+						if (!SourceData.SequenceEqual(DestData))
+						{
+							return false;
+						}
+					}
+					return true;
 				}
 			}
-			return true;
+			
 		}
 
 		private bool RequiresOBB(bool bDisallowPackageInAPK, string OBBLocation)
