Fix -shared flag so sent to clang instead of passed only to linker

#ue4
#android

diff --git a/Engine/Source/Programs/UnrealBuildTool/Platform/Android/AndroidToolChain.cs b/Engine/Source/Programs/UnrealBuildTool/Platform/Android/AndroidToolChain.cs
--- a/Engine/Source/Programs/UnrealBuildTool/Platform/Android/AndroidToolChain.cs
+++ b/Engine/Source/Programs/UnrealBuildTool/Platform/Android/AndroidToolChain.cs
@@ -847,7 +847,8 @@ namespace UnrealBuildTool
 
 			//Result += " -nostdlib";
 			Result += " -no-canonical-prefixes";
-			Result += " -Wl,-shared,-Bsymbolic";
+			Result += " -shared";
+			Result += " -Wl,-Bsymbolic";
 			Result += " -Wl,--no-undefined";
 			if (bEnableGcSections)
 			{
