- Removed some leftover debug logging

diff --git a/Engine/Source/Programs/UnrealBuildTool/System/ConfigHierarchy.cs b/Engine/Source/Programs/UnrealBuildTool/System/ConfigHierarchy.cs
--- a/Engine/Source/Programs/UnrealBuildTool/System/ConfigHierarchy.cs
+++ b/Engine/Source/Programs/UnrealBuildTool/System/ConfigHierarchy.cs
@@ -869,7 +869,6 @@ namespace UnrealBuildTool
 			string BaseIniName = Enum.GetName(typeof(ConfigHierarchyType), Type);
 			string PlatformName = GetIniPlatformName(Platform);
 
-			Console.WriteLine("Ini hiererachy: {0} / {1} / {2}", PlatformName, ProjectDir != null ? ProjectDir.FullName : "NoProject", GetUserDir() != null ? GetUserDir() : "NoUserDir");
 			foreach (string Layer in ConfigLayers)
 			{
 				bool bHasPlatformTag = Layer.Contains("{PLATFORM}");
@@ -924,7 +923,6 @@ namespace UnrealBuildTool
 							yield return new FileReference(PerformFinalExpansions(ExpandedPath, "", ProjectDir));
 						}
 					}
-					Console.WriteLine("");
 				}
 				else
 				{
