Moving Binaries and Build NotForLicensees files around

diff --git a/Engine/Source/Programs/AutomationTool/AutomationUtils/ProjectParams.cs b/Engine/Source/Programs/AutomationTool/AutomationUtils/ProjectParams.cs
--- a/Engine/Source/Programs/AutomationTool/AutomationUtils/ProjectParams.cs
+++ b/Engine/Source/Programs/AutomationTool/AutomationUtils/ProjectParams.cs
@@ -625,7 +625,7 @@ namespace AutomationTool
 			this.SignedPak = !String.IsNullOrEmpty(this.SignPak) || GetParamValueIfNotSpecified(Command, SignedPak, this.SignedPak, "signedpak");
 			if (string.IsNullOrEmpty(this.SignPak))
 			{
-				this.SignPak = Path.Combine(RawProjectPath.Directory.FullName, @"Build\NoRedist\Keys.txt");
+				this.SignPak = Path.Combine(RawProjectPath.Directory.FullName, @"Restricted\NoRedist\Build\Keys.txt");
 				if (!File.Exists(this.SignPak))
 				{
 					this.SignPak = null;
diff --git a/Engine/Source/Programs/AutomationTool/AutomationUtils/UE4Build.cs b/Engine/Source/Programs/AutomationTool/AutomationUtils/UE4Build.cs
--- a/Engine/Source/Programs/AutomationTool/AutomationUtils/UE4Build.cs
+++ b/Engine/Source/Programs/AutomationTool/AutomationUtils/UE4Build.cs
@@ -317,7 +317,7 @@ namespace AutomationTool
 		/// </summary>
 		public List<FileReference> UpdateVersionFiles(bool ActuallyUpdateVersionFiles = true, int? ChangelistNumberOverride = null, int? CompatibleChangelistNumberOverride = null, string Build = null, bool? IsPromotedOverride = null)
 		{
-			bool bIsLicenseeVersion = ParseParam("Licensee") || !FileReference.Exists(FileReference.Combine(CommandUtils.EngineDirectory, "Build", "NotForLicensees", "EpicInternal.txt"));
+			bool bIsLicenseeVersion = ParseParam("Licensee") || !FileReference.Exists(FileReference.Combine(CommandUtils.EngineDirectory, "Restricted", "NotForLicensees", "Build", "EpicInternal.txt"));
 			bool bIsPromotedBuild = IsPromotedOverride.HasValue? IsPromotedOverride.Value : (ParseParamInt("Promoted", 1) != 0);
 			bool bDoUpdateVersionFiles = CommandUtils.P4Enabled && ActuallyUpdateVersionFiles;		
 			int ChangelistNumber = 0;
diff --git a/Engine/Source/Programs/AutomationTool/Scripts/Localisation.Automation.cs b/Engine/Source/Programs/AutomationTool/Scripts/Localisation.Automation.cs
--- a/Engine/Source/Programs/AutomationTool/Scripts/Localisation.Automation.cs
+++ b/Engine/Source/Programs/AutomationTool/Scripts/Localisation.Automation.cs
@@ -233,7 +233,7 @@ class Localize : BuildCommand
 		if (P4Enabled)
 		{
 			var ChangeListCommitMessage = String.Format("Localization Automation using CL {0}", P4Env.Changelist);
-			if (File.Exists(CombinePaths(CmdEnv.LocalRoot, @"Engine/Build/NotForLicensees/EpicInternal.txt")))
+			if (File.Exists(CombinePaths(CmdEnv.LocalRoot, @"Engine/Restricted/NotForLicensees/Build/EpicInternal.txt")))
 			{
 				ChangeListCommitMessage += "\n#okforgithub ignore";
 			}
diff --git a/Engine/Source/Programs/AutomationTool/Scripts/SharedCookedBuild.cs b/Engine/Source/Programs/AutomationTool/Scripts/SharedCookedBuild.cs
--- a/Engine/Source/Programs/AutomationTool/Scripts/SharedCookedBuild.cs
+++ b/Engine/Source/Programs/AutomationTool/Scripts/SharedCookedBuild.cs
@@ -331,7 +331,7 @@ public class SharedCookedBuild
 
 			FileReference PreviousManifest = FileReference.Combine(PlatformInstallPath, ".build", "Current.manifest");
 
-			FileReference BPTI = FileReference.Combine(CommandUtils.RootDirectory, "Engine", "Binaries", "Win64", "NotForLicensees", "BuildPatchToolInstaller.exe");
+			FileReference BPTI = FileReference.Combine(CommandUtils.RootDirectory, "Engine", "Restricted", "NotForLicensees", "Binaries", "Win64", "BuildPatchToolInstaller.exe");
 			if (!FileReference.Exists(BPTI))
 			{
 				CommandUtils.LogInformation("Could not locate BuildPatchToolInstaller.exe");
diff --git a/Engine/Source/Programs/UnrealBuildTool/Platform/Android/UEDeployAndroid.cs b/Engine/Source/Programs/UnrealBuildTool/Platform/Android/UEDeployAndroid.cs
--- a/Engine/Source/Programs/UnrealBuildTool/Platform/Android/UEDeployAndroid.cs
+++ b/Engine/Source/Programs/UnrealBuildTool/Platform/Android/UEDeployAndroid.cs
@@ -3222,7 +3222,11 @@ namespace UnrealBuildTool
 			string IntermediateAndroidPath = Path.Combine(ProjectDirectory, "Intermediate", "Android");
 			string UE4JavaFilePath = Path.Combine(ProjectDirectory, "Build", "Android", GetUE4JavaSrcPath());
 			string UE4BuildFilesPath = GetUE4BuildFilePath(EngineDirectory);
+			string UE4BuildFilesPath_NFL = GetUE4BuildFilePath(Path.Combine(EngineDirectory, "Restricted/NotForLicensees"));
+			string UE4BuildFilesPath_NR = GetUE4BuildFilePath(Path.Combine(EngineDirectory, "Restricted/NoRedist"));
 			string GameBuildFilesPath = Path.Combine(ProjectDirectory, "Build", "Android");
+			string GameBuildFilesPath_NFL = Path.Combine(Path.Combine(EngineDirectory, "Restricted/NotForLicensees"), "Build", "Android");
+			string GameBuildFilesPath_NR = Path.Combine(Path.Combine(EngineDirectory, "Restricted/NoRedist"), "Build", "Android");
 
 			// get a list of unique NDK architectures enabled for build
 			List<string> NDKArches = new List<string>();
@@ -3554,11 +3558,11 @@ namespace UnrealBuildTool
 				//  - Game
 				//  - Game NoRedist (for Epic secret files)
 				CopyFileDirectory(UE4BuildFilesPath, UE4BuildPath, Replacements);
-				CopyFileDirectory(UE4BuildFilesPath + "/NotForLicensees", UE4BuildPath, Replacements);
-				CopyFileDirectory(UE4BuildFilesPath + "/NoRedist", UE4BuildPath, Replacements);
+				CopyFileDirectory(UE4BuildFilesPath_NFL, UE4BuildPath, Replacements);
+				CopyFileDirectory(UE4BuildFilesPath_NR, UE4BuildPath, Replacements);
 				CopyFileDirectory(GameBuildFilesPath, UE4BuildPath, Replacements);
-				CopyFileDirectory(GameBuildFilesPath + "/NotForLicensees", UE4BuildPath, Replacements);
-				CopyFileDirectory(GameBuildFilesPath + "/NoRedist", UE4BuildPath, Replacements);
+				CopyFileDirectory(GameBuildFilesPath_NFL, UE4BuildPath, Replacements);
+				CopyFileDirectory(GameBuildFilesPath_NR, UE4BuildPath, Replacements);
 
 				//Generate Gradle AAR dependencies
 				GenerateGradleAARImports(EngineDirectory, UE4BuildPath, NDKArches);
diff --git a/Engine/Source/Programs/UnrealBuildTool/Platform/IOS/UEDeployIOS.cs b/Engine/Source/Programs/UnrealBuildTool/Platform/IOS/UEDeployIOS.cs
--- a/Engine/Source/Programs/UnrealBuildTool/Platform/IOS/UEDeployIOS.cs
+++ b/Engine/Source/Programs/UnrealBuildTool/Platform/IOS/UEDeployIOS.cs
@@ -901,6 +901,7 @@ namespace UnrealBuildTool
 			string AppDirectory = PayloadDirectory + "/" + GameName + ".app";
 			string CookedContentDirectory = AppDirectory + "/cookeddata";
 			string BuildDirectory = InProjectDirectory + "/Build/" + SubDir;
+			string BuildDirectory_NFL = InProjectDirectory + "/Restricted/NotForLicensees/Build/" + SubDir;
 			string IntermediateDirectory = (bIsUE4Game ? InEngineDir : InProjectDirectory) + "/Intermediate/" + SubDir;
 
 			Directory.CreateDirectory(BinaryPath);
@@ -932,31 +933,33 @@ namespace UnrealBuildTool
 			FileInfo DestFileInfo;
 			// always look for provisions in the IOS dir, even for TVOS
 			string ProvisionWithPrefix = InEngineDir + "/Build/IOS/UE4Game.mobileprovision";
-			if (File.Exists(BuildDirectory + "/" + InProjectName + ".mobileprovision"))
+
+			string ProjectProvision = InProjectName + ".mobileprovision";
+			if (File.Exists(Path.Combine(BuildDirectory, ProjectProvision)))
 			{
-				ProvisionWithPrefix = BuildDirectory + "/" + InProjectName + ".mobileprovision";
+				ProvisionWithPrefix = Path.Combine(BuildDirectory, ProjectProvision);
 			}
 			else
 			{
-				if (File.Exists(BuildDirectory + "/NotForLicensees/" + InProjectName + ".mobileprovision"))
+				if (File.Exists(Path.Combine(BuildDirectory_NFL, ProjectProvision)))
 				{
-					ProvisionWithPrefix = BuildDirectory + "/NotForLicensees/" + InProjectName + ".mobileprovision";
+					ProvisionWithPrefix = Path.Combine(BuildDirectory_NFL, BuildDirectory, ProjectProvision);
 				}
 				else if (!File.Exists(ProvisionWithPrefix))
 				{
-					ProvisionWithPrefix = InEngineDir + "/Build/" + SubDir + "/NotForLicensees/UE4Game.mobileprovision";
+					ProvisionWithPrefix = Path.Combine(InEngineDir, "Restricted/NotForLicensees/Build", SubDir, "UE4Game.mobileprovision");
 				}
 			}
 			if (File.Exists(ProvisionWithPrefix))
 			{
 				Directory.CreateDirectory(Environment.GetEnvironmentVariable("HOME") + "/Library/MobileDevice/Provisioning Profiles/");
-				if (File.Exists(Environment.GetEnvironmentVariable("HOME") + "/Library/MobileDevice/Provisioning Profiles/" + InProjectName + ".mobileprovision"))
+				if (File.Exists(Environment.GetEnvironmentVariable("HOME") + "/Library/MobileDevice/Provisioning Profiles/" + ProjectProvision))
 				{
-					DestFileInfo = new FileInfo(Environment.GetEnvironmentVariable("HOME") + "/Library/MobileDevice/Provisioning Profiles/" + InProjectName + ".mobileprovision");
+					DestFileInfo = new FileInfo(Environment.GetEnvironmentVariable("HOME") + "/Library/MobileDevice/Provisioning Profiles/" + ProjectProvision);
 					DestFileInfo.Attributes = DestFileInfo.Attributes & ~FileAttributes.ReadOnly;
 				}
-				File.Copy(ProvisionWithPrefix, Environment.GetEnvironmentVariable("HOME") + "/Library/MobileDevice/Provisioning Profiles/" + InProjectName + ".mobileprovision", true);
-				DestFileInfo = new FileInfo(Environment.GetEnvironmentVariable("HOME") + "/Library/MobileDevice/Provisioning Profiles/" + InProjectName + ".mobileprovision");
+				File.Copy(ProvisionWithPrefix, Environment.GetEnvironmentVariable("HOME") + "/Library/MobileDevice/Provisioning Profiles/" + ProjectProvision, true);
+				DestFileInfo = new FileInfo(Environment.GetEnvironmentVariable("HOME") + "/Library/MobileDevice/Provisioning Profiles/" + ProjectProvision);
 				DestFileInfo.Attributes = DestFileInfo.Attributes & ~FileAttributes.ReadOnly;
 			}
 			if (!File.Exists(ProvisionWithPrefix) || Environment.GetEnvironmentVariable("IsBuildMachine") == "1")
@@ -1008,19 +1011,20 @@ namespace UnrealBuildTool
 
 			// install the distribution provision
 			ProvisionWithPrefix = InEngineDir + "/Build/IOS/UE4Game_Distro.mobileprovision";
-			if (File.Exists(BuildDirectory + "/" + InProjectName + "_Distro.mobileprovision"))
+			string ProjectDistroProvision = InProjectName + "_Distro.mobileprovision";
+			if (File.Exists(Path.Combine(BuildDirectory, ProjectDistroProvision )))
 			{
-				ProvisionWithPrefix = BuildDirectory + "/" + InProjectName + "_Distro.mobileprovision";
+				ProvisionWithPrefix = Path.Combine(BuildDirectory, ProjectDistroProvision);
 			}
 			else
 			{
-				if (File.Exists(BuildDirectory + "/NotForLicensees/" + InProjectName + "_Distro.mobileprovision"))
+				if (File.Exists(Path.Combine(BuildDirectory_NFL, ProjectDistroProvision)))
 				{
-					ProvisionWithPrefix = BuildDirectory + "/NotForLicensees/" + InProjectName + "_Distro.mobileprovision";
+					ProvisionWithPrefix = Path.Combine(BuildDirectory_NFL, ProjectDistroProvision);
 				}
 				else if (!File.Exists(ProvisionWithPrefix))
 				{
-					ProvisionWithPrefix = InEngineDir + "/Build/IOS/NotForLicensees/UE4Game_Distro.mobileprovision";
+					ProvisionWithPrefix = Path.Combine(InEngineDir, "Restricted/NotForLicensees/Build", SubDir, "UE4Game_Distro.mobileprovision");
 				}
 			}
 			if (File.Exists(ProvisionWithPrefix))
diff --git a/Engine/Source/Programs/UnrealBuildTool/ToolChain/AppleToolChain.cs b/Engine/Source/Programs/UnrealBuildTool/ToolChain/AppleToolChain.cs
--- a/Engine/Source/Programs/UnrealBuildTool/ToolChain/AppleToolChain.cs
+++ b/Engine/Source/Programs/UnrealBuildTool/ToolChain/AppleToolChain.cs
@@ -233,7 +233,7 @@ namespace UnrealBuildTool
 			// if the installed one is too old, use a fixed up one if it can
 			if (bUseInstalledDsymutil == false)
 			{
-				FileReference PatchedDsymutilLocation = FileReference.Combine(UnrealBuildTool.EngineDirectory, "Binaries/Mac/NotForLicensees/LLVM/bin/dsymutil");
+				FileReference PatchedDsymutilLocation = FileReference.Combine(UnrealBuildTool.EngineDirectory, "Restricted/NotForLicensees/Binaries/Mac/LLVM/bin/dsymutil");
 
 				if (File.Exists(PatchedDsymutilLocation.FullName))
 				{
diff --git a/Engine/Source/Programs/UnrealBuildTool/ToolChain/RemoteMac.cs b/Engine/Source/Programs/UnrealBuildTool/ToolChain/RemoteMac.cs
--- a/Engine/Source/Programs/UnrealBuildTool/ToolChain/RemoteMac.cs
+++ b/Engine/Source/Programs/UnrealBuildTool/ToolChain/RemoteMac.cs
@@ -310,12 +310,12 @@ namespace UnrealBuildTool
 			Locations.Add(DirectoryReference.Combine(DirectoryReference.GetSpecialFolder(Environment.SpecialFolder.Personal), "Unreal Engine", "UnrealBuildTool"));
 			if (ProjectFile != null)
 			{
-				Locations.Add(DirectoryReference.Combine(ProjectFile.Directory, "Build", "NotForLicensees"));
-				Locations.Add(DirectoryReference.Combine(ProjectFile.Directory, "Build", "NoRedist"));
+				Locations.Add(DirectoryReference.Combine(ProjectFile.Directory, "Restricted", "NotForLicensees", "Build"));
+				Locations.Add(DirectoryReference.Combine(ProjectFile.Directory, "Restricted", "NoRedist", "Build"));
 				Locations.Add(DirectoryReference.Combine(ProjectFile.Directory, "Build"));
 			}
-			Locations.Add(DirectoryReference.Combine(UnrealBuildTool.EngineDirectory, "Build", "NotForLicensees"));
-			Locations.Add(DirectoryReference.Combine(UnrealBuildTool.EngineDirectory, "Build", "NoRedist"));
+			Locations.Add(DirectoryReference.Combine(UnrealBuildTool.EngineDirectory, "Restricted", "NotForLicensees", "Build"));
+			Locations.Add(DirectoryReference.Combine(UnrealBuildTool.EngineDirectory, "Restricted", "NoRedist", "Build"));
 			Locations.Add(DirectoryReference.Combine(UnrealBuildTool.EngineDirectory, "Build"));
 
 			// Find the first that exists
diff --git a/Engine/Source/Runtime/Core/Private/Misc/EngineBuildSettings.cpp b/Engine/Source/Runtime/Core/Private/Misc/EngineBuildSettings.cpp
--- a/Engine/Source/Runtime/Core/Private/Misc/EngineBuildSettings.cpp
+++ b/Engine/Source/Runtime/Core/Private/Misc/EngineBuildSettings.cpp
@@ -6,7 +6,7 @@
 
 bool FEngineBuildSettings::IsInternalBuild()
 {
-	static bool bIsInternalBuild = FPaths::FileExists( FPaths::EngineDir() / TEXT("Build/NotForLicensees/EpicInternal.txt") );
+	static bool bIsInternalBuild = FPaths::FileExists( FPaths::EngineDir() / TEXT("Restricted/NotForLicensees/Build/EpicInternal.txt") );
 	return bIsInternalBuild;
 }
 
diff --git a/Engine/Source/Runtime/IOS/IOSRuntimeSettings/Private/IOSRuntimeSettings.cpp b/Engine/Source/Runtime/IOS/IOSRuntimeSettings/Private/IOSRuntimeSettings.cpp
--- a/Engine/Source/Runtime/IOS/IOSRuntimeSettings/Private/IOSRuntimeSettings.cpp
+++ b/Engine/Source/Runtime/IOS/IOSRuntimeSettings/Private/IOSRuntimeSettings.cpp
@@ -132,11 +132,11 @@ void UIOSRuntimeSettings::PostInitProperties()
 		FString Path = FPlatformMisc::GetEnvironmentVariable(TEXT("APPDATA"));
 
 		TArray<FString> PossibleKeyLocations;
-		PossibleKeyLocations.Add(FPaths::Combine(*FPaths::ProjectDir(), TEXT("Build"), TEXT("NotForLicensees"), *RelativeFilePathLocation));
-		PossibleKeyLocations.Add(FPaths::Combine(*FPaths::ProjectDir(), TEXT("Build"), TEXT("NoRedist"), *RelativeFilePathLocation));
+		PossibleKeyLocations.Add(FPaths::Combine(*FPaths::ProjectDir(), TEXT("Restricted"), TEXT("NotForLicensees"), TEXT("Build"), *RelativeFilePathLocation));
+		PossibleKeyLocations.Add(FPaths::Combine(*FPaths::ProjectDir(), TEXT("Restricted"), TEXT("NoRedist"), TEXT("Build"), *RelativeFilePathLocation));
 		PossibleKeyLocations.Add(FPaths::Combine(*FPaths::ProjectDir(), TEXT("Build"), *RelativeFilePathLocation));
-		PossibleKeyLocations.Add(FPaths::Combine(*FPaths::EngineDir(), TEXT("Build"), TEXT("NotForLicensees"), *RelativeFilePathLocation));
-		PossibleKeyLocations.Add(FPaths::Combine(*FPaths::EngineDir(), TEXT("Build"), TEXT("NoRedist"), *RelativeFilePathLocation));
+		PossibleKeyLocations.Add(FPaths::Combine(*FPaths::EngineDir(), TEXT("Restricted"), TEXT("NotForLicensees"), TEXT("Build"), TEXT("NotForLicensees"), *RelativeFilePathLocation));
+		PossibleKeyLocations.Add(FPaths::Combine(*FPaths::EngineDir(), TEXT("Restricted"), TEXT("NoRedist"), TEXT("Build"), *RelativeFilePathLocation));
 		PossibleKeyLocations.Add(FPaths::Combine(*FPaths::EngineDir(), TEXT("Build"), *RelativeFilePathLocation));
 		PossibleKeyLocations.Add(FPaths::Combine(*Path, TEXT("Unreal Engine"), TEXT("UnrealBuildTool"), *RelativeFilePathLocation));
 
