- Added ability for a module to add additional directories (ConditionalAddModuleDirectory). Used for adding source files in Restricted (without needing the extra Build.cs files the way PlatformExtensions are set up).

diff --git a/Engine/Source/Programs/UnrealBuildTool/Configuration/ModuleRules.cs b/Engine/Source/Programs/UnrealBuildTool/Configuration/ModuleRules.cs
--- a/Engine/Source/Programs/UnrealBuildTool/Configuration/ModuleRules.cs
+++ b/Engine/Source/Programs/UnrealBuildTool/Configuration/ModuleRules.cs
@@ -483,6 +483,11 @@ namespace UnrealBuildTool
 		/// </summary>
 		internal Dictionary<Type, DirectoryReference> DirectoriesForModuleSubClasses;
 
+		/// <summary>
+		/// Additional directories that contribute to this module but are not based on a subclass (NotForLicensees, etc)
+		/// </summary>
+		private List<DirectoryReference> AdditionalModuleDirectories = new List<DirectoryReference>();
+
 		/// <summary>
 		/// Plugin containing this module
 		/// </summary>
@@ -1350,12 +1355,32 @@ namespace UnrealBuildTool
 		}
 
 		/// <summary>
-		/// Returns the directories for all subclasses of this module
+		/// Returns the directories for all subclasses of this module, as well as any additional directories specified by the rules
 		/// </summary>
 		/// <returns>List of directories, or null if none were added</returns>
-		public DirectoryReference[] GetModuleDirectoriesForAllSubClasses()
+		public DirectoryReference[] GetAllModuleDirectories()
+		{
+			List<DirectoryReference> AllDirectories = new List<DirectoryReference> { Directory };
+			AllDirectories.AddRange(AdditionalModuleDirectories);
+
+			if (DirectoriesForModuleSubClasses != null)
+			{
+				AllDirectories.AddRange(DirectoriesForModuleSubClasses.Values);
+			}
+
+			return AllDirectories.ToArray();
+		}
+
+		/// <summary>
+		/// Adds an additional module directory, if it exists (useful for NotForLicensees/NoRedist)
+		/// </summary>
+		/// <param name="Directory"></param>
+		protected void ConditionalAddModuleDirectory(DirectoryReference Directory)
 		{
-			return DirectoriesForModuleSubClasses == null ? null : DirectoriesForModuleSubClasses.Values.ToArray();
+			if (DirectoryReference.Exists(Directory))
+			{
+				AdditionalModuleDirectories.Add(Directory);
+			}
 		}
 	}
 }
diff --git a/Engine/Source/Programs/UnrealBuildTool/Configuration/UEBuildModule.cs b/Engine/Source/Programs/UnrealBuildTool/Configuration/UEBuildModule.cs
--- a/Engine/Source/Programs/UnrealBuildTool/Configuration/UEBuildModule.cs
+++ b/Engine/Source/Programs/UnrealBuildTool/Configuration/UEBuildModule.cs
@@ -244,16 +244,8 @@ namespace UnrealBuildTool
 			WhitelistRestrictedFolders = new HashSet<DirectoryReference>(Rules.WhitelistRestrictedFolders.Select(x => DirectoryReference.Combine(ModuleDirectory, x)));
 			AliasRestrictedFolders = new Dictionary<string, string>(Rules.AliasRestrictedFolders);
 
-			// merge the main directory and any others set in the Rules
-			List<DirectoryReference> MergedDirectories = new List<DirectoryReference> { ModuleDirectory };
-			DirectoryReference[] ExtraModuleDirectories = Rules.GetModuleDirectoriesForAllSubClasses();
-			if (ExtraModuleDirectories != null)
-			{
-				MergedDirectories.AddRange(ExtraModuleDirectories);
-			}
-
-			// cache the results (it will always at least have the ModuleDirectory)
-			ModuleDirectories = MergedDirectories.ToArray();
+			// get the module directories from the module
+			ModuleDirectories = Rules.GetAllModuleDirectories();
 		}
 
 		/// <summary>
diff --git a/Engine/Source/Programs/UnrealBuildTool/ProjectFiles/ProjectFileGenerator.cs b/Engine/Source/Programs/UnrealBuildTool/ProjectFiles/ProjectFileGenerator.cs
--- a/Engine/Source/Programs/UnrealBuildTool/ProjectFiles/ProjectFileGenerator.cs
+++ b/Engine/Source/Programs/UnrealBuildTool/ProjectFiles/ProjectFileGenerator.cs
@@ -803,7 +803,7 @@ namespace UnrealBuildTool
 						AddEngineExtrasFiles(EngineProject);
 
 						// Platform Extension files
-						AddPlatformExtensionFiles(EngineProject);
+						AddEngineExtensionFiles(EngineProject);
 
 						// Engine localization files
 						if( bIncludeLocalizationFiles )
@@ -1266,21 +1266,25 @@ namespace UnrealBuildTool
 						GameProject.AddFilesToProject(SourceFileSearch.FindFiles(ExtensionDir, SearchSubdirectories: false), ExtensionDir);
 					}
 
+					// Game restricted source files, since they won't be added via a module FileReference
+					foreach (DirectoryReference GameRestrictedSourceDirectory in UnrealBuildTool.GetExtensionDirs(GameProjectDirectory, "Source", bIncludePlatformDirectories:false, bIncludeBaseDirectory:false))
+					{
+						GameProject.AddFilesToProject(SourceFileSearch.FindFiles(GameRestrictedSourceDirectory), GameRestrictedSourceDirectory);
+					}
+
 					// Game config files
 					if ( bIncludeConfigFiles )
 					{
-						DirectoryReference GameConfigDirectory = DirectoryReference.Combine(GameProjectDirectory, "Config");
-						if( DirectoryReference.Exists(GameConfigDirectory) )
+						foreach (DirectoryReference GameConfigDirectory in UnrealBuildTool.GetExtensionDirs(GameProjectDirectory, "Config"))
 						{
-							GameProject.AddFilesToProject( SourceFileSearch.FindFiles( GameConfigDirectory ), GameProjectDirectory );
+							GameProject.AddFilesToProject(SourceFileSearch.FindFiles(GameConfigDirectory), GameProjectDirectory);
 						}
 					}
 
 					// Game build files
 					if( bIncludeBuildSystemFiles )
 					{
-						DirectoryReference GameBuildDirectory = DirectoryReference.Combine(GameProjectDirectory, "Build");
-						if( DirectoryReference.Exists(GameBuildDirectory) )
+						foreach (DirectoryReference GameBuildDirectory in UnrealBuildTool.GetExtensionDirs(GameProjectDirectory, "Build"))
 						{
 							List<string> SubdirectoryNamesToExclude = new List<string>();
 							SubdirectoryNamesToExclude.Add("Receipts");
@@ -1292,8 +1296,7 @@ namespace UnrealBuildTool
 						}
 					}
 
-					DirectoryReference GameShaderDirectory = DirectoryReference.Combine(GameProjectDirectory, "Shaders");
-					if (DirectoryReference.Exists(GameShaderDirectory))
+					foreach (DirectoryReference GameShaderDirectory in UnrealBuildTool.GetExtensionDirs(GameProjectDirectory, "Shaders"))
 					{
 						GameProject.AddFilesToProject(SourceFileSearch.FindFiles(GameShaderDirectory), GameProjectDirectory);
 					}
@@ -1340,7 +1343,7 @@ namespace UnrealBuildTool
 		}
 
 		/// Adds additional files from the platform extensions folder
-		protected virtual void AddPlatformExtensionFiles( ProjectFile EngineProject )
+		protected virtual void AddEngineExtensionFiles( ProjectFile EngineProject )
 		{
 			// @todo: this will add the same files to the solution (like the UBT source files that also get added to UnrealBuildTool project).
 			// not sure of a good filtering method here
diff --git a/Engine/Source/Programs/UnrealBuildTool/UnrealBuildTool.cs b/Engine/Source/Programs/UnrealBuildTool/UnrealBuildTool.cs
--- a/Engine/Source/Programs/UnrealBuildTool/UnrealBuildTool.cs
+++ b/Engine/Source/Programs/UnrealBuildTool/UnrealBuildTool.cs
@@ -144,15 +144,16 @@ namespace UnrealBuildTool
 
 		// cached dictionary of BaseDir to extension directories
 		private static Dictionary<DirectoryReference, Tuple<List<DirectoryReference>, List<DirectoryReference>>> CachedExtensionDirectories = new Dictionary<DirectoryReference, Tuple<List<DirectoryReference>, List<DirectoryReference>>>();
-		
+
 		/// <summary>
 		/// Finds all the extension directories for the given base directory. This includes platform extensions and restricted folders.
 		/// </summary>
 		/// <param name="BaseDir">Location of the base directory</param>
 		/// <param name="bIncludePlatformDirectories">If true, platform subdirectories are included (will return platform directories under Restricted dirs, even if bIncludeRestrictedDirectories is false)</param>
 		/// <param name="bIncludeRestrictedDirectories">If true, restricted (NotForLicensees, NoRedist) subdirectories are included</param>
+		/// <param name="bIncludeBaseDirectory">If true, BaseDir is included</param>
 		/// <returns>List of extension directories, including the given base directory</returns>
-		public static List<DirectoryReference> GetExtensionDirs(DirectoryReference BaseDir, bool bIncludePlatformDirectories=true, bool bIncludeRestrictedDirectories=true)
+		public static List<DirectoryReference> GetExtensionDirs(DirectoryReference BaseDir, bool bIncludePlatformDirectories=true, bool bIncludeRestrictedDirectories=true, bool bIncludeBaseDirectory=true)
 		{
 			Tuple<List<DirectoryReference>, List<DirectoryReference>> CachedDirs;
 			if (!CachedExtensionDirectories.TryGetValue(BaseDir, out CachedDirs))
@@ -186,7 +187,11 @@ namespace UnrealBuildTool
 			}
 
 			// now return what the caller wanted (always include BaseDir)
-			List<DirectoryReference> ExtensionDirs = new List<DirectoryReference> { BaseDir };
+			List<DirectoryReference> ExtensionDirs = new List<DirectoryReference>();
+			if (bIncludeBaseDirectory)
+			{
+				ExtensionDirs.Add(BaseDir);
+			}
 			if (bIncludePlatformDirectories)
 			{
 				ExtensionDirs.AddRange(CachedDirs.Item1);
@@ -205,10 +210,11 @@ namespace UnrealBuildTool
 		/// <param name="SubDir">The subdirectory to find</param>
 		/// <param name="bIncludePlatformDirectories">If true, platform subdirectories are included (will return platform directories under Restricted dirs, even if bIncludeRestrictedDirectories is false)</param>
 		/// <param name="bIncludeRestrictedDirectories">If true, restricted (NotForLicensees, NoRedist) subdirectories are included</param>
+		/// <param name="bIncludeBaseDirectory">If true, BaseDir is included</param>
 		/// <returns>List of extension directories, including the given base directory</returns>
-		public static List<DirectoryReference> GetExtensionDirs(DirectoryReference BaseDir, string SubDir, bool bIncludePlatformDirectories = true, bool bIncludeRestrictedDirectories = true)
+		public static List<DirectoryReference> GetExtensionDirs(DirectoryReference BaseDir, string SubDir, bool bIncludePlatformDirectories=true, bool bIncludeRestrictedDirectories=true, bool bIncludeBaseDirectory=true)
 		{
-			return GetExtensionDirs(BaseDir, bIncludePlatformDirectories, bIncludeRestrictedDirectories).Select(x => DirectoryReference.Combine(x, SubDir)).Where(x => DirectoryReference.Exists(x)).ToList();
+			return GetExtensionDirs(BaseDir, bIncludePlatformDirectories, bIncludeRestrictedDirectories, bIncludeBaseDirectory).Select(x => DirectoryReference.Combine(x, SubDir)).Where(x => DirectoryReference.Exists(x)).ToList();
 		}
 
 		/// <summary>
