- Fixed include paths with UAC enabled (local builds don't have it on by default, so missed these in previous checkin)

diff --git a/Engine/Source/Programs/UnrealBuildTool/Configuration/ModuleRules.cs b/Engine/Source/Programs/UnrealBuildTool/Configuration/ModuleRules.cs
--- a/Engine/Source/Programs/UnrealBuildTool/Configuration/ModuleRules.cs
+++ b/Engine/Source/Programs/UnrealBuildTool/Configuration/ModuleRules.cs
@@ -1375,12 +1375,16 @@ namespace UnrealBuildTool
 		/// Adds an additional module directory, if it exists (useful for NotForLicensees/NoRedist)
 		/// </summary>
 		/// <param name="Directory"></param>
-		protected void ConditionalAddModuleDirectory(DirectoryReference Directory)
+		/// <returns>true if the directory exists</returns>
+		protected bool ConditionalAddModuleDirectory(DirectoryReference Directory)
 		{
 			if (DirectoryReference.Exists(Directory))
 			{
 				AdditionalModuleDirectories.Add(Directory);
+				return true;
 			}
+
+			return false;
 		}
 	}
 }
