Disable --gc-sections on link if not compiling with -ffunction-sections and -fdata-sections for Android

#android

diff --git a/Engine/Source/Programs/UnrealBuildTool/Platform/Android/AndroidToolChain.cs b/Engine/Source/Programs/UnrealBuildTool/Platform/Android/AndroidToolChain.cs
--- a/Engine/Source/Programs/UnrealBuildTool/Platform/Android/AndroidToolChain.cs
+++ b/Engine/Source/Programs/UnrealBuildTool/Platform/Android/AndroidToolChain.cs
@@ -850,9 +850,9 @@ namespace UnrealBuildTool
 			Result += " -shared";
 			Result += " -Wl,-Bsymbolic";
 			Result += " -Wl,--no-undefined";
-			if (bEnableGcSections)
+			if (bEnableGcSections && !DisableFunctionDataSections())
 			{
-				Result += " -Wl,-gc-sections"; // Enable garbage collection of unused input sections. works best with -ffunction-sections, -fdata-sections
+				Result += " -Wl,--gc-sections"; // Enable garbage collection of unused input sections. works best with -ffunction-sections, -fdata-sections
 			}
 
 			if (!LinkEnvironment.bCreateDebugInfo)
@@ -2185,6 +2185,8 @@ namespace UnrealBuildTool
 						SetupActionToExecuteCompilerThroughShell(ref LinkAction, LinkAction.CommandPath.FullName, LinkAction.CommandArguments, "Link");
 					}
 
+					//Log.TraceInformation("Link: {0} {1}", LinkAction.CommandPath.FullName, LinkAction.CommandArguments);
+
 					// Windows can run into an issue with too long of a commandline when clang tries to call ld to link.
 					// To work around this we call clang to just get the command it would execute and generate a
 					// second response file to directly call ld with the right arguments instead of calling through clang.
