- Fixed a bug in UBT finding platform directories under Restricted

diff --git a/Engine/Source/Programs/UnrealBuildTool/UnrealBuildTool.cs b/Engine/Source/Programs/UnrealBuildTool/UnrealBuildTool.cs
--- a/Engine/Source/Programs/UnrealBuildTool/UnrealBuildTool.cs
+++ b/Engine/Source/Programs/UnrealBuildTool/UnrealBuildTool.cs
@@ -177,7 +177,7 @@ namespace UnrealBuildTool
 					// also look for nested platforms in the restricted 
 					foreach (DirectoryReference RestrictedDir in RestrictedDirs)
 					{
-						DirectoryReference RestrictedPlatformExtensionBaseDir = DirectoryReference.Combine(RestrictedBaseDir, "Platforms");
+						DirectoryReference RestrictedPlatformExtensionBaseDir = DirectoryReference.Combine(RestrictedDir, "Platforms");
 						if (DirectoryReference.Exists(RestrictedPlatformExtensionBaseDir))
 						{
 							CachedDirs.Item1.AddRange(DirectoryReference.EnumerateDirectories(RestrictedPlatformExtensionBaseDir));
