Fix path to Android keystore causing build failures.

diff --git a/Engine/Source/Programs/UnrealBuildTool/Platform/Android/UEDeployAndroid.cs b/Engine/Source/Programs/UnrealBuildTool/Platform/Android/UEDeployAndroid.cs
--- a/Engine/Source/Programs/UnrealBuildTool/Platform/Android/UEDeployAndroid.cs
+++ b/Engine/Source/Programs/UnrealBuildTool/Platform/Android/UEDeployAndroid.cs
@@ -3225,8 +3225,8 @@ namespace UnrealBuildTool
 			string UE4BuildFilesPath_NFL = GetUE4BuildFilePath(Path.Combine(EngineDirectory, "Restricted/NotForLicensees"));
 			string UE4BuildFilesPath_NR = GetUE4BuildFilePath(Path.Combine(EngineDirectory, "Restricted/NoRedist"));
 			string GameBuildFilesPath = Path.Combine(ProjectDirectory, "Build", "Android");
-			string GameBuildFilesPath_NFL = Path.Combine(Path.Combine(EngineDirectory, "Restricted/NotForLicensees"), "Build", "Android");
-			string GameBuildFilesPath_NR = Path.Combine(Path.Combine(EngineDirectory, "Restricted/NoRedist"), "Build", "Android");
+			string GameBuildFilesPath_NFL = Path.Combine(Path.Combine(ProjectDirectory, "Restricted/NotForLicensees"), "Build", "Android");
+			string GameBuildFilesPath_NR = Path.Combine(Path.Combine(ProjectDirectory, "Restricted/NoRedist"), "Build", "Android");
 
 			// get a list of unique NDK architectures enabled for build
 			List<string> NDKArches = new List<string>();
