PR #6857: Fix rider projectfilegenerator (Contributed by [REDACTED_PERSON])

#rnx

diff --git a/Engine/Source/Programs/UnrealBuildTool/ProjectFiles/Rider/RiderProjectFile.cs b/Engine/Source/Programs/UnrealBuildTool/ProjectFiles/Rider/RiderProjectFile.cs
--- a/Engine/Source/Programs/UnrealBuildTool/ProjectFiles/Rider/RiderProjectFile.cs
+++ b/Engine/Source/Programs/UnrealBuildTool/ProjectFiles/Rider/RiderProjectFile.cs
@@ -2,6 +2,7 @@
 
 using System;
 using System.Collections.Generic;
+using System.IO;
 using System.Linq;
 using Tools.DotNETCommon;
 
@@ -39,9 +40,9 @@ namespace UnrealBuildTool
 			PlatformProjectGeneratorCollection PlatformProjectGenerators)
 		{
 			string ProjectName = ProjectFilePath.GetFileNameWithoutAnyExtensions();
-			DirectoryReference projectRootFolder = DirectoryReference.Combine(RootPath, ".Rider");
-			List<Tuple<FileReference, UEBuildTarget>> fileToTarget = new List<Tuple<FileReference, UEBuildTarget>>();
-			foreach (UnrealTargetPlatform Platform in InPlatforms.Where(it => it != UnrealTargetPlatform.Win32))
+			DirectoryReference ProjectRootFolder = DirectoryReference.Combine(RootPath, ".Rider");
+			List<Tuple<FileReference, UEBuildTarget>> FileToTarget = new List<Tuple<FileReference, UEBuildTarget>>();
+			foreach (UnrealTargetPlatform Platform in InPlatforms)
 			{
 				foreach (UnrealTargetConfiguration Configuration in InConfigurations)
 				{
@@ -61,32 +62,32 @@ namespace UnrealBuildTool
 							continue;
 						}
 						
-						DirectoryReference ConfigurationFolder = DirectoryReference.Combine(projectRootFolder, Platform.ToString(), Configuration.ToString());
+						DirectoryReference ConfigurationFolder = DirectoryReference.Combine(ProjectRootFolder, Platform.ToString(), Configuration.ToString());
 
 						DirectoryReference TargetFolder =
 							DirectoryReference.Combine(ConfigurationFolder, ProjectTarget.TargetRules.Type.ToString());
 
 						string DefaultArchitecture = UEBuildPlatform
-							.GetBuildPlatform(BuildHostPlatform.Current.Platform)
+							.GetBuildPlatform(Platform)
 							.GetDefaultArchitecture(ProjectTarget.UnrealProjectFilePath);
 						TargetDescriptor TargetDesc = new TargetDescriptor(ProjectTarget.UnrealProjectFilePath, ProjectTarget.Name,
-							BuildHostPlatform.Current.Platform, UnrealTargetConfiguration.Development,
-							DefaultArchitecture, Arguments);
+							Platform, Configuration, DefaultArchitecture, Arguments);
 						try
 						{
 							UEBuildTarget BuildTarget = UEBuildTarget.Create(TargetDesc, false, false);
 						
 							FileReference OutputFile = FileReference.Combine(TargetFolder, $"{ProjectName}.json");
-							fileToTarget.Add(Tuple.Create(OutputFile, BuildTarget));
+							FileToTarget.Add(Tuple.Create(OutputFile, BuildTarget));
 						}
 						catch(Exception Ex)
 						{
-							Log.TraceWarning("Exception while generating include data for {0}: {1}", TargetDesc.Name, Ex.ToString());
+							Log.TraceWarning("Exception while generating include data for Target:{0}, Platform: {1}, Configuration: {2}", TargetDesc.Name, Platform.ToString(), Configuration.ToString());
+							Log.TraceWarning(Ex.ToString());
 						}
 					}
 				}
 			}
-			foreach (Tuple<FileReference,UEBuildTarget> tuple in fileToTarget)
+			foreach (Tuple<FileReference,UEBuildTarget> tuple in FileToTarget)
 			{
 				SerializeTarget(tuple.Item1, tuple.Item2);
 			}
@@ -134,17 +135,25 @@ namespace UnrealBuildTool
 				}
 				Writer.WriteArrayEnd();
 			}
-
-			HashSet<string> moduleNames = new HashSet<string>();
+			
+			CppCompileEnvironment GlobalCompileEnvironment = Target.CreateCompileEnvironmentForProjectFiles();
+			HashSet<string> ModuleNames = new HashSet<string>();
 			Writer.WriteObjectStart("Modules");
 			foreach (UEBuildBinary Binary in Target.Binaries)
 			{
+				CppCompileEnvironment BinaryCompileEnvironment = Binary.CreateBinaryCompileEnvironment(GlobalCompileEnvironment);
 				foreach (UEBuildModule Module in Binary.Modules)
 				{
-					if(moduleNames.Add(Module.Name))
+					if(ModuleNames.Add(Module.Name))
 					{
 						Writer.WriteObjectStart(Module.Name);
 						ExportModule(Module, Binary.OutputDir, Target.GetExecutableDir(), Writer);
+						UEBuildModuleCPP ModuleCpp = Module as UEBuildModuleCPP;
+						if (ModuleCpp != null)
+						{
+							CppCompileEnvironment ModuleCompileEnvironment = ModuleCpp.CreateCompileEnvironmentForIntellisense(Target.Rules, BinaryCompileEnvironment);
+							ExportModuleCpp(ModuleCpp, ModuleCompileEnvironment, Writer);
+						}
 						Writer.WriteObjectEnd();
 					}
 				}
@@ -156,6 +165,40 @@ namespace UnrealBuildTool
 			Writer.WriteObjectEnd();
 		}
 
+		private static void ExportModuleCpp(UEBuildModuleCPP ModuleCPP, CppCompileEnvironment ModuleCompileEnvironment, JsonWriter Writer)
+		{
+			Writer.WriteValue("GeneratedCodeDirectory", ModuleCPP.GeneratedCodeDirectory != null ? ModuleCPP.GeneratedCodeDirectory.FullName : string.Empty);
+			
+			if (ModuleCompileEnvironment.PrecompiledHeaderIncludeFilename != null)
+			{
+				string CorrectFilePathPch;
+				if(ExtractWrappedIncludeFile(ModuleCompileEnvironment.PrecompiledHeaderIncludeFilename, out CorrectFilePathPch))
+					Writer.WriteValue("SharedPCHFilePath", CorrectFilePathPch);
+			}
+		}
+
+		private static bool ExtractWrappedIncludeFile(FileSystemReference FileRef, out string CorrectFilePathPch)
+		{
+			CorrectFilePathPch = "";
+			try
+			{
+				using (StreamReader Reader = new StreamReader(FileRef.FullName))
+				{
+					string Line = Reader.ReadLine();
+					if (Line != null)
+					{
+						CorrectFilePathPch = Line.Substring("// PCH for ".Length).Trim();
+						return true;
+					}
+				}
+			}
+			finally
+			{
+				Log.TraceVerbose("Couldn't extract path to PCH from {0}", FileRef);
+			}
+			return false;
+		}
+
 		/// <summary>
 		/// Write a Module to a JSON writer. If array is empty, don't write anything
 		/// </summary>
@@ -170,12 +213,6 @@ namespace UnrealBuildTool
 			Writer.WriteValue("Rules", Module.RulesFile.FullName);
 			Writer.WriteValue("PCHUsage", Module.Rules.PCHUsage.ToString());
 
-			UEBuildModuleCPP ModuleCPP = Module as UEBuildModuleCPP;
-			if (ModuleCPP != null)
-			{
-				Writer.WriteValue("GeneratedCodeDirectory", ModuleCPP.GeneratedCodeDirectory != null ? ModuleCPP.GeneratedCodeDirectory.FullName : string.Empty);
-			}
-
 			if (Module.Rules.PrivatePCHHeaderFile != null)
 			{
 				Writer.WriteValue("PrivatePCH", FileReference.Combine(Module.ModuleDirectory, Module.Rules.PrivatePCHHeaderFile).FullName);
@@ -199,12 +236,11 @@ namespace UnrealBuildTool
 			
 			ExportJsonStringArray(Writer, "PrivateIncludePaths", Module.PrivateIncludePaths.Select(x => x.FullName));
 			ExportJsonStringArray(Writer, "PublicLibraryPaths", Module.PublicSystemLibraryPaths.Select(x => x.FullName));
-			ExportJsonStringArray(Writer, "PublicAdditionalLibraries", Enumerable.Concat(Module.PublicSystemLibraries, Module.PublicLibraries.Select(x => x.FullName)));
+			ExportJsonStringArray(Writer, "PublicAdditionalLibraries", Module.PublicSystemLibraries.Concat(Module.PublicLibraries.Select(x => x.FullName)));
 			ExportJsonStringArray(Writer, "PublicFrameworks", Module.PublicFrameworks);
 			ExportJsonStringArray(Writer, "PublicWeakFrameworks", Module.PublicWeakFrameworks);
 			ExportJsonStringArray(Writer, "PublicDelayLoadDLLs", Module.PublicDelayLoadDLLs);
 			ExportJsonStringArray(Writer, "PublicDefinitions", Module.PublicDefinitions);
-			
 			ExportJsonStringArray(Writer, "PrivateDefinitions", Module.Rules.PrivateDefinitions);
 			ExportJsonStringArray(Writer, "ProjectDefinitions", /* TODO: Add method ShouldAddProjectDefinitions */ !Module.Rules.bTreatAsEngineModule ? Module.Rules.Target.ProjectDefinitions : new string[0]);
 			ExportJsonStringArray(Writer, "ApiDefinitions", Module.GetEmptyApiMacros());
diff --git a/Engine/Source/Programs/UnrealBuildTool/ProjectFiles/Rider/RiderProjectFileGenerator.cs b/Engine/Source/Programs/UnrealBuildTool/ProjectFiles/Rider/RiderProjectFileGenerator.cs
--- a/Engine/Source/Programs/UnrealBuildTool/ProjectFiles/Rider/RiderProjectFileGenerator.cs
+++ b/Engine/Source/Programs/UnrealBuildTool/ProjectFiles/Rider/RiderProjectFileGenerator.cs
@@ -98,31 +98,42 @@ namespace UnrealBuildTool
 		{
 			using (ProgressWriter Progress = new ProgressWriter("Writing project files...", true))
 			{
-				List<ProjectFile> projectsToGenerate = new List<ProjectFile>(GeneratedProjectFiles);
+				List<ProjectFile> ProjectsToGenerate = new List<ProjectFile>(GeneratedProjectFiles);
 				if (ProjectNames.Any())
 				{
-					projectsToGenerate = projectsToGenerate.Where(it =>
+					ProjectsToGenerate = ProjectsToGenerate.Where(it =>
 						ProjectNames.Contains(it.ProjectFilePath.GetFileNameWithoutAnyExtensions())).ToList();
 				}
 
-				int TotalProjectFileCount = projectsToGenerate.Count;
+				int TotalProjectFileCount = ProjectsToGenerate.Count;
 
-				HashSet<UnrealTargetPlatform> platformsToGenerate = new HashSet<UnrealTargetPlatform>(SupportedPlatforms);
+				HashSet<UnrealTargetPlatform> PlatformsToGenerate = new HashSet<UnrealTargetPlatform>(SupportedPlatforms);
 				if (Platforms.Any())
 				{
-					platformsToGenerate.IntersectWith(Platforms);
+					PlatformsToGenerate.IntersectWith(Platforms);
 				}
 
-				HashSet<UnrealTargetConfiguration> configurationsToGenerate = new HashSet<UnrealTargetConfiguration>(SupportedConfigurations);
+				List<UnrealTargetPlatform> FilteredPlatforms = PlatformsToGenerate.Where(it =>
+				{
+					if (UEBuildPlatform.IsPlatformAvailable(it))
+						return true;
+					Log.TraceWarning(
+						"Platform {0} is not a valid platform to build. Check that the SDK is installed properly",
+						it);
+					Log.TraceWarning("Platform will be ignored in project file generation");
+					return false;
+				}).ToList();
+
+				HashSet<UnrealTargetConfiguration> ConfigurationsToGenerate = new HashSet<UnrealTargetConfiguration>(SupportedConfigurations);
 				if (TargetConfigurations.Any())
 				{
-					configurationsToGenerate.IntersectWith(TargetConfigurations);
+					ConfigurationsToGenerate.IntersectWith(TargetConfigurations);
 				}
 
-				for (int ProjectFileIndex = 0; ProjectFileIndex < projectsToGenerate.Count; ++ProjectFileIndex)
+				for (int ProjectFileIndex = 0; ProjectFileIndex < ProjectsToGenerate.Count; ++ProjectFileIndex)
 				{
-					ProjectFile CurProject = projectsToGenerate[ProjectFileIndex];
-					if (!CurProject.WriteProjectFile(platformsToGenerate.ToList(), configurationsToGenerate.ToList(),
+					ProjectFile CurProject = ProjectsToGenerate[ProjectFileIndex];
+					if (!CurProject.WriteProjectFile(FilteredPlatforms, ConfigurationsToGenerate.ToList(),
 						PlatformProjectGenerators))
 					{
 						return false;
