UBT: allow multiple -SingleFile arguments. Now stores SpecifcFilesToCompile list rather than single file reference.

diff --git a/Engine/Source/Programs/UnrealBuildTool/Configuration/TargetDescriptor.cs b/Engine/Source/Programs/UnrealBuildTool/Configuration/TargetDescriptor.cs
--- a/Engine/Source/Programs/UnrealBuildTool/Configuration/TargetDescriptor.cs
+++ b/Engine/Source/Programs/UnrealBuildTool/Configuration/TargetDescriptor.cs
@@ -34,10 +34,10 @@ namespace UnrealBuildTool
 		public HashSet<string> OnlyModuleNames = new HashSet<string>(StringComparer.OrdinalIgnoreCase);
 
 		/// <summary>
-		/// Single file to compile
+		/// Individual file(s) to compile
 		/// </summary>
 		[CommandLine("-SingleFile=")]
-		public FileReference SingleFileToCompile = null;
+		public List<FileReference> SpecificFilesToCompile = new List<FileReference>();
 
 		/// <summary>
 		/// Whether to perform hot reload for this target
diff --git a/Engine/Source/Programs/UnrealBuildTool/Configuration/UEBuildBinary.cs b/Engine/Source/Programs/UnrealBuildTool/Configuration/UEBuildBinary.cs
--- a/Engine/Source/Programs/UnrealBuildTool/Configuration/UEBuildBinary.cs
+++ b/Engine/Source/Programs/UnrealBuildTool/Configuration/UEBuildBinary.cs
@@ -169,12 +169,12 @@ namespace UnrealBuildTool
 		/// <param name="ToolChain">The toolchain which to use for building</param>
 		/// <param name="CompileEnvironment">The environment to compile the binary in</param>
 		/// <param name="LinkEnvironment">The environment to link the binary in</param>
-		/// <param name="SingleFileToCompile">If non-null, specifies a single cpp file to be compiled</param>
+		/// <param name="SpecificFilesToCompile">If non-empty, specifies individual cpp files to be compiled</param>
 		/// <param name="WorkingSet">The working set of source files</param>
 		/// <param name="ExeDir">Directory containing the output executable</param>
 		/// <param name="Graph">The graph being built</param>
 		/// <returns>Set of built products</returns>
-		public List<FileItem> Build(ReadOnlyTargetRules Target, UEToolChain ToolChain, CppCompileEnvironment CompileEnvironment, LinkEnvironment LinkEnvironment, FileReference SingleFileToCompile, ISourceFileWorkingSet WorkingSet, DirectoryReference ExeDir, IActionGraphBuilder Graph)
+		public List<FileItem> Build(ReadOnlyTargetRules Target, UEToolChain ToolChain, CppCompileEnvironment CompileEnvironment, LinkEnvironment LinkEnvironment, List<FileReference> SpecificFilesToCompile, ISourceFileWorkingSet WorkingSet, DirectoryReference ExeDir, IActionGraphBuilder Graph)
 		{
 			// Return nothing if we're using precompiled binaries. If we're not linking, we might want just one module to be compiled (eg. a foreign plugin), so allow any actions to run.
 			if (bUsePrecompiled && !(Target.LinkType == TargetLinkType.Monolithic && Target.bDisableLinking))
@@ -183,7 +183,7 @@ namespace UnrealBuildTool
 			}
 
 			// Setup linking environment.
-			LinkEnvironment BinaryLinkEnvironment = SetupBinaryLinkEnvironment(Target, ToolChain, LinkEnvironment, CompileEnvironment, SingleFileToCompile, WorkingSet, ExeDir, Graph);
+			LinkEnvironment BinaryLinkEnvironment = SetupBinaryLinkEnvironment(Target, ToolChain, LinkEnvironment, CompileEnvironment, SpecificFilesToCompile, WorkingSet, ExeDir, Graph);
 
 			// If we're generating projects, we only need include paths and definitions, there is no need to run the linking logic.
 			if (ProjectFileGenerator.bGenerateProjectFiles)
@@ -691,7 +691,7 @@ namespace UnrealBuildTool
 			return BinaryCompileEnvironment;
 		}
 
-		private LinkEnvironment SetupBinaryLinkEnvironment(ReadOnlyTargetRules Target, UEToolChain ToolChain, LinkEnvironment LinkEnvironment, CppCompileEnvironment CompileEnvironment, FileReference SingleFileToCompile, ISourceFileWorkingSet WorkingSet, DirectoryReference ExeDir, IActionGraphBuilder Graph)
+		private LinkEnvironment SetupBinaryLinkEnvironment(ReadOnlyTargetRules Target, UEToolChain ToolChain, LinkEnvironment LinkEnvironment, CppCompileEnvironment CompileEnvironment, List<FileReference> SpecificFilesToCompile, ISourceFileWorkingSet WorkingSet, DirectoryReference ExeDir, IActionGraphBuilder Graph)
 		{
 			LinkEnvironment BinaryLinkEnvironment = new LinkEnvironment(LinkEnvironment);
 			HashSet<UEBuildModule> LinkEnvironmentVisitedModules = new HashSet<UEBuildModule>();
@@ -711,7 +711,7 @@ namespace UnrealBuildTool
 				{
 					// Compile each module.
 					Log.TraceVerbose("Compile module: " + Module.Name);
-					LinkInputFiles = Module.Compile(Target, ToolChain, BinaryCompileEnvironment, SingleFileToCompile, WorkingSet, Graph);
+					LinkInputFiles = Module.Compile(Target, ToolChain, BinaryCompileEnvironment, SpecificFilesToCompile, WorkingSet, Graph);
 
 					// Save the module outputs. In monolithic builds, this is just the object files.
 					if (Target.LinkType == TargetLinkType.Monolithic)
diff --git a/Engine/Source/Programs/UnrealBuildTool/Configuration/UEBuildModule.cs b/Engine/Source/Programs/UnrealBuildTool/Configuration/UEBuildModule.cs
--- a/Engine/Source/Programs/UnrealBuildTool/Configuration/UEBuildModule.cs
+++ b/Engine/Source/Programs/UnrealBuildTool/Configuration/UEBuildModule.cs
@@ -792,7 +792,7 @@ namespace UnrealBuildTool
 		/// <summary>
 		/// Compiles the module, and returns a list of files output by the compiler.
 		/// </summary>
-		public virtual List<FileItem> Compile(ReadOnlyTargetRules Target, UEToolChain ToolChain, CppCompileEnvironment CompileEnvironment, FileReference SingleFileToCompile, ISourceFileWorkingSet WorkingSet, IActionGraphBuilder Graph)
+		public virtual List<FileItem> Compile(ReadOnlyTargetRules Target, UEToolChain ToolChain, CppCompileEnvironment CompileEnvironment, List<FileReference> SpecificFilesToCompile, ISourceFileWorkingSet WorkingSet, IActionGraphBuilder Graph)
 		{
 			// Generate type libraries for Windows
 			foreach(ModuleRules.TypeLibrary TypeLibrary in Rules.TypeLibraries)
diff --git a/Engine/Source/Programs/UnrealBuildTool/Configuration/UEBuildModuleCPP.cs b/Engine/Source/Programs/UnrealBuildTool/Configuration/UEBuildModuleCPP.cs
--- a/Engine/Source/Programs/UnrealBuildTool/Configuration/UEBuildModuleCPP.cs
+++ b/Engine/Source/Programs/UnrealBuildTool/Configuration/UEBuildModuleCPP.cs
@@ -276,11 +276,11 @@ namespace UnrealBuildTool
 		}
 
 		// UEBuildModule interface.
-		public override List<FileItem> Compile(ReadOnlyTargetRules Target, UEToolChain ToolChain, CppCompileEnvironment BinaryCompileEnvironment, FileReference SingleFileToCompile, ISourceFileWorkingSet WorkingSet, IActionGraphBuilder Graph)
+		public override List<FileItem> Compile(ReadOnlyTargetRules Target, UEToolChain ToolChain, CppCompileEnvironment BinaryCompileEnvironment, List<FileReference> SpecificFilesToCompile, ISourceFileWorkingSet WorkingSet, IActionGraphBuilder Graph)
 		{
 			//UEBuildPlatform BuildPlatform = UEBuildPlatform.GetBuildPlatform(BinaryCompileEnvironment.Platform);
 
-			List<FileItem> LinkInputFiles = base.Compile(Target, ToolChain, BinaryCompileEnvironment, SingleFileToCompile, WorkingSet, Graph);
+			List<FileItem> LinkInputFiles = base.Compile(Target, ToolChain, BinaryCompileEnvironment, SpecificFilesToCompile, WorkingSet, Graph);
 
 			CppCompileEnvironment ModuleCompileEnvironment = CreateModuleCompileEnvironment(Target, BinaryCompileEnvironment);
 
@@ -321,15 +321,16 @@ namespace UnrealBuildTool
 				Graph.AddSourceFiles(Pair.Key, Pair.Value);
 			}
 
-			// If we're compiling a single file, strip out anything else. This prevents us clobbering response files for anything we're 
+			// If we're compiling only specific files, strip out anything else. This prevents us clobbering response files for anything we're 
 			// not going to build, triggering a larger build than necessary when we do a regular build again.
-			if(SingleFileToCompile != null)
+			if(SpecificFilesToCompile.Count > 0)
 			{
-				InputFiles.CPPFiles.RemoveAll(x => x.Location != SingleFileToCompile);
-				InputFiles.CCFiles.RemoveAll(x => x.Location != SingleFileToCompile);
-				InputFiles.CFiles.RemoveAll(x => x.Location != SingleFileToCompile);
+				InputFiles.CPPFiles.RemoveAll(x => !SpecificFilesToCompile.Contains(x.Location));
+				InputFiles.CCFiles.RemoveAll(x => !SpecificFilesToCompile.Contains(x.Location));
+				InputFiles.CFiles.RemoveAll(x => !SpecificFilesToCompile.Contains(x.Location));
 
-				if(InputFiles.CPPFiles.Count == 0 && InputFiles.CCFiles.Count == 0 && InputFiles.CFiles.Count == 0 && !ContainsFile(SingleFileToCompile))
+				if (InputFiles.CPPFiles.Count == 0 && InputFiles.CCFiles.Count == 0 && InputFiles.CFiles.Count == 0 &&
+					!SpecificFilesToCompile.Any(x => ContainsFile(x)))
 				{
 					return new List<FileItem>();
 				}
@@ -432,7 +433,7 @@ namespace UnrealBuildTool
 			}
 
 			// Compile all the generated CPP files
-			if (GeneratedCppDirectories != null && !CompileEnvironment.bHackHeaderGenerator && SingleFileToCompile == null)
+			if (GeneratedCppDirectories != null && !CompileEnvironment.bHackHeaderGenerator && SpecificFilesToCompile.Count == 0)
 			{
 				List<string> GeneratedFiles = new List<string>();
 				foreach (string GeneratedDir in GeneratedCppDirectories)
@@ -468,7 +469,7 @@ namespace UnrealBuildTool
 					foreach (string GeneratedFilename in GeneratedFiles)
 					{
 						FileItem GeneratedCppFileItem = FileItem.GetItemByPath(GeneratedFilename);
-						if (SingleFileToCompile == null || GeneratedCppFileItem.Location == SingleFileToCompile)
+						if (SpecificFilesToCompile.Count == 0 || SpecificFilesToCompile.Contains(GeneratedCppFileItem.Location))
 						{
 							GeneratedFileItems.Add(GeneratedCppFileItem);
 						}
diff --git a/Engine/Source/Programs/UnrealBuildTool/Configuration/UEBuildTarget.cs b/Engine/Source/Programs/UnrealBuildTool/Configuration/UEBuildTarget.cs
--- a/Engine/Source/Programs/UnrealBuildTool/Configuration/UEBuildTarget.cs
+++ b/Engine/Source/Programs/UnrealBuildTool/Configuration/UEBuildTarget.cs
@@ -757,8 +757,8 @@ namespace UnrealBuildTool
 				RulesObject.DependencyListFileNames.Add(DependencyListFile);
 			}
 
-			// If we're compiling just a single file, we need to prevent unity builds from running
-			if(Descriptor.SingleFileToCompile != null)
+			// If we're compiling only specific files, we need to prevent unity builds from running
+			if(Descriptor.SpecificFilesToCompile.Count > 0)
 			{
 				RulesObject.bUseUnityBuild = false;
 				RulesObject.bForceUnityBuild = false;
@@ -1071,7 +1071,7 @@ namespace UnrealBuildTool
 			RulesAssembly = InRulesAssembly;
 			TargetType = Rules.Type;
 			ForeignPlugin = InDescriptor.ForeignPlugin;
-			bDeployAfterCompile = InRules.bDeployAfterCompile && !InRules.bDisableLinking && InDescriptor.SingleFileToCompile == null;
+			bDeployAfterCompile = InRules.bDeployAfterCompile && !InRules.bDisableLinking && InDescriptor.SpecificFilesToCompile.Count == 0;
 
 			// now that we have the platform, we can set the intermediate path to include the platform/architecture name
 			PlatformIntermediateFolder = GetPlatformIntermediateFolder(Platform, Architecture);
@@ -1588,7 +1588,7 @@ namespace UnrealBuildTool
 		/// <summary>
 		/// Builds the target, appending list of output files and returns building result.
 		/// </summary>
-		public TargetMakefile Build(BuildConfiguration BuildConfiguration, ISourceFileWorkingSet WorkingSet, bool bIsAssemblingBuild, FileReference SingleFileToCompile)
+		public TargetMakefile Build(BuildConfiguration BuildConfiguration, ISourceFileWorkingSet WorkingSet, bool bIsAssemblingBuild, List<FileReference> SpecificFilesToCompile)
 		{
 			CppConfiguration CppConfiguration = GetCppConfiguration(Configuration);
 
@@ -1774,14 +1774,20 @@ namespace UnrealBuildTool
 				}
 			}
 
-			// Find the set of binaries to build. If we're just compiling a single file, filter the list of binaries to only include the file we're interested in.
+			// Find the set of binaries to build. If we're compiling only specific files, filter the list of binaries to only include the files we're interested in.
 			List<UEBuildBinary> BuildBinaries = Binaries;
-			if (SingleFileToCompile != null)
+			foreach (FileReference SpecificFile in SpecificFilesToCompile)
 			{
-				BuildBinaries = Binaries.Where(x => x.Modules.Any(y => y.ContainsFile(SingleFileToCompile))).ToList();
-				if (BuildBinaries.Count == 0)
+				UEBuildBinary Binary = Binaries.Find(x => x.Modules.Any(y => y.ContainsFile(SpecificFile)));
+
+				if (Binary == null)
+				{
+					throw new BuildException("Couldn't find any module containing {0} in {1}.", SpecificFile, TargetName);
+				}
+
+				if (!BuildBinaries.Contains(Binary))
 				{
-					throw new BuildException("Couldn't find any module containing {0} in {1}.", SingleFileToCompile, TargetName);
+					BuildBinaries.Add(Binary);
 				}
 			}
 
@@ -1791,7 +1797,7 @@ namespace UnrealBuildTool
 			{
 				foreach (UEBuildBinary Binary in BuildBinaries)
 				{
-					List<FileItem> BinaryOutputItems = Binary.Build(Rules, TargetToolChain, GlobalCompileEnvironment, GlobalLinkEnvironment, SingleFileToCompile, WorkingSet, ExeDir, Makefile);
+					List<FileItem> BinaryOutputItems = Binary.Build(Rules, TargetToolChain, GlobalCompileEnvironment, GlobalLinkEnvironment, SpecificFilesToCompile, WorkingSet, ExeDir, Makefile);
 					Makefile.OutputItems.AddRange(BinaryOutputItems);
 				}
 			}
diff --git a/Engine/Source/Programs/UnrealBuildTool/Modes/BuildMode.cs b/Engine/Source/Programs/UnrealBuildTool/Modes/BuildMode.cs
--- a/Engine/Source/Programs/UnrealBuildTool/Modes/BuildMode.cs
+++ b/Engine/Source/Programs/UnrealBuildTool/Modes/BuildMode.cs
@@ -165,10 +165,10 @@ namespace UnrealBuildTool
 					TargetDescriptors = TargetDescriptor.ParseCommandLine(Arguments, BuildConfiguration.bUsePrecompiled, BuildConfiguration.bSkipRulesCompile);
 				}
 
-				// Hack for single file compile; don't build the ShaderCompileWorker target that's added to the command line for generated project files
+				// Hack for specific files compile; don't build the ShaderCompileWorker target that's added to the command line for generated project files
 				if(TargetDescriptors.Count >= 2)
 				{
-					TargetDescriptors.RemoveAll(x => (x.Name == "ShaderCompileWorker" || x.Name == "LiveCodingConsole") && x.SingleFileToCompile != null);
+					TargetDescriptors.RemoveAll(x => (x.Name == "ShaderCompileWorker" || x.Name == "LiveCodingConsole") && x.SpecificFilesToCompile.Count > 0);
 				}
 
 				// Handle remote builds
@@ -357,13 +357,16 @@ namespace UnrealBuildTool
 					}
 
 					// Plan the actions to execute for the build. For single file compiles, always rebuild the source file regardless of whether it's out of date.
-					if (TargetDescriptor.SingleFileToCompile == null)
+					if (TargetDescriptor.SpecificFilesToCompile.Count == 0)
 					{
 						ActionGraph.GatherAllOutdatedActions(PrerequisiteActions, History, ActionToOutdatedFlag, CppDependencies, BuildConfiguration.bIgnoreOutdatedImportLibraries);
 					}
 					else
 					{
-						ActionToOutdatedFlag[PrerequisiteActions.FirstOrDefault(x => x.PrerequisiteItems.Any(y => y.Location == TargetDescriptor.SingleFileToCompile))] = true;
+						foreach (FileReference SpecificFile in TargetDescriptor.SpecificFilesToCompile)
+						{
+							ActionToOutdatedFlag[PrerequisiteActions.FirstOrDefault(x => x.PrerequisiteItems.Any(y => y.Location == SpecificFile))] = true;
+						}
 					}
 				}
 
@@ -532,7 +535,7 @@ namespace UnrealBuildTool
 		{
 			// Get the path to the makefile for this target
 			FileReference MakefileLocation = null;
-			if(BuildConfiguration.bUseUBTMakefiles && TargetDescriptor.SingleFileToCompile == null)
+			if(BuildConfiguration.bUseUBTMakefiles && TargetDescriptor.SpecificFilesToCompile.Count == 0)
 			{
 				MakefileLocation = TargetMakefile.GetLocation(TargetDescriptor.ProjectFile, TargetDescriptor.Name, TargetDescriptor.Platform, TargetDescriptor.Configuration);
 			}
@@ -599,7 +602,7 @@ namespace UnrealBuildTool
 				using(Timeline.ScopeEvent("UEBuildTarget.Build()"))
 				{
 					const bool bIsAssemblingBuild = true;
-					Makefile = Target.Build(BuildConfiguration, WorkingSet, bIsAssemblingBuild, TargetDescriptor.SingleFileToCompile);
+					Makefile = Target.Build(BuildConfiguration, WorkingSet, bIsAssemblingBuild, TargetDescriptor.SpecificFilesToCompile);
 				}
 
 				// Save the pre-build scripts onto the makefile
@@ -653,11 +656,13 @@ namespace UnrealBuildTool
 		/// <returns>List of actions that need to be executed</returns>
 		static void GatherOutputItems(TargetDescriptor TargetDescriptor, TargetMakefile Makefile, HashSet<FileItem> OutputItems)
 		{
-			if(TargetDescriptor.SingleFileToCompile != null)
+			if(TargetDescriptor.SpecificFilesToCompile.Count > 0)
 			{
-				// If we're just compiling a single file, set the target items to be all the derived items
-				FileItem FileToCompile = FileItem.GetItemByFileReference(TargetDescriptor.SingleFileToCompile);
-				OutputItems.UnionWith(Makefile.Actions.Where(x => x.PrerequisiteItems.Contains(FileToCompile)).SelectMany(x => x.ProducedItems));
+				// If we're just compiling a specific files, set the target items to be all the derived items
+				List<FileItem> FilesToCompile = TargetDescriptor.SpecificFilesToCompile.ConvertAll(x => FileItem.GetItemByFileReference(x));
+				OutputItems.UnionWith(
+					Makefile.Actions.Where(x => x.PrerequisiteItems.Any(y => FilesToCompile.Contains(y)))
+					.SelectMany(x => x.ProducedItems));
 			}
 			else if(TargetDescriptor.OnlyModuleNames.Count > 0)
 			{
diff --git a/Engine/Source/Programs/UnrealBuildTool/Modes/JsonExportMode.cs b/Engine/Source/Programs/UnrealBuildTool/Modes/JsonExportMode.cs
--- a/Engine/Source/Programs/UnrealBuildTool/Modes/JsonExportMode.cs
+++ b/Engine/Source/Programs/UnrealBuildTool/Modes/JsonExportMode.cs
@@ -55,7 +55,7 @@ namespace UnrealBuildTool
 
 						// Create the makefile
 						const bool bIsAssemblingBuild = true;
-						TargetMakefile Makefile = Target.Build(BuildConfiguration, WorkingSet, bIsAssemblingBuild, TargetDescriptor.SingleFileToCompile);
+						TargetMakefile Makefile = Target.Build(BuildConfiguration, WorkingSet, bIsAssemblingBuild, TargetDescriptor.SpecificFilesToCompile);
 						ActionGraph.Link(Makefile.Actions);
 
 						// Filter all the actions to execute
diff --git a/Engine/Source/Programs/UnrealBuildTool/System/HotReload.cs b/Engine/Source/Programs/UnrealBuildTool/System/HotReload.cs
--- a/Engine/Source/Programs/UnrealBuildTool/System/HotReload.cs
+++ b/Engine/Source/Programs/UnrealBuildTool/System/HotReload.cs
@@ -181,7 +181,7 @@ namespace UnrealBuildTool
 			if (TargetDescriptor.HotReloadMode == HotReloadMode.Disabled)
 			{
 				// Make sure we're not doing a partial build from the editor (eg. compiling a new plugin)
-				if (TargetDescriptor.ForeignPlugin == null && TargetDescriptor.SingleFileToCompile == null)
+				if (TargetDescriptor.ForeignPlugin == null && TargetDescriptor.SpecificFilesToCompile.Count == 0)
 				{
 					// Delete the previous state file
 					FileReference StateFile = HotReloadState.GetLocation(TargetDescriptor);
