Allow overriding of CanUseAdvancedLinkerFeatures

diff --git a/Engine/Source/Programs/UnrealBuildTool/Platform/Linux/LinuxToolChain.cs b/Engine/Source/Programs/UnrealBuildTool/Platform/Linux/LinuxToolChain.cs
--- a/Engine/Source/Programs/UnrealBuildTool/Platform/Linux/LinuxToolChain.cs
+++ b/Engine/Source/Programs/UnrealBuildTool/Platform/Linux/LinuxToolChain.cs
@@ -1168,7 +1168,7 @@ namespace UnrealBuildTool
 		/// <summary>
 		/// Checks if we actually can use LTO/PGO with this set of tools
 		/// </summary>
-		private bool CanUseAdvancedLinkerFeatures(string Architecture)
+		protected virtual bool CanUseAdvancedLinkerFeatures(string Architecture)
 		{
 			return UsingLld(Architecture) && !String.IsNullOrEmpty(LlvmArPath);
 		}
