Fix Storyboard flow for Remote Build

diff --git a/Engine/Source/Programs/UnrealBuildTool/Platform/IOS/UEDeployIOS.cs b/Engine/Source/Programs/UnrealBuildTool/Platform/IOS/UEDeployIOS.cs
--- a/Engine/Source/Programs/UnrealBuildTool/Platform/IOS/UEDeployIOS.cs
+++ b/Engine/Source/Programs/UnrealBuildTool/Platform/IOS/UEDeployIOS.cs
@@ -800,62 +800,49 @@ namespace UnrealBuildTool
 			CopyFiles(InEngineDir + "/Build/IOS/Cloud", AppDirectory, "*.*", true);
 		}
 
-		protected virtual void CopyGraphicsResources(bool bSkipDefaultPNGs, bool bSkipIcons, string InEngineDir, string AppDirectory, string BuildDirectory, string IntermediateDir, bool bSupportsPortrait, bool bSupportsLandscape)
+		protected virtual void CopyLaunchScreenResources(string InEngineDir, string AppDirectory, string BuildDirectory)
 		{
-			// copy engine assets in (IOS and TVOS shared in IOS)
-			if (bSkipDefaultPNGs)
+			if (VersionUtilities.bCustomLaunchscreenStoryboard)
 			{
-				if (VersionUtilities.bCustomLaunchscreenStoryboard)
+				if (Directory.Exists(BuildDirectory + "/Resources/Interface/LaunchScreen.storyboardc"))
 				{
-					if (Directory.Exists(BuildDirectory + "/Resources/Interface/LaunchScreen.storyboardc"))
-					{
-						CopyFolder(BuildDirectory + "/Resources/Interface/LaunchScreen.storyboardc", AppDirectory + "/LaunchScreen.storyboardc", true);
-						CopyFiles(BuildDirectory + "/Resources/Interface/Assets", AppDirectory, "*", true);
-					}
-					else
-					{
-						CopyFolder(InEngineDir + "/Build/IOS/Resources/Interface/LaunchScreen.storyboardc", AppDirectory + "/LaunchScreen.storyboardc", true);
-						CopyFiles(InEngineDir + "/Build/IOS/Resources/Interface/Assets", AppDirectory, "*", true);
-					}
+					CopyFolder(BuildDirectory + "/Resources/Interface/LaunchScreen.storyboardc", AppDirectory + "/LaunchScreen.storyboardc", true);
+					CopyFiles(BuildDirectory + "/Resources/Interface/Assets", AppDirectory, "*", true);
 				}
 				else
 				{
-					if (File.Exists(BuildDirectory + "/Resources/Graphics/LaunchScreenIOS.png"))
-					{
-						CopyFolder(InEngineDir + "/Build/IOS/Resources/Interface/LaunchScreen.storyboardc", AppDirectory + "/LaunchScreen.storyboardc", true);
-						CopyFiles(BuildDirectory + "/Resources/Graphics", AppDirectory, "LaunchScreenIOS.png", true);
-					}
-					else
-					{
-						CopyFolder(InEngineDir + "/Build/IOS/Resources/Interface/LaunchScreen.storyboardc", AppDirectory + "/LaunchScreen.storyboardc", true);
-						CopyFiles(InEngineDir + "/Build/IOS/Resources/Graphics", AppDirectory, "LaunchScreenIOS.png", true);
-					}
-
+					CopyFolder(InEngineDir + "/Build/IOS/Resources/Interface/LaunchScreen.storyboardc", AppDirectory + "/LaunchScreen.storyboardc", true);
+					CopyFiles(InEngineDir + "/Build/IOS/Resources/Interface/Assets", AppDirectory, "*", true);
 				}
-				// we still want default icons
-				if (!bSkipIcons)
+			}
+			else
+			{
+				if (File.Exists(BuildDirectory + "/Resources/Graphics/LaunchScreenIOS.png"))
 				{
-					CopyFiles(InEngineDir + "/Build/IOS/Resources/Graphics", AppDirectory, "Icon*.png", true);
+					CopyFolder(InEngineDir + "/Build/IOS/Resources/Interface/LaunchScreen.storyboardc", AppDirectory + "/LaunchScreen.storyboardc", true);
+					CopyFiles(BuildDirectory + "/Resources/Graphics", AppDirectory, "LaunchScreenIOS.png", true);
 				}
 				else
 				{
-					if (!bSkipIcons)
-					{
-						CopyFiles(InEngineDir + "/Build/IOS/Resources/Graphics", AppDirectory, "Icon*.png", true);
-					}
+					CopyFolder(InEngineDir + "/Build/IOS/Resources/Interface/LaunchScreen.storyboardc", AppDirectory + "/LaunchScreen.storyboardc", true);
+					CopyFiles(InEngineDir + "/Build/IOS/Resources/Graphics", AppDirectory, "LaunchScreenIOS.png", true);
 				}
+
 			}
+		}
+
+		protected virtual void CopyIconResources(string InEngineDir, string AppDirectory, string BuildDirectory)
+		{
+			CopyFiles(InEngineDir + "/Build/IOS/Resources/Graphics", AppDirectory, "Icon*.png", true);
+	
 			// merge game assets on top
 			// @todo tvos: Do we want to copy IOS and TVOS both in? (Engine/IOS -> Game/IOS -> Game/TVOS)?
 			if (Directory.Exists(BuildDirectory + "/Resources/Graphics"))
 			{
-				if (!bSkipIcons)
-				{
 					CopyFiles(BuildDirectory + "/Resources/Graphics", AppDirectory, "Icon*.png", true);
-				}
 			}
 		}
-		
+
 		protected virtual void CopyLocalizationsResources(string InEngineDir, string AppDirectory, string BuildDirectory, string IntermediateDir)
 		{
 			string LocalizationsPath = BuildDirectory + "/Resources/Localizations";
@@ -1054,9 +1041,14 @@ namespace UnrealBuildTool
 			// copy the GameName binary
 			File.Copy(BinaryPath + "/" + GameExeName, AppDirectory + "/" + GameName, true);
 
+			CopyLaunchScreenResources(InEngineDir, AppDirectory, BuildDirectory);
+
 			if (!bCreateStubIPA)
 			{
-				CopyGraphicsResources(true, bSkipIcons, InEngineDir, AppDirectory, BuildDirectory, IntermediateDirectory, bSupportsPortrait, bSupportsLandscape);
+				if (!bSkipIcons)
+				{
+					CopyIconResources(InEngineDir, AppDirectory, BuildDirectory);
+				}
 				CopyLocalizationsResources(InEngineDir, AppDirectory, BuildDirectory, IntermediateDirectory);
 
 				CopyCloudResources(InProjectDirectory, AppDirectory);
diff --git a/Engine/Source/Programs/UnrealBuildTool/Platform/TVOS/UEDeployTVOS.cs b/Engine/Source/Programs/UnrealBuildTool/Platform/TVOS/UEDeployTVOS.cs
--- a/Engine/Source/Programs/UnrealBuildTool/Platform/TVOS/UEDeployTVOS.cs
+++ b/Engine/Source/Programs/UnrealBuildTool/Platform/TVOS/UEDeployTVOS.cs
@@ -361,9 +361,16 @@ namespace UnrealBuildTool
 			return GenerateTVOSPList(ProjectDirectory, bIsUE4Game, GameName, bIsClient, ProjectName, InEngineDir, AppDirectory, null, BundleID);
 		}
 
-        protected override void CopyGraphicsResources(bool bSkipDefaultPNGs, bool bSkipIcons, string InEngineDir, string AppDirectory, string BuildDirectory, string IntermediateDir, bool bSupportsPortrait, bool bSupportsLandscape)
-        {
-            // do nothing on TVOS
-        }
-    }
+		protected virtual void CopyIconResources(string InEngineDir, string AppDirectory, string BuildDirectory)
+		{
+			// do nothing on TVOS
+		}
+
+		protected virtual void CopyLaunchScreenResources(string InEngineDir, string AppDirectory, string BuildDirectory)
+		{
+			// do nothing on TVOS
+		}
+
+
+	}
 }
