Fix warning for tvOS

diff --git a/Engine/Source/Programs/UnrealBuildTool/Platform/TVOS/UEDeployTVOS.cs b/Engine/Source/Programs/UnrealBuildTool/Platform/TVOS/UEDeployTVOS.cs
--- a/Engine/Source/Programs/UnrealBuildTool/Platform/TVOS/UEDeployTVOS.cs
+++ b/Engine/Source/Programs/UnrealBuildTool/Platform/TVOS/UEDeployTVOS.cs
@@ -361,12 +361,12 @@ namespace UnrealBuildTool
 			return GenerateTVOSPList(ProjectDirectory, bIsUE4Game, GameName, bIsClient, ProjectName, InEngineDir, AppDirectory, null, BundleID);
 		}
 
-		protected virtual void CopyIconResources(string InEngineDir, string AppDirectory, string BuildDirectory)
+		protected override void CopyIconResources(string InEngineDir, string AppDirectory, string BuildDirectory)
 		{
 			// do nothing on TVOS
 		}
 
-		protected virtual void CopyLaunchScreenResources(string InEngineDir, string AppDirectory, string BuildDirectory)
+		protected override void CopyLaunchScreenResources(string InEngineDir, string AppDirectory, string BuildDirectory)
 		{
 			// do nothing on TVOS
 		}
