Helper function to add a debug string to the manifest localised string resources.

diff --git a/Engine/Source/Programs/UnrealBuildTool/Platform/Windows/VCManifestGenerator.cs b/Engine/Source/Programs/UnrealBuildTool/Platform/Windows/VCManifestGenerator.cs
--- a/Engine/Source/Programs/UnrealBuildTool/Platform/Windows/VCManifestGenerator.cs
+++ b/Engine/Source/Programs/UnrealBuildTool/Platform/Windows/VCManifestGenerator.cs
@@ -517,7 +517,20 @@ namespace UnrealBuildTool
 			}
 		}
 
-        protected virtual XName GetName( string BaseName, string SchemaName )
+		protected string AddDebugResourceString(string ResourceEntryName, string Value)
+		{
+			DefaultResourceWriter.AddResource(ResourceEntryName, Value);
+
+			foreach (var CultureId in CulturesToStage)
+			{
+				var Writer = PerCultureResourceWriters[CultureId];
+				Writer.AddResource(ResourceEntryName, Value);
+			}
+
+			return "ms-resource:" + ResourceEntryName;
+		}
+
+		protected virtual XName GetName( string BaseName, string SchemaName )
 		{
 			return XName.Get(BaseName);
 		}
