Manifest string resource helper function returns the manifest string reference too.

diff --git a/Engine/Source/Programs/UnrealBuildTool/Platform/Windows/VCManifestGenerator.cs b/Engine/Source/Programs/UnrealBuildTool/Platform/Windows/VCManifestGenerator.cs
--- a/Engine/Source/Programs/UnrealBuildTool/Platform/Windows/VCManifestGenerator.cs
+++ b/Engine/Source/Programs/UnrealBuildTool/Platform/Windows/VCManifestGenerator.cs
@@ -454,7 +454,7 @@ namespace UnrealBuildTool
 			}
 		}
 
-        protected void AddResourceEntry(string ResourceEntryName, string ConfigKey, string GenericINISection, string GenericINIKey, string DefaultValue, string ValueSuffix = "")
+        protected string AddResourceEntry(string ResourceEntryName, string ConfigKey, string GenericINISection, string GenericINIKey, string DefaultValue, string ValueSuffix = "")
 		{
 			string ConfigScratchValue = null;
 
@@ -466,7 +466,7 @@ namespace UnrealBuildTool
 				if (!ConfigHierarchy.TryParse(DefaultCultureScratchValue, out Values))
 				{
 					Log.TraceError("Invalid default culture string resources: \"{0}\". Unable to add resource entry.", DefaultCultureScratchValue);
-					return;
+					return "";
 				}
 
 				ConfigScratchValue = Values[ConfigKey];
@@ -515,6 +515,8 @@ namespace UnrealBuildTool
 					}
 				}
 			}
+
+			return "ms-resource:" + ResourceEntryName;
 		}
 
 		protected string AddDebugResourceString(string ResourceEntryName, string Value)
