Check presence of view controller ID when copying the LaunchScreen

diff --git a/Engine/Source/Programs/UnrealBuildTool/Platform/IOS/UEDeployIOS.cs b/Engine/Source/Programs/UnrealBuildTool/Platform/IOS/UEDeployIOS.cs
--- a/Engine/Source/Programs/UnrealBuildTool/Platform/IOS/UEDeployIOS.cs
+++ b/Engine/Source/Programs/UnrealBuildTool/Platform/IOS/UEDeployIOS.cs
@@ -828,7 +828,11 @@ namespace UnrealBuildTool
 					CopyFolder(InEngineDir + "/Build/IOS/Resources/Interface/LaunchScreen.storyboardc", AppDirectory + "/LaunchScreen.storyboardc", true);
 					CopyFiles(InEngineDir + "/Build/IOS/Resources/Graphics", AppDirectory, "LaunchScreenIOS.png", true);
 				}
+			}
 
+			if (!File.Exists(AppDirectory + "/LaunchScreen.storyboardc/LaunchScreen.nib"))
+			{
+				Log.TraceError("Launchscreen.storyboard ViewController needs an ID named LaunchScreen");
 			}
 		}
 
