Fixed UBT warnings due to missing documentation comment.

diff --git a/Engine/Source/Programs/UnrealBuildTool/ProjectFiles/Xcode/XcodeProject.cs b/Engine/Source/Programs/UnrealBuildTool/ProjectFiles/Xcode/XcodeProject.cs
--- a/Engine/Source/Programs/UnrealBuildTool/ProjectFiles/Xcode/XcodeProject.cs
+++ b/Engine/Source/Programs/UnrealBuildTool/ProjectFiles/Xcode/XcodeProject.cs
@@ -135,6 +135,7 @@ namespace UnrealBuildTool
 		/// <param name="InOnlyGameProject"></param>
 		/// <param name="IsForDistribution">True for distribution builds</param>
 		/// <param name="BundleID">Override option for bundle identifier</param>
+		/// <param name="InAppName"></param>
 		public XcodeProjectFile(FileReference InitFilePath, FileReference InOnlyGameProject, bool IsForDistribution, string BundleID, string InAppName)
 			: base(InitFilePath)
 		{
