Fix output directory in project files for programs build from extension folders.

diff --git a/Engine/Source/Programs/UnrealBuildTool/Configuration/UEBuildTarget.cs b/Engine/Source/Programs/UnrealBuildTool/Configuration/UEBuildTarget.cs
--- a/Engine/Source/Programs/UnrealBuildTool/Configuration/UEBuildTarget.cs
+++ b/Engine/Source/Programs/UnrealBuildTool/Configuration/UEBuildTarget.cs
@@ -2896,6 +2896,17 @@ namespace UnrealBuildTool
 			}
 		}
 
+		/// <summary>
+		/// Gets the output directory for a target
+		/// </summary>
+		/// <param name="BaseDir">The base directory for output files</param>
+		/// <param name="TargetFile">Path to the target file</param>
+		/// <returns>The output directory for this target</returns>
+		public static DirectoryReference GetOutputDirectoryForExecutable(DirectoryReference BaseDir, FileReference TargetFile)
+		{
+			return UnrealBuildTool.GetExtensionDirs(BaseDir).Where(x => TargetFile.IsUnderDirectory(x)).OrderByDescending(x => x.FullName.Length).FirstOrDefault() ?? UnrealBuildTool.EngineDirectory;
+		}
+
 		/// <summary>
 		/// Finds the base output directory for build products of the given module
 		/// </summary>
@@ -3386,14 +3397,15 @@ namespace UnrealBuildTool
 			}
 
 			// Construct the output paths for this target's executable
-			List<DirectoryReference> PossibleOutputDirectories = new List<DirectoryReference>();
-			PossibleOutputDirectories.AddRange(UnrealBuildTool.GetExtensionDirs(UnrealBuildTool.EngineDirectory));
-			if (ProjectFile != null && (bCompileMonolithic || !bUseSharedBuildEnvironment))
+			DirectoryReference OutputDirectory;
+			if (ProjectFile != null && (bCompileMonolithic || !bUseSharedBuildEnvironment) && Rules.File.IsUnderDirectory(ProjectDirectory))
 			{
-				PossibleOutputDirectories.AddRange(UnrealBuildTool.GetExtensionDirs(ProjectDirectory));
+				OutputDirectory = GetOutputDirectoryForExecutable(ProjectDirectory, Rules.File);
+			}
+			else
+			{
+				OutputDirectory = GetOutputDirectoryForExecutable(UnrealBuildTool.EngineDirectory, Rules.File);
 			}
-
-			DirectoryReference OutputDirectory = PossibleOutputDirectories.Where(x => Rules.File.IsUnderDirectory(x)).OrderByDescending(x => x.FullName.Length).FirstOrDefault() ?? UnrealBuildTool.EngineDirectory;
 
 			bool bCompileAsDLL = Rules.bShouldCompileAsDLL && bCompileMonolithic;
 			List<FileReference> OutputPaths = MakeBinaryPaths(OutputDirectory, bCompileMonolithic ? TargetName : AppName, Platform, Configuration, bCompileAsDLL ? UEBuildBinaryType.DynamicLinkLibrary : UEBuildBinaryType.Executable, Rules.Architecture, Rules.UndecoratedConfiguration, bCompileMonolithic && ProjectFile != null, Rules.ExeBinariesSubFolder, ProjectFile, Rules);
diff --git a/Engine/Source/Programs/UnrealBuildTool/ProjectFiles/VisualStudio/VCProject.cs b/Engine/Source/Programs/UnrealBuildTool/ProjectFiles/VisualStudio/VCProject.cs
--- a/Engine/Source/Programs/UnrealBuildTool/ProjectFiles/VisualStudio/VCProject.cs
+++ b/Engine/Source/Programs/UnrealBuildTool/ProjectFiles/VisualStudio/VCProject.cs
@@ -1410,23 +1410,22 @@ namespace UnrealBuildTool
 						bShouldCompileMonolithic = (Combination.ProjectTarget.CreateRulesDelegate(Platform, Configuration).LinkType == TargetLinkType.Monolithic);
 					}
 
+					// Get the .uproject directory
+					DirectoryReference UProjectDirectory = DirectoryReference.FromFile(Combination.ProjectTarget.UnrealProjectFilePath);
+
 					// Get the output directory
-					DirectoryReference RootDirectory = UnrealBuildTool.EngineDirectory;
-					if (TargetRulesObject.Type != TargetType.Program && (bShouldCompileMonolithic || TargetRulesObject.BuildEnvironment == TargetBuildEnvironment.Unique))
+					DirectoryReference RootOutputDirectory;
+					if (UProjectDirectory != null && (bShouldCompileMonolithic || TargetRulesObject.BuildEnvironment == TargetBuildEnvironment.Unique) && TargetRulesObject.File.IsUnderDirectory(UProjectDirectory))
 					{
-						if(Combination.ProjectTarget.UnrealProjectFilePath != null)
-						{
-							RootDirectory = Combination.ProjectTarget.UnrealProjectFilePath.Directory;
-						}
+						RootOutputDirectory = UEBuildTarget.GetOutputDirectoryForExecutable(UProjectDirectory, TargetRulesObject.File);
 					}
-
-					if (TargetRulesObject.Type == TargetType.Program && Combination.ProjectTarget.UnrealProjectFilePath != null)
+					else
 					{
-						RootDirectory = Combination.ProjectTarget.UnrealProjectFilePath.Directory;
+						RootOutputDirectory = UEBuildTarget.GetOutputDirectoryForExecutable(UnrealBuildTool.EngineDirectory, TargetRulesObject.File);
 					}
 
 					// Get the output directory
-					DirectoryReference OutputDirectory = DirectoryReference.Combine(RootDirectory, "Binaries", UBTPlatformName);
+					DirectoryReference OutputDirectory = DirectoryReference.Combine(RootOutputDirectory, "Binaries", UBTPlatformName);
 
 					if (!string.IsNullOrEmpty(TargetRulesObject.ExeBinariesSubFolder))
 					{
@@ -1553,7 +1552,7 @@ namespace UnrealBuildTool
 					if (TargetRulesObject.Type == TargetType.Game || TargetRulesObject.Type == TargetType.Client || TargetRulesObject.Type == TargetType.Server)
 					{
 						// Allow platforms to add any special properties they require... like aumid override for Xbox One
-						PlatformProjectGenerators.GenerateGamePlatformSpecificProperties(Platform, Configuration, TargetRulesObject.Type, VCProjectFileContent, RootDirectory, TargetFilePath);
+						PlatformProjectGenerators.GenerateGamePlatformSpecificProperties(Platform, Configuration, TargetRulesObject.Type, VCProjectFileContent, RootOutputDirectory, TargetFilePath);
 					}
 
 					VCProjectFileContent.AppendLine("  </PropertyGroup>");
