Set a new minimum SDK level for Android using Gradle

#ue4
#android

diff --git a/Engine/Source/Programs/UnrealBuildTool/Platform/Android/UEDeployAndroid.cs b/Engine/Source/Programs/UnrealBuildTool/Platform/Android/UEDeployAndroid.cs
--- a/Engine/Source/Programs/UnrealBuildTool/Platform/Android/UEDeployAndroid.cs
+++ b/Engine/Source/Programs/UnrealBuildTool/Platform/Android/UEDeployAndroid.cs
@@ -30,8 +30,8 @@ namespace UnrealBuildTool
 		// Minimum SDK version needed for App Bundles
 		readonly int MinimumSDKLevelForBundle = 21;
 
-		// Minimum SDK version needed for Gradle based on active plugins (14 is for Google Play Services 11.0.4)
-		private int MinimumSDKLevelForGradle = 14;
+		// Minimum SDK version needed for Gradle based on active plugins
+		private int MinimumSDKLevelForGradle = 19;
 
 		// Reserved Java keywords not allowed in package names without modification
 		static private string[] JavaReservedKeywords = new string[] {
