Check for change in active plugins on Android and clean Intermediate/Android if necessary to prevent packaging failures

#ue4
#android

diff --git a/Engine/Source/Programs/UnrealBuildTool/Platform/Android/UEDeployAndroid.cs b/Engine/Source/Programs/UnrealBuildTool/Platform/Android/UEDeployAndroid.cs
--- a/Engine/Source/Programs/UnrealBuildTool/Platform/Android/UEDeployAndroid.cs
+++ b/Engine/Source/Programs/UnrealBuildTool/Platform/Android/UEDeployAndroid.cs
@@ -65,6 +65,7 @@ namespace UnrealBuildTool
 		}
 
 		private UnrealPluginLanguage UPL = null;
+		private string ActiveUPLFiles = "";
 		private string UPLHashCode = null;
 		private bool ARCorePluginEnabled = false;
 		private bool FacebookPluginEnabled = false;
@@ -87,8 +88,11 @@ namespace UnrealBuildTool
 			FacebookPluginEnabled = false;
 			GoogleVRPluginEnabled = false;
 			OculusMobilePluginEnabled = false;
+			ActiveUPLFiles = "";
 			foreach (string Plugin in inPluginExtraData)
 			{
+				ActiveUPLFiles += Plugin + "\n";
+
 				// check if the Facebook plugin was enabled
 				if (Plugin.Contains("OnlineSubsystemFacebook_UPL"))
 				{
@@ -3250,6 +3254,20 @@ namespace UnrealBuildTool
 				string BuildTypeContents = File.ReadAllText(BuildTypeFilename);
 				if (BuildTypeID != BuildTypeContents)
 				{
+					Log.TraceInformation("Build type changed, forcing clean");
+					bCreateFromScratch = true;
+				}
+			}
+
+			// check if the enabled plugins has changed
+			string PluginListFilename = Path.Combine(IntermediateAndroidPath, "ActiveUPL.txt");
+			string PluginListContents = ActiveUPLFiles.ToString();
+			if (File.Exists(PluginListFilename))
+			{
+				string PreviousPluginListContents = File.ReadAllText(PluginListFilename);
+				if (PluginListContents != PreviousPluginListContents)
+				{
+					Log.TraceInformation("Active UPL files changed, forcing clean");
 					bCreateFromScratch = true;
 				}
 			}
@@ -3266,6 +3284,9 @@ namespace UnrealBuildTool
 				System.IO.Directory.CreateDirectory(IntermediateAndroidPath);
 			}
 
+			// write enabled plugins list
+			File.WriteAllText(PluginListFilename, PluginListContents);
+
 			// write build type
 			File.WriteAllText(BuildTypeFilename, BuildTypeID);
 
