Allow -armv7 etc on -architectures for Android in UBT

#ue4
#android

diff --git a/Engine/Source/Programs/UnrealBuildTool/Platform/Android/AndroidToolChain.cs b/Engine/Source/Programs/UnrealBuildTool/Platform/Android/AndroidToolChain.cs
--- a/Engine/Source/Programs/UnrealBuildTool/Platform/Android/AndroidToolChain.cs
+++ b/Engine/Source/Programs/UnrealBuildTool/Platform/Android/AndroidToolChain.cs
@@ -345,22 +345,22 @@ namespace UnrealBuildTool
 			Arches = new List<string>();
 			bool bBuild = true;
 			if (Ini.GetBool("/Script/AndroidRuntimeSettings.AndroidRuntimeSettings", "bBuildForArmV7", out bBuild) && bBuild
-				|| (AdditionalArches != null && AdditionalArches.Contains("armv7", StringComparer.OrdinalIgnoreCase)))
+				|| (AdditionalArches != null && (AdditionalArches.Contains("armv7", StringComparer.OrdinalIgnoreCase) || AdditionalArches.Contains("-armv7", StringComparer.OrdinalIgnoreCase))))
 			{
 				Arches.Add("-armv7");
 			}
 			if (Ini.GetBool("/Script/AndroidRuntimeSettings.AndroidRuntimeSettings", "bBuildForArm64", out bBuild) && bBuild
-				|| (AdditionalArches != null && AdditionalArches.Contains("arm64", StringComparer.OrdinalIgnoreCase)))
+				|| (AdditionalArches != null && (AdditionalArches.Contains("arm64", StringComparer.OrdinalIgnoreCase) || AdditionalArches.Contains("-arm64", StringComparer.OrdinalIgnoreCase))))
 			{
 				Arches.Add("-arm64");
 			}
 			if (Ini.GetBool("/Script/AndroidRuntimeSettings.AndroidRuntimeSettings", "bBuildForx86", out bBuild) && bBuild
-				|| (AdditionalArches != null && AdditionalArches.Contains("x86", StringComparer.OrdinalIgnoreCase)))
+				|| (AdditionalArches != null && (AdditionalArches.Contains("x86", StringComparer.OrdinalIgnoreCase) || AdditionalArches.Contains("-x86", StringComparer.OrdinalIgnoreCase)))
 			{
 				Arches.Add("-x86");
 			}
 			if (Ini.GetBool("/Script/AndroidRuntimeSettings.AndroidRuntimeSettings", "bBuildForx8664", out bBuild) && bBuild
-				|| (AdditionalArches != null && AdditionalArches.Contains("x64", StringComparer.OrdinalIgnoreCase)))
+				|| (AdditionalArches != null && (AdditionalArches.Contains("x64", StringComparer.OrdinalIgnoreCase) || AdditionalArches.Contains("-x64", StringComparer.OrdinalIgnoreCase))))
 			{
 				Arches.Add("-x64");
 			}
