Fix typo

#ue4
#android

diff --git a/Engine/Source/Programs/UnrealBuildTool/Platform/Android/AndroidToolChain.cs b/Engine/Source/Programs/UnrealBuildTool/Platform/Android/AndroidToolChain.cs
--- a/Engine/Source/Programs/UnrealBuildTool/Platform/Android/AndroidToolChain.cs
+++ b/Engine/Source/Programs/UnrealBuildTool/Platform/Android/AndroidToolChain.cs
@@ -355,7 +355,7 @@ namespace UnrealBuildTool
 				Arches.Add("-arm64");
 			}
 			if (Ini.GetBool("/Script/AndroidRuntimeSettings.AndroidRuntimeSettings", "bBuildForx86", out bBuild) && bBuild
-				|| (AdditionalArches != null && (AdditionalArches.Contains("x86", StringComparer.OrdinalIgnoreCase) || AdditionalArches.Contains("-x86", StringComparer.OrdinalIgnoreCase)))
+				|| (AdditionalArches != null && (AdditionalArches.Contains("x86", StringComparer.OrdinalIgnoreCase) || AdditionalArches.Contains("-x86", StringComparer.OrdinalIgnoreCase))))
 			{
 				Arches.Add("-x86");
 			}
