Change the property used to set include paths for Visual Studio project file. This signifantly reduces memory usage in the IDE (measured ~1,100mb before, ~500mb after) without changing intellisense results, because the property isn't propagated to each item. Also remove unnecessary force included files in non-C++ items.

diff --git a/Engine/Source/Programs/UnrealBuildTool/ProjectFiles/VisualStudio/VCProject.cs b/Engine/Source/Programs/UnrealBuildTool/ProjectFiles/VisualStudio/VCProject.cs
--- a/Engine/Source/Programs/UnrealBuildTool/ProjectFiles/VisualStudio/VCProject.cs
+++ b/Engine/Source/Programs/UnrealBuildTool/ProjectFiles/VisualStudio/VCProject.cs
@@ -954,7 +954,10 @@ namespace UnrealBuildTool
 				//      data being stored into the project file, and might make the IDE perform worse when switching configurations!
 				VCProjectFileContent.AppendLine("  <PropertyGroup>");
 				VCProjectFileContent.AppendLine("    <NMakePreprocessorDefinitions>$(NMakePreprocessorDefinitions){0}</NMakePreprocessorDefinitions>", (VCPreprocessorDefinitions.Length > 0 ? (";" + VCPreprocessorDefinitions) : ""));
-				VCProjectFileContent.AppendLine("    <NMakeIncludeSearchPath>$(NMakeIncludeSearchPath){0}</NMakeIncludeSearchPath>", (VCIncludeSearchPaths.Length > 0 ? (";" + VCIncludeSearchPaths) : ""));
+				// NOTE: Setting the IncludePath property rather than NMakeIncludeSearchPath results in significantly less
+				// memory usage, because NMakeIncludeSearchPath metadata is duplicated to each output item. Functionality should be identical for
+				// intellisense results.
+				VCProjectFileContent.AppendLine("    <IncludePath>$(IncludePath){0}</IncludePath>", (VCIncludeSearchPaths.Length > 0 ? (";" + VCIncludeSearchPaths) : ""));
 				VCProjectFileContent.AppendLine("    <NMakeForcedIncludes>$(NMakeForcedIncludes)</NMakeForcedIncludes>");
 				VCProjectFileContent.AppendLine("    <NMakeAssemblySearchPath>$(NMakeAssemblySearchPath)</NMakeAssemblySearchPath>");
 				VCProjectFileContent.AppendLine("    <AdditionalOptions>{0}</AdditionalOptions>",
@@ -1032,7 +1035,7 @@ namespace UnrealBuildTool
 					}
 
 					string VCFileType = GetVCFileType(AliasedFile.FileSystemPath);
-					if (ForceIncludePaths == null)
+					if (ForceIncludePaths == null)// || VCFileType != "ClCompile")
 					{
 						VCProjectFileContent.AppendLine("    <{0} Include=\"{1}\"/>", VCFileType, EscapeFileName(AliasedFile.FileSystemPath));
 					}
