Temporarily revert change to set IncludePaths variable, since it causes an error setting the INCLUDE environment variable when building large projects.

diff --git a/Engine/Source/Programs/UnrealBuildTool/ProjectFiles/VisualStudio/VCProject.cs b/Engine/Source/Programs/UnrealBuildTool/ProjectFiles/VisualStudio/VCProject.cs
--- a/Engine/Source/Programs/UnrealBuildTool/ProjectFiles/VisualStudio/VCProject.cs
+++ b/Engine/Source/Programs/UnrealBuildTool/ProjectFiles/VisualStudio/VCProject.cs
@@ -957,7 +957,9 @@ namespace UnrealBuildTool
 				// NOTE: Setting the IncludePath property rather than NMakeIncludeSearchPath results in significantly less
 				// memory usage, because NMakeIncludeSearchPath metadata is duplicated to each output item. Functionality should be identical for
 				// intellisense results.
-				VCProjectFileContent.AppendLine("    <IncludePath>$(IncludePath){0}</IncludePath>", (VCIncludeSearchPaths.Length > 0 ? (";" + VCIncludeSearchPaths) : ""));
+				// VCProjectFileContent.AppendLine("    <IncludePath>$(IncludePath){0}</IncludePath>", (VCIncludeSearchPaths.Length > 0 ? (";" + VCIncludeSearchPaths) : ""));
+				// NOTE 2: Temporarily reverted due to errors setting INCLUDE environment variable to a string exceeding the max allowed when building large projects
+				VCProjectFileContent.AppendLine("    <NMakeIncludeSearchPath>$(NMakeIncludeSearchPath){0}</NMakeIncludeSearchPath>", (VCIncludeSearchPaths.Length > 0 ? (";" + VCIncludeSearchPaths) : ""));
 				VCProjectFileContent.AppendLine("    <NMakeForcedIncludes>$(NMakeForcedIncludes)</NMakeForcedIncludes>");
 				VCProjectFileContent.AppendLine("    <NMakeAssemblySearchPath>$(NMakeAssemblySearchPath)</NMakeAssemblySearchPath>");
 				VCProjectFileContent.AppendLine("    <AdditionalOptions>{0}</AdditionalOptions>",
