Another attempt at fixing Visual Studio memory uses. Now sets the most common include paths, up to 24k, in the IncludePath property, and the remainder on individual items.

diff --git a/Engine/Source/Programs/UnrealBuildTool/ProjectFiles/Project.cs b/Engine/Source/Programs/UnrealBuildTool/ProjectFiles/Project.cs
--- a/Engine/Source/Programs/UnrealBuildTool/ProjectFiles/Project.cs
+++ b/Engine/Source/Programs/UnrealBuildTool/ProjectFiles/Project.cs
@@ -305,8 +305,15 @@ namespace UnrealBuildTool
 		public virtual void AddModule(UEBuildModuleCPP Module, CppCompileEnvironment CompileEnvironment)
 		{
 			AddIntelliSensePreprocessorDefinitions(CompileEnvironment.Definitions);
-			AddIntelliSenseIncludePaths(CompileEnvironment.SystemIncludePaths, true);
-			AddIntelliSenseIncludePaths(CompileEnvironment.UserIncludePaths, false);
+			AddIntelliSenseIncludePaths(SystemIncludePaths, CompileEnvironment.SystemIncludePaths);
+			AddIntelliSenseIncludePaths(UserIncludePaths, CompileEnvironment.UserIncludePaths);
+
+			foreach (DirectoryReference BaseDir in Module.ModuleDirectories)
+			{
+				AddIntelliSenseIncludePaths(BaseDir, BaseDirToSystemIncludePaths, CompileEnvironment.SystemIncludePaths);
+				AddIntelliSenseIncludePaths(BaseDir, BaseDirToUserIncludePaths, CompileEnvironment.UserIncludePaths);
+			}
+
 			SetIntelliSenseCppVersion(Module.Rules.CppStandard);
 		}
 
@@ -355,38 +362,38 @@ namespace UnrealBuildTool
 			}
 		}
 
+		/// <summary>
+		/// Adds all of the specified include paths to this VCProject's list of include paths for all modules in the project
+		/// </summary>
+		/// <param name="BaseDir">The base directory to which the include paths apply</param>
+		/// <param name="BaseDirToCollection">Map of base directory to include paths</param>
+		/// <param name="NewIncludePaths">List of include paths to add</param>
+		public void AddIntelliSenseIncludePaths(DirectoryReference BaseDir, Dictionary<DirectoryReference, IncludePathsCollection> BaseDirToCollection, IEnumerable<DirectoryReference> NewIncludePaths)
+		{
+			IncludePathsCollection ModuleUserPaths;
+			if (!BaseDirToCollection.TryGetValue(BaseDir, out ModuleUserPaths))
+			{
+				ModuleUserPaths = new IncludePathsCollection();
+				BaseDirToCollection.Add(BaseDir, ModuleUserPaths);
+			}
+			AddIntelliSenseIncludePaths(ModuleUserPaths, NewIncludePaths);
+		}
 
 		/// <summary>
 		/// Adds all of the specified include paths to this VCProject's list of include paths for all modules in the project
 		/// </summary>
+		/// <param name="Collection">The collection to add to</param>
 		/// <param name="NewIncludePaths">List of include paths to add</param>
-		/// <param name="bAddingSystemIncludes"></param>
-		public void AddIntelliSenseIncludePaths(HashSet<DirectoryReference> NewIncludePaths, bool bAddingSystemIncludes)
+		public void AddIntelliSenseIncludePaths(IncludePathsCollection Collection, IEnumerable<DirectoryReference> NewIncludePaths)
 		{
 			foreach (DirectoryReference CurPath in NewIncludePaths)
 			{
-				if (bAddingSystemIncludes ? KnownIntelliSenseSystemIncludeSearchPaths.Add(CurPath.FullName) : KnownIntelliSenseIncludeSearchPaths.Add(CurPath.FullName))
+				if(Collection.AbsolutePaths.Add(CurPath))
 				{
 					// Incoming include paths are relative to the solution directory, but we need these paths to be
 					// relative to the project file's directory
 					string PathRelativeToProjectFile = NormalizeProjectPath(CurPath);
-
-					// Make sure that it doesn't exist already
-					bool AlreadyExists = false;
-					List<string> SearchPaths = bAddingSystemIncludes ? IntelliSenseSystemIncludeSearchPaths : IntelliSenseIncludeSearchPaths;
-					foreach (string ExistingPath in SearchPaths)
-					{
-						if (PathRelativeToProjectFile == ExistingPath)
-						{
-							AlreadyExists = true;
-							break;
-						}
-					}
-
-					if (!AlreadyExists)
-					{
-						SearchPaths.Add(PathRelativeToProjectFile);
-					}
+					Collection.RelativePaths.Add(PathRelativeToProjectFile);
 				}
 			}
 		}
@@ -472,7 +479,7 @@ namespace UnrealBuildTool
 		/// <summary>
 		/// Takes the given path and tries to rebase it relative to the project or solution directory variables.
 		/// </summary>
-		public string NormalizeProjectPath(string InputPath)
+		public static string NormalizeProjectPath(string InputPath)
 		{
 			// If the path is rooted in an environment variable, leave it be.
 			if (InputPath.StartsWith("$("))
@@ -492,7 +499,7 @@ namespace UnrealBuildTool
 		/// <summary>
 		/// Takes the given path and tries to rebase it relative to the project.
 		/// </summary>
-		public string NormalizeProjectPath(FileSystemReference InputPath)
+		public static string NormalizeProjectPath(FileSystemReference InputPath)
 		{
 			// Try to make it relative to the solution directory.
 			if (InputPath.IsUnderDirectory(ProjectFileGenerator.MasterProjectPath))
@@ -526,23 +533,74 @@ namespace UnrealBuildTool
 			return ProjectFilePath.ToString();
 		}
 
+		/// <summary>
 		/// Map of file paths to files in the project.
+		/// </summary>
 		private readonly Dictionary<FileReference, SourceFile> SourceFileMap = new Dictionary<FileReference, SourceFile>();
 
+		/// <summary>
 		/// Files in this project
+		/// </summary>
 		public readonly List<SourceFile> SourceFiles = new List<SourceFile>();
 
-		/// Include paths for every single module in the project file, merged together
-		public readonly List<string> IntelliSenseIncludeSearchPaths = new List<string>();
-		public readonly List<string> IntelliSenseSystemIncludeSearchPaths = new List<string>();
-		public readonly HashSet<string> KnownIntelliSenseIncludeSearchPaths = new HashSet<string>(StringComparer.InvariantCultureIgnoreCase);
-		public readonly HashSet<string> KnownIntelliSenseSystemIncludeSearchPaths = new HashSet<string>(StringComparer.InvariantCultureIgnoreCase);
+		/// <summary>
+		/// Collection of include paths
+		/// </summary>
+		public class IncludePathsCollection
+		{
+			public List<string> RelativePaths = new List<string>();
+			public HashSet<DirectoryReference> AbsolutePaths = new HashSet<DirectoryReference>();
+		}
+
+		/// <summary>
+		/// Merged list of include paths for the project
+		/// </summary>
+		IncludePathsCollection UserIncludePaths = new IncludePathsCollection();
+
+		/// <summary>
+		/// Merged list of include paths for the project
+		/// </summary>
+		IncludePathsCollection SystemIncludePaths = new IncludePathsCollection();
+
+		/// <summary>
+		/// Map of base directory to user include paths
+		/// </summary>
+		protected Dictionary<DirectoryReference, IncludePathsCollection> BaseDirToUserIncludePaths = new Dictionary<DirectoryReference, IncludePathsCollection>();
+
+		/// <summary>
+		/// Map of base directory to system include paths
+		/// </summary>
+		protected Dictionary<DirectoryReference, IncludePathsCollection> BaseDirToSystemIncludePaths = new Dictionary<DirectoryReference, IncludePathsCollection>();
+
+		/// <summary>
+		/// Legacy accessor for user search paths
+		/// </summary>
+		public List<string> IntelliSenseIncludeSearchPaths
+		{
+			get { return UserIncludePaths.RelativePaths; }
+		}
+
+		/// <summary>
+		/// Legacy accessor for system include paths
+		/// </summary>
+		public List<string> IntelliSenseSystemIncludeSearchPaths
+		{
+			get { return SystemIncludePaths.RelativePaths; }
+		}
 
-		/// Preprocessor definitions for every single module in the project file, merged together
+		/// <summary>
+		/// List of preprocessor definitions for the project
+		/// </summary>
 		public readonly List<string> IntelliSensePreprocessorDefinitions = new List<string>();
-		public readonly HashSet<string> KnownIntelliSensePreprocessorDefinitions = new HashSet<string>(StringComparer.InvariantCultureIgnoreCase);
 
+		/// <summary>
+		/// Set of unique preprocessor definitions
+		/// </summary>
+		HashSet<string> KnownIntelliSensePreprocessorDefinitions = new HashSet<string>(StringComparer.InvariantCultureIgnoreCase);
+
+		/// <summary>
 		/// Projects that this project is dependent on
+		/// </summary>
 		public readonly List<ProjectFile> DependsOnProjects = new List<ProjectFile>();
 	}
 
diff --git a/Engine/Source/Programs/UnrealBuildTool/ProjectFiles/VisualStudio/VCProject.cs b/Engine/Source/Programs/UnrealBuildTool/ProjectFiles/VisualStudio/VCProject.cs
--- a/Engine/Source/Programs/UnrealBuildTool/ProjectFiles/VisualStudio/VCProject.cs
+++ b/Engine/Source/Programs/UnrealBuildTool/ProjectFiles/VisualStudio/VCProject.cs
@@ -718,6 +718,33 @@ namespace UnrealBuildTool
 			return ProjectFiles;
 		}
 
+		/// <summary>
+		/// Append a list of include paths to a property list
+		/// </summary>
+		/// <param name="Builder">String builder for the property value</param>
+		/// <param name="BaseDir">Directory containing the source file</param>
+		/// <param name="BaseDirToIncludePaths">Map of base directory to set of include paths</param>
+		/// <param name="IgnorePaths">Set of paths to ignore</param>
+		static void AppendIncludePaths(StringBuilder Builder, DirectoryReference BaseDir, Dictionary<DirectoryReference, IncludePathsCollection> BaseDirToIncludePaths, HashSet<DirectoryReference> IgnorePaths)
+		{
+			for (DirectoryReference CurrentDir = BaseDir; CurrentDir != null; CurrentDir = CurrentDir.ParentDirectory)
+			{
+				IncludePathsCollection Collection;
+				if (BaseDirToIncludePaths.TryGetValue(CurrentDir, out Collection))
+				{
+					foreach(DirectoryReference IncludePath in Collection.AbsolutePaths)
+					{
+						if (!IgnorePaths.Contains(IncludePath))
+						{
+							Builder.Append(NormalizeProjectPath(IncludePath.FullName));
+							Builder.Append(';');
+						}
+					}
+					break;
+				}
+			}
+		}
+
 		/// Implements Project interface
 		public override bool WriteProjectFile(List<UnrealTargetPlatform> InPlatforms, List<UnrealTargetConfiguration> InConfigurations, PlatformProjectGeneratorCollection PlatformProjectGenerators)
 		{
@@ -725,28 +752,74 @@ namespace UnrealBuildTool
 
 			bool bSuccess = true;
 
+			// Merge as many include paths as possible into the shared list
+			HashSet<DirectoryReference> SharedIncludeSearchPathsSet = new HashSet<DirectoryReference>();
+
 			// Build up the new include search path string
-			StringBuilder VCIncludeSearchPaths = new StringBuilder();
+			StringBuilder SharedIncludeSearchPaths = new StringBuilder();
 			{
-				foreach (string CurPath in IntelliSenseIncludeSearchPaths)
+				// Find out how many source files there are in each directory
+				Dictionary<DirectoryReference, int> SourceDirToCount = new Dictionary<DirectoryReference, int>();
+				foreach (SourceFile SourceFile in SourceFiles)
 				{
-					VCIncludeSearchPaths.Append(CurPath + ";");
+					if(SourceFile.Reference.HasExtension(".cpp"))
+					{
+						DirectoryReference SourceDir = SourceFile.Reference.Directory;
+
+						int Count;
+						SourceDirToCount.TryGetValue(SourceDir, out Count);
+						SourceDirToCount[SourceDir] = Count + 1;
+					}
+				}
+
+				// Figure out the most common include paths
+				Dictionary<DirectoryReference, int> IncludePathToCount = new Dictionary<DirectoryReference, int>();
+				foreach (KeyValuePair<DirectoryReference, int> Pair in SourceDirToCount)
+				{
+					for (DirectoryReference CurrentDir = Pair.Key; CurrentDir != null; CurrentDir = CurrentDir.ParentDirectory)
+					{
+						IncludePathsCollection IncludePaths;
+						if (BaseDirToUserIncludePaths.TryGetValue(CurrentDir, out IncludePaths))
+						{
+							foreach (DirectoryReference IncludePath in IncludePaths.AbsolutePaths)
+							{
+								int Count;
+								IncludePathToCount.TryGetValue(IncludePath, out Count);
+								IncludePathToCount[IncludePath] = Count + Pair.Value;
+							}
+							break;
+						}
+					}
 				}
-				foreach (string CurPath in IntelliSenseSystemIncludeSearchPaths)
+
+				// Append the most common include paths to the search list.
+				const int MaxSharedIncludePathsLength = 24 * 1024;
+				foreach (DirectoryReference IncludePath in IncludePathToCount.OrderByDescending(x => x.Value).Select(x => x.Key))
 				{
-					VCIncludeSearchPaths.Append(CurPath + ";");
+					string RelativePath = NormalizeProjectPath(IncludePath);
+					if (SharedIncludeSearchPaths.Length + RelativePath.Length < MaxSharedIncludePathsLength)
+					{
+						SharedIncludeSearchPaths.AppendFormat("{0};", RelativePath);
+						SharedIncludeSearchPathsSet.Add(IncludePath);
+					}
+					else
+					{
+						break;
+					}
 				}
+
+				// Add all the default system include paths
 				if (InPlatforms.Contains(UnrealTargetPlatform.Win64))
 				{
-					VCIncludeSearchPaths.Append(VCToolChain.GetVCIncludePaths(UnrealTargetPlatform.Win64, GetCompilerForIntellisense(), null) + ";");
+					SharedIncludeSearchPaths.Append(VCToolChain.GetVCIncludePaths(UnrealTargetPlatform.Win64, GetCompilerForIntellisense(), null) + ";");
 				}
 				else if (InPlatforms.Contains(UnrealTargetPlatform.Win32))
 				{
-					VCIncludeSearchPaths.Append(VCToolChain.GetVCIncludePaths(UnrealTargetPlatform.Win32, GetCompilerForIntellisense(), null) + ";");
+					SharedIncludeSearchPaths.Append(VCToolChain.GetVCIncludePaths(UnrealTargetPlatform.Win32, GetCompilerForIntellisense(), null) + ";");
 				}
 				else if (InPlatforms.Contains(UnrealTargetPlatform.HoloLens))
 				{
-					VCIncludeSearchPaths.Append(HoloLensToolChain.GetVCIncludePaths(UnrealTargetPlatform.HoloLens, GetCompilerForIntellisense()) + ";");
+					SharedIncludeSearchPaths.Append(HoloLensToolChain.GetVCIncludePaths(UnrealTargetPlatform.HoloLens, GetCompilerForIntellisense()) + ";");
 				}
 			}
 
@@ -957,9 +1030,7 @@ namespace UnrealBuildTool
 				// NOTE: Setting the IncludePath property rather than NMakeIncludeSearchPath results in significantly less
 				// memory usage, because NMakeIncludeSearchPath metadata is duplicated to each output item. Functionality should be identical for
 				// intellisense results.
-				// VCProjectFileContent.AppendLine("    <IncludePath>$(IncludePath){0}</IncludePath>", (VCIncludeSearchPaths.Length > 0 ? (";" + VCIncludeSearchPaths) : ""));
-				// NOTE 2: Temporarily reverted due to errors setting INCLUDE environment variable to a string exceeding the max allowed when building large projects
-				VCProjectFileContent.AppendLine("    <NMakeIncludeSearchPath>$(NMakeIncludeSearchPath){0}</NMakeIncludeSearchPath>", (VCIncludeSearchPaths.Length > 0 ? (";" + VCIncludeSearchPaths) : ""));
+				VCProjectFileContent.AppendLine("    <IncludePath>$(IncludePath){0}</IncludePath>", (SharedIncludeSearchPaths.Length > 0 ? (";" + SharedIncludeSearchPaths) : ""));
 				VCProjectFileContent.AppendLine("    <NMakeForcedIncludes>$(NMakeForcedIncludes)</NMakeForcedIncludes>");
 				VCProjectFileContent.AppendLine("    <NMakeAssemblySearchPath>$(NMakeAssemblySearchPath)</NMakeAssemblySearchPath>");
 				VCProjectFileContent.AppendLine("    <AdditionalOptions>{0}</AdditionalOptions>",
@@ -1012,6 +1083,7 @@ namespace UnrealBuildTool
 				HashSet<string> FilterDirectories = new HashSet<string>();
 				//UEBuildPlatform BuildPlatform = UEBuildPlatform.GetBuildPlatform(BuildHostPlatform.Current.Platform);
 
+				Dictionary<DirectoryReference, string> DirectoryToIncludeSearchPaths = new Dictionary<DirectoryReference, string>();
 				Dictionary<DirectoryReference, string> DirectoryToForceIncludePaths = new Dictionary<DirectoryReference, string>();
 				foreach (AliasedFile AliasedFile in LocalAliasedFiles)
 				{
@@ -1021,29 +1093,43 @@ namespace UnrealBuildTool
 						FiltersFileIsNeeded = EnsureFilterPathExists(AliasedFile.ProjectPath, VCFiltersFileContent, FilterDirectories);
 					}
 
-					DirectoryReference Directory = AliasedFile.Location.Directory;
-
-					string ForceIncludePaths;
-					if (!DirectoryToForceIncludePaths.TryGetValue(Directory, out ForceIncludePaths))
-					{
-						for(DirectoryReference ParentDir = Directory; ParentDir != null; ParentDir = ParentDir.ParentDirectory)
-						{
-							if(ModuleDirToForceIncludePaths.TryGetValue(ParentDir, out ForceIncludePaths))
-							{
-								break;
-							}
-						}
-						DirectoryToForceIncludePaths[Directory] = ForceIncludePaths;
-					}
-
 					string VCFileType = GetVCFileType(AliasedFile.FileSystemPath);
-					if (ForceIncludePaths == null)// || VCFileType != "ClCompile")
+					if (VCFileType != "ClCompile")
 					{
 						VCProjectFileContent.AppendLine("    <{0} Include=\"{1}\"/>", VCFileType, EscapeFileName(AliasedFile.FileSystemPath));
 					}
 					else
 					{
+						DirectoryReference Directory = AliasedFile.Location.Directory;
+
+						// Find the force-included headers
+						string ForceIncludePaths;
+						if (!DirectoryToForceIncludePaths.TryGetValue(Directory, out ForceIncludePaths))
+						{
+							for (DirectoryReference ParentDir = Directory; ParentDir != null; ParentDir = ParentDir.ParentDirectory)
+							{
+								if (ModuleDirToForceIncludePaths.TryGetValue(ParentDir, out ForceIncludePaths))
+								{
+									break;
+								}
+							}
+							DirectoryToForceIncludePaths[Directory] = ForceIncludePaths;
+						}
+
+						// Find the include search paths
+						string IncludeSearchPaths;
+						if (!DirectoryToIncludeSearchPaths.TryGetValue(Directory, out IncludeSearchPaths))
+						{
+							StringBuilder Builder = new StringBuilder();
+							AppendIncludePaths(Builder, Directory, BaseDirToUserIncludePaths, SharedIncludeSearchPathsSet);
+							AppendIncludePaths(Builder, Directory, BaseDirToSystemIncludePaths, SharedIncludeSearchPathsSet);
+							IncludeSearchPaths = Builder.ToString();
+
+							DirectoryToIncludeSearchPaths.Add(Directory, IncludeSearchPaths);
+						}
+
 						VCProjectFileContent.AppendLine("    <{0} Include=\"{1}\">", VCFileType, EscapeFileName(AliasedFile.FileSystemPath));
+						VCProjectFileContent.AppendLine("      <AdditionalIncludeDirectories>$(NMakeIncludeSearchPath);{0}</AdditionalIncludeDirectories>", IncludeSearchPaths);
 						VCProjectFileContent.AppendLine("      <ForcedIncludeFiles>$(NMakeForcedIncludes);{0}</ForcedIncludeFiles>", ForceIncludePaths);
 						VCProjectFileContent.AppendLine("    </{0}>", VCFileType);
 					}
