Clear out the INCLUDE environment variable before building. Since the IncludePath property from the Visual C++ project files gets propagated to this environment variable, it gets convered into /I arguments by the compiler driver. That results in a long argument list, which causes an error in Visual Studio 2017.

diff --git a/Engine/Source/Programs/UnrealBuildTool/Platform/Windows/VCToolChain.cs b/Engine/Source/Programs/UnrealBuildTool/Platform/Windows/VCToolChain.cs
--- a/Engine/Source/Programs/UnrealBuildTool/Platform/Windows/VCToolChain.cs
+++ b/Engine/Source/Programs/UnrealBuildTool/Platform/Windows/VCToolChain.cs
@@ -52,6 +52,11 @@ namespace UnrealBuildTool
 		public override void SetEnvironmentVariables()
 		{
 			EnvVars.SetEnvironmentVariables();
+
+			// Don't allow the INCLUDE environment variable to propagate. It's set by the IDE based on the IncludePath property in the project files which we
+			// add to improve Visual Studio memory usage, but we don't actually need it to set when invoking the compiler. Doing so results in it being converted
+			// into /I arguments by the CL driver, which results in errors due to the command line not fitting into the PDB debug record.
+			Environment.SetEnvironmentVariable("INCLUDE", null);
 		}
 
 		/// <summary>
