Fix crash on tvOS startup due to missing storyboard

diff --git a/Engine/Source/Programs/UnrealBuildTool/Platform/TVOS/UEDeployTVOS.cs b/Engine/Source/Programs/UnrealBuildTool/Platform/TVOS/UEDeployTVOS.cs
--- a/Engine/Source/Programs/UnrealBuildTool/Platform/TVOS/UEDeployTVOS.cs
+++ b/Engine/Source/Programs/UnrealBuildTool/Platform/TVOS/UEDeployTVOS.cs
@@ -363,12 +363,5 @@ namespace UnrealBuildTool
 		{
 			// do nothing on TVOS
 		}
-
-		protected override void CopyLaunchScreenResources(string InEngineDir, string AppDirectory, string BuildDirectory)
-		{
-			// do nothing on TVOS
-		}
-
-
 	}
 }
