Exclude binary files from generated projects.

diff --git a/Engine/Source/Programs/UnrealBuildTool/System/SourceFileSearch.cs b/Engine/Source/Programs/UnrealBuildTool/System/SourceFileSearch.cs
--- a/Engine/Source/Programs/UnrealBuildTool/System/SourceFileSearch.cs
+++ b/Engine/Source/Programs/UnrealBuildTool/System/SourceFileSearch.cs
@@ -27,6 +27,7 @@ namespace UnrealBuildTool
 			".csproj",				// Visual Studio C# project files
 			".userprefs",			// MonoDevelop project settings
 			".ds_store",			// Mac Desktop Services Store hidden files
+			".bin",					// Binary files
 			Path.DirectorySeparatorChar + "do_not_delete.txt",		// Perforce placeholder file
 		};
 
