Clear the promoted flag for binaries built with the installed engine build. This fixes issues where a code project will only build once.

diff --git a/Engine/Source/Programs/UnrealBuildTool/Configuration/UEBuildTarget.cs b/Engine/Source/Programs/UnrealBuildTool/Configuration/UEBuildTarget.cs
--- a/Engine/Source/Programs/UnrealBuildTool/Configuration/UEBuildTarget.cs
+++ b/Engine/Source/Programs/UnrealBuildTool/Configuration/UEBuildTarget.cs
@@ -1434,6 +1434,12 @@ namespace UnrealBuildTool
 				}
 			}
 
+			// If this is an installed engine build, clear the promoted flag on the output binaries. This will ensure we will rebuild them.
+			if (Version.IsPromotedBuild && UnrealBuildTool.IsEngineInstalled())
+			{
+				Version.IsPromotedBuild = false;
+			}
+
 			// Create the receipt
 			TargetReceipt Receipt = new TargetReceipt(ProjectFile, TargetName, TargetType, Platform, Configuration, Version, Architecture);
 
