Update SetupAndroid scripts to install Android NDK 21b with fixed clang for x86_64

#ue4
#android

diff --git a/Engine/Extras/Android/SetupAndroid.bat b/Engine/Extras/Android/SetupAndroid.bat
--- a/Engine/Extras/Android/SetupAndroid.bat
+++ b/Engine/Extras/Android/SetupAndroid.bat
@@ -42,7 +42,7 @@ if DEFINED JAVA_HOME (set a=1) ELSE (
 	set JAVA_HOME=%STUDIO_PATH%\jre
 	setx JAVA_HOME %STUDIO_PATH%\jre
 )
-set NDKINSTALLPATH=%STUDIO_SDK_PATH%\ndk\21.0.6113669
+set NDKINSTALLPATH=%STUDIO_SDK_PATH%\ndk\21.1.6352462
 set PLATFORMTOOLS=%STUDIO_SDK_PATH%\platform-tools;%STUDIO_SDK_PATH%\tools
 
 set KEY_NAME=HKCU\Environment
@@ -72,7 +72,7 @@ IF EXIST "%SDKMANAGER%" (
 	)
 )
 
-call "%SDKMANAGER%" "platform-tools" "platforms;android-28" "build-tools;28.0.3" "lldb;3.1" "cmake;3.10.2.4988404" "ndk;21.0.6113669"
+call "%SDKMANAGER%" "platform-tools" "platforms;android-28" "build-tools;28.0.3" "lldb;3.1" "cmake;3.10.2.4988404" "ndk;21.1.6352462"
 
 IF /I "%ERRORLEVEL%" NEQ "0" (
 	echo Update failed. Please check the Android Studio install.
@@ -83,6 +83,7 @@ IF /I "%ERRORLEVEL%" NEQ "0" (
 if EXIST %NDKINSTALLPATH% (
 	echo Success!
 	setx NDKROOT %NDKINSTALLPATH%
+	setx NDK_ROOT %NDKINSTALLPATH%
 ) ELSE (
 	echo Update failed. Did you accept the license agreement?
 	pause
diff --git a/Engine/Extras/Android/SetupAndroid.command b/Engine/Extras/Android/SetupAndroid.command
--- a/Engine/Extras/Android/SetupAndroid.command
+++ b/Engine/Extras/Android/SetupAndroid.command
@@ -39,7 +39,7 @@ if ! grep -q "export JAVA_HOME=\"$JAVA_HOME\"" ~/.bash_profile
 then
 	echo "export JAVA_HOME=\"$JAVA_HOME\"" >>~/.bash_profile
 fi
-NDKINSTALLPATH="$STUDIO_SDK_PATH/ndk/21.0.6113669"
+NDKINSTALLPATH="$STUDIO_SDK_PATH/ndk/21.1.6352462"
 PLATFORMTOOLS="$STUDIO_SDK_PATH/platform-tools:$STUDIO_SDK_PATH/build-tools/28.0.3:$%STUDIO_SDK_PATH/tools/bin"
 
 retVal=$(type -P "adb")
@@ -59,7 +59,7 @@ if [ ! -d "$SDKMANAGERPATH" ]; then
 	fi
 fi
 
-"$SDKMANAGERPATH/sdkmanager" "platform-tools" "platforms;android-28" "build-tools;28.0.3" "lldb;3.1" "cmake;3.10.2.4988404" "ndk;21.0.6113669"
+"$SDKMANAGERPATH/sdkmanager" "platform-tools" "platforms;android-28" "build-tools;28.0.3" "lldb;3.1" "cmake;3.10.2.4988404" "ndk;21.1.6352462"
 
 retVal=$?
 if [ $retVal -ne 0 ]; then
@@ -92,6 +92,7 @@ echo Success!
 if ! grep -q "export NDKROOT=\"$NDKINSTALLPATH\"" ~/.bash_profile
 then
 	echo "export NDKROOT=\"$NDKINSTALLPATH\"" >>~/.bash_profile
+	echo "export NDK_ROOT=\"$NDKINSTALLPATH\"" >>~/.bash_profile
 fi
 
 exit 0
diff --git a/Engine/Extras/Android/SetupAndroid.sh b/Engine/Extras/Android/SetupAndroid.sh
--- a/Engine/Extras/Android/SetupAndroid.sh
+++ b/Engine/Extras/Android/SetupAndroid.sh
@@ -37,7 +37,7 @@ if ! grep -q "export JAVA_HOME=\"$JAVA_HOME\"" $HOME/.bashrc
 then
 	echo "export JAVA_HOME=\"$JAVA_HOME\"" >>$HOME/.bashrc
 fi
-NDKINSTALLPATH="$STUDIO_SDK_PATH/ndk/21.0.6113669"
+NDKINSTALLPATH="$STUDIO_SDK_PATH/ndk/21.1.6352462"
 PLATFORMTOOLS="$STUDIO_SDK_PATH/platform-tools:$STUDIO_SDK_PATH/build-tools/28.0.3:$%STUDIO_SDK_PATH/tools/bin"
 
 retVal=$(type -P "adb")
@@ -57,7 +57,7 @@ if [ ! -d "$SDKMANAGERPATH" ]; then
 	fi
 fi
 
-"$SDKMANAGERPATH/sdkmanager" "platform-tools" "platforms;android-28" "build-tools;28.0.3" "lldb;3.1" "cmake;3.10.2.4988404" "ndk;21.0.6113669"
+"$SDKMANAGERPATH/sdkmanager" "platform-tools" "platforms;android-28" "build-tools;28.0.3" "lldb;3.1" "cmake;3.10.2.4988404" "ndk;21.1.6352462"
 
 retVal=$?
 if [ $retVal -ne 0 ]; then
@@ -90,6 +90,7 @@ echo Success!
 if ! grep -q "export NDKROOT=\"$NDKINSTALLPATH\"" $HOME/.bashrc
 then
 	echo "export NDKROOT=\"$NDKINSTALLPATH\"" >>$HOME/.bashrc
+	echo "export NDK_ROOT=\"$NDKINSTALLPATH\"" >>$HOME/.bashrc
 fi
 
 exit 0
diff --git a/Engine/Source/Programs/UnrealBuildTool/Platform/Android/AndroidToolChain.cs b/Engine/Source/Programs/UnrealBuildTool/Platform/Android/AndroidToolChain.cs
--- a/Engine/Source/Programs/UnrealBuildTool/Platform/Android/AndroidToolChain.cs
+++ b/Engine/Source/Programs/UnrealBuildTool/Platform/Android/AndroidToolChain.cs
@@ -14,9 +14,9 @@ namespace UnrealBuildTool
 	class AndroidToolChain : ISPCToolChain, IAndroidToolChain
 	{
 		// Android NDK toolchain that must be used for C++ compiling
-		readonly int MinimumNDKToolchain = 200200;
-		readonly int MaximumNDKToolchain = 220100;
-		readonly int RecommendedNDKToolchain = 210100;
+		readonly int MinimumNDKToolchain = 210100;
+		readonly int MaximumNDKToolchain = 230100;
+		readonly int RecommendedNDKToolchain = 210200;
 
 		public static readonly string[] AllCpuSuffixes =
 		{
@@ -109,6 +109,14 @@ namespace UnrealBuildTool
 			return string.Format("{0}.{1}.{2}", ClangVersionMajor, ClangVersionMinor, ClangVersionPatch);
 		}
 
+		public void ShowNDKWarnings()
+		{
+			if (NDKDefineInt < RecommendedNDKToolchain)
+			{
+				Log.TraceInformation("Note: Android toolchain NDK " + ToolchainIntToString(RecommendedNDKToolchain) + " recommended");
+			}
+		}
+
 		/// <summary>
 		/// Checks if compiler version matches the requirements
 		/// </summary>
diff --git a/Engine/Source/Programs/UnrealBuildTool/Platform/Android/UEBuildAndroid.cs b/Engine/Source/Programs/UnrealBuildTool/Platform/Android/UEBuildAndroid.cs
--- a/Engine/Source/Programs/UnrealBuildTool/Platform/Android/UEBuildAndroid.cs
+++ b/Engine/Source/Programs/UnrealBuildTool/Platform/Android/UEBuildAndroid.cs
@@ -338,6 +338,7 @@ namespace UnrealBuildTool
 
 			int NDKVersionInt = ToolChain.GetNdkApiLevelInt();
 			Log.TraceInformation("NDK toolchain: {0}, NDK version: {1}, GccVersion: {2}, ClangVersion: {3}", NDKToolchainVersion, NDKVersionInt.ToString(), GccVersion, ToolChain.GetClangVersionString());
+			ToolChain.ShowNDKWarnings();
 
 			CompileEnvironment.Definitions.Add("PLATFORM_DESKTOP=0");
 			CompileEnvironment.Definitions.Add("PLATFORM_CAN_SUPPORT_EDITORONLY_DATA=0");
