Workaround bug in Lumin clang toolchain

diff --git a/Engine/Source/Programs/UnrealBuildTool/Platform/Lumin/LuminToolChain.cs b/Engine/Source/Programs/UnrealBuildTool/Platform/Lumin/LuminToolChain.cs
--- a/Engine/Source/Programs/UnrealBuildTool/Platform/Lumin/LuminToolChain.cs
+++ b/Engine/Source/Programs/UnrealBuildTool/Platform/Lumin/LuminToolChain.cs
@@ -361,7 +361,13 @@ namespace UnrealBuildTool
 			{
 				QuotedCommandPath = "'" + CommandPath + "'";
 			}
-			if (BuildHostPlatform.Current.ShellType != ShellType.Cmd)
+			if (BuildHostPlatform.Current.ShellType == ShellType.Cmd)
+			{
+				// Workaround for UE-91990: The first argument is sensitive to whitespaces even when
+				// properly quoted, so make it something trivial as a workaround.
+				CompileOrLinkAction.CommandArguments = String.Format("/c \"{0} -DLUMIN_CLANG_CMD_WORKAROUND=1 {1}\"", QuotedCommandPath, CommandArguments);
+			}
+			else
 			{
 				// When quoting the command for the shell, we also need to escape any quotes. Otherwise
 				// they inadvertently unquote the interior arguments. For example it would unquote the
