Prevent writing to engine directory in un-promoted installed build, causing errors when creating modules.

diff --git a/Engine/Source/Programs/UnrealBuildTool/Configuration/UEBuildTarget.cs b/Engine/Source/Programs/UnrealBuildTool/Configuration/UEBuildTarget.cs
--- a/Engine/Source/Programs/UnrealBuildTool/Configuration/UEBuildTarget.cs
+++ b/Engine/Source/Programs/UnrealBuildTool/Configuration/UEBuildTarget.cs
@@ -1770,12 +1770,16 @@ namespace UnrealBuildTool
 				{
 					if(!Rules.bFormalBuild)
 					{
-						CppCompileEnvironment DefaultResourceCompileEnvironment = new CppCompileEnvironment(GlobalCompileEnvironment);
+						FileReference DefaultResourceLocation = FileReference.Combine(UnrealBuildTool.EngineDirectory, "Build", "Windows", "Resources", "Default.rc2");
+						if (!UnrealBuildTool.IsFileInstalled(DefaultResourceLocation))
+						{
+							CppCompileEnvironment DefaultResourceCompileEnvironment = new CppCompileEnvironment(GlobalCompileEnvironment);
 
-						FileItem DefaultResourceFile = FileItem.GetItemByFileReference(FileReference.Combine(UnrealBuildTool.EngineDirectory, "Build", "Windows", "Resources", "Default.rc2"));
+							FileItem DefaultResourceFile = FileItem.GetItemByFileReference(DefaultResourceLocation);
 
-						CPPOutput DefaultResourceOutput = TargetToolChain.CompileRCFiles(DefaultResourceCompileEnvironment, new List<FileItem> { DefaultResourceFile }, EngineIntermediateDirectory, Makefile);
-						GlobalLinkEnvironment.DefaultResourceFiles.AddRange(DefaultResourceOutput.ObjectFiles);
+							CPPOutput DefaultResourceOutput = TargetToolChain.CompileRCFiles(DefaultResourceCompileEnvironment, new List<FileItem> { DefaultResourceFile }, EngineIntermediateDirectory, Makefile);
+							GlobalLinkEnvironment.DefaultResourceFiles.AddRange(DefaultResourceOutput.ObjectFiles);
+						}
 					}
 				}
 			}
