Fix Windows Remote Build iOS for Storyboards
Copy correct LaunchScreenIOS on tvOS

diff --git a/Engine/Source/Programs/UnrealBuildTool/Platform/IOS/UEBuildIOS.cs b/Engine/Source/Programs/UnrealBuildTool/Platform/IOS/UEBuildIOS.cs
--- a/Engine/Source/Programs/UnrealBuildTool/Platform/IOS/UEBuildIOS.cs
+++ b/Engine/Source/Programs/UnrealBuildTool/Platform/IOS/UEBuildIOS.cs
@@ -956,6 +956,20 @@ namespace UnrealBuildTool
 			return false;
 		}
 
+		public bool HasCustomLaunchScreen(DirectoryReference ProjectDirectoryName)
+		{
+			string LaunchScreenDir = Path.Combine(ProjectDirectoryName.FullName, "Build", "IOS", "Resources", "Graphics");
+			if (Directory.Exists(LaunchScreenDir))
+			{
+				if (File.Exists(LaunchScreenDir + "\\LaunchScreenIOS.png"))
+				{
+					Log.TraceInformation("Requiring custom build because project {0} has a custom LaunchScreen", Path.GetFileName(ProjectDirectoryName.FullName));
+					return true;
+				}
+			}
+			return false;
+		}
+
 		/// <summary>
 		/// Check for the default configuration
 		/// return true if the project uses the default build config
@@ -975,7 +989,7 @@ namespace UnrealBuildTool
 			};
 
 			// check for custom icons
-			if (HasCustomIcons(ProjectDirectoryName))
+			if (HasCustomIcons(ProjectDirectoryName) || HasCustomLaunchScreen(ProjectDirectoryName))
 			{
 				return false;
 			}
@@ -1274,5 +1288,4 @@ namespace UnrealBuildTool
 			UEBuildPlatform.RegisterPlatformWithGroup(UnrealTargetPlatform.IOS, UnrealPlatformGroup.IOS);
 		}
 	}
-}
-
+}
\ No newline at end of file
diff --git a/Engine/Source/Programs/UnrealBuildTool/Platform/IOS/UEDeployIOS.cs b/Engine/Source/Programs/UnrealBuildTool/Platform/IOS/UEDeployIOS.cs
--- a/Engine/Source/Programs/UnrealBuildTool/Platform/IOS/UEDeployIOS.cs
+++ b/Engine/Source/Programs/UnrealBuildTool/Platform/IOS/UEDeployIOS.cs
@@ -1046,7 +1046,15 @@ namespace UnrealBuildTool
 			// copy the GameName binary
 			File.Copy(BinaryPath + "/" + GameExeName, AppDirectory + "/" + GameName, true);
 
-			CopyLaunchScreenResources(InEngineDir, AppDirectory, BuildDirectory);
+			if (SubDir == GetTargetPlatformName())
+			{
+				string BuildDirectoryFortvOS = InProjectDirectory + "/Build/IOS";
+				CopyLaunchScreenResources(InEngineDir, AppDirectory, BuildDirectoryFortvOS);
+			}
+			else
+			{
+				CopyLaunchScreenResources(InEngineDir, AppDirectory, BuildDirectory);
+			}
 
 			if (!bCreateStubIPA)
 			{
