Adding toolchain test builds.

#rnx

diff --git a/Engine/Source/Programs/UnrealBuildTool/Platform/Windows/UEBuildWindows.cs b/Engine/Source/Programs/UnrealBuildTool/Platform/Windows/UEBuildWindows.cs
--- a/Engine/Source/Programs/UnrealBuildTool/Platform/Windows/UEBuildWindows.cs
+++ b/Engine/Source/Programs/UnrealBuildTool/Platform/Windows/UEBuildWindows.cs
@@ -135,6 +135,7 @@ namespace UnrealBuildTool
 		/// WindowsPlatform.DefaultToolChainVersion (otherwise, we use the latest version).
 		/// </summary>
 		[XmlConfigFile(Category = "WindowsPlatform")]
+		[CommandLine("-CompilerVersion")]
 		public string CompilerVersion = null;
 
 		/// <summary>
