Fix CompileEnvironment being mutated during a call to VCToolChain.CompileFiles().

diff --git a/Engine/Source/Programs/UnrealBuildTool/Platform/Windows/VCToolChain.cs b/Engine/Source/Programs/UnrealBuildTool/Platform/Windows/VCToolChain.cs
--- a/Engine/Source/Programs/UnrealBuildTool/Platform/Windows/VCToolChain.cs
+++ b/Engine/Source/Programs/UnrealBuildTool/Platform/Windows/VCToolChain.cs
@@ -381,7 +381,7 @@ namespace UnrealBuildTool
 			else
 			{
 				// This is required to disable exception handling in VC platform headers.
-				CompileEnvironment.Definitions.Add("_HAS_EXCEPTIONS=0");
+				Arguments.Add("/D_HAS_EXCEPTIONS=0");
 			}
 
 			// If enabled, create debug information.
