Give a warning if a non-source version of engine tries to compile for x64 on Android

#ue4
#android

diff --git a/Engine/Source/Programs/UnrealBuildTool/Platform/Android/AndroidToolChain.cs b/Engine/Source/Programs/UnrealBuildTool/Platform/Android/AndroidToolChain.cs
--- a/Engine/Source/Programs/UnrealBuildTool/Platform/Android/AndroidToolChain.cs
+++ b/Engine/Source/Programs/UnrealBuildTool/Platform/Android/AndroidToolChain.cs
@@ -365,12 +365,26 @@ namespace UnrealBuildTool
 			if (Ini.GetBool("/Script/AndroidRuntimeSettings.AndroidRuntimeSettings", "bBuildForx86", out bBuild) && bBuild
 				|| (AdditionalArches != null && (AdditionalArches.Contains("x86", StringComparer.OrdinalIgnoreCase) || AdditionalArches.Contains("-x86", StringComparer.OrdinalIgnoreCase))))
 			{
-				Arches.Add("-x86");
+				if (File.Exists(Path.Combine(UnrealBuildTool.EngineDirectory.FullName, "Build", "InstalledBuild.txt")))
+				{
+					Log.TraceWarningOnce("Please install source to build for x86 (-x86); ignoring this architecture target.");
+				}
+				else
+				{
+					Arches.Add("-x86");
+				}
 			}
 			if (Ini.GetBool("/Script/AndroidRuntimeSettings.AndroidRuntimeSettings", "bBuildForx8664", out bBuild) && bBuild
 				|| (AdditionalArches != null && (AdditionalArches.Contains("x64", StringComparer.OrdinalIgnoreCase) || AdditionalArches.Contains("-x64", StringComparer.OrdinalIgnoreCase))))
 			{
-				Arches.Add("-x64");
+				if (File.Exists(Path.Combine(UnrealBuildTool.EngineDirectory.FullName, "Build", "InstalledBuild.txt")))
+				{
+					Log.TraceWarningOnce("Please install source to build for x86_64 (-x64); ignoring this architecture target.");
+				}
+				else
+				{
+					Arches.Add("-x64");
+				}
 			}
 
 			// force armv7 if something went wrong
