Remove redundant compiler flags on x86_64 architectures

#ue4
#android

diff --git a/Engine/Source/Programs/UnrealBuildTool/Platform/Android/AndroidToolChain.cs b/Engine/Source/Programs/UnrealBuildTool/Platform/Android/AndroidToolChain.cs
--- a/Engine/Source/Programs/UnrealBuildTool/Platform/Android/AndroidToolChain.cs
+++ b/Engine/Source/Programs/UnrealBuildTool/Platform/Android/AndroidToolChain.cs
@@ -774,7 +774,6 @@ namespace UnrealBuildTool
 			{
 				Result += " -funwind-tables";           // Just generates any needed static data, affects no code
 				Result += " -fstack-protector-strong";  // Emits extra code to check for buffer overflows
-				Result += " -fstrict-aliasing";
 				Result += " -fPIC";                     // Generates position-independent code (PIC) suitable for use in a shared library
 				Result += " -fno-omit-frame-pointer";
 				Result += " -fno-strict-aliasing";
@@ -793,7 +792,6 @@ namespace UnrealBuildTool
 			{
 				Result += " -funwind-tables";           // Just generates any needed static data, affects no code
 				Result += " -fstack-protector-strong";  // Emits extra code to check for buffer overflows
-				Result += " -fstrict-aliasing";
 				Result += " -fPIC";                     // Generates position-independent code (PIC) suitable for use in a shared library
 				Result += " -fno-omit-frame-pointer";
 				Result += " -fno-strict-aliasing";
