Error if x64 or x86 is only platform and in binarybuild

diff --git a/Engine/Source/Programs/UnrealBuildTool/Platform/Android/AndroidToolChain.cs b/Engine/Source/Programs/UnrealBuildTool/Platform/Android/AndroidToolChain.cs
--- a/Engine/Source/Programs/UnrealBuildTool/Platform/Android/AndroidToolChain.cs
+++ b/Engine/Source/Programs/UnrealBuildTool/Platform/Android/AndroidToolChain.cs
@@ -352,6 +352,8 @@ namespace UnrealBuildTool
 			ConfigHierarchy Ini = ConfigCache.ReadHierarchy(ConfigHierarchyType.Engine, DirectoryReference.FromFile(ProjectFile), UnrealTargetPlatform.Android);
 			Arches = new List<string>();
 			bool bBuild = true;
+			bool bUnsupportedBinaryBuildArch = false;
+
 			if (Ini.GetBool("/Script/AndroidRuntimeSettings.AndroidRuntimeSettings", "bBuildForArmV7", out bBuild) && bBuild
 				|| (AdditionalArches != null && (AdditionalArches.Contains("armv7", StringComparer.OrdinalIgnoreCase) || AdditionalArches.Contains("-armv7", StringComparer.OrdinalIgnoreCase))))
 			{
@@ -367,6 +369,7 @@ namespace UnrealBuildTool
 			{
 				if (File.Exists(Path.Combine(UnrealBuildTool.EngineDirectory.FullName, "Build", "InstalledBuild.txt")))
 				{
+					bUnsupportedBinaryBuildArch = true;
 					Log.TraceWarningOnce("Please install source to build for x86 (-x86); ignoring this architecture target.");
 				}
 				else
@@ -379,6 +382,7 @@ namespace UnrealBuildTool
 			{
 				if (File.Exists(Path.Combine(UnrealBuildTool.EngineDirectory.FullName, "Build", "InstalledBuild.txt")))
 				{
+					bUnsupportedBinaryBuildArch = true;
 					Log.TraceWarningOnce("Please install source to build for x86_64 (-x64); ignoring this architecture target.");
 				}
 				else
@@ -390,7 +394,14 @@ namespace UnrealBuildTool
 			// force armv7 if something went wrong
 			if (Arches.Count == 0)
 			{
-				Arches.Add("-armv7");
+				if (bUnsupportedBinaryBuildArch)
+				{
+					throw new BuildException("Only architectures unsupported by binary-only engine selected.");
+				}
+				else
+				{
+					Arches.Add("-armv7");
+				}
 			}
 
 			// For android just set the GPUArchitecture to an empty string
