Fix incremental builds triggering UHT to be rebuilt if a script plugin is enabled.

diff --git a/Engine/Source/Programs/UnrealBuildTool/Configuration/TargetRules.cs b/Engine/Source/Programs/UnrealBuildTool/Configuration/TargetRules.cs
--- a/Engine/Source/Programs/UnrealBuildTool/Configuration/TargetRules.cs
+++ b/Engine/Source/Programs/UnrealBuildTool/Configuration/TargetRules.cs
@@ -1524,6 +1524,15 @@ namespace UnrealBuildTool
 				}
 			}
 
+			// Get the directory to use for crypto settings. We can build engine targets (eg. UHT) with 
+			// a project file, but we can't use that to determine crypto settings without triggering
+			// constant rebuilds of UHT.
+			DirectoryReference CryptoSettingsDir = DirectoryReference.FromFile(ProjectFile);
+			if (CryptoSettingsDir != null && !File.IsUnderDirectory(CryptoSettingsDir))
+			{
+				CryptoSettingsDir = null;
+			}
+
 			// Setup macros for signing and encryption keys
 			EncryptionAndSigning.CryptoSettings CryptoSettings = EncryptionAndSigning.ParseCryptoSettings(DirectoryReference.FromFile(ProjectFile), Platform);
 			if (CryptoSettings.IsAnyEncryptionEnabled())
