Fix exception when packaging projects, due to UAT not setting TargetRules.File.

diff --git a/Engine/Source/Programs/UnrealBuildTool/Configuration/TargetRules.cs b/Engine/Source/Programs/UnrealBuildTool/Configuration/TargetRules.cs
--- a/Engine/Source/Programs/UnrealBuildTool/Configuration/TargetRules.cs
+++ b/Engine/Source/Programs/UnrealBuildTool/Configuration/TargetRules.cs
@@ -1528,13 +1528,13 @@ namespace UnrealBuildTool
 			// a project file, but we can't use that to determine crypto settings without triggering
 			// constant rebuilds of UHT.
 			DirectoryReference CryptoSettingsDir = DirectoryReference.FromFile(ProjectFile);
-			if (CryptoSettingsDir != null && !File.IsUnderDirectory(CryptoSettingsDir))
+			if (CryptoSettingsDir != null && File != null && !File.IsUnderDirectory(CryptoSettingsDir))
 			{
 				CryptoSettingsDir = null;
 			}
 
 			// Setup macros for signing and encryption keys
-			EncryptionAndSigning.CryptoSettings CryptoSettings = EncryptionAndSigning.ParseCryptoSettings(DirectoryReference.FromFile(ProjectFile), Platform);
+			EncryptionAndSigning.CryptoSettings CryptoSettings = EncryptionAndSigning.ParseCryptoSettings(CryptoSettingsDir, Platform);
 			if (CryptoSettings.IsAnyEncryptionEnabled())
 			{
 				ProjectDefinitions.Add(String.Format("IMPLEMENT_ENCRYPTION_KEY_REGISTRATION()=UE_REGISTER_ENCRYPTION_KEY({0})", FormatHexBytes(CryptoSettings.EncryptionKey.Key)));
