Use relative paths where possible when linking on Mac. This solves the "Argument list too long" error during linking

diff --git a/Engine/Source/Programs/UnrealBuildTool/Platform/Mac/MacToolChain.cs b/Engine/Source/Programs/UnrealBuildTool/Platform/Mac/MacToolChain.cs
--- a/Engine/Source/Programs/UnrealBuildTool/Platform/Mac/MacToolChain.cs
+++ b/Engine/Source/Programs/UnrealBuildTool/Platform/Mac/MacToolChain.cs
@@ -816,7 +816,13 @@ namespace UnrealBuildTool
 			List<string> InputFileNames = new List<string>();
 			foreach (FileItem InputFile in LinkEnvironment.InputFiles)
 			{
-				InputFileNames.Add(string.Format("\"{0}\"", InputFile.AbsolutePath));
+				string InputFilePath = InputFile.AbsolutePath;
+				if (InputFile.Location.IsUnderDirectory(UnrealBuildTool.RootDirectory))
+				{
+					InputFilePath = InputFile.Location.MakeRelativeTo(UnrealBuildTool.EngineSourceDirectory);
+				}
+
+				InputFileNames.Add(string.Format("\"{0}\"", InputFilePath));
 				LinkAction.PrerequisiteItems.Add(InputFile);
 			}
 
