Allow Shared PCHs in non-engine modules

diff --git a/Engine/Source/Programs/UnrealBuildTool/Configuration/UEBuildTarget.cs b/Engine/Source/Programs/UnrealBuildTool/Configuration/UEBuildTarget.cs
--- a/Engine/Source/Programs/UnrealBuildTool/Configuration/UEBuildTarget.cs
+++ b/Engine/Source/Programs/UnrealBuildTool/Configuration/UEBuildTarget.cs
@@ -2653,13 +2653,6 @@ namespace UnrealBuildTool
 				}
 			}
 
-			// Shared PCHs are only supported for engine modules at the moment. Check there are no game modules in the list.
-			List<UEBuildModuleCPP> NonEngineSharedPCHs = SharedPCHModules.Where(x => !x.Rules.Context.bCanUseForSharedPCH).ToList();
-			if(NonEngineSharedPCHs.Count > 0)
-			{
-				throw new BuildException("Shared PCHs are only supported for engine modules (found {0}).", String.Join(", ", NonEngineSharedPCHs.Select(x => x.Name)));
-			}
-
 			// Find a priority for each shared PCH, determined as the number of other shared PCHs it includes.
 			Dictionary<UEBuildModuleCPP, int> SharedPCHModuleToPriority = new Dictionary<UEBuildModuleCPP, int>();
 			foreach(UEBuildModuleCPP SharedPCHModule in SharedPCHModules)
