- XboxOne Build/Resources files moving to PlatformExtensions
#rnx

diff --git a/Engine/Source/Programs/UnrealBuildTool/Platform/Windows/VCManifestGenerator.cs b/Engine/Source/Programs/UnrealBuildTool/Platform/Windows/VCManifestGenerator.cs
--- a/Engine/Source/Programs/UnrealBuildTool/Platform/Windows/VCManifestGenerator.cs
+++ b/Engine/Source/Programs/UnrealBuildTool/Platform/Windows/VCManifestGenerator.cs
@@ -20,7 +20,6 @@ namespace UnrealBuildTool
 
 		protected virtual string IniSection_PlatformTargetSettings { get { return string.Format( "/Script/{0}PlatformEditor.{0}TargetSettings", Platform.ToString() ); } }
 		protected virtual string IniSection_GeneralProjectSettings { get { return "/Script/EngineSettings.GeneralProjectSettings"; } }
-		protected virtual string BuildResourceProjectRelativePath { get { return "Build\\" + Platform.ToString(); } }
 
 		protected const string BuildResourceSubPath = "Resources";
 		protected const string EngineResourceSubPath = "DefaultImages";
@@ -278,23 +277,29 @@ namespace UnrealBuildTool
         protected bool CopyAndReplaceBinaryIntermediate(string ResourceFileName, bool AllowEngineFallback = true)
 		{
 			string TargetPath = Path.Combine(IntermediatePath, BuildResourceSubPath);
-			string SourcePath = Path.Combine(ProjectPath, BuildResourceProjectRelativePath, BuildResourceSubPath);
 
-			// Try falling back to the engine defaults if requested
+			// look in project normal Build location
+			string SourcePath = Path.Combine(ProjectPath, "Build", Platform.ToString(), BuildResourceSubPath);
 			bool bFileExists = File.Exists(Path.Combine(SourcePath, ResourceFileName));
+
+			// look in Platform Extensions next
 			if (!bFileExists)
 			{
-				if (AllowEngineFallback)
+				SourcePath = Path.Combine(ProjectPath, "Platforms", Platform.ToString(), "Build", BuildResourceSubPath);
+				bFileExists = File.Exists(Path.Combine(SourcePath, ResourceFileName));
+			}
+
+			// look in Engine, if allowed
+			if (!bFileExists && AllowEngineFallback)
+			{
+				SourcePath = Path.Combine(UnrealBuildTool.EngineDirectory.FullName, "Build", Platform.ToString(), EngineResourceSubPath);
+				bFileExists = File.Exists(Path.Combine(SourcePath, ResourceFileName));
+
+				// look in Platform extensions too
+				if (!bFileExists)
 				{
-					SourcePath = Path.Combine(UnrealBuildTool.EngineDirectory.FullName, BuildResourceProjectRelativePath, EngineResourceSubPath);
+					SourcePath = Path.Combine(UnrealBuildTool.EngineDirectory.FullName, "Platforms", Platform.ToString(), "Build", EngineResourceSubPath);
 					bFileExists = File.Exists(Path.Combine(SourcePath, ResourceFileName));
-
-					// look in Platform extensions too
-					if (!bFileExists)
-					{
-						SourcePath = Path.Combine(UnrealBuildTool.EngineDirectory.FullName, "Platforms", Platform.ToString(), "Build", EngineResourceSubPath);
-						bFileExists = File.Exists(Path.Combine(SourcePath, ResourceFileName));
-					}
 				}
 			}
 
