UBT: When generating visual studio project files, projects with multiple editor targets will use the default editor target config setting, if it exists, to determine which editor target should be added to the main project file. All others will generate separate projects inside the VS solution.

diff --git a/Engine/Source/Programs/UnrealBuildTool/ProjectFiles/ProjectFileGenerator.cs b/Engine/Source/Programs/UnrealBuildTool/ProjectFiles/ProjectFileGenerator.cs
--- a/Engine/Source/Programs/UnrealBuildTool/ProjectFiles/ProjectFileGenerator.cs
+++ b/Engine/Source/Programs/UnrealBuildTool/ProjectFiles/ProjectFileGenerator.cs
@@ -2226,6 +2226,9 @@ namespace UnrealBuildTool
 			//DirectoryReference EngineSourceProgramsDirectory = DirectoryReference.Combine(UnrealBuildTool.EngineSourceDirectory, "Programs");
 			DirectoryReference EnterpriseSourceProgramsDirectory = DirectoryReference.Combine(UnrealBuildTool.EnterpriseSourceDirectory, "Programs");
 
+			// Keep a cache of the default editor target for each project so that we don't have to interrogate the configs for each target
+			Dictionary<string, string> DefaultProjectEditorTargetCache = new Dictionary<string, string>();
+
 			foreach( FileReference TargetFilePath in AllTargetFiles )
 			{
 				string TargetName = TargetFilePath.GetFileNameWithoutAnyExtensions();		// Remove both ".cs" and ".Target"
@@ -2317,12 +2320,44 @@ namespace UnrealBuildTool
 						ProjectFileNameBase = ProjectInfo.GetFileNameWithoutExtension();
 					}
 
+					// Look at the project engine config to see if it has specified a default editor target
+					FileReference ProjectLocation = GetProjectLocation(ProjectFileNameBase);
+					string DefaultEditorTarget = null;
+					if (!DefaultProjectEditorTargetCache.TryGetValue(ProjectFileNameBase, out DefaultEditorTarget))
+					{
+						if (GameFolder != null)
+						{
+							FileReference DefaultEngineIniFilename = FileReference.Combine(GameFolder, "Config", "DefaultEngine.ini");
+							ConfigFile ProjectDefaultEngineIni;
+							ConfigCache.TryReadFile(DefaultEngineIniFilename, out ProjectDefaultEngineIni);
+							if (ProjectDefaultEngineIni != null)
+							{
+								ConfigFileSection Section;
+								if (ProjectDefaultEngineIni.TryGetSection("/Script/BuildSettings.BuildSettings", out Section))
+								{
+									ConfigLine Line;
+									if (Section.TryGetLine("DefaultEditorTarget", out Line))
+									{
+										DefaultEditorTarget = Line.Value;
+									}
+								}
+							}
+						}
+
+						DefaultProjectEditorTargetCache.Add(ProjectFileNameBase, DefaultEditorTarget);
+					}
+
 					// Get the suffix to use for this project file. If we have multiple targets of the same type, we'll have to split them out into separate IDE project files.
 					string GeneratedProjectName = TargetRulesObject.GeneratedProjectName;
 					if(GeneratedProjectName == null)
 					{
 						ProjectFile ExistingProjectFile;
-						if( ProjectFileMap.TryGetValue( GetProjectLocation(ProjectFileNameBase), out ExistingProjectFile ) && ExistingProjectFile.ProjectTargets.Any(x => x.TargetRules.Type == TargetRulesObject.Type))
+						// We should create a separate project for targets which aren't the default for this target type.
+						// Note that we currently only support changing the default editor target and not any other types, so
+						// if we aren't an editor, we'll just fall back to previous behavior and assume the first one we encounter
+						// should be added to the main project
+						bool bIsDefaultTargetForType = (TargetRulesObject.Type != TargetType.Editor) || TargetRulesObject.Name == DefaultEditorTarget;
+						if( !bIsDefaultTargetForType || (ProjectFileMap.TryGetValue(ProjectLocation, out ExistingProjectFile ) && ExistingProjectFile.ProjectTargets.Any(x => x.TargetRules.Type == TargetRulesObject.Type)))
 						{
 							GeneratedProjectName = TargetRulesObject.Name;
 						}
