ParallelExecutor scheduling tweak to make ispc jobs execute early

ispc jobs tend to run long so getting them out of the way quickly helps prevents ispc from making the build longer than necessary

The logic is pretty crude (it checks to see if the StatusDecription ends with .ispc) but works for now -- perhaps we can implement a cleaner prioritization mechanism later

diff --git a/Engine/Source/Programs/UnrealBuildTool/Executors/ParallelExecutor.cs b/Engine/Source/Programs/UnrealBuildTool/Executors/ParallelExecutor.cs
--- a/Engine/Source/Programs/UnrealBuildTool/Executors/ParallelExecutor.cs
+++ b/Engine/Source/Programs/UnrealBuildTool/Executors/ParallelExecutor.cs
@@ -107,6 +107,12 @@ namespace UnrealBuildTool
 				BuildAction Action = new BuildAction();
 				Action.SortIndex = Idx;
 				Action.Inner = InputActions[Idx];
+
+				if (!Action.Inner.StatusDescription.EndsWith(".ispc"))
+				{
+					Action.SortIndex += 10000;
+				}
+
 				Actions.Add(Action);
 			}
 
