WhichClang should check if our toolchain is being used
Helps with the generation of CMake project files which needs a valid clang compiler on Linux

diff --git a/Engine/Source/Programs/UnrealBuildTool/Platform/Linux/LinuxCommon.cs b/Engine/Source/Programs/UnrealBuildTool/Platform/Linux/LinuxCommon.cs
--- a/Engine/Source/Programs/UnrealBuildTool/Platform/Linux/LinuxCommon.cs
+++ b/Engine/Source/Programs/UnrealBuildTool/Platform/Linux/LinuxCommon.cs
@@ -38,6 +38,12 @@ namespace UnrealBuildTool
 
 		public static string WhichClang()
 		{
+			string InternalSDKPath = LinuxPlatformSDK.GetInternalSDKPath();
+			if (!String.IsNullOrEmpty(InternalSDKPath))
+			{
+				return Path.Combine(InternalSDKPath, "bin", "clang++");
+			}
+
 			string[] ClangNames = { "clang++", "clang++-7.0", "clang++-6.0" };
 			string ClangPath;
 			foreach (string ClangName in ClangNames)
@@ -48,6 +54,7 @@ namespace UnrealBuildTool
 					return ClangPath;
 				}
 			}
+
 			return null;
 		}
 
diff --git a/Engine/Source/Programs/UnrealBuildTool/Platform/Linux/UEBuildLinux.cs b/Engine/Source/Programs/UnrealBuildTool/Platform/Linux/UEBuildLinux.cs
--- a/Engine/Source/Programs/UnrealBuildTool/Platform/Linux/UEBuildLinux.cs
+++ b/Engine/Source/Programs/UnrealBuildTool/Platform/Linux/UEBuildLinux.cs
@@ -678,6 +678,31 @@ namespace UnrealBuildTool
 			return TargetPlatformName;
 		}
 
+		/// <summary>
+		/// Returns a path to the internal SDK
+		/// </summary>
+		/// <returns>Valid path to the internal SDK, null otherwise</returns>
+		static public string GetInternalSDKPath()
+		{
+			string SDKRoot = Environment.GetEnvironmentVariable(SDKRootEnvVar);
+			if (!String.IsNullOrEmpty(SDKRoot))
+			{
+				string AutoSDKPath = Path.Combine(SDKRoot, "Host" + BuildHostPlatform.Current.Platform, TargetPlatformName, ExpectedSDKVersion, LinuxPlatform.DefaultHostArchitecture);
+				if (DirectoryReference.Exists(new DirectoryReference(AutoSDKPath)))
+				{
+					return AutoSDKPath;
+				}
+			}
+
+			string InTreeSDKPath = Path.Combine(LinuxPlatformSDK.GetInTreeSDKRoot().FullName, ExpectedSDKVersion, LinuxPlatform.DefaultHostArchitecture);
+			if (DirectoryReference.Exists(new DirectoryReference(InTreeSDKPath)))
+			{
+				return InTreeSDKPath;
+			}
+
+			return null;
+		}
+
 		/// <summary>
 		/// Returns SDK string as required by the platform
 		/// </summary>
