Add an option to configure the maximum length of shared include paths to use when generating Visual Studio project files. Despite improving Visual Studio stability for some users, it seems to be causing VAX crashes for others.

To disable the use of shared include paths (at the expensive of Visual Studio using more memory), put this in BuildConfiguration.xml:

<?xml version="1.0" encoding="utf-8" ?>
<Configuration xmlns="https://www.unrealengine.com/BuildConfiguration">
	<VCProjectFileGenerator>
		<MaxSharedIncludePaths>0</MaxSharedIncludePaths>
	</VCProjectFileGenerator>
</Configuration>

The setting specifies the maximum length of the shared include paths to use, in bytes. The default value is 24,576 (24k).

diff --git a/Engine/Source/Programs/UnrealBuildTool/ProjectFiles/VisualStudio/VCProject.cs b/Engine/Source/Programs/UnrealBuildTool/ProjectFiles/VisualStudio/VCProject.cs
--- a/Engine/Source/Programs/UnrealBuildTool/ProjectFiles/VisualStudio/VCProject.cs
+++ b/Engine/Source/Programs/UnrealBuildTool/ProjectFiles/VisualStudio/VCProject.cs
@@ -255,13 +255,10 @@ namespace UnrealBuildTool
 	{
 		//FileReference OnlyGameProject;
 		VCProjectFileFormat ProjectFileFormat;
-		bool bUseFastPDB;
-		bool bUsePerFileIntellisense;
 		bool bUsePrecompiled;
-		bool bEditorDependsOnShaderCompileWorker;
-		bool bBuildLiveCodingConsole;
 		string BuildToolOverride;
 		Dictionary<DirectoryReference, string> ModuleDirToForceIncludePaths = new Dictionary<DirectoryReference, string>();
+		VCProjectFileSettings Settings;
 
 		/// This is the platform name that Visual Studio is always guaranteed to support.  We'll use this as
 		/// a platform for any project configurations where our actual platform is not supported by the
@@ -277,26 +274,18 @@ namespace UnrealBuildTool
 		/// <summary>
 		/// Constructs a new project file object
 		/// </summary>
-		/// <param name="InFilePath">The path to the project file on disk</param>
-		/// <param name="InOnlyGameProject"></param>
-		/// <param name="InProjectFileFormat">Visual C++ project file version</param>
-		/// <param name="bUseFastPDB">If true, adds the -FastPDB argument to build command lines</param>
-		/// <param name="bUsePerFileIntellisense">If true, generates per-file intellisense data</param>
+		/// <param name="FilePath">The path to the project file on disk</param>
+		/// <param name="ProjectFileFormat">Visual C++ project file version</param>
 		/// <param name="bUsePrecompiled">Whether to add the -UsePrecompiled argumemnt when building targets</param>
-		/// <param name="bEditorDependsOnShaderCompileWorker">Whether editor targets should also build ShaderCompileWorker</param>
-		/// <param name="bBuildLiveCodingConsole">Whether targets using live coding should also build LiveCodingConsole</param>
 		/// <param name="BuildToolOverride">Optional arguments to pass to UBT when building</param>
-		public VCProjectFile(FileReference InFilePath, FileReference InOnlyGameProject, VCProjectFileFormat InProjectFileFormat, bool bUseFastPDB, bool bUsePerFileIntellisense, bool bUsePrecompiled, bool bEditorDependsOnShaderCompileWorker, bool bBuildLiveCodingConsole, string BuildToolOverride)
-			: base(InFilePath)
+		/// <param name="Settings">Other settings</param>
+		public VCProjectFile(FileReference FilePath, VCProjectFileFormat ProjectFileFormat, bool bUsePrecompiled, string BuildToolOverride, VCProjectFileSettings Settings)
+			: base(FilePath)
 		{
-			//OnlyGameProject = InOnlyGameProject;
-			ProjectFileFormat = InProjectFileFormat;
-			this.bUseFastPDB = bUseFastPDB;
-			this.bUsePerFileIntellisense = bUsePerFileIntellisense;
+			this.ProjectFileFormat = ProjectFileFormat;
 			this.bUsePrecompiled = bUsePrecompiled;
-			this.bEditorDependsOnShaderCompileWorker = bEditorDependsOnShaderCompileWorker;
-			this.bBuildLiveCodingConsole = bBuildLiveCodingConsole;
 			this.BuildToolOverride = BuildToolOverride;
+			this.Settings = Settings;
 		}
 
 		/// <summary>
@@ -516,7 +505,7 @@ namespace UnrealBuildTool
 		{
 			base.AddModule(Module, CompileEnvironment);
 
-			if (bUsePerFileIntellisense)
+			if (Settings.bUsePerFileIntellisense)
 			{
 				foreach (DirectoryReference ModuleDirectory in Module.ModuleDirectories)
 				{
@@ -793,19 +782,18 @@ namespace UnrealBuildTool
 				}
 
 				// Append the most common include paths to the search list.
-				const int MaxSharedIncludePathsLength = 24 * 1024;
-				foreach (DirectoryReference IncludePath in IncludePathToCount.OrderByDescending(x => x.Value).Select(x => x.Key))
+				if (Settings.MaxSharedIncludePaths > 0)
 				{
-					string RelativePath = NormalizeProjectPath(IncludePath);
-					if (SharedIncludeSearchPaths.Length + RelativePath.Length < MaxSharedIncludePathsLength)
+					foreach (DirectoryReference IncludePath in IncludePathToCount.OrderByDescending(x => x.Value).Select(x => x.Key))
 					{
+						string RelativePath = NormalizeProjectPath(IncludePath);
+						if (SharedIncludeSearchPaths.Length + RelativePath.Length >= Settings.MaxSharedIncludePaths)
+						{
+							break;
+						}
 						SharedIncludeSearchPaths.AppendFormat("{0};", RelativePath);
 						SharedIncludeSearchPathsSet.Add(IncludePath);
 					}
-					else
-					{
-						break;
-					}
 				}
 
 				// Add all the default system include paths
@@ -1576,11 +1564,11 @@ namespace UnrealBuildTool
 					List<string> ExtraTargets = new List<string>();
 					if (!bUsePrecompiled)
 					{
-						if (TargetRulesObject.Type == TargetType.Editor && bEditorDependsOnShaderCompileWorker && !UnrealBuildTool.IsEngineInstalled())
+						if (TargetRulesObject.Type == TargetType.Editor && Settings.bEditorDependsOnShaderCompileWorker && !UnrealBuildTool.IsEngineInstalled())
 						{
 							ExtraTargets.Add("ShaderCompileWorker Win64 Development");
 						}
-						if (TargetRulesObject.bWithLiveCoding && bBuildLiveCodingConsole && !UnrealBuildTool.IsEngineInstalled() && TargetRulesObject.Name != "LiveCodingConsole")
+						if (TargetRulesObject.bWithLiveCoding && Settings.bBuildLiveCodingConsole && !UnrealBuildTool.IsEngineInstalled() && TargetRulesObject.Name != "LiveCodingConsole")
 						{
 							ExtraTargets.Add(TargetRulesObject.bUseDebugLiveCodingConsole? "LiveCodingConsole Win64 Debug" : "LiveCodingConsole Win64 Development");
 						}
@@ -1608,7 +1596,7 @@ namespace UnrealBuildTool
 					// Always include a flag to format log messages for MSBuild
 					BuildArguments.Append(" -FromMsBuild");
 
-					if (bUseFastPDB)
+					if (Settings.bAddFastPDBToProjects)
 					{
 						// Pass Fast PDB option to make use of Visual Studio's /DEBUG:FASTLINK option
 						BuildArguments.Append(" -FastPDB");
diff --git a/Engine/Source/Programs/UnrealBuildTool/ProjectFiles/VisualStudio/VCProjectFileGenerator.cs b/Engine/Source/Programs/UnrealBuildTool/ProjectFiles/VisualStudio/VCProjectFileGenerator.cs
--- a/Engine/Source/Programs/UnrealBuildTool/ProjectFiles/VisualStudio/VCProjectFileGenerator.cs
+++ b/Engine/Source/Programs/UnrealBuildTool/ProjectFiles/VisualStudio/VCProjectFileGenerator.cs
@@ -35,52 +35,68 @@ namespace UnrealBuildTool
 		VisualStudio2019,
 	}
 
-	/// <summary>
-	/// Visual C++ project file generator implementation
-	/// </summary>
-	class VCProjectFileGenerator : ProjectFileGenerator
+	class VCProjectFileSettings
 	{
 		/// <summary>
 		/// The version of Visual Studio to generate project files for.
 		/// </summary>
 		[XmlConfigFile(Name = "Version")]
-		protected VCProjectFileFormat ProjectFileFormat = VCProjectFileFormat.Default;
+		public VCProjectFileFormat ProjectFileFormat = VCProjectFileFormat.Default;
+
+		/// <summary>
+		/// Puts the most common include paths in the IncludePath property in the MSBuild project. This significantly reduces Visual Studio
+		/// memory usage (measured 1.1GB -> 500mb), but seems to be causing issues with Visual Assist. Value here specifies maximum length
+		/// of the include path list in KB.
+		/// </summary>
+		[XmlConfigFile(Category = "VCProjectFileGenerator")]
+		public int MaxSharedIncludePaths = 24 * 1024;
 
 		/// <summary>
 		/// Whether to write a solution option (suo) file for the sln.
 		/// </summary>
 		[XmlConfigFile(Category = "BuildConfiguration")]
-		protected bool bWriteSolutionOptionFile = true;
+		public bool bWriteSolutionOptionFile = true;
 
 		/// <summary>
 		/// Forces UBT to be built in debug configuration, regardless of the solution configuration
 		/// </summary>
-		[XmlConfigFile]
-		protected bool bBuildUBTInDebug = false;
+		[XmlConfigFile(Category = "VCProjectFileGenerator")]
+		public bool bBuildUBTInDebug = false;
 
 		/// <summary>
 		/// Whether to add the -FastPDB option to build command lines by default.
 		/// </summary>
 		[XmlConfigFile(Category = "BuildConfiguration")]
-		bool bAddFastPDBToProjects = false;
+		public bool bAddFastPDBToProjects = false;
 
 		/// <summary>
 		/// Whether to generate per-file intellisense data.
 		/// </summary>
 		[XmlConfigFile(Category = "BuildConfiguration")]
-		bool bUsePerFileIntellisense = true;
+		public bool bUsePerFileIntellisense = true;
 
 		/// <summary>
 		/// Whether to include a dependency on ShaderCompileWorker when generating project files for the editor.
 		/// </summary>
 		[XmlConfigFile(Category = "BuildConfiguration")]
-		bool bEditorDependsOnShaderCompileWorker = true;
+		public bool bEditorDependsOnShaderCompileWorker = true;
 
 		/// <summary>
 		/// Whether to include a dependency on LiveCodingConsole when building targets that support live coding.
 		/// </summary>
-		[XmlConfigFile]
-		bool bBuildLiveCodingConsole = false;
+		[XmlConfigFile(Category = "VCProjectFileGenerator")]
+		public bool bBuildLiveCodingConsole = false;
+	}
+
+	/// <summary>
+	/// Visual C++ project file generator implementation
+	/// </summary>
+	class VCProjectFileGenerator : ProjectFileGenerator
+	{
+		/// <summary>
+		/// The settings object
+		/// </summary>
+		protected VCProjectFileSettings Settings = new VCProjectFileSettings();
 
 		/// <summary>
 		/// Override for the build tool to use in generated projects. If the compiler version is specified on the command line, we use the same argument on the 
@@ -97,11 +113,11 @@ namespace UnrealBuildTool
 		public VCProjectFileGenerator(FileReference InOnlyGameProject, VCProjectFileFormat InProjectFileFormat, CommandLineArguments InArguments)
 			: base(InOnlyGameProject)
 		{
-			XmlConfig.ApplyTo(this);
+			XmlConfig.ApplyTo(Settings);
 
 			if(InProjectFileFormat != VCProjectFileFormat.Default)
 			{
-				ProjectFileFormat = InProjectFileFormat;
+				Settings.ProjectFileFormat = InProjectFileFormat;
 			}
 
 			if(InArguments.HasOption("-2015"))
@@ -185,7 +201,7 @@ namespace UnrealBuildTool
 		/// <returns>The newly allocated project file object</returns>
 		protected override ProjectFile AllocateProjectFile(FileReference InitFilePath)
 		{
-			return new VCProjectFile(InitFilePath, OnlyGameProject, ProjectFileFormat, bAddFastPDBToProjects, bUsePerFileIntellisense, bUsePrecompiled, bEditorDependsOnShaderCompileWorker, bBuildLiveCodingConsole, BuildToolOverride);
+			return new VCProjectFile(InitFilePath, Settings.ProjectFileFormat, bUsePrecompiled, BuildToolOverride, Settings);
 		}
 
 
@@ -256,48 +272,48 @@ namespace UnrealBuildTool
 			base.SetupSupportedPlatformsAndConfigurations(IncludeAllPlatforms, out SupportedPlatformNames);
 
 			// If we have a non-default setting for visual studio, check the compiler exists. If not, revert to the default.
-			if(ProjectFileFormat == VCProjectFileFormat.VisualStudio2015)
+			if(Settings.ProjectFileFormat == VCProjectFileFormat.VisualStudio2015)
 			{
 				if (!WindowsPlatform.HasCompiler(WindowsCompiler.VisualStudio2015_DEPRECATED))
 				{
 					Log.TraceWarning("Visual Studio C++ 2015 installation not found - ignoring preferred project file format.");
-					ProjectFileFormat = VCProjectFileFormat.Default;
+					Settings.ProjectFileFormat = VCProjectFileFormat.Default;
 				}
 			}
-			else if(ProjectFileFormat == VCProjectFileFormat.VisualStudio2017)
+			else if(Settings.ProjectFileFormat == VCProjectFileFormat.VisualStudio2017)
 			{
 				if (!WindowsPlatform.HasCompiler(WindowsCompiler.VisualStudio2017))
 				{
 					Log.TraceWarning("Visual Studio C++ 2017 installation not found - ignoring preferred project file format.");
-					ProjectFileFormat = VCProjectFileFormat.Default;
+					Settings.ProjectFileFormat = VCProjectFileFormat.Default;
 				}
 			}
-			else if(ProjectFileFormat == VCProjectFileFormat.VisualStudio2019)
+			else if(Settings.ProjectFileFormat == VCProjectFileFormat.VisualStudio2019)
 			{
 				if (!WindowsPlatform.HasCompiler(WindowsCompiler.VisualStudio2019))
 				{
 					Log.TraceWarning("Visual Studio C++ 2019 installation not found - ignoring preferred project file format.");
-					ProjectFileFormat = VCProjectFileFormat.Default;
+					Settings.ProjectFileFormat = VCProjectFileFormat.Default;
 				}
 			}
 
 			// Certain platforms override the project file format because their debugger add-ins may not yet support the latest
 			// version of Visual Studio.  This is their chance to override that.
 			// ...but only if the user didn't override this via the command-line.
-			if (ProjectFileFormat == VCProjectFileFormat.Default)
+			if (Settings.ProjectFileFormat == VCProjectFileFormat.Default)
 			{
 				// Pick the best platform installed by default
 				if (WindowsPlatform.HasCompiler(WindowsCompiler.VisualStudio2019) && WindowsPlatform.HasIDE(WindowsCompiler.VisualStudio2019))
 				{
-					ProjectFileFormat = VCProjectFileFormat.VisualStudio2019;
+					Settings.ProjectFileFormat = VCProjectFileFormat.VisualStudio2019;
 				}
 				else if (WindowsPlatform.HasCompiler(WindowsCompiler.VisualStudio2017) && WindowsPlatform.HasIDE(WindowsCompiler.VisualStudio2017))
 				{
-					ProjectFileFormat = VCProjectFileFormat.VisualStudio2017;
+					Settings.ProjectFileFormat = VCProjectFileFormat.VisualStudio2017;
 				}
 				else if (WindowsPlatform.HasCompiler(WindowsCompiler.VisualStudio2015_DEPRECATED) && WindowsPlatform.HasIDE(WindowsCompiler.VisualStudio2015_DEPRECATED))
 				{
-					ProjectFileFormat = VCProjectFileFormat.VisualStudio2015;
+					Settings.ProjectFileFormat = VCProjectFileFormat.VisualStudio2015;
 				}
 
 				// Allow the SDKs to override
@@ -314,9 +330,9 @@ namespace UnrealBuildTool
 							if (ProposedFormat != VCProjectFileFormat.Default)
 							{
 								// Reduce the Visual Studio version to the max supported by each platform we plan to include.
-								if (ProjectFileFormat == VCProjectFileFormat.Default || ProposedFormat < ProjectFileFormat)
+								if (Settings.ProjectFileFormat == VCProjectFileFormat.Default || ProposedFormat < Settings.ProjectFileFormat)
 								{
-									ProjectFileFormat = ProposedFormat;
+									Settings.ProjectFileFormat = ProposedFormat;
 								}
 							}
 						}
@@ -429,7 +445,7 @@ namespace UnrealBuildTool
 				DirectoryReference AutomationToolDir = DirectoryReference.Combine(UnrealBuildTool.EngineSourceDirectory, "Programs", "AutomationTool");
 				new XDocument(
 					new XElement(NS + "Project",
-						new XAttribute("ToolsVersion", VCProjectFileGenerator.GetProjectFileToolVersionString(ProjectFileFormat)),
+						new XAttribute("ToolsVersion", VCProjectFileGenerator.GetProjectFileToolVersionString(Settings.ProjectFileFormat)),
 						new XAttribute("DefaultTargets", "Build"),
 						new XElement(NS + "ItemGroup",
 							from AutomationProject in AutomationProjectFiles
@@ -458,7 +474,7 @@ namespace UnrealBuildTool
 			StringBuilder VCSolutionFileContent = new StringBuilder();
 
 			// Solution file header. Note that a leading newline is required for file type detection to work correclty in the shell.
-			if (ProjectFileFormat == VCProjectFileFormat.VisualStudio2019)
+			if (Settings.ProjectFileFormat == VCProjectFileFormat.VisualStudio2019)
 			{
 				VCSolutionFileContent.AppendLine();
 				VCSolutionFileContent.AppendLine("Microsoft Visual Studio Solution File, Format Version 12.00");
@@ -466,7 +482,7 @@ namespace UnrealBuildTool
 				VCSolutionFileContent.AppendLine("VisualStudioVersion = 16.0.28315.86");
 				VCSolutionFileContent.AppendLine("MinimumVisualStudioVersion = 10.0.40219.1");
 			}
-			else if (ProjectFileFormat == VCProjectFileFormat.VisualStudio2017)
+			else if (Settings.ProjectFileFormat == VCProjectFileFormat.VisualStudio2017)
 			{
 				VCSolutionFileContent.AppendLine();
 				VCSolutionFileContent.AppendLine("Microsoft Visual Studio Solution File, Format Version 12.00");
@@ -474,7 +490,7 @@ namespace UnrealBuildTool
 				VCSolutionFileContent.AppendLine("VisualStudioVersion = 15.0.25807.0");
 				VCSolutionFileContent.AppendLine("MinimumVisualStudioVersion = 10.0.40219.1");
 			}
-			else if (ProjectFileFormat == VCProjectFileFormat.VisualStudio2015)
+			else if (Settings.ProjectFileFormat == VCProjectFileFormat.VisualStudio2015)
 			{
 				VCSolutionFileContent.AppendLine();
 				VCSolutionFileContent.AppendLine("Microsoft Visual Studio Solution File, Format Version 12.00");
@@ -482,13 +498,13 @@ namespace UnrealBuildTool
 				VCSolutionFileContent.AppendLine("VisualStudioVersion = 14.0.22310.1");
 				VCSolutionFileContent.AppendLine("MinimumVisualStudioVersion = 10.0.40219.1");
 			}
-			else if (ProjectFileFormat == VCProjectFileFormat.VisualStudio2013)
+			else if (Settings.ProjectFileFormat == VCProjectFileFormat.VisualStudio2013)
 			{
 				VCSolutionFileContent.AppendLine();
 				VCSolutionFileContent.AppendLine("Microsoft Visual Studio Solution File, Format Version 12.00");
 				VCSolutionFileContent.AppendLine("# Visual Studio 2013");
             }
-            else if (ProjectFileFormat == VCProjectFileFormat.VisualStudio2012)
+            else if (Settings.ProjectFileFormat == VCProjectFileFormat.VisualStudio2012)
             {
 				VCSolutionFileContent.AppendLine();
                 VCSolutionFileContent.AppendLine("Microsoft Visual Studio Solution File, Format Version 12.00");
@@ -704,7 +720,7 @@ namespace UnrealBuildTool
 								MSBuildProjectContext ProjectContext = CurProject.GetMatchingProjectContext(SolutionConfigCombination.TargetConfigurationName, SolutionConfigCombination.Configuration, SolutionConfigCombination.Platform, PlatformProjectGenerators);
 
 								// Override the configuration to build for UBT
-								if (bBuildUBTInDebug && CurProject == UBTProject)
+								if (Settings.bBuildUBTInDebug && CurProject == UBTProject)
 								{
 									ProjectContext.ConfigurationName = "Debug";
 								}
@@ -788,11 +804,11 @@ namespace UnrealBuildTool
 
 
 			// Save a solution config file which selects the development editor configuration by default.
-			if (bSuccess && bWriteSolutionOptionFile)
+			if (bSuccess && Settings.bWriteSolutionOptionFile)
 			{
 				// Figure out the filename for the SUO file. VS will automatically import the options from earlier versions if necessary.
 				FileReference SolutionOptionsFileName;
-				switch (ProjectFileFormat)
+				switch (Settings.ProjectFileFormat)
                 {
                     case VCProjectFileFormat.VisualStudio2012:
 						SolutionOptionsFileName = FileReference.Combine(MasterProjectPath, Path.ChangeExtension(SolutionFileName, "v11.suo"));
@@ -818,7 +834,7 @@ namespace UnrealBuildTool
 				{
 					DirectoryReference.CreateDirectory(SolutionOptionsFileName.Directory);
 
-					VCSolutionOptions Options = new VCSolutionOptions(ProjectFileFormat);
+					VCSolutionOptions Options = new VCSolutionOptions(Settings.ProjectFileFormat);
 
 					// Set the default configuration and startup project
 					VCSolutionConfigCombination DefaultConfig = SolutionConfigCombinations.Find(x => x.Configuration == UnrealTargetConfiguration.Development && x.Platform == UnrealTargetPlatform.Win64 && x.TargetConfigurationName == TargetType.Editor);
@@ -835,7 +851,7 @@ namespace UnrealBuildTool
 
 					// Mark all the projects as closed by default, apart from the startup project
 					VCSolutionExplorerState ExplorerState = new VCSolutionExplorerState();
-					if(ProjectFileFormat >= VCProjectFileFormat.VisualStudio2017)
+					if(Settings.ProjectFileFormat >= VCProjectFileFormat.VisualStudio2017)
 					{
 						BuildSolutionExplorerState_VS2017(RootFolder, "", ExplorerState, DefaultProject);
 					}
diff --git a/Engine/Source/Programs/UnrealBuildTool/ProjectFiles/VisualStudioMac/VisualStudioMacProjectFileGenerator.cs b/Engine/Source/Programs/UnrealBuildTool/ProjectFiles/VisualStudioMac/VisualStudioMacProjectFileGenerator.cs
--- a/Engine/Source/Programs/UnrealBuildTool/ProjectFiles/VisualStudioMac/VisualStudioMacProjectFileGenerator.cs
+++ b/Engine/Source/Programs/UnrealBuildTool/ProjectFiles/VisualStudioMac/VisualStudioMacProjectFileGenerator.cs
@@ -25,8 +25,8 @@ namespace UnrealBuildTool
         public VCMacProjectFileGenerator(FileReference InOnlyGameProject, CommandLineArguments InArguments)
 			: base(InOnlyGameProject, VCProjectFileFormat.Default, InArguments)
         {
-            // no suo file, requires ole32
-            bWriteSolutionOptionFile = false;
+			// no suo file, requires ole32
+			Settings.bWriteSolutionOptionFile = false;
         }
 
         /// True if we should include IntelliSense data in the generated project files when possible
@@ -41,9 +41,9 @@ namespace UnrealBuildTool
         /// <returns>True if successful</returns>
         protected override bool WriteProjectFiles(PlatformProjectGeneratorCollection PlatformProjectGenerators)
         {
-            // This can be reset by higher level code when it detects that we don't have
-            // VS2015 installed (TODO - add custom format for Mac?)
-            ProjectFileFormat = VCProjectFileFormat.VisualStudio2015;
+			// This can be reset by higher level code when it detects that we don't have
+			// VS2015 installed (TODO - add custom format for Mac?)
+			Settings.ProjectFileFormat = VCProjectFileFormat.VisualStudio2015;
 
             // we can't generate native projects so clear them here, we will just
             // write out OtherProjectFiles and AutomationProjectFiles
@@ -63,7 +63,7 @@ namespace UnrealBuildTool
                 DirectoryReference AutomationToolDir = DirectoryReference.Combine(UnrealBuildTool.EngineSourceDirectory, "Programs", "AutomationTool");
                 new XDocument(
                     new XElement(NS + "Project",
-                        new XAttribute("ToolsVersion", VCProjectFileGenerator.GetProjectFileToolVersionString(ProjectFileFormat)),
+                        new XAttribute("ToolsVersion", VCProjectFileGenerator.GetProjectFileToolVersionString(Settings.ProjectFileFormat)),
                         new XAttribute("DefaultTargets", "Build"),
                         new XElement(NS + "ItemGroup",
                             from AutomationProject in AutomationProjectFiles
