Fix BuildConfiguration.xml not being read from NotForLicensees folder (no longer returned by XmlConfig.FindInputFiles)

diff --git a/Engine/Source/Programs/UnrealBuildTool/ProjectFiles/Project.cs b/Engine/Source/Programs/UnrealBuildTool/ProjectFiles/Project.cs
--- a/Engine/Source/Programs/UnrealBuildTool/ProjectFiles/Project.cs
+++ b/Engine/Source/Programs/UnrealBuildTool/ProjectFiles/Project.cs
@@ -236,6 +236,11 @@ namespace UnrealBuildTool
 			}
 		}
 
+		/// <summary>
+		/// Set of all files that have been added
+		/// </summary>
+		HashSet<FileReference> AliasedFilesSet = new HashSet<FileReference>();
+
 		/// Aliased (i.e. files is custom filter tree) in this project
 		public readonly List<AliasedFile> AliasedFiles = new List<AliasedFile>();
 
@@ -245,7 +250,10 @@ namespace UnrealBuildTool
 		/// <param name="File">Aliased file.</param>
 		public void AddAliasedFileToProject(AliasedFile File)
 		{
-			AliasedFiles.Add(File);
+			if (AliasedFilesSet.Add(File.Location))
+			{
+				AliasedFiles.Add(File);
+			}
 		}
 
 		/// <summary>
@@ -255,6 +263,12 @@ namespace UnrealBuildTool
 		/// <param name="BaseFolder">The directory the path within the project will be relative to</param>
 		public void AddFileToProject(FileReference FilePath, DirectoryReference BaseFolder)
 		{
+			// Check if hasn't already been added as an aliased file
+			if (AliasedFilesSet.Contains(FilePath))
+			{
+				return;
+			}
+
 			// Don't add duplicates
 			SourceFile ExistingFile = null;
 			if (SourceFileMap.TryGetValue(FilePath, out ExistingFile))
diff --git a/Engine/Source/Programs/UnrealBuildTool/System/XmlConfig.cs b/Engine/Source/Programs/UnrealBuildTool/System/XmlConfig.cs
--- a/Engine/Source/Programs/UnrealBuildTool/System/XmlConfig.cs
+++ b/Engine/Source/Programs/UnrealBuildTool/System/XmlConfig.cs
@@ -204,13 +204,12 @@ namespace UnrealBuildTool
 			// Skip all the config files under the Engine folder if it's an installed build
 			if(!UnrealBuildTool.IsEngineInstalled())
 			{
-				// @todo benm: This file is getting added again when we scan Programs - easiest to just not add it here, but it will put it under Programs, not Engine
-// 				// Check for the config file under /Engine/Programs/NotForLicensees/UnrealBuildTool
-// 				FileReference NotForLicenseesConfigLocation = FileReference.Combine(UnrealBuildTool.EngineDirectory, "Restricted", "NotForLicensees", "Programs", "UnrealBuildTool", "BuildConfiguration.xml");
-// 				if(FileReference.Exists(NotForLicenseesConfigLocation))
-// 				{
-// 					InputFiles.Add(new InputFile { Location = NotForLicenseesConfigLocation, FolderName = "NotForLicensees" });
-// 				}
+ 				// Check for the config file under /Engine/Programs/NotForLicensees/UnrealBuildTool
+ 				FileReference NotForLicenseesConfigLocation = FileReference.Combine(UnrealBuildTool.EngineDirectory, "Restricted", "NotForLicensees", "Programs", "UnrealBuildTool", "BuildConfiguration.xml");
+ 				if(FileReference.Exists(NotForLicenseesConfigLocation))
+ 				{
+ 					InputFiles.Add(new InputFile { Location = NotForLicenseesConfigLocation, FolderName = "NotForLicensees" });
+ 				}
 
 				// Check for the user config file under /Engine/Saved/UnrealBuildTool
 				FileReference UserConfigLocation = FileReference.Combine(UnrealBuildTool.EngineDirectory, "Saved", "UnrealBuildTool", "BuildConfiguration.xml");
