Fix macros for UE_4_X_OR_LATER (add UE_5_X_OR_LATER)

diff --git a/Engine/Source/Programs/UnrealBuildTool/System/RulesAssembly.cs b/Engine/Source/Programs/UnrealBuildTool/System/RulesAssembly.cs
--- a/Engine/Source/Programs/UnrealBuildTool/System/RulesAssembly.cs
+++ b/Engine/Source/Programs/UnrealBuildTool/System/RulesAssembly.cs
@@ -201,10 +201,14 @@ namespace UnrealBuildTool
 			BuildVersion Version;
 			if (BuildVersion.TryRead(BuildVersion.GetDefaultFileName(), out Version))
 			{
-				for(int MinorVersion = 17; MinorVersion <= Version.MinorVersion; MinorVersion++)
+				for(int MinorVersion = 17; MinorVersion <= 30; MinorVersion++)
 				{
 					PreprocessorDefines.Add(String.Format("UE_4_{0}_OR_LATER", MinorVersion));
 				}
+				for (int MinorVersion = 0; MinorVersion <= 0; MinorVersion++)
+				{
+					PreprocessorDefines.Add(String.Format("UE_5_{0}_OR_LATER", MinorVersion));
+				}
 			}
 			return PreprocessorDefines;
 		}
