- Undid a debug change that was left in by mistake

diff --git a/Engine/Source/Programs/UnrealBuildTool/Configuration/UEBuildPlatformSDK.cs b/Engine/Source/Programs/UnrealBuildTool/Configuration/UEBuildPlatformSDK.cs
--- a/Engine/Source/Programs/UnrealBuildTool/Configuration/UEBuildPlatformSDK.cs
+++ b/Engine/Source/Programs/UnrealBuildTool/Configuration/UEBuildPlatformSDK.cs
@@ -51,7 +51,7 @@ namespace UnrealBuildTool
 		/// </summary>
 		protected const string SDKEnvironmentVarsFile = "OutputEnvVars.txt";
 
-		protected const string SDKRootEnvVar = "UE_SDKS_ROOT_LEGACY";
+		protected const string SDKRootEnvVar = "UE_SDKS_ROOT";
 
 		protected const string AutoSetupEnvVar = "AutoSDKSetup";
 
