- Big TargetPlatform Module cleanup:
  - Removed the function to return a single TP, and wrapped the functionality in a simple single required function (platform just has to add TPs to an array, high level code manages init and cleanup of the array)
  - Stripped unncessary code from all TPs
  - Collapsed the desktop targetplatform modules into 1 per platform
  - Renamed LinuxAArch64NoEditorTargetPlatfortm to LinuxAArch64TargetPlatform to match the other non-editor platforms
  - Deleted AllDesktopPlatform

diff --git a/Engine/Config/LinuxAArch64/DataDrivenPlatformInfo.ini b/Engine/Config/LinuxAArch64/DataDrivenPlatformInfo.ini
--- a/Engine/Config/LinuxAArch64/DataDrivenPlatformInfo.ini
+++ b/Engine/Config/LinuxAArch64/DataDrivenPlatformInfo.ini
@@ -2,8 +2,8 @@
 IniParent=Unix
 bIsConfidential=false
 
-[PlatformInfo LinuxAArch64NoEditor]
-TargetPlatformName=LinuxAArch64NoEditor
+[PlatformInfo LinuxAArch64]
+TargetPlatformName=LinuxAArch64
 DisplayName=Linux (AArch64)
 PlatformType=Game
 PlatformFlags=BuildFlavor
diff --git a/Engine/Plugins/Media/MediaPlayerEditor/Source/MediaPlayerEditor/Private/Customizations/BaseMediaSourceCustomization.cpp b/Engine/Plugins/Media/MediaPlayerEditor/Source/MediaPlayerEditor/Private/Customizations/BaseMediaSourceCustomization.cpp
--- a/Engine/Plugins/Media/MediaPlayerEditor/Source/MediaPlayerEditor/Private/Customizations/BaseMediaSourceCustomization.cpp
+++ b/Engine/Plugins/Media/MediaPlayerEditor/Source/MediaPlayerEditor/Private/Customizations/BaseMediaSourceCustomization.cpp
@@ -137,7 +137,7 @@ TSharedRef<SWidget> FBaseMediaSourceCustomization::MakePlatformPlayerNamesValueW
 
 	for (const PlatformInfo::FPlatformInfo& PlatformInfo : PlatformInfo::GetPlatformInfoArray())
 	{
-		if (PlatformInfo.IsVanilla() && (PlatformInfo.PlatformType == EBuildTargetType::Game) && (PlatformInfo.PlatformInfoName != TEXT("AllDesktop")))
+		if (PlatformInfo.IsVanilla() && (PlatformInfo.PlatformType == EBuildTargetType::Game))
 		{
 			AvailablePlatforms.Add(&PlatformInfo);
 		}
diff --git a/Engine/Plugins/Media/MediaPlayerEditor/Source/MediaPlayerEditor/Private/Customizations/PlatformMediaSourceCustomization.cpp b/Engine/Plugins/Media/MediaPlayerEditor/Source/MediaPlayerEditor/Private/Customizations/PlatformMediaSourceCustomization.cpp
--- a/Engine/Plugins/Media/MediaPlayerEditor/Source/MediaPlayerEditor/Private/Customizations/PlatformMediaSourceCustomization.cpp
+++ b/Engine/Plugins/Media/MediaPlayerEditor/Source/MediaPlayerEditor/Private/Customizations/PlatformMediaSourceCustomization.cpp
@@ -75,7 +75,7 @@ TSharedRef<SWidget> FPlatformMediaSourceCustomization::MakePlatformMediaSourcesV
 
 	for (const PlatformInfo::FPlatformInfo& PlatformInfo : PlatformInfo::GetPlatformInfoArray())
 	{
-		if (PlatformInfo.IsVanilla() && (PlatformInfo.PlatformType == EBuildTargetType::Game) && (PlatformInfo.PlatformInfoName != TEXT("AllDesktop")))
+		if (PlatformInfo.IsVanilla() && (PlatformInfo.PlatformType == EBuildTargetType::Game))
 		{
 			if (PlatformInfo.PlatformInfoName == TEXT("TVOS"))
 			{
diff --git a/Engine/Plugins/Runtime/AR/Microsoft/HoloLensAR/Source/HoloLensTargetPlatform/Private/HoloLensTargetPlatformModule.cpp b/Engine/Plugins/Runtime/AR/Microsoft/HoloLensAR/Source/HoloLensTargetPlatform/Private/HoloLensTargetPlatformModule.cpp
--- a/Engine/Plugins/Runtime/AR/Microsoft/HoloLensAR/Source/HoloLensTargetPlatform/Private/HoloLensTargetPlatformModule.cpp
+++ b/Engine/Plugins/Runtime/AR/Microsoft/HoloLensAR/Source/HoloLensTargetPlatform/Private/HoloLensTargetPlatformModule.cpp
@@ -13,13 +13,6 @@
 #define LOCTEXT_NAMESPACE "FHoloLensTargetPlatformModule"
 
 
-
-/**
- * Holds the target platform singleton.
- */
-static ITargetPlatform* HoloLensTargetSingleton = nullptr;
-
-
 /**
  * Module for the HoloLens target platform.
  */
@@ -27,34 +20,16 @@ class FHoloLensTargetPlatformModule
 	: public ITargetPlatformModule
 {
 public:
-
-	/** Default constructor. */
-	FHoloLensTargetPlatformModule( )
-	{ }
-
-	/** Destructor. */
-	~FHoloLensTargetPlatformModule( )
+	virtual void GetTargetPlatforms(TArray<ITargetPlatform*>& TargetPlatforms) override
 	{
-	}
-
-public:
-	
-	// ITargetPlatformModule interface
-	
-	virtual ITargetPlatform* GetTargetPlatform( )
-	{
-		if (HoloLensTargetSingleton == nullptr)
+		if (FHoloLensTargetPlatform::IsUsable())
 		{
 			//@todo HoloLens: Check for SDK?
-
-			HoloLensTargetSingleton = new FHoloLensTargetPlatform();
+			TargetPlatforms.Add(new FHoloLensTargetPlatform());
 		}
-		
-		return HoloLensTargetSingleton;
 	}
 };
 
-
 #undef LOCTEXT_NAMESPACE
 
 
diff --git a/Engine/Source/Developer/AllDesktopTargetPlatform/AllDesktopTargetPlatform.Build.cs b/Engine/Source/Developer/AllDesktopTargetPlatform/AllDesktopTargetPlatform.Build.cs
deleted file mode 100644
diff --git a/Engine/Source/Developer/AllDesktopTargetPlatform/Private/AllDesktopTargetPlatform.cpp b/Engine/Source/Developer/AllDesktopTargetPlatform/Private/AllDesktopTargetPlatform.cpp
deleted file mode 100644
diff --git a/Engine/Source/Developer/AllDesktopTargetPlatform/Private/AllDesktopTargetPlatform.h b/Engine/Source/Developer/AllDesktopTargetPlatform/Private/AllDesktopTargetPlatform.h
deleted file mode 100644
diff --git a/Engine/Source/Developer/AllDesktopTargetPlatform/Private/AllDesktopTargetPlatformModule.cpp b/Engine/Source/Developer/AllDesktopTargetPlatform/Private/AllDesktopTargetPlatformModule.cpp
deleted file mode 100644
diff --git a/Engine/Source/Developer/Android/AndroidTargetPlatform/Private/AndroidTargetPlatformModule.cpp b/Engine/Source/Developer/Android/AndroidTargetPlatform/Private/AndroidTargetPlatformModule.cpp
--- a/Engine/Source/Developer/Android/AndroidTargetPlatform/Private/AndroidTargetPlatformModule.cpp
+++ b/Engine/Source/Developer/Android/AndroidTargetPlatform/Private/AndroidTargetPlatformModule.cpp
@@ -20,27 +20,9 @@ class FAndroidTargetPlatformModule : public IAndroidTargetPlatformModule
 {
 public:
 
-	/**
-	 * Destructor.
-	 */
-	~FAndroidTargetPlatformModule( )
+	virtual void GetTargetPlatforms(TArray<ITargetPlatform*>& TargetPlatforms) override
 	{
-		for (ITargetPlatform* TP : TargetPlatforms)
-		{
-			delete TP;
-		}
-		TargetPlatforms.Empty();
-		MultiPlatforms.Empty();
-	}
-
-
-public:
-	
-	// Begin ITargetPlatformModule interface
-
-	virtual TArray<ITargetPlatform*> GetTargetPlatforms() override
-	{
-		if (TargetPlatforms.Num() == 0 && FAndroidTargetPlatform::IsUsable())
+		if (FAndroidTargetPlatform::IsUsable())
 		{
 			for (int32 Type = 0; Type < 2; Type++)
 			{
@@ -61,8 +43,6 @@ public:
 			// set up the multi platforms now that we have all the other platforms ready to go
 			NotifyMultiSelectedFormatsChanged();
 		}
-
-		return TargetPlatforms;
 	}
 
 
@@ -75,18 +55,10 @@ public:
 		// @todo multi needs to be passed this event!
 	}
 
-	// End ITargetPlatformModule interface
-
-public:
-	// Begin IModuleInterface interface
-	virtual void StartupModule() override { }
-	virtual void ShutdownModule() override { }
-	// End IModuleInterface interface
 
 private:
 
-	/** Holds the target platforms. */
-	TArray<ITargetPlatform*> TargetPlatforms;
+	/** Holds the specific types of target platforms for NotifyMultiSelectedFormatsChanged */
 	TArray<FAndroidTargetPlatform*> SinglePlatforms;
 	TArray<FAndroid_MultiTargetPlatform*> MultiPlatforms;
 };
diff --git a/Engine/Source/Developer/DesktopPlatform/Private/PlatformInfo.cpp b/Engine/Source/Developer/DesktopPlatform/Private/PlatformInfo.cpp
--- a/Engine/Source/Developer/DesktopPlatform/Private/PlatformInfo.cpp
+++ b/Engine/Source/Developer/DesktopPlatform/Private/PlatformInfo.cpp
@@ -85,14 +85,6 @@ void BuildPlatformInfo(const FName& InPlatformInfoName, const FName& InTargetPla
 	PlatformInfo.bTargetPlatformCanUseCrashReporter = InTargetPlatformCanUseCrashReporter;
 }
 
-
-void BuildHardcodedPlatforms()
-{
-	// PlatformInfoName									TargetPlatformName			DisplayName														PlatformType			PlatformFlags					IconPaths																																		UATCommandLine										AutoSDKPath			SDKStatus						SDKTutorial																								bEnabledForUse										BinaryFolderName	IniPlatformName		FbUsesHostCompiler		bUATClosesAfterLaunch	bIsConfidential	UBTTargetId (match UBT's UnrealTargetPlatform enum)		PlatformSubMenu		bTargetPlatformCanUseCrashReporter
-	BuildPlatformInfo(TEXT("AllDesktop"),				TEXT("AllDesktop"),			LOCTEXT("DesktopTargetPlatDisplay", "Desktop (Win+Mac+Linux)"),	EBuildTargetType::Game,	EPlatformFlags::None,			FPlatformIconPaths(TEXT("Launcher/Desktop/Platform_Desktop_24x"), TEXT("Launcher/Desktop/Platform_Desktop_128x")),					TEXT(""),											TEXT(""),			EPlatformSDKStatus::Unknown,	TEXT(""),																								PLATFORM_WINDOWS /* see note below */,						TEXT(""),			TEXT(""),			false,					true,					false,			TEXT("AllDesktop"),	TEXT("Desktop"),	TEXT(""),	true);
-	// Note: For "AllDesktop" bEnabledForUse value, see SProjectTargetPlatformSettings::Construct !!!! IsAvailableOnWindows || IsAvailableOnMac || IsAvailableOnLinux
-}
-
 // Gets a string from a section, or empty string if it didn't exist
 FString GetSectionString(const FConfigSection& Section, FName Key)
 {
@@ -205,7 +197,6 @@ void LoadDataDrivenPlatforms()
 
 FDelayedAutoRegisterHelper GPlatformInfoInit(EDelayedRegisterRunPhase::FileSystemReady, []()
 {
-	BuildHardcodedPlatforms();
 	LoadDataDrivenPlatforms();
 });
 
diff --git a/Engine/Source/Developer/IOS/IOSTargetPlatform/Private/IOSTargetPlatformModule.cpp b/Engine/Source/Developer/IOS/IOSTargetPlatform/Private/IOSTargetPlatformModule.cpp
--- a/Engine/Source/Developer/IOS/IOSTargetPlatform/Private/IOSTargetPlatformModule.cpp
+++ b/Engine/Source/Developer/IOS/IOSTargetPlatform/Private/IOSTargetPlatformModule.cpp
@@ -12,38 +12,15 @@ class FIOSTargetPlatformModule : public ITargetPlatformModule
 {
 public:
 
-	/** Destructor. */
-	~FIOSTargetPlatformModule()
+	virtual void GetTargetPlatforms(TArray<ITargetPlatform*>& TargetPlatforms) override
 	{
-		for (ITargetPlatform* TP : TargetPlatforms)
-		{
-			delete TP;
-		}
-		TargetPlatforms.Empty();
-	}
-
-	//~ ITargetPlatformModule interface
-
-	virtual TArray<ITargetPlatform*> GetTargetPlatforms() override
-	{
-		if (TargetPlatforms.Num() == 0 && FIOSTargetPlatform::IsUsable())
+		if (FIOSTargetPlatform::IsUsable())
 		{
 			TargetPlatforms.Add(new FIOSTargetPlatform(false, true));
 			TargetPlatforms.Add(new FIOSTargetPlatform(false, false));
 		}
-
-		return TargetPlatforms;
 	}
 
-	//~ IModuleInterface interface
-
-	virtual void StartupModule() override { }
-	virtual void ShutdownModule() override { }
-
-protected:
-	/** Holds the target platforms. */
-	TArray<ITargetPlatform*> TargetPlatforms;
-	
 };
 
 
diff --git a/Engine/Source/Developer/IOS/TVOSTargetPlatform/Private/TVOSTargetPlatformModule.cpp b/Engine/Source/Developer/IOS/TVOSTargetPlatform/Private/TVOSTargetPlatformModule.cpp
--- a/Engine/Source/Developer/IOS/TVOSTargetPlatform/Private/TVOSTargetPlatformModule.cpp
+++ b/Engine/Source/Developer/IOS/TVOSTargetPlatform/Private/TVOSTargetPlatformModule.cpp
@@ -12,42 +12,16 @@ class FTVOSTargetPlatformModule	: public ITargetPlatformModule
 {
 public:
 
-	/**
-	 * Destructor.
-	 */
-	~FTVOSTargetPlatformModule()
+	virtual void GetTargetPlatforms(TArray<ITargetPlatform*>& TargetPlatforms) override
 	{
-		for (ITargetPlatform* TP : TargetPlatforms)
-		{
-			delete TP;
-		}
-		TargetPlatforms.Empty();
-	}
-
-	//~ ITargetPlatformModule interface
-
-	virtual TArray<ITargetPlatform*> GetTargetPlatforms() override
-	{
-		if (TargetPlatforms.Num() == 0 && FIOSTargetPlatform::IsUsable())
+		if (FIOSTargetPlatform::IsUsable())
 		{
 			// add Game and Client TPs
 			TargetPlatforms.Add(new FIOSTargetPlatform(true, true));
 			TargetPlatforms.Add(new FIOSTargetPlatform(true, false));
 		}
-
-		return TargetPlatforms;
 	}
 
-public:
-
-	//~ IModuleInterface interface
-
-	virtual void StartupModule() override { }
-	virtual void ShutdownModule() override { }
-
-protected:
-	/** Holds the target platforms. */
-	TArray<ITargetPlatform*> TargetPlatforms;
 };
 
 
diff --git a/Engine/Source/Developer/Linux/LinuxAArch64ClientTargetPlatform/LinuxAArch64ClientTargetPlatform.Build.cs b/Engine/Source/Developer/Linux/LinuxAArch64ClientTargetPlatform/LinuxAArch64ClientTargetPlatform.Build.cs
deleted file mode 100644
diff --git a/Engine/Source/Developer/Linux/LinuxAArch64ClientTargetPlatform/Private/LinuxAArch64ClientTargetPlatformModule.cpp b/Engine/Source/Developer/Linux/LinuxAArch64ClientTargetPlatform/Private/LinuxAArch64ClientTargetPlatformModule.cpp
deleted file mode 100644
diff --git a/Engine/Source/Developer/Linux/LinuxAArch64NoEditorTargetPlatform/LinuxAArch64NoEditorTargetPlatform.Build.cs b/Engine/Source/Developer/Linux/LinuxAArch64NoEditorTargetPlatform/LinuxAArch64NoEditorTargetPlatform.Build.cs
deleted file mode 100644
diff --git a/Engine/Source/Developer/Linux/LinuxAArch64NoEditorTargetPlatform/Private/LinuxAArch64NoEditorTargetPlatformModule.cpp b/Engine/Source/Developer/Linux/LinuxAArch64NoEditorTargetPlatform/Private/LinuxAArch64NoEditorTargetPlatformModule.cpp
deleted file mode 100644
diff --git a/Engine/Source/Developer/Linux/LinuxAArch64ServerTargetPlatform/LinuxAArch64ServerTargetPlatform.Build.cs b/Engine/Source/Developer/Linux/LinuxAArch64ServerTargetPlatform/LinuxAArch64ServerTargetPlatform.Build.cs
deleted file mode 100644
diff --git a/Engine/Source/Developer/Linux/LinuxAArch64ServerTargetPlatform/Private/LinuxAArch64ServerTargetPlatformModule.cpp b/Engine/Source/Developer/Linux/LinuxAArch64ServerTargetPlatform/Private/LinuxAArch64ServerTargetPlatformModule.cpp
deleted file mode 100644
diff --git a/Engine/Source/Developer/Linux/LinuxAArch64TargetPlatform/LinuxAArch64TargetPlatform.Build.cs b/Engine/Source/Developer/Linux/LinuxAArch64TargetPlatform/LinuxAArch64TargetPlatform.Build.cs
new file mode 100644
--- /dev/null
+++ b/Engine/Source/Developer/Linux/LinuxAArch64TargetPlatform/LinuxAArch64TargetPlatform.Build.cs
@@ -0,0 +1,36 @@
+// Copyright Epic Games, Inc. All Rights Reserved.
+
+using UnrealBuildTool;
+
+public class LinuxAArch64TargetPlatform : ModuleRules
+{
+    public LinuxAArch64TargetPlatform(ReadOnlyTargetRules Target) : base(Target)
+    {
+        BinariesSubFolder = "LinuxAArch64";
+
+        PrivateDependencyModuleNames.AddRange(
+            new string[] {
+				"Core",
+				"TargetPlatform",
+				"DesktopPlatform",
+				"Projects"
+			}
+        );
+
+        if (Target.bCompileAgainstEngine)
+        {
+            PrivateDependencyModuleNames.AddRange(new string[] {
+				"Engine"
+				}
+            );
+
+            PrivateIncludePathModuleNames.Add("TextureCompressor");
+        }
+
+        PrivateIncludePaths.AddRange(
+            new string[] {
+				"Developer/Linux/LinuxTargetPlatform/Private"
+			}
+        );
+    }
+}
diff --git a/Engine/Source/Developer/Linux/LinuxAArch64TargetPlatform/Private/LinuxAArch64TargetPlatformModule.cpp b/Engine/Source/Developer/Linux/LinuxAArch64TargetPlatform/Private/LinuxAArch64TargetPlatformModule.cpp
new file mode 100644
--- /dev/null
+++ b/Engine/Source/Developer/Linux/LinuxAArch64TargetPlatform/Private/LinuxAArch64TargetPlatformModule.cpp
@@ -0,0 +1,45 @@
+// Copyright Epic Games, Inc. All Rights Reserved.
+
+/*=============================================================================
+	LinuxAArch64NoEditorTargetPlatformModule.cpp: Implements the FLinuxAArch64NoEditorTargetPlatformModule class.
+=============================================================================*/
+
+#include "CoreMinimal.h"
+#include "Modules/ModuleManager.h"
+
+#include "Interfaces/ITargetPlatformModule.h"
+
+#include "LinuxTargetDevice.h"
+#include "LinuxTargetPlatform.h"
+
+
+/**
+ * Module for the Linux target platforms
+ */
+class FLinuxAArch64TargetPlatformModule
+	: public ITargetPlatformModule
+{
+public:
+
+	virtual void GetTargetPlatforms(TArray<ITargetPlatform*>& TargetPlatforms) override
+	{
+		// Server TP
+		if (TLinuxTargetPlatform<FLinuxPlatformProperties<false, true, false, true> >::IsUsable())
+		{
+			TargetPlatforms.Add(new TLinuxTargetPlatform<FLinuxPlatformProperties<false, true, false, true> >());
+		}
+		// Game TP
+		if (TLinuxTargetPlatform<FLinuxPlatformProperties<false, false, false, true> >::IsUsable())
+		{
+			TargetPlatforms.Add(new TLinuxTargetPlatform<FLinuxPlatformProperties<false, false, false, true> >());
+		}
+		// Client TP
+		if (TLinuxTargetPlatform<FLinuxPlatformProperties<false, false, true, true> >::IsUsable())
+		{
+			TargetPlatforms.Add(new TLinuxTargetPlatform<FLinuxPlatformProperties<false, false, true, true> >());
+		}
+	}
+};
+
+
+IMPLEMENT_MODULE(FLinuxAArch64TargetPlatformModule, LinuxAArch64TargetPlatform);
diff --git a/Engine/Source/Developer/Linux/LinuxClientTargetPlatform/LinuxClientTargetPlatform.Build.cs b/Engine/Source/Developer/Linux/LinuxClientTargetPlatform/LinuxClientTargetPlatform.Build.cs
deleted file mode 100644
diff --git a/Engine/Source/Developer/Linux/LinuxClientTargetPlatform/Private/LinuxClientTargetPlatformModule.cpp b/Engine/Source/Developer/Linux/LinuxClientTargetPlatform/Private/LinuxClientTargetPlatformModule.cpp
deleted file mode 100644
diff --git a/Engine/Source/Developer/Linux/LinuxNoEditorTargetPlatform/LinuxNoEditorTargetPlatform.Build.cs b/Engine/Source/Developer/Linux/LinuxNoEditorTargetPlatform/LinuxNoEditorTargetPlatform.Build.cs
deleted file mode 100644
diff --git a/Engine/Source/Developer/Linux/LinuxNoEditorTargetPlatform/Private/LinuxNoEditorTargetPlatformModule.cpp b/Engine/Source/Developer/Linux/LinuxNoEditorTargetPlatform/Private/LinuxNoEditorTargetPlatformModule.cpp
deleted file mode 100644
diff --git a/Engine/Source/Developer/Linux/LinuxServerTargetPlatform/LinuxServerTargetPlatform.Build.cs b/Engine/Source/Developer/Linux/LinuxServerTargetPlatform/LinuxServerTargetPlatform.Build.cs
deleted file mode 100644
diff --git a/Engine/Source/Developer/Linux/LinuxServerTargetPlatform/Private/LinuxServerTargetPlatformModule.cpp b/Engine/Source/Developer/Linux/LinuxServerTargetPlatform/Private/LinuxServerTargetPlatformModule.cpp
deleted file mode 100644
diff --git a/Engine/Source/Developer/Linux/LinuxTargetPlatform/LinuxTargetPlatform.Build.cs b/Engine/Source/Developer/Linux/LinuxTargetPlatform/LinuxTargetPlatform.Build.cs
--- a/Engine/Source/Developer/Linux/LinuxTargetPlatform/LinuxTargetPlatform.Build.cs
+++ b/Engine/Source/Developer/Linux/LinuxTargetPlatform/LinuxTargetPlatform.Build.cs
@@ -7,10 +7,6 @@ public class LinuxTargetPlatform : ModuleRules
     public LinuxTargetPlatform(ReadOnlyTargetRules Target) : base(Target)
 	{
         BinariesSubFolder = "Linux";
-		if (Target.Platform == UnrealTargetPlatform.LinuxAArch64)
-		{
-			BinariesSubFolder += "AArch64";
-		}
 
 		PrivateDependencyModuleNames.AddRange(
             new string[] {
diff --git a/Engine/Source/Developer/Linux/LinuxTargetPlatform/Private/LinuxTargetPlatformModule.cpp b/Engine/Source/Developer/Linux/LinuxTargetPlatform/Private/LinuxTargetPlatformModule.cpp
--- a/Engine/Source/Developer/Linux/LinuxTargetPlatform/Private/LinuxTargetPlatformModule.cpp
+++ b/Engine/Source/Developer/Linux/LinuxTargetPlatform/Private/LinuxTargetPlatformModule.cpp
@@ -17,11 +17,6 @@
 #define LOCTEXT_NAMESPACE "FLinuxTargetPlatformModule"
 
 
-/**
- * Holds the target platform singleton.
- */
-static ITargetPlatform* Singleton = NULL;
-
 
 /**
  * Module for the Linux target platform.
@@ -30,25 +25,28 @@ class FLinuxTargetPlatformModule
 	: public ITargetPlatformModule
 {
 public:
-
-	/** Destructor. */
-	~FLinuxTargetPlatformModule( )
+	virtual void GetTargetPlatforms(TArray<ITargetPlatform*>& TargetPlatforms) override
 	{
-		Singleton = NULL;
-	}
-
-public:
-	
-	// ITargetPlatformModule interface
-
-	virtual ITargetPlatform* GetTargetPlatform( ) override
-	{
-		if (Singleton == NULL && TLinuxTargetPlatform<FLinuxPlatformProperties<true, false, false, false> >::IsUsable())
+		// Editor TP
+		if (TLinuxTargetPlatform<FLinuxPlatformProperties<true, false, false, false> >::IsUsable())
+		{
+			TargetPlatforms.Add(new TLinuxTargetPlatform<FLinuxPlatformProperties<true, false, false, false> >());
+		}
+		// Server TP
+		if (TLinuxTargetPlatform<FLinuxPlatformProperties<false, true, false, false> >::IsUsable())
+		{
+			TargetPlatforms.Add(new TLinuxTargetPlatform<FLinuxPlatformProperties<false, true, false, false> >());
+		}
+		// NoEditor TP
+		if (TLinuxTargetPlatform<FLinuxPlatformProperties<false, false, false, false> >::IsUsable())
+		{
+			TargetPlatforms.Add(new TLinuxTargetPlatform<FLinuxPlatformProperties<false, false, false, false> >());
+		}
+		// Client TP
+		if (TLinuxTargetPlatform<FLinuxPlatformProperties<false, false, true, false> >::IsUsable())
 		{
-			Singleton = new TLinuxTargetPlatform<FLinuxPlatformProperties<true, false, false, false> >();
+			TargetPlatforms.Add(new TLinuxTargetPlatform<FLinuxPlatformProperties<false, false, true, false> >());
 		}
-		
-		return Singleton;
 	}
 
 public:
diff --git a/Engine/Source/Developer/Lumin/LuminTargetPlatform/Private/LuminTargetPlatformModule.cpp b/Engine/Source/Developer/Lumin/LuminTargetPlatform/Private/LuminTargetPlatformModule.cpp
--- a/Engine/Source/Developer/Lumin/LuminTargetPlatform/Private/LuminTargetPlatformModule.cpp
+++ b/Engine/Source/Developer/Lumin/LuminTargetPlatform/Private/LuminTargetPlatformModule.cpp
@@ -13,50 +13,14 @@ class FLuminTargetPlatformModule : public ITargetPlatformModule
 {
 public:
 
-	/**
-	 * Destructor.
-	 */
-	virtual ~FLuminTargetPlatformModule( )
+	virtual void GetTargetPlatforms(TArray<ITargetPlatform*>& TargetPlatforms) override
 	{
-		for (ITargetPlatform* TP : TargetPlatforms)
-		{
-			delete TP;
-		}
-		TargetPlatforms.Empty();
-	}
-
-
-public:
-	
-	// Begin ITargetPlatformModule interface
-
-	virtual TArray<ITargetPlatform*> GetTargetPlatforms() override
-	{
-		if (TargetPlatforms.Num() == 0 && FLuminTargetPlatform::IsUsable())
+		if (FLuminTargetPlatform::IsUsable())
 		{
 			TargetPlatforms.Add(new FLuminTargetPlatform(false));
 			TargetPlatforms.Add(new FLuminTargetPlatform(true));
 		}
-		
-		return TargetPlatforms;
 	}
-
-	// End ITargetPlatformModule interface
-
-public:
-	// Begin IModuleInterface interface
-	virtual void StartupModule() override
-	{
-	}
-
-	virtual void ShutdownModule() override
-	{
-	}
-	// End IModuleInterface interface
-
-private:
-	/** Holds the target platforms. */
-	TArray<ITargetPlatform*> TargetPlatforms;
 };
 
 
diff --git a/Engine/Source/Developer/Mac/MacClientTargetPlatform/MacClientTargetPlatform.Build.cs b/Engine/Source/Developer/Mac/MacClientTargetPlatform/MacClientTargetPlatform.Build.cs
deleted file mode 100644
diff --git a/Engine/Source/Developer/Mac/MacClientTargetPlatform/Private/MacClientTargetPlatformModule.cpp b/Engine/Source/Developer/Mac/MacClientTargetPlatform/Private/MacClientTargetPlatformModule.cpp
deleted file mode 100644
diff --git a/Engine/Source/Developer/Mac/MacNoEditorTargetPlatform/MacNoEditorTargetPlatform.Build.cs b/Engine/Source/Developer/Mac/MacNoEditorTargetPlatform/MacNoEditorTargetPlatform.Build.cs
deleted file mode 100644
diff --git a/Engine/Source/Developer/Mac/MacNoEditorTargetPlatform/Private/MacNoEditorTargetPlatformModule.cpp b/Engine/Source/Developer/Mac/MacNoEditorTargetPlatform/Private/MacNoEditorTargetPlatformModule.cpp
deleted file mode 100644
diff --git a/Engine/Source/Developer/Mac/MacServerTargetPlatform/MacServerTargetPlatform.Build.cs b/Engine/Source/Developer/Mac/MacServerTargetPlatform/MacServerTargetPlatform.Build.cs
deleted file mode 100644
diff --git a/Engine/Source/Developer/Mac/MacServerTargetPlatform/Private/MacServerTargetPlatformModule.cpp b/Engine/Source/Developer/Mac/MacServerTargetPlatform/Private/MacServerTargetPlatformModule.cpp
deleted file mode 100644
diff --git a/Engine/Source/Developer/Mac/MacTargetPlatform/Private/MacTargetPlatformModule.cpp b/Engine/Source/Developer/Mac/MacTargetPlatform/Private/MacTargetPlatformModule.cpp
--- a/Engine/Source/Developer/Mac/MacTargetPlatform/Private/MacTargetPlatformModule.cpp
+++ b/Engine/Source/Developer/Mac/MacTargetPlatform/Private/MacTargetPlatformModule.cpp
@@ -13,11 +13,6 @@
 #define LOCTEXT_NAMESPACE "FMacTargetPlatformModule"
 
 
-/**
- * Holds the target platform singleton.
- */
-static ITargetPlatform* Singleton = NULL;
-
 
 /**
  * Module for Mac as a target platform
@@ -27,31 +22,30 @@ class FMacTargetPlatformModule
 {
 public:
 
-	/**
-	 * Destructor.
-	 */
-	~FMacTargetPlatformModule()
-	{
-		Singleton = NULL;
-	}
-
-
-public:
-
-	// Begin ITargetPlatformModule interface
-
-	virtual ITargetPlatform* GetTargetPlatform() override
+	virtual void GetTargetPlatforms(TArray<ITargetPlatform*>& TargetPlatforms) override
 	{
-		if (Singleton == NULL && TGenericMacTargetPlatform<true, false, false>::IsUsable())
+		// Editor TP
+		if (TGenericMacTargetPlatform<true, false, false>::IsUsable())
 		{
-			Singleton = new TGenericMacTargetPlatform<true, false, false>();
+			TargetPlatforms.Add(new TGenericMacTargetPlatform<true, false, false>());
+		}
+		// Server TP
+		if (TGenericMacTargetPlatform<false, true, false>::IsUsable())
+		{
+			TargetPlatforms.Add(new TGenericMacTargetPlatform<false, true, false>());
+		}
+		// NoEditor TP
+		if (TGenericMacTargetPlatform<false, false, false>::IsUsable())
+		{
+			TargetPlatforms.Add(new TGenericMacTargetPlatform<false, false, false>());
+		}
+		// Client TP
+		if (TGenericMacTargetPlatform<false, false, true>::IsUsable())
+		{
+			TargetPlatforms.Add(new TGenericMacTargetPlatform<false, false, true>());
 		}
-
-		return Singleton;
 	}
 
-	// End ITargetPlatformModule interface
-
 
 public:
 
diff --git a/Engine/Source/Developer/TargetPlatform/Private/TargetPlatformManagerModule.cpp b/Engine/Source/Developer/TargetPlatform/Private/TargetPlatformManagerModule.cpp
--- a/Engine/Source/Developer/TargetPlatform/Private/TargetPlatformManagerModule.cpp
+++ b/Engine/Source/Developer/TargetPlatform/Private/TargetPlatformManagerModule.cpp
@@ -197,14 +197,6 @@ public:
 
 		for (int32 Index = 0; Index < TargetPlatforms.Num(); Index++)
 		{
-			//@todo-lh:
-			// FAllDesktopPlatformProperties will be removed soon as it's no longer maintained
-			// and will be replaced by the platform specific subclasses eventually, so skip "AllDesktop.
-			// Find platform specific subclass instead.
-			if (TargetPlatforms[Index]->PlatformName() == TEXT("AllDesktop"))
-			{
-				continue;
-			}
 			if (TargetPlatforms[Index]->SupportsValueForType(SupportType, RequiredSupportedValue))
 			{
 				return TargetPlatforms[Index];
diff --git a/Engine/Source/Developer/TargetPlatform/Public/Interfaces/ITargetPlatformModule.h b/Engine/Source/Developer/TargetPlatform/Public/Interfaces/ITargetPlatformModule.h
--- a/Engine/Source/Developer/TargetPlatform/Public/Interfaces/ITargetPlatformModule.h
+++ b/Engine/Source/Developer/TargetPlatform/Public/Interfaces/ITargetPlatformModule.h
@@ -4,8 +4,7 @@
 
 #include "CoreMinimal.h"
 #include "Modules/ModuleInterface.h"
-
-class ITargetPlatform;
+#include "Interfaces/ITargetPlatform.h"
 
 /**
  * Interface for target platform modules.
@@ -13,15 +12,18 @@ class ITargetPlatform;
 class ITargetPlatformModule
 	: public IModuleInterface
 {
+
 public:
 
-protected:
-	virtual ITargetPlatform* GetTargetPlatform()
+	/** Virtual destructor. */
+	virtual ~ITargetPlatformModule()
 	{
-		return nullptr;
-	};
-	
-public:
+		for (ITargetPlatform* TP : AllTargetPlatforms)
+		{
+			delete TP;
+		}
+		AllTargetPlatforms.Empty();
+	}
 
 	/**
 	 * Gets the module's target platforms. This should be overridden by each platform, but 
@@ -29,19 +31,24 @@ public:
 	 *
 	 * @return The target platform.
 	 */
-	virtual TArray<ITargetPlatform*> GetTargetPlatforms()
+	TArray<ITargetPlatform*> GetTargetPlatforms()
 	{
-		TArray<ITargetPlatform*> TargetPlatforms;
-		ITargetPlatform* TargetPlatform = GetTargetPlatform();
-		if (TargetPlatform != nullptr)
+		if (AllTargetPlatforms.Num() == 0)
 		{
-			TargetPlatforms.Add(TargetPlatform);
+			GetTargetPlatforms(AllTargetPlatforms);
 		}
-		return TargetPlatforms;
+
+		return AllTargetPlatforms;
 	}
 
-public:
+protected:
 
-	/** Virtual destructor. */
-	virtual ~ITargetPlatformModule() { }
+	/**
+	 * This is where each platform module will fill out an array
+	*/
+	virtual void GetTargetPlatforms(TArray<ITargetPlatform*>& TargetPlatforms) = 0;
+
+private:
+	/** Holds the target platforms. */
+	TArray<ITargetPlatform*> AllTargetPlatforms;
 };
diff --git a/Engine/Source/Developer/TargetPlatform/TargetPlatform.Build.cs b/Engine/Source/Developer/TargetPlatform/TargetPlatform.Build.cs
--- a/Engine/Source/Developer/TargetPlatform/TargetPlatform.Build.cs
+++ b/Engine/Source/Developer/TargetPlatform/TargetPlatform.Build.cs
@@ -60,9 +60,6 @@ public class TargetPlatform : ModuleRules
 					DynamicallyLoadedModuleNames.Add("IOSTargetPlatform");
 					DynamicallyLoadedModuleNames.Add("TVOSTargetPlatform");
 					DynamicallyLoadedModuleNames.Add("MacTargetPlatform");
-					DynamicallyLoadedModuleNames.Add("MacNoEditorTargetPlatform");
-					DynamicallyLoadedModuleNames.Add("MacServerTargetPlatform");
-					DynamicallyLoadedModuleNames.Add("MacClientTargetPlatform");
 				}
 			}
 			else if (Target.Platform == UnrealTargetPlatform.Mac)
diff --git a/Engine/Source/Developer/Windows/WindowsClientTargetPlatform/Private/WindowsClientTargetPlatformModule.cpp b/Engine/Source/Developer/Windows/WindowsClientTargetPlatform/Private/WindowsClientTargetPlatformModule.cpp
deleted file mode 100644
diff --git a/Engine/Source/Developer/Windows/WindowsClientTargetPlatform/WindowsClientTargetPlatform.Build.cs b/Engine/Source/Developer/Windows/WindowsClientTargetPlatform/WindowsClientTargetPlatform.Build.cs
deleted file mode 100644
diff --git a/Engine/Source/Developer/Windows/WindowsNoEditorTargetPlatform/Private/WindowsNoEditorTargetPlatformModule.cpp b/Engine/Source/Developer/Windows/WindowsNoEditorTargetPlatform/Private/WindowsNoEditorTargetPlatformModule.cpp
deleted file mode 100644
diff --git a/Engine/Source/Developer/Windows/WindowsNoEditorTargetPlatform/WindowsNoEditorTargetPlatform.Build.cs b/Engine/Source/Developer/Windows/WindowsNoEditorTargetPlatform/WindowsNoEditorTargetPlatform.Build.cs
deleted file mode 100644
diff --git a/Engine/Source/Developer/Windows/WindowsServerTargetPlatform/Private/WindowsServerTargetPlatformModule.cpp b/Engine/Source/Developer/Windows/WindowsServerTargetPlatform/Private/WindowsServerTargetPlatformModule.cpp
deleted file mode 100644
diff --git a/Engine/Source/Developer/Windows/WindowsServerTargetPlatform/WindowsServerTargetPlatform.Build.cs b/Engine/Source/Developer/Windows/WindowsServerTargetPlatform/WindowsServerTargetPlatform.Build.cs
deleted file mode 100644
diff --git a/Engine/Source/Developer/Windows/WindowsTargetPlatform/Private/WindowsTargetPlatformModule.cpp b/Engine/Source/Developer/Windows/WindowsTargetPlatform/Private/WindowsTargetPlatformModule.cpp
--- a/Engine/Source/Developer/Windows/WindowsTargetPlatform/Private/WindowsTargetPlatformModule.cpp
+++ b/Engine/Source/Developer/Windows/WindowsTargetPlatform/Private/WindowsTargetPlatformModule.cpp
@@ -11,9 +11,6 @@
 #define LOCTEXT_NAMESPACE "FWindowsTargetPlatformModule"
 
 
-/** Holds the target platform singleton. */
-static ITargetPlatform* Singleton = nullptr;
-
 
 /**
  * Implements the Windows target platform module.
@@ -21,14 +18,6 @@ static ITargetPlatform* Singleton = nullptr;
 class FWindowsTargetPlatformModule
 	: public ITargetPlatformModule
 {
-public:
-
-	/** Destructor. */
-	~FWindowsTargetPlatformModule( )
-	{
-		Singleton = nullptr;
-	}
-
 public:
 
 	// this is an example of a hotfix, declared here for no particular reason. Once we have other examples, it can be deleted.
@@ -45,16 +34,28 @@ public:
 
 public:
 
-	// ITargetPlatformModule interface
-
-	virtual ITargetPlatform* GetTargetPlatform( ) override
+	virtual void GetTargetPlatforms(TArray<ITargetPlatform*>& TargetPlatforms) override
 	{
-		if (Singleton == nullptr && TGenericWindowsTargetPlatform<true, false, false>::IsUsable())
+		// Editor TP
+		if (TGenericWindowsTargetPlatform<true, false, false>::IsUsable())
 		{
-			Singleton = new TGenericWindowsTargetPlatform<true, false, false>();
+			TargetPlatforms.Add(new TGenericWindowsTargetPlatform<true, false, false>());
+		}
+		// Server TP
+		if (TGenericWindowsTargetPlatform<false, true, false>::IsUsable())
+		{
+			TargetPlatforms.Add(new TGenericWindowsTargetPlatform<false, true, false>());
+		}
+		// NoEditor TP
+		if (TGenericWindowsTargetPlatform<false, false, false>::IsUsable())
+		{
+			TargetPlatforms.Add(new TGenericWindowsTargetPlatform<false, false, false>());
+		}
+		// Client TP
+		if (TGenericWindowsTargetPlatform<false, false, true>::IsUsable())
+		{
+			TargetPlatforms.Add(new TGenericWindowsTargetPlatform<false, false, true>());
 		}
-
-		return Singleton;
 	}
 
 public:
@@ -78,10 +79,11 @@ public:
 
 		ISettingsModule* SettingsModule = FModuleManager::GetModulePtr<ISettingsModule>("Settings");
 	}
-
+	
 };
 
 
+
 #undef LOCTEXT_NAMESPACE
 
 
diff --git a/Engine/Source/Editor/ProjectTargetPlatformEditor/Private/Widgets/SProjectTargetPlatformSettings.cpp b/Engine/Source/Editor/ProjectTargetPlatformEditor/Private/Widgets/SProjectTargetPlatformSettings.cpp
--- a/Engine/Source/Editor/ProjectTargetPlatformEditor/Private/Widgets/SProjectTargetPlatformSettings.cpp
+++ b/Engine/Source/Editor/ProjectTargetPlatformEditor/Private/Widgets/SProjectTargetPlatformSettings.cpp
@@ -26,13 +26,7 @@ void SProjectTargetPlatformSettings::Construct(const FArguments& InArgs)
 	{
 		if(PlatformInfo.IsVanilla() && PlatformInfo.PlatformType == EBuildTargetType::Game)
 		{
-#if !PLATFORM_WINDOWS
-			// @todo AllDesktop now only works on Windows (it can compile D3D shaders, and it can remote compile Metal shaders)
-			if (PlatformInfo.PlatformInfoName != TEXT("AllDesktop"))
-#endif
-			{
-				AvailablePlatforms.Add(&PlatformInfo);
-			}
+			AvailablePlatforms.Add(&PlatformInfo);
 		}
 	}
 
diff --git a/Engine/Source/Editor/PropertyEditor/Private/SPerPlatformPropertiesWidget.cpp b/Engine/Source/Editor/PropertyEditor/Private/SPerPlatformPropertiesWidget.cpp
--- a/Engine/Source/Editor/PropertyEditor/Private/SPerPlatformPropertiesWidget.cpp
+++ b/Engine/Source/Editor/PropertyEditor/Private/SPerPlatformPropertiesWidget.cpp
@@ -59,7 +59,7 @@ void SPerPlatformPropertiesWidget::ConstructChildren()
 		TMultiMap<FName, FName> GroupToPlatform;
 						
 		// Create mapping from platform to platform groups and remove postfixes and invalid platform names
-		const TArray<FString> Filters = { TEXT("NoEditor"), TEXT("Client"), TEXT("Server"), TEXT("AllDesktop") };
+		const TArray<FString> Filters = { TEXT("NoEditor"), TEXT("Client"), TEXT("Server") };
 		for (const FName& PlatformName : VanillaPlatformNameArray)
 		{
 			const PlatformInfo::FPlatformInfo* PlatformInfo = PlatformInfo::FindPlatformInfo(PlatformName);
diff --git a/Engine/Source/Programs/AutomationTool/AllDesktop/AllDesktop.Automation.csproj b/Engine/Source/Programs/AutomationTool/AllDesktop/AllDesktop.Automation.csproj
deleted file mode 100644
diff --git a/Engine/Source/Programs/AutomationTool/AllDesktop/AllDesktopPlatform.Automation.cs b/Engine/Source/Programs/AutomationTool/AllDesktop/AllDesktopPlatform.Automation.cs
deleted file mode 100644
diff --git a/Engine/Source/Programs/AutomationTool/AllDesktop/Properties/AssemblyInfo.cs b/Engine/Source/Programs/AutomationTool/AllDesktop/Properties/AssemblyInfo.cs
deleted file mode 100644
diff --git a/Engine/Source/Programs/UnrealBuildTool/Configuration/UEBuildTarget.cs b/Engine/Source/Programs/UnrealBuildTool/Configuration/UEBuildTarget.cs
--- a/Engine/Source/Programs/UnrealBuildTool/Configuration/UEBuildTarget.cs
+++ b/Engine/Source/Programs/UnrealBuildTool/Configuration/UEBuildTarget.cs
@@ -257,11 +257,6 @@ namespace UnrealBuildTool
 		/// </summary>
 		public static UnrealTargetPlatform LinuxAArch64 = FindOrAddByName("LinuxAArch64");
 
-		/// <summary>
-		/// All desktop platforms
-		/// </summary>
-		public static UnrealTargetPlatform AllDesktop = FindOrAddByName("AllDesktop");
-
 		/// <summary>
 		/// TVOS
 		/// </summary>
@@ -533,11 +528,6 @@ namespace UnrealBuildTool
 		/// </summary>
 		public static UnrealPlatformGroup XboxCommon = FindOrAddByName("XboxCommon");
 
-		/// <summary>
-		/// Target all desktop platforms (Win64, Mac, Linux) simultaneously
-		/// </summary>
-		public static UnrealPlatformGroup AllDesktop = FindOrAddByName("AllDesktop");
-
 		/// <summary>
 		/// Desktop group - used by UnrealPlatformClass.Desktop
 		/// </summary>
diff --git a/Engine/Source/Programs/UnrealBuildTool/Platform/Linux/UEBuildLinux.cs b/Engine/Source/Programs/UnrealBuildTool/Platform/Linux/UEBuildLinux.cs
--- a/Engine/Source/Programs/UnrealBuildTool/Platform/Linux/UEBuildLinux.cs
+++ b/Engine/Source/Programs/UnrealBuildTool/Platform/Linux/UEBuildLinux.cs
@@ -352,7 +352,7 @@ namespace UnrealBuildTool
 				return;
 			}
 
-			if ((Target.Platform == UnrealTargetPlatform.Win32) || (Target.Platform == UnrealTargetPlatform.Win64))
+			if (Target.Platform == UnrealTargetPlatform.Win64)
 			{
 				if (!Target.bBuildRequiresCookedData)
 				{
@@ -361,12 +361,7 @@ namespace UnrealBuildTool
 						if (Target.bBuildDeveloperTools)
 						{
 							Rules.DynamicallyLoadedModuleNames.Add("LinuxTargetPlatform");
-							Rules.DynamicallyLoadedModuleNames.Add("LinuxNoEditorTargetPlatform");
-							Rules.DynamicallyLoadedModuleNames.Add("LinuxAArch64NoEditorTargetPlatform");
-							Rules.DynamicallyLoadedModuleNames.Add("LinuxClientTargetPlatform");
-							Rules.DynamicallyLoadedModuleNames.Add("LinuxAArch64ClientTargetPlatform");
-							Rules.DynamicallyLoadedModuleNames.Add("LinuxServerTargetPlatform");
-							Rules.DynamicallyLoadedModuleNames.Add("LinuxAArch64ServerTargetPlatform");
+							Rules.DynamicallyLoadedModuleNames.Add("LinuxAArch64TargetPlatform");
 						}
 					}
 				}
@@ -375,12 +370,7 @@ namespace UnrealBuildTool
 				if (Target.bForceBuildTargetPlatforms && ModuleName == "TargetPlatform")
 				{
 					Rules.DynamicallyLoadedModuleNames.Add("LinuxTargetPlatform");
-					Rules.DynamicallyLoadedModuleNames.Add("LinuxNoEditorTargetPlatform");
-					Rules.DynamicallyLoadedModuleNames.Add("LinuxAArch64NoEditorTargetPlatform");
-					Rules.DynamicallyLoadedModuleNames.Add("LinuxClientTargetPlatform");
-					Rules.DynamicallyLoadedModuleNames.Add("LinuxAArch64ClientTargetPlatform");
-					Rules.DynamicallyLoadedModuleNames.Add("LinuxServerTargetPlatform");
-					Rules.DynamicallyLoadedModuleNames.Add("LinuxAArch64ServerTargetPlatform");
+					Rules.DynamicallyLoadedModuleNames.Add("LinuxAArch64TargetPlatform");
 				}
 			}
 		}
@@ -410,13 +400,7 @@ namespace UnrealBuildTool
 				if (Target.bForceBuildTargetPlatforms)
 				{
 					Rules.DynamicallyLoadedModuleNames.Add("LinuxTargetPlatform");
-					Rules.DynamicallyLoadedModuleNames.Add("LinuxNoEditorTargetPlatform");
-					Rules.DynamicallyLoadedModuleNames.Add("LinuxAArch64NoEditorTargetPlatform");
-					Rules.DynamicallyLoadedModuleNames.Add("LinuxClientTargetPlatform");
-					Rules.DynamicallyLoadedModuleNames.Add("LinuxAArch64ClientTargetPlatform");
-					Rules.DynamicallyLoadedModuleNames.Add("LinuxServerTargetPlatform");
-					Rules.DynamicallyLoadedModuleNames.Add("LinuxAArch64ServerTargetPlatform");
-					Rules.DynamicallyLoadedModuleNames.Add("AllDesktopTargetPlatform");
+					Rules.DynamicallyLoadedModuleNames.Add("LinuxAArch64TargetPlatform");
 				}
 
 				if (bBuildShaderFormats)
diff --git a/Engine/Source/Programs/UnrealBuildTool/Platform/Mac/UEBuildMac.cs b/Engine/Source/Programs/UnrealBuildTool/Platform/Mac/UEBuildMac.cs
--- a/Engine/Source/Programs/UnrealBuildTool/Platform/Mac/UEBuildMac.cs
+++ b/Engine/Source/Programs/UnrealBuildTool/Platform/Mac/UEBuildMac.cs
@@ -281,10 +281,6 @@ namespace UnrealBuildTool
 				if (Target.bForceBuildTargetPlatforms)
 				{
 					Rules.DynamicallyLoadedModuleNames.Add("MacTargetPlatform");
-					Rules.DynamicallyLoadedModuleNames.Add("MacNoEditorTargetPlatform");
-					Rules.DynamicallyLoadedModuleNames.Add("MacClientTargetPlatform");
-					Rules.DynamicallyLoadedModuleNames.Add("MacServerTargetPlatform");
-					Rules.DynamicallyLoadedModuleNames.Add("AllDesktopTargetPlatform");
 				}
 
 				if (bBuildShaderFormats)
diff --git a/Engine/Source/Programs/UnrealBuildTool/Platform/Windows/UEBuildWindows.cs b/Engine/Source/Programs/UnrealBuildTool/Platform/Windows/UEBuildWindows.cs
--- a/Engine/Source/Programs/UnrealBuildTool/Platform/Windows/UEBuildWindows.cs
+++ b/Engine/Source/Programs/UnrealBuildTool/Platform/Windows/UEBuildWindows.cs
@@ -2044,10 +2044,6 @@ namespace UnrealBuildTool
 				if (Target.bForceBuildTargetPlatforms)
 				{
 					Rules.DynamicallyLoadedModuleNames.Add("WindowsTargetPlatform");
-					Rules.DynamicallyLoadedModuleNames.Add("WindowsNoEditorTargetPlatform");
-					Rules.DynamicallyLoadedModuleNames.Add("WindowsServerTargetPlatform");
-					Rules.DynamicallyLoadedModuleNames.Add("WindowsClientTargetPlatform");
-					Rules.DynamicallyLoadedModuleNames.Add("AllDesktopTargetPlatform");
 				}
 
 				if (bBuildShaderFormats)
diff --git a/Engine/Source/Runtime/Core/Public/Linux/LinuxPlatformProperties.h b/Engine/Source/Runtime/Core/Public/Linux/LinuxPlatformProperties.h
--- a/Engine/Source/Runtime/Core/Public/Linux/LinuxPlatformProperties.h
+++ b/Engine/Source/Runtime/Core/Public/Linux/LinuxPlatformProperties.h
@@ -72,7 +72,7 @@ struct FLinuxPlatformProperties
 			return IS_AARCH64 ? "LinuxAArch64Client" : "LinuxClient";
 		}
 
-		return IS_AARCH64 ? "LinuxAArch64NoEditor" : "LinuxNoEditor";
+		return IS_AARCH64 ? "LinuxAArch64" : "LinuxNoEditor";
 	}
 
 	static FORCEINLINE bool RequiresCookedData( )
diff --git a/Engine/Source/Runtime/Engine/Engine.Build.cs b/Engine/Source/Runtime/Engine/Engine.Build.cs
--- a/Engine/Source/Runtime/Engine/Engine.Build.cs
+++ b/Engine/Source/Runtime/Engine/Engine.Build.cs
@@ -262,10 +262,6 @@ public class Engine : ModuleRules
 				DynamicallyLoadedModuleNames.AddRange(
 					new string[] {
 						"WindowsTargetPlatform",
-						"WindowsNoEditorTargetPlatform",
-						"WindowsServerTargetPlatform",
-						"WindowsClientTargetPlatform",
-						"AllDesktopTargetPlatform",
 						"WindowsPlatformEditor",
 					}
 				);
@@ -275,10 +271,6 @@ public class Engine : ModuleRules
 				DynamicallyLoadedModuleNames.AddRange(
 					new string[] {
 						"MacTargetPlatform",
-						"MacNoEditorTargetPlatform",
-						"MacServerTargetPlatform",
-						"MacClientTargetPlatform",
-						"AllDesktopTargetPlatform",
 						"MacPlatformEditor",
 					}
 				);
@@ -288,13 +280,6 @@ public class Engine : ModuleRules
 				DynamicallyLoadedModuleNames.AddRange(
 					new string[] {
 						"LinuxTargetPlatform",
-						"LinuxNoEditorTargetPlatform",
-						"LinuxAArch64NoEditorTargetPlatform",
-						"LinuxServerTargetPlatform",
-						"LinuxAArch64ServerTargetPlatform",
-						"LinuxClientTargetPlatform",
-						"LinuxAArch64ClientTargetPlatform",
-						"AllDesktopTargetPlatform",
 						"LinuxPlatformEditor",
 					}
 				);
