Added option to UBT project generation to allow users to specify include paths to be omitted from project files. This speeds up parsing, reduces memory footprint, and allows irrelevant symbols to be excluded from intellisense suggestions.

Example usage in BuildConfiguration.xml -

<VCProjectFileGenerator>
    <ExcludedIncludePaths>ThirdParty\WebRTC;ThirdParty\PhysX</ExcludedIncludePaths>
</VCProjectFileGenerator>

Example savings -

Adding ThirdParty\PhysX = 47% vcproj reduction (28MB -> 15MB)
Adding ThirdParty\ = 61% vcproj reduction (28MB -> 11MB)

diff --git a/Engine/Source/Programs/UnrealBuildTool/ProjectFiles/VisualStudio/VCProject.cs b/Engine/Source/Programs/UnrealBuildTool/ProjectFiles/VisualStudio/VCProject.cs
--- a/Engine/Source/Programs/UnrealBuildTool/ProjectFiles/VisualStudio/VCProject.cs
+++ b/Engine/Source/Programs/UnrealBuildTool/ProjectFiles/VisualStudio/VCProject.cs
@@ -707,6 +707,35 @@ namespace UnrealBuildTool
 			return ProjectFiles;
 		}
 
+		private string[] FilteredList = null;
+
+		bool IncludePathIsFilteredOut(DirectoryReference IncludePath)
+		{			
+			// Turn the filter string into an array, remove whitespace, and normalize any path statements the first time
+			// we are asked to check a path.
+			if (FilteredList == null)
+			{
+				IEnumerable<string> CleanPaths = Settings.ExcludedIncludePaths.Split(new[] { ';' }, StringSplitOptions.RemoveEmptyEntries)
+					.Select(P => P.Trim())
+					.Select(P => P.Replace('/', Path.DirectorySeparatorChar));
+
+				FilteredList = CleanPaths.ToArray();
+			}
+
+			if (FilteredList.Length > 0)
+			{
+				foreach (string Entry in FilteredList)
+				{
+					if (IncludePath.FullName.Contains(Entry))
+					{
+						return true;
+					}
+				}
+			}
+
+			return false;
+		}
+
 		/// <summary>
 		/// Append a list of include paths to a property list
 		/// </summary>
@@ -714,7 +743,7 @@ namespace UnrealBuildTool
 		/// <param name="BaseDir">Directory containing the source file</param>
 		/// <param name="BaseDirToIncludePaths">Map of base directory to set of include paths</param>
 		/// <param name="IgnorePaths">Set of paths to ignore</param>
-		static void AppendIncludePaths(StringBuilder Builder, DirectoryReference BaseDir, Dictionary<DirectoryReference, IncludePathsCollection> BaseDirToIncludePaths, HashSet<DirectoryReference> IgnorePaths)
+		void AppendIncludePaths(StringBuilder Builder, DirectoryReference BaseDir, Dictionary<DirectoryReference, IncludePathsCollection> BaseDirToIncludePaths, HashSet<DirectoryReference> IgnorePaths)
 		{
 			for (DirectoryReference CurrentDir = BaseDir; CurrentDir != null; CurrentDir = CurrentDir.ParentDirectory)
 			{
@@ -723,7 +752,7 @@ namespace UnrealBuildTool
 				{
 					foreach(DirectoryReference IncludePath in Collection.AbsolutePaths)
 					{
-						if (!IgnorePaths.Contains(IncludePath))
+						if (!IgnorePaths.Contains(IncludePath) && !IncludePathIsFilteredOut(IncludePath))
 						{
 							Builder.Append(NormalizeProjectPath(IncludePath.FullName));
 							Builder.Append(';');
@@ -783,7 +812,7 @@ namespace UnrealBuildTool
 
 				// Append the most common include paths to the search list.
 				if (Settings.MaxSharedIncludePaths > 0)
-				{
+				{	
 					foreach (DirectoryReference IncludePath in IncludePathToCount.OrderByDescending(x => x.Value).Select(x => x.Key))
 					{
 						string RelativePath = NormalizeProjectPath(IncludePath);
@@ -791,8 +820,12 @@ namespace UnrealBuildTool
 						{
 							break;
 						}
-						SharedIncludeSearchPaths.AppendFormat("{0};", RelativePath);
-						SharedIncludeSearchPathsSet.Add(IncludePath);
+
+						if (!IncludePathIsFilteredOut(IncludePath))
+						{
+							SharedIncludeSearchPathsSet.Add(IncludePath);
+							SharedIncludeSearchPaths.AppendFormat("{0};", RelativePath);
+						}						
 					}
 				}
 
@@ -1084,7 +1117,10 @@ namespace UnrealBuildTool
 					string VCFileType = GetVCFileType(AliasedFile.FileSystemPath);
 					if (VCFileType != "ClCompile")
 					{
-						VCProjectFileContent.AppendLine("    <{0} Include=\"{1}\"/>", VCFileType, EscapeFileName(AliasedFile.FileSystemPath));
+						if (!IncludePathIsFilteredOut(new DirectoryReference(AliasedFile.FileSystemPath)))
+						{
+							VCProjectFileContent.AppendLine("    <{0} Include=\"{1}\"/>", VCFileType, EscapeFileName(AliasedFile.FileSystemPath));
+						}
 					}
 					else
 					{
diff --git a/Engine/Source/Programs/UnrealBuildTool/ProjectFiles/VisualStudio/VCProjectFileGenerator.cs b/Engine/Source/Programs/UnrealBuildTool/ProjectFiles/VisualStudio/VCProjectFileGenerator.cs
--- a/Engine/Source/Programs/UnrealBuildTool/ProjectFiles/VisualStudio/VCProjectFileGenerator.cs
+++ b/Engine/Source/Programs/UnrealBuildTool/ProjectFiles/VisualStudio/VCProjectFileGenerator.cs
@@ -51,6 +51,13 @@ namespace UnrealBuildTool
 		[XmlConfigFile(Category = "VCProjectFileGenerator")]
 		public int MaxSharedIncludePaths = 24 * 1024;
 
+		/// <summary>
+		/// Semi-colon separated list of paths that should not be added to the projects include paths. Useful for omitting third-party headers
+		/// (e.g ThirdParty/WebRTC) from intellisense suggestions and reducing memory footprints.
+		/// </summary>
+		[XmlConfigFile(Category = "VCProjectFileGenerator")]
+		public string ExcludedIncludePaths = "";
+
 		/// <summary>
 		/// Whether to write a solution option (suo) file for the sln.
 		/// </summary>
