- Removed HTML5 from UnrealTargetPlatform (it's a PlatformExtension, and that should add the UTP)

diff --git a/Engine/Source/Programs/UnrealBuildTool/Configuration/UEBuildTarget.cs b/Engine/Source/Programs/UnrealBuildTool/Configuration/UEBuildTarget.cs
--- a/Engine/Source/Programs/UnrealBuildTool/Configuration/UEBuildTarget.cs
+++ b/Engine/Source/Programs/UnrealBuildTool/Configuration/UEBuildTarget.cs
@@ -243,10 +243,6 @@ namespace UnrealBuildTool
 		/// </summary>
 		public static UnrealTargetPlatform Android = FindOrAddByName("Android");
 
-		/// <summary>
-		/// </summary>
-		public static UnrealTargetPlatform HTML5 = FindOrAddByName("HTML5");
-
 		/// <summary>
 		/// Linux
 		/// </summary>
