- Start capturing log to a StartupTraceListener as soon as UBT starts (allowing the ToolMode to decide if it wants it or not via ToolModeOptions)

diff --git a/Engine/Source/Programs/UnrealBuildTool/Modes/BuildMode.cs b/Engine/Source/Programs/UnrealBuildTool/Modes/BuildMode.cs
--- a/Engine/Source/Programs/UnrealBuildTool/Modes/BuildMode.cs
+++ b/Engine/Source/Programs/UnrealBuildTool/Modes/BuildMode.cs
@@ -42,7 +42,7 @@ namespace UnrealBuildTool
 	/// <summary>
 	/// Builds a target
 	/// </summary>
-	[ToolMode("Build", ToolModeOptions.XmlConfig | ToolModeOptions.BuildPlatforms | ToolModeOptions.SingleInstance | ToolModeOptions.StartPrefetchingEngine | ToolModeOptions.ShowExecutionTime)]
+	[ToolMode("Build", ToolModeOptions.XmlConfig | ToolModeOptions.BuildPlatforms | ToolModeOptions.SingleInstance | ToolModeOptions.StartPrefetchingEngine | ToolModeOptions.ShowExecutionTime | ToolModeOptions.UseStartupTraceListener)]
 	class BuildMode : ToolMode
 	{
 		/// <summary>
@@ -97,10 +97,6 @@ namespace UnrealBuildTool
 		{
 			Arguments.ApplyTo(this);
 
-			// Initialize the log system, buffering the output until we can create the log file
-			StartupTraceListener StartupListener = new StartupTraceListener();
-			Trace.Listeners.Add(StartupListener);
-
 			// Write the command line
 			Log.TraceLog("Command line: {0}", Environment.CommandLine);
 
@@ -120,6 +116,9 @@ namespace UnrealBuildTool
 				BaseLogFileName = FileReference.Combine(UnrealBuildTool.EngineProgramSavedDirectory, "UnrealBuildTool", "Log.txt").FullName;
 			}
 
+			// find the StartupTraceListener in the listeners that was added super early on
+			StartupTraceListener StartupListener = Trace.Listeners.OfType<StartupTraceListener>().First();
+
 			// Create the log file, and flush the startup listener to it
 			if (!Arguments.HasOption("-NoLog") && !Log.HasFileWriter())
 			{
diff --git a/Engine/Source/Programs/UnrealBuildTool/System/ToolMode.cs b/Engine/Source/Programs/UnrealBuildTool/System/ToolMode.cs
--- a/Engine/Source/Programs/UnrealBuildTool/System/ToolMode.cs
+++ b/Engine/Source/Programs/UnrealBuildTool/System/ToolMode.cs
@@ -54,6 +54,11 @@ namespace UnrealBuildTool
 		/// Print out the total time taken to execute
 		/// </summary>
 		ShowExecutionTime = 64,
+
+		/// <summary>
+		/// Capture logs as early as possible in a StartupTraceListener object
+		/// </summary>
+		UseStartupTraceListener = 128,
 	}
 
 	/// <summary>
diff --git a/Engine/Source/Programs/UnrealBuildTool/UnrealBuildTool.cs b/Engine/Source/Programs/UnrealBuildTool/UnrealBuildTool.cs
--- a/Engine/Source/Programs/UnrealBuildTool/UnrealBuildTool.cs
+++ b/Engine/Source/Programs/UnrealBuildTool/UnrealBuildTool.cs
@@ -442,7 +442,11 @@ namespace UnrealBuildTool
 				Log.OutputLevel = Options.LogOutputLevel;
 				Log.IncludeTimestamps = Options.bLogTimestamps;
 				Log.IncludeProgramNameWithSeverityPrefix = Options.bLogFromMsBuild;
-				
+
+				// Always start capturing logs as early as possible to later copy to a log file if the ToolMode desires it (we have to start capturing before we get the ToolModeOptions below)
+				StartupTraceListener StartupTrace = new StartupTraceListener();
+				Log.AddTraceListener(StartupTrace);
+
 				// Configure the progress writer
 				ProgressWriter.bWriteMarkup = Options.bWriteProgressMarkup;
 
@@ -489,6 +493,12 @@ namespace UnrealBuildTool
 				// Get the options for which systems have to be initialized for this mode
 				ToolModeOptions ModeOptions = ModeType.GetCustomAttribute<ToolModeAttribute>().Options;
 
+				// if we don't care about the trace listener, toss it now
+				if ((ModeOptions & ToolModeOptions.UseStartupTraceListener) == 0)
+				{
+					Log.RemoveTraceListener(StartupTrace);
+				}
+
 				// Start prefetching the contents of the engine folder
 				if((ModeOptions & ToolModeOptions.StartPrefetchingEngine) != 0)
 				{
