- Moved Platform SDK classes to their own files to prepare for moving to BuildUtilities
- Deleted some empty subclasses

diff --git a/Engine/Source/Programs/UnrealBuildTool/Platform/Android/AndroidPlatformSDK.cs b/Engine/Source/Programs/UnrealBuildTool/Platform/Android/AndroidPlatformSDK.cs
new file mode 100644
--- /dev/null
+++ b/Engine/Source/Programs/UnrealBuildTool/Platform/Android/AndroidPlatformSDK.cs
@@ -0,0 +1,277 @@
+// Copyright Epic Games, Inc. All Rights Reserved.
+
+using System;
+using System.Collections.Generic;
+using System.IO;
+using Tools.DotNETCommon;
+using System.Text.RegularExpressions;
+using System.Linq;
+
+namespace UnrealBuildTool
+{
+	class AndroidPlatformSDK : UEBuildPlatformSDK
+	{
+		public override string GetDesiredVersion()
+		{
+			return "-23";
+		}
+
+		public override void GetValidVersionRange(out string MinVersion, out string MaxVersion)
+		{
+			MinVersion = "r21a";
+			MaxVersion = "r23a";
+		}
+
+		// the version we prefer
+		private string RecommendedVersion = "21b";
+
+		public override string GetInstalledSDKVersion()
+		{
+			string NDKPath = Environment.GetEnvironmentVariable("NDKROOT");
+
+			// don't register if we don't have an NDKROOT specified
+			if (String.IsNullOrEmpty(NDKPath))
+			{
+				return null;
+			}
+
+			NDKPath = NDKPath.Replace("\"", "");
+
+			// figure out the NDK version
+			string NDKToolchainVersion = null;
+			string SourcePropFilename = Path.Combine(NDKPath, "source.properties");
+			if (File.Exists(SourcePropFilename))
+			{
+				string RevisionString = "";
+				string[] PropertyContents = File.ReadAllLines(SourcePropFilename);
+				foreach (string PropertyLine in PropertyContents)
+				{
+					if (PropertyLine.StartsWith("Pkg.Revision"))
+					{
+						RevisionString = PropertyLine;
+						break;
+					}
+				}
+
+				int EqualsIndex = RevisionString.IndexOf('=');
+				if (EqualsIndex > 0)
+				{
+					string[] RevisionParts = RevisionString.Substring(EqualsIndex + 1).Trim().Split('.');
+					int RevisionMinor = int.Parse(RevisionParts.Length > 1 ? RevisionParts[1] : "0");
+					char RevisionLetter = Convert.ToChar('a' + RevisionMinor);
+					NDKToolchainVersion = "r" + RevisionParts[0] + (RevisionMinor > 0 ? Char.ToString(RevisionLetter) : "");
+				}
+			}
+			else
+			{
+				string ReleaseFilename = Path.Combine(NDKPath, "RELEASE.TXT");
+				if (File.Exists(ReleaseFilename))
+				{
+					string[] PropertyContents = File.ReadAllLines(SourcePropFilename);
+					NDKToolchainVersion = PropertyContents[0];
+				}
+			}
+
+			// return something like r10e, or null if anything went wrong above
+			return NDKToolchainVersion;
+
+		}
+
+
+		public override bool TryConvertVersionToInt(string StringValue, out UInt64 OutValue)
+		{
+			// convert r<num>[letter] to hex
+			Match Result = Regex.Match(StringValue, @"^r(\d*)([a-z])?");
+
+			if (Result.Success)
+			{
+				// 8 for number, 8 for letter, 8 for unused patch
+				int RevisionNumber = 0;
+				if (Result.Groups[2].Success)
+				{
+					RevisionNumber = (Result.Groups[2].Value[0] - 'a') + 1;
+				}
+				string VersionString = string.Format("{0}{1:00}{2:00}", Result.Groups[1], RevisionNumber, 0);
+				return UInt64.TryParse(VersionString, out OutValue);
+			}
+
+			OutValue = 0;
+			return false;
+		}
+
+
+		public override void PrintSDKInfo(LogEventType Verbosity, LogFormatOptions Options)
+		{
+			base.PrintSDKInfo(Verbosity, Options);
+
+			if (GetInstalledVersion() != RecommendedVersion)
+			{
+				Log.WriteLine(Verbosity, Options, "Note: Android toolchain NDK {0} recommended", RecommendedVersion);
+			}
+		}
+
+
+		protected override bool PlatformSupportsAutoSDKs()
+		{
+			return true;
+		}
+
+		protected override string GetRequiredSDKString()
+		{
+			return GetDesiredVersion();
+		}
+
+		protected override String GetRequiredScriptVersionString()
+		{
+			return "3.6";
+		}
+
+		// prefer auto sdk on android as correct 'manual' sdk detection isn't great at the moment.
+		protected override bool PreferAutoSDK()
+		{
+			return true;
+		}
+
+		private static bool ExtractPath(string Source, out string Path)
+		{
+			int start = Source.IndexOf('"');
+			int end = Source.LastIndexOf('"');
+			if (start != 1 && end != -1 && start < end)
+			{
+				++start;
+				Path = Source.Substring(start, end - start);
+				return true;
+			}
+			else
+			{
+				Path = "";
+			}
+
+			return false;
+		}
+
+		public static bool GetPath(ConfigHierarchy Ini, string SectionName, string Key, out string Value)
+		{
+			string temp;
+			if (Ini.TryGetValue(SectionName, Key, out temp))
+			{
+				return ExtractPath(temp, out Value);
+			}
+			else
+			{
+				Value = "";
+			}
+
+			return false;
+		}
+
+		/// <summary>
+		/// checks if the sdk is installed or has been synced
+		/// </summary>
+		/// <returns></returns>
+		protected virtual bool HasAnySDK()
+		{
+			string NDKPath = Environment.GetEnvironmentVariable("NDKROOT");
+			{
+				ConfigHierarchy configCacheIni = ConfigCache.ReadHierarchy(ConfigHierarchyType.Engine, (DirectoryReference)null, BuildHostPlatform.Current.Platform);
+				Dictionary<string, string> AndroidEnv = new Dictionary<string, string>();
+
+				Dictionary<string, string> EnvVarNames = new Dictionary<string, string> {
+														 {"ANDROID_HOME", "SDKPath"},
+														 {"NDKROOT", "NDKPath"},
+														 {"JAVA_HOME", "JavaPath"}
+														 };
+
+				string path;
+				foreach (KeyValuePair<string, string> kvp in EnvVarNames)
+				{
+					if (GetPath(configCacheIni, "/Script/AndroidPlatformEditor.AndroidSDKSettings", kvp.Value, out path) && !string.IsNullOrEmpty(path))
+					{
+						AndroidEnv.Add(kvp.Key, path);
+					}
+					else
+					{
+						string envValue = Environment.GetEnvironmentVariable(kvp.Key);
+						if (!String.IsNullOrEmpty(envValue))
+						{
+							AndroidEnv.Add(kvp.Key, envValue);
+						}
+					}
+				}
+
+				// If we are on Mono and we are still missing a key then go and find it from the .bash_profile
+				if (Utils.IsRunningOnMono && !EnvVarNames.All(s => AndroidEnv.ContainsKey(s.Key)))
+				{
+					string BashProfilePath = Path.Combine(Environment.GetFolderPath(Environment.SpecialFolder.Personal), ".bash_profile");
+					if (!File.Exists(BashProfilePath))
+					{
+						// Try .bashrc if didn't fine .bash_profile
+						BashProfilePath = Path.Combine(Environment.GetFolderPath(Environment.SpecialFolder.Personal), ".bashrc");
+					}
+					if (File.Exists(BashProfilePath))
+					{
+						string[] BashProfileContents = File.ReadAllLines(BashProfilePath);
+
+						// Walk backwards so we keep the last export setting instead of the first
+						for (int LineIndex = BashProfileContents.Length - 1; LineIndex >= 0; --LineIndex)
+						{
+							foreach (KeyValuePair<string, string> kvp in EnvVarNames)
+							{
+								if (AndroidEnv.ContainsKey(kvp.Key))
+								{
+									continue;
+								}
+
+								if (BashProfileContents[LineIndex].StartsWith("export " + kvp.Key + "="))
+								{
+									string PathVar = BashProfileContents[LineIndex].Split('=')[1].Replace("\"", "");
+									AndroidEnv.Add(kvp.Key, PathVar);
+								}
+							}
+						}
+					}
+				}
+
+				// Set for the process
+				foreach (KeyValuePair<string, string> kvp in AndroidEnv)
+				{
+					Environment.SetEnvironmentVariable(kvp.Key, kvp.Value);
+				}
+
+				// See if we have an NDK path now...
+				AndroidEnv.TryGetValue("NDKROOT", out NDKPath);
+			}
+
+			// we don't have an NDKROOT specified
+			if (String.IsNullOrEmpty(NDKPath))
+			{
+				return false;
+			}
+
+			NDKPath = NDKPath.Replace("\"", "");
+
+			// need a supported llvm
+			if (!Directory.Exists(Path.Combine(NDKPath, @"toolchains/llvm")))
+			{
+				return false;
+			}
+			return true;
+		}
+
+		protected override SDKStatus HasRequiredManualSDKInternal()
+		{
+			// if any autosdk setup has been done then the local process environment is suspect
+			if (HasSetupAutoSDK())
+			{
+				return SDKStatus.Invalid;
+			}
+
+			if (HasAnySDK())
+			{
+				return SDKStatus.Valid;
+			}
+
+			return SDKStatus.Invalid;
+		}
+	}
+}
diff --git a/Engine/Source/Programs/UnrealBuildTool/Platform/Android/UEBuildAndroid.cs b/Engine/Source/Programs/UnrealBuildTool/Platform/Android/UEBuildAndroid.cs
--- a/Engine/Source/Programs/UnrealBuildTool/Platform/Android/UEBuildAndroid.cs
+++ b/Engine/Source/Programs/UnrealBuildTool/Platform/Android/UEBuildAndroid.cs
@@ -3,12 +3,8 @@
 using System;
 using System.Collections.Generic;
 using System.Text;
-using System.Diagnostics;
 using System.IO;
-using System.Xml;
-using System.Linq;
 using Tools.DotNETCommon;
-using System.Text.RegularExpressions;
 
 namespace UnrealBuildTool
 {
@@ -443,271 +439,6 @@ namespace UnrealBuildTool
 		}
 	}
 
-	class AndroidPlatformSDK : UEBuildPlatformSDK
-	{
-		public override string GetDesiredVersion()
-		{
-			return "-23";
-		}
-
- 		public override void GetValidVersionRange(out string MinVersion, out string MaxVersion)
- 		{
-			MinVersion = "r21a";
- 			MaxVersion = "r23a";
- 		}
-
-		// the version we prefer
-		private string RecommendedVersion = "21b";
-
-		public override string GetInstalledSDKVersion()
-		{
-			string NDKPath = Environment.GetEnvironmentVariable("NDKROOT");
-
-			// don't register if we don't have an NDKROOT specified
-			if (String.IsNullOrEmpty(NDKPath))
-			{
-				return null;
-			}
-
-			NDKPath = NDKPath.Replace("\"", "");
-
-			// figure out the NDK version
-			string NDKToolchainVersion = null;
-			string SourcePropFilename = Path.Combine(NDKPath, "source.properties");
-			if (File.Exists(SourcePropFilename))
-			{
-				string RevisionString = "";
-				string[] PropertyContents = File.ReadAllLines(SourcePropFilename);
-				foreach (string PropertyLine in PropertyContents)
-				{
-					if (PropertyLine.StartsWith("Pkg.Revision"))
-					{
-						RevisionString = PropertyLine;
-						break;
-					}
-				}
-
-				int EqualsIndex = RevisionString.IndexOf('=');
-				if (EqualsIndex > 0)
-				{
-					string[] RevisionParts = RevisionString.Substring(EqualsIndex + 1).Trim().Split('.');
-					int RevisionMinor = int.Parse(RevisionParts.Length > 1 ? RevisionParts[1] : "0");
-					char RevisionLetter = Convert.ToChar('a' + RevisionMinor);
-					NDKToolchainVersion = "r" + RevisionParts[0] + (RevisionMinor > 0 ? Char.ToString(RevisionLetter) : "");
-				}
-			}
-			else
-			{
-				string ReleaseFilename = Path.Combine(NDKPath, "RELEASE.TXT");
-				if (File.Exists(ReleaseFilename))
-				{
-					string[] PropertyContents = File.ReadAllLines(SourcePropFilename);
-					NDKToolchainVersion = PropertyContents[0];
-				}
-			}
-
-			// return something like r10e, or null if anything went wrong above
-			return NDKToolchainVersion;
-
-		}
-
-
-		public override bool TryConvertVersionToInt(string StringValue, out UInt64 OutValue)
-		{
-			// convert r<num>[letter] to hex
-			Match Result = Regex.Match(StringValue, @"^r(\d*)([a-z])?");
-
-			if (Result.Success)
-			{
-				// 8 for number, 8 for letter, 8 for unused patch
-				int RevisionNumber = 0;
-				if (Result.Groups[2].Success)
-				{
-					RevisionNumber = (Result.Groups[2].Value[0] - 'a') + 1;
-				}
-				string VersionString = string.Format("{0}{1:00}{2:00}", Result.Groups[1], RevisionNumber, 0);
-				return UInt64.TryParse(VersionString, out OutValue);
-			}
-
-			OutValue = 0;
-			return false;
-		}
-
-
-		public override void PrintSDKInfo(LogEventType Verbosity, LogFormatOptions Options)
-		{
-			base.PrintSDKInfo(Verbosity, Options);
-
-			if (GetInstalledVersion() != RecommendedVersion)
-			{
-				Log.WriteLine(Verbosity, Options, "Note: Android toolchain NDK {0} recommended", RecommendedVersion);
-			}
-		}
-
-
-		protected override bool PlatformSupportsAutoSDKs()
-		{
-			return true;
-		}
-
-		protected override string GetRequiredSDKString()
-		{
-			return GetDesiredVersion();
-		}
-
-		protected override String GetRequiredScriptVersionString()
-		{
-			return "3.6";
-		}
-
-		// prefer auto sdk on android as correct 'manual' sdk detection isn't great at the moment.
-		protected override bool PreferAutoSDK()
-		{
-			return true;
-		}
-
-		private static bool ExtractPath(string Source, out string Path)
-		{
-			int start = Source.IndexOf('"');
-			int end = Source.LastIndexOf('"');
-			if (start != 1 && end != -1 && start < end)
-			{
-				++start;
-				Path = Source.Substring(start, end - start);
-				return true;
-			}
-			else
-			{
-				Path = "";
-			}
-
-			return false;
-		}
-
-		public static bool GetPath(ConfigHierarchy Ini, string SectionName, string Key, out string Value)
-		{
-			string temp;
-			if (Ini.TryGetValue(SectionName, Key, out temp))
-			{
-				return ExtractPath(temp, out Value);
-			}
-			else
-			{
-				Value = "";
-			}
-
-			return false;
-		}
-
-		/// <summary>
-		/// checks if the sdk is installed or has been synced
-		/// </summary>
-		/// <returns></returns>
-		protected virtual bool HasAnySDK()
-		{
-			string NDKPath = Environment.GetEnvironmentVariable("NDKROOT");
-			{
-				ConfigHierarchy configCacheIni = ConfigCache.ReadHierarchy(ConfigHierarchyType.Engine, (DirectoryReference)null, BuildHostPlatform.Current.Platform);
-				Dictionary<string, string> AndroidEnv = new Dictionary<string, string>();
-
-				Dictionary<string, string> EnvVarNames = new Dictionary<string, string> { 
-                                                         {"ANDROID_HOME", "SDKPath"}, 
-                                                         {"NDKROOT", "NDKPath"}, 
-                                                         {"JAVA_HOME", "JavaPath"}
-                                                         };
-
-				string path;
-				foreach (KeyValuePair<string, string> kvp in EnvVarNames)
-				{
-					if (GetPath(configCacheIni, "/Script/AndroidPlatformEditor.AndroidSDKSettings", kvp.Value, out path) && !string.IsNullOrEmpty(path))
-					{
-						AndroidEnv.Add(kvp.Key, path);
-					}
-					else
-					{
-						string envValue = Environment.GetEnvironmentVariable(kvp.Key);
-						if (!String.IsNullOrEmpty(envValue))
-						{
-							AndroidEnv.Add(kvp.Key, envValue);
-						}
-					}
-				}
-
-				// If we are on Mono and we are still missing a key then go and find it from the .bash_profile
-				if (Utils.IsRunningOnMono && !EnvVarNames.All(s => AndroidEnv.ContainsKey(s.Key)))
-				{
-					string BashProfilePath = Path.Combine(Environment.GetFolderPath(Environment.SpecialFolder.Personal), ".bash_profile");
-					if (!File.Exists(BashProfilePath))
-					{
-						// Try .bashrc if didn't fine .bash_profile
-						BashProfilePath = Path.Combine(Environment.GetFolderPath(Environment.SpecialFolder.Personal), ".bashrc");
-					}
-					if (File.Exists(BashProfilePath))
-					{
-						string[] BashProfileContents = File.ReadAllLines(BashProfilePath);
-
-						// Walk backwards so we keep the last export setting instead of the first
-						for (int LineIndex = BashProfileContents.Length - 1; LineIndex >= 0; --LineIndex)
-						{
-							foreach (KeyValuePair<string, string> kvp in EnvVarNames)
-							{
-								if (AndroidEnv.ContainsKey(kvp.Key))
-								{
-									continue;
-								}
-
-								if (BashProfileContents[LineIndex].StartsWith("export " + kvp.Key + "="))
-								{
-									string PathVar = BashProfileContents[LineIndex].Split('=')[1].Replace("\"", "");
-									AndroidEnv.Add(kvp.Key, PathVar);
-								}
-							}
-						}
-					}
-				}
-
-				// Set for the process
-				foreach (KeyValuePair<string, string> kvp in AndroidEnv)
-				{
-					Environment.SetEnvironmentVariable(kvp.Key, kvp.Value);
-				}
-
-				// See if we have an NDK path now...
-				AndroidEnv.TryGetValue("NDKROOT", out NDKPath);
-			}
-
-			// we don't have an NDKROOT specified
-			if (String.IsNullOrEmpty(NDKPath))
-			{
-				return false;
-			}
-
-			NDKPath = NDKPath.Replace("\"", "");
-
-			// need a supported llvm
-			if (!Directory.Exists(Path.Combine(NDKPath, @"toolchains/llvm")))
-			{
-				return false;
-			}
-			return true;
-		}
-
-		protected override SDKStatus HasRequiredManualSDKInternal()
-		{
-			// if any autosdk setup has been done then the local process environment is suspect
-			if (HasSetupAutoSDK())
-			{
-				return SDKStatus.Invalid;
-			}
-
-			if (HasAnySDK())
-			{
-				return SDKStatus.Valid;
-			}
-
-			return SDKStatus.Invalid;
-		}
-	}
 
 	class AndroidPlatformFactory : UEBuildPlatformFactory
 	{
diff --git a/Engine/Source/Programs/UnrealBuildTool/Platform/HoloLens/UEBuildHoloLens.cs b/Engine/Source/Programs/UnrealBuildTool/Platform/HoloLens/UEBuildHoloLens.cs
--- a/Engine/Source/Programs/UnrealBuildTool/Platform/HoloLens/UEBuildHoloLens.cs
+++ b/Engine/Source/Programs/UnrealBuildTool/Platform/HoloLens/UEBuildHoloLens.cs
@@ -142,7 +142,8 @@ namespace UnrealBuildTool
 		public static readonly Version MinimumSDKVersionForD3D12RHI = new Version(10, 0, 15063, 0);
 
 
-		public HoloLens(UnrealTargetPlatform InPlatform, HoloLensPlatformSDK InSDK) : base(InPlatform, InSDK)
+		public HoloLens(MicrosoftPlatformSDK InSDK) 
+			: base(UnrealTargetPlatform.HoloLens, InSDK)
 		{
 		}
 
@@ -743,11 +744,6 @@ namespace UnrealBuildTool
 
 
 
-	class HoloLensPlatformSDK : MicrosoftPlatformSDK
-	{
-
-	}
-
 	class HoloLensPlatformFactory : UEBuildPlatformFactory
 	{
 		public override UnrealTargetPlatform TargetPlatform
@@ -760,9 +756,9 @@ namespace UnrealBuildTool
 		/// </summary>
 		public override void RegisterBuildPlatforms()
 		{
-			HoloLensPlatformSDK SDK = new HoloLensPlatformSDK();
+			MicrosoftPlatformSDK SDK = new MicrosoftPlatformSDK();
 
-			UEBuildPlatform.RegisterBuildPlatform(new HoloLens(UnrealTargetPlatform.HoloLens, SDK));
+			UEBuildPlatform.RegisterBuildPlatform(new HoloLens(SDK));
 			UEBuildPlatform.RegisterPlatformWithGroup(UnrealTargetPlatform.HoloLens, UnrealPlatformGroup.Microsoft);
 			UEBuildPlatform.RegisterPlatformWithGroup(UnrealTargetPlatform.HoloLens, UnrealPlatformGroup.HoloLens);
 		}
diff --git a/Engine/Source/Programs/UnrealBuildTool/Platform/Linux/LinuxPlatformSDK.cs b/Engine/Source/Programs/UnrealBuildTool/Platform/Linux/LinuxPlatformSDK.cs
new file mode 100644
--- /dev/null
+++ b/Engine/Source/Programs/UnrealBuildTool/Platform/Linux/LinuxPlatformSDK.cs
@@ -0,0 +1,286 @@
+// Copyright Epic Games, Inc. All Rights Reserved.
+
+using System;
+using System.Collections.Generic;
+using System.IO;
+using Tools.DotNETCommon;
+using System.Text.RegularExpressions;
+
+namespace UnrealBuildTool
+{
+	class LinuxPlatformSDK : UEBuildPlatformSDK
+	{
+		public override string GetDesiredVersion()
+		{
+			return "v15_clang-8.0.1-centos7";
+			//			return "v16_clang-9.0.1-centos7";
+		}
+
+		public override void GetValidVersionRange(out string MinVersion, out string MaxVersion)
+		{
+			// all that matters is the number after the v, according to TryConvertVersionToInt()
+			MinVersion = "v10_clang-5.0.0-centos7";
+			MaxVersion = "v16_clang-9.0.1-centos7";
+		}
+
+		public override string GetAutoSDKPlatformName()
+		{
+			return TargetPlatformName;
+		}
+
+
+		public override string GetInstalledSDKVersion()
+		{
+			// @todo turnkey: ForceUseSystemCompiler() returns true, we should probably run system clang -V or similar to get version
+
+			string SDKDir = GetSDKLocation();
+			if (string.IsNullOrEmpty(SDKDir))
+			{
+				return null;
+			}
+
+			string VersionFile = Path.Combine(SDKDir, SDKVersionFileName());
+
+			if (!File.Exists(VersionFile))
+			{
+				// ErrorMessage = "Cannot use an old toolchain (missing " + PlatformSDK.SDKVersionFileName() + " file, assuming version earlier than v11)";
+				return null;
+			}
+
+			StreamReader SDKVersionFile = new StreamReader(VersionFile);
+			string SDKVersionString = SDKVersionFile.ReadLine();
+			SDKVersionFile.Close();
+
+			return SDKVersionString;
+		}
+
+
+		public override bool TryConvertVersionToInt(string StringValue, out UInt64 OutValue)
+		{
+			// Example: v11_clang-5.0.0-centos7
+			string FullVersionPattern = @"^v([0-9]+)_.*$";
+			Match Result = Regex.Match(StringValue, FullVersionPattern);
+			if (Result.Success)
+			{
+				return UInt64.TryParse(Result.Groups[1].Value, out OutValue);
+			}
+
+			OutValue = 0;
+			return false;
+		}
+
+
+
+
+
+		protected override bool PlatformSupportsAutoSDKs()
+		{
+			return true;
+		}
+
+
+
+
+		/// <summary>
+		/// Platform name (embeds architecture for now)
+		/// </summary>
+		private const string TargetPlatformName = "Linux_x64";
+
+		/// <summary>
+		/// Force using system compiler and error out if not possible
+		/// </summary>
+		private int bForceUseSystemCompiler = -1;
+
+		/// <summary>
+		/// Whether to compile with the verbose flag
+		/// </summary>
+		public bool bVerboseCompiler = false;
+
+		/// <summary>
+		/// Whether to link with the verbose flag
+		/// </summary>
+		public bool bVerboseLinker = false;
+
+		/// <summary>
+		/// Whether platform supports switching SDKs during runtime
+		/// </summary>
+		/// <returns>true if supports</returns>
+
+		/// <summary>
+		/// Returns a path to the internal SDK
+		/// </summary>
+		/// <returns>Valid path to the internal SDK, null otherwise</returns>
+		public string GetInternalSDKPath()
+		{
+			string SDKRoot = Environment.GetEnvironmentVariable(SDKRootEnvVar);
+			if (!String.IsNullOrEmpty(SDKRoot))
+			{
+				string AutoSDKPath = Path.Combine(SDKRoot, "Host" + BuildHostPlatform.Current.Platform, TargetPlatformName, GetDesiredVersion(), LinuxPlatform.DefaultHostArchitecture);
+				if (DirectoryReference.Exists(new DirectoryReference(AutoSDKPath)))
+				{
+					return AutoSDKPath;
+				}
+			}
+
+			string InTreeSDKPath = Path.Combine(LinuxPlatformSDK.GetInTreeSDKRoot().FullName, GetDesiredVersion(), LinuxPlatform.DefaultHostArchitecture);
+			if (DirectoryReference.Exists(new DirectoryReference(InTreeSDKPath)))
+			{
+				return InTreeSDKPath;
+			}
+
+			return null;
+		}
+
+		/// <summary>
+		/// Returns SDK string as required by the platform
+		/// </summary>
+		/// <returns>Valid SDK string</returns>
+		protected override string GetRequiredSDKString()
+		{
+			return GetDesiredVersion();
+		}
+
+		protected override String GetRequiredScriptVersionString()
+		{
+			return "3.0";
+		}
+
+		protected override bool PreferAutoSDK()
+		{
+			// having LINUX_ROOT set (for legacy reasons or for convenience of cross-compiling certain third party libs) should not make UBT skip AutoSDKs
+			return true;
+		}
+
+		public string HaveLinuxDependenciesFile()
+		{
+			// This file must have no extension so that GitDeps considers it a binary dependency - it will only be pulled by the Setup script if Linux is enabled.
+			return "HaveLinuxDependencies";
+		}
+
+		public string SDKVersionFileName()
+		{
+			return "ToolchainVersion.txt";
+		}
+
+		/// <summary>
+		/// Returns the in-tree root for the Linux Toolchain for this host platform.
+		/// </summary>
+		private static DirectoryReference GetInTreeSDKRoot()
+		{
+			return DirectoryReference.Combine(UnrealBuildTool.RootDirectory, "Engine/Extras/ThirdPartyNotUE/SDKs", "Host" + BuildHostPlatform.Current.Platform, TargetPlatformName);
+		}
+
+		/// <summary>
+		/// Whether a host can use its system sdk for this platform
+		/// </summary>
+		public virtual bool ForceUseSystemCompiler()
+		{
+			// by default tools chains don't parse arguments, but we want to be able to check the -bForceUseSystemCompiler flag.
+			if (bForceUseSystemCompiler == -1)
+			{
+				bForceUseSystemCompiler = 0;
+				string[] CmdLine = Environment.GetCommandLineArgs();
+
+				foreach (string CmdLineArg in CmdLine)
+				{
+					if (CmdLineArg.Equals("-ForceUseSystemCompiler", StringComparison.OrdinalIgnoreCase))
+					{
+						bForceUseSystemCompiler = 1;
+						break;
+					}
+				}
+			}
+
+			return bForceUseSystemCompiler == 1;
+		}
+
+		/// <summary>
+		/// Returns the root SDK path for all architectures
+		/// WARNING: Do not cache this value - it may be changed after sourcing OutputEnvVars.txt
+		/// </summary>
+		/// <returns>Valid SDK string</returns>
+		public virtual string GetSDKLocation()
+		{
+			// if new multi-arch toolchain is used, prefer it
+			string MultiArchRoot = Environment.GetEnvironmentVariable("LINUX_MULTIARCH_ROOT");
+
+			if (String.IsNullOrEmpty(MultiArchRoot))
+			{
+				// check if in-tree SDK is available
+				DirectoryReference InTreeSDKVersionRoot = GetInTreeSDKRoot();
+				if (InTreeSDKVersionRoot != null)
+				{
+					DirectoryReference InTreeSDKVersionPath = DirectoryReference.Combine(InTreeSDKVersionRoot, GetDesiredVersion());
+					if (DirectoryReference.Exists(InTreeSDKVersionPath))
+					{
+						MultiArchRoot = InTreeSDKVersionPath.FullName;
+					}
+				}
+			}
+			return MultiArchRoot;
+		}
+
+		/// <summary>
+		/// Returns the SDK path for a specific architecture
+		/// WARNING: Do not cache this value - it may be changed after sourcing OutputEnvVars.txt
+		/// </summary>
+		/// <returns>Valid SDK string</returns>
+		public virtual string GetBaseLinuxPathForArchitecture(string Architecture)
+		{
+			// if new multi-arch toolchain is used, prefer it
+			string MultiArchRoot = GetSDKLocation();
+			string BaseLinuxPath;
+
+			if (!String.IsNullOrEmpty(MultiArchRoot))
+			{
+				BaseLinuxPath = Path.Combine(MultiArchRoot, Architecture);
+			}
+			else
+			{
+				// use cross linux toolchain if LINUX_ROOT is specified
+				BaseLinuxPath = Environment.GetEnvironmentVariable("LINUX_ROOT");
+			}
+			return BaseLinuxPath;
+		}
+
+		/// <summary>
+		/// Whether the path contains a valid clang version
+		/// </summary>
+		private static bool IsValidClangPath(DirectoryReference BaseLinuxPath)
+		{
+			FileReference ClangPath = FileReference.Combine(BaseLinuxPath, @"bin", (BuildHostPlatform.Current.Platform == UnrealTargetPlatform.Win64) ? "clang++.exe" : "clang++");
+			return FileReference.Exists(ClangPath);
+		}
+
+		/// <summary>
+		/// Whether the required external SDKs are installed for this platform
+		/// </summary>
+		protected override SDKStatus HasRequiredManualSDKInternal()
+		{
+			// FIXME: UBT should loop across all the architectures and compile for all the selected ones.
+
+			// do not cache this value - it may be changed after sourcing OutputEnvVars.txt
+			string BaseLinuxPath = GetBaseLinuxPathForArchitecture(LinuxPlatform.DefaultHostArchitecture);
+
+			if (ForceUseSystemCompiler())
+			{
+				if (!String.IsNullOrEmpty(LinuxCommon.WhichClang()) || !String.IsNullOrEmpty(LinuxCommon.WhichGcc()))
+				{
+					return SDKStatus.Valid;
+				}
+			}
+			else if (!String.IsNullOrEmpty(BaseLinuxPath))
+			{
+				// paths to our toolchains if BaseLinuxPath is specified
+				BaseLinuxPath = BaseLinuxPath.Replace("\"", "");
+
+				if (IsValidClangPath(new DirectoryReference(BaseLinuxPath)))
+				{
+					return SDKStatus.Valid;
+				}
+			}
+
+			return SDKStatus.Invalid;
+		}
+	}
+}
diff --git a/Engine/Source/Programs/UnrealBuildTool/Platform/Linux/UEBuildLinux.cs b/Engine/Source/Programs/UnrealBuildTool/Platform/Linux/UEBuildLinux.cs
--- a/Engine/Source/Programs/UnrealBuildTool/Platform/Linux/UEBuildLinux.cs
+++ b/Engine/Source/Programs/UnrealBuildTool/Platform/Linux/UEBuildLinux.cs
@@ -609,281 +609,6 @@ namespace UnrealBuildTool
 		}
 	}
 
-	class LinuxPlatformSDK : UEBuildPlatformSDK
-	{
-		public override string GetDesiredVersion()
-		{
-			return "v15_clang-8.0.1-centos7";
-//			return "v16_clang-9.0.1-centos7";
-		}
-
-		public override void GetValidVersionRange(out string MinVersion, out string MaxVersion)
- 		{
-			// all that matters is the number after the v, according to TryConvertVersionToInt()
-			MinVersion = "v10_clang-5.0.0-centos7";
- 			MaxVersion = "v16_clang-9.0.1-centos7";
- 		}
-
-		public override string GetAutoSDKPlatformName()
-		{
-			return TargetPlatformName;
-		}
-
-
-		public override string GetInstalledSDKVersion()
-		{
-			// @todo turnkey: ForceUseSystemCompiler() returns true, we should probably run system clang -V or similar to get version
-
-			string SDKDir = GetSDKLocation();
-			if (string.IsNullOrEmpty(SDKDir))
-			{
-				return null;
-			}
-
-			string VersionFile = Path.Combine(SDKDir, SDKVersionFileName());
-
-			if (!File.Exists(VersionFile))
-			{
-				// ErrorMessage = "Cannot use an old toolchain (missing " + PlatformSDK.SDKVersionFileName() + " file, assuming version earlier than v11)";
-				return null;
-			}
-
-			StreamReader SDKVersionFile = new StreamReader(VersionFile);
-			string SDKVersionString = SDKVersionFile.ReadLine();
-			SDKVersionFile.Close();
-
-			return SDKVersionString;
-		}
-
-
-		public override bool TryConvertVersionToInt(string StringValue, out UInt64 OutValue)
-		{
-			// Example: v11_clang-5.0.0-centos7
-			string FullVersionPattern = @"^v([0-9]+)_.*$";
-			Match Result = Regex.Match(StringValue, FullVersionPattern);
-			if (Result.Success)
-			{
-				return UInt64.TryParse(Result.Groups[1].Value, out OutValue);
-			}
-
-			OutValue = 0;
-			return false;
-		}
-
-
-
-
-
-		protected override bool PlatformSupportsAutoSDKs()
-		{
-			return true;
-		}
-
-
-
-
-		/// <summary>
-		/// Platform name (embeds architecture for now)
-		/// </summary>
-		private const string TargetPlatformName = "Linux_x64";
-
-		/// <summary>
-		/// Force using system compiler and error out if not possible
-		/// </summary>
-		private int bForceUseSystemCompiler = -1;
-
-		/// <summary>
-		/// Whether to compile with the verbose flag
-		/// </summary>
-		public bool bVerboseCompiler = false;
-
-		/// <summary>
-		/// Whether to link with the verbose flag
-		/// </summary>
-		public bool bVerboseLinker = false;
-
-		/// <summary>
-		/// Whether platform supports switching SDKs during runtime
-		/// </summary>
-		/// <returns>true if supports</returns>
-
-		/// <summary>
-		/// Returns a path to the internal SDK
-		/// </summary>
-		/// <returns>Valid path to the internal SDK, null otherwise</returns>
-		public string GetInternalSDKPath()
-		{
-			string SDKRoot = Environment.GetEnvironmentVariable(SDKRootEnvVar);
-			if (!String.IsNullOrEmpty(SDKRoot))
-			{
-				string AutoSDKPath = Path.Combine(SDKRoot, "Host" + BuildHostPlatform.Current.Platform, TargetPlatformName, GetDesiredVersion(), LinuxPlatform.DefaultHostArchitecture);
-				if (DirectoryReference.Exists(new DirectoryReference(AutoSDKPath)))
-				{
-					return AutoSDKPath;
-				}
-			}
-
-			string InTreeSDKPath = Path.Combine(LinuxPlatformSDK.GetInTreeSDKRoot().FullName, GetDesiredVersion(), LinuxPlatform.DefaultHostArchitecture);
-			if (DirectoryReference.Exists(new DirectoryReference(InTreeSDKPath)))
-			{
-				return InTreeSDKPath;
-			}
-
-			return null;
-		}
-
-		/// <summary>
-		/// Returns SDK string as required by the platform
-		/// </summary>
-		/// <returns>Valid SDK string</returns>
-		protected override string GetRequiredSDKString()
-		{
-			return GetDesiredVersion();
-		}
-
-		protected override String GetRequiredScriptVersionString()
-		{
-			return "3.0";
-		}
-
-		protected override bool PreferAutoSDK()
-		{
-			// having LINUX_ROOT set (for legacy reasons or for convenience of cross-compiling certain third party libs) should not make UBT skip AutoSDKs
-			return true;
-		}
-
-		public string HaveLinuxDependenciesFile()
-		{
-			// This file must have no extension so that GitDeps considers it a binary dependency - it will only be pulled by the Setup script if Linux is enabled.
-			return "HaveLinuxDependencies";
-		}
-
-		public string SDKVersionFileName()
-		{
-			return "ToolchainVersion.txt";
-		}
-
-		/// <summary>
-		/// Returns the in-tree root for the Linux Toolchain for this host platform.
-		/// </summary>
-		private static DirectoryReference GetInTreeSDKRoot()
-		{
-			return DirectoryReference.Combine(UnrealBuildTool.RootDirectory, "Engine/Extras/ThirdPartyNotUE/SDKs", "Host" + BuildHostPlatform.Current.Platform, TargetPlatformName);
-		}
-
-		/// <summary>
-		/// Whether a host can use its system sdk for this platform
-		/// </summary>
-		public virtual bool ForceUseSystemCompiler()
-		{
-			// by default tools chains don't parse arguments, but we want to be able to check the -bForceUseSystemCompiler flag.
-			if (bForceUseSystemCompiler == -1)
-			{
-				bForceUseSystemCompiler = 0;
-				string[] CmdLine = Environment.GetCommandLineArgs();
-
-				foreach (string CmdLineArg in CmdLine)
-				{
-					if (CmdLineArg.Equals("-ForceUseSystemCompiler", StringComparison.OrdinalIgnoreCase))
-					{
-						bForceUseSystemCompiler = 1;
-						break;
-					}
-				}
-			}
-
-			return bForceUseSystemCompiler == 1;
-		}
-
-		/// <summary>
-		/// Returns the root SDK path for all architectures
-		/// WARNING: Do not cache this value - it may be changed after sourcing OutputEnvVars.txt
-		/// </summary>
-		/// <returns>Valid SDK string</returns>
-		public virtual string GetSDKLocation()
-		{
-			// if new multi-arch toolchain is used, prefer it
-			string MultiArchRoot = Environment.GetEnvironmentVariable("LINUX_MULTIARCH_ROOT");
-
-			if (String.IsNullOrEmpty(MultiArchRoot))
-			{
-				// check if in-tree SDK is available
-				DirectoryReference InTreeSDKVersionRoot = GetInTreeSDKRoot();
-				if (InTreeSDKVersionRoot != null)
-				{
-					DirectoryReference InTreeSDKVersionPath = DirectoryReference.Combine(InTreeSDKVersionRoot, GetDesiredVersion());
-					if (DirectoryReference.Exists(InTreeSDKVersionPath))
-					{
-						MultiArchRoot = InTreeSDKVersionPath.FullName;
-					}
-				}
-			}
-			return MultiArchRoot;
-		}
-
-		/// <summary>
-		/// Returns the SDK path for a specific architecture
-		/// WARNING: Do not cache this value - it may be changed after sourcing OutputEnvVars.txt
-		/// </summary>
-		/// <returns>Valid SDK string</returns>
-		public virtual string GetBaseLinuxPathForArchitecture(string Architecture)
-		{
-			// if new multi-arch toolchain is used, prefer it
-			string MultiArchRoot = GetSDKLocation();
-			string BaseLinuxPath;
-
-			if (!String.IsNullOrEmpty(MultiArchRoot))
-			{
-				BaseLinuxPath = Path.Combine(MultiArchRoot, Architecture);
-			}
-			else
-			{
-				// use cross linux toolchain if LINUX_ROOT is specified
-				BaseLinuxPath = Environment.GetEnvironmentVariable("LINUX_ROOT");
-			} 
-			return BaseLinuxPath;
-		}
-
-		/// <summary>
-		/// Whether the path contains a valid clang version
-		/// </summary>
-		private static bool IsValidClangPath(DirectoryReference BaseLinuxPath)
-		{
-			FileReference ClangPath = FileReference.Combine(BaseLinuxPath, @"bin", (BuildHostPlatform.Current.Platform == UnrealTargetPlatform.Win64) ? "clang++.exe" : "clang++");
-			return FileReference.Exists(ClangPath);
-		}
-
-		/// <summary>
-		/// Whether the required external SDKs are installed for this platform
-		/// </summary>
-		protected override SDKStatus HasRequiredManualSDKInternal()
-		{
-			// FIXME: UBT should loop across all the architectures and compile for all the selected ones.
-
-			// do not cache this value - it may be changed after sourcing OutputEnvVars.txt
-			string BaseLinuxPath = GetBaseLinuxPathForArchitecture(LinuxPlatform.DefaultHostArchitecture);
-
-			if (ForceUseSystemCompiler())
-			{
-				if (!String.IsNullOrEmpty(LinuxCommon.WhichClang()) || !String.IsNullOrEmpty(LinuxCommon.WhichGcc()))
-				{
-					return SDKStatus.Valid;
-				}
-			}
-			else if (!String.IsNullOrEmpty(BaseLinuxPath))
-			{
-				// paths to our toolchains if BaseLinuxPath is specified
-				BaseLinuxPath = BaseLinuxPath.Replace("\"", "");
-
-				if (IsValidClangPath(new DirectoryReference(BaseLinuxPath)))
-				{
-					return SDKStatus.Valid;
-				}
-			}
-
-			return SDKStatus.Invalid;
-		}
-	}
 
 	class LinuxPlatformFactory : UEBuildPlatformFactory
 	{
diff --git a/Engine/Source/Programs/UnrealBuildTool/Platform/Lumin/LuminPlatformSDK.cs b/Engine/Source/Programs/UnrealBuildTool/Platform/Lumin/LuminPlatformSDK.cs
new file mode 100644
--- /dev/null
+++ b/Engine/Source/Programs/UnrealBuildTool/Platform/Lumin/LuminPlatformSDK.cs
@@ -0,0 +1,105 @@
+// Copyright Epic Games, Inc. All Rights Reserved.
+
+using System;
+using System.Collections.Generic;
+using System.IO;
+using Tools.DotNETCommon;
+using System.Text.RegularExpressions;
+
+namespace UnrealBuildTool
+{
+	class LuminPlatformSDK : UEBuildPlatformSDK
+	{
+
+		public override string GetDesiredVersion()
+		{
+			return "0.23";
+		}
+
+		// not needed just returning GetDesiredVersion()
+		// 		public override void GetValidVersionRange(out string MinVersion, out string MaxVersion)
+		// 		{
+		// 			MinVersion = "0.23";
+		// 			MaxVersion = "0.23";
+		// 		}
+
+
+		public override string GetInstalledSDKVersion()
+		{
+			string EnvVarKey = "MLSDK";
+
+			string MLSDKPath = Environment.GetEnvironmentVariable(EnvVarKey);
+			if (String.IsNullOrEmpty(MLSDKPath))
+			{
+				return null;
+			}
+
+			String VersionFile = Path.Combine(MLSDKPath, "include/ml_version.h");
+			if (File.Exists(VersionFile))
+			{
+				string[] VersionText = File.ReadAllLines(VersionFile);
+
+				String MajorVersion = FindVersionNumber("MLSDK_VERSION_MAJOR", VersionText);
+				String MinorVersion = FindVersionNumber("MLSDK_VERSION_MINOR", VersionText);
+
+				return string.Format("{0}.{1}", MajorVersion, MinorVersion);
+			}
+
+			return null;
+		}
+
+
+		public override bool TryConvertVersionToInt(string StringValue, out UInt64 OutValue)
+		{
+			// make major high 16 bits, minor low 16
+			Match Result = Regex.Match(StringValue, @"^(\d+).(\d+)$");
+			if (Result.Success)
+			{
+				OutValue = UInt64.Parse(Result.Groups[1].Value) << 16 | UInt64.Parse(Result.Groups[2].Value) << 0;
+				return true;
+			}
+			OutValue = 0;
+			return false;
+		}
+
+
+
+
+		private string FindVersionNumber(string StringToFind, string[] AllLines)
+		{
+			string FoundVersion = "Unknown";
+			foreach (string CurrentLine in AllLines)
+			{
+				int Index = CurrentLine.IndexOf(StringToFind);
+				if (Index != -1)
+				{
+					FoundVersion = CurrentLine.Substring(Index + StringToFind.Length);
+					break;
+				}
+			}
+			return FoundVersion.Trim();
+		}
+
+
+		protected override bool PlatformSupportsAutoSDKs()
+		{
+			return true;
+		}
+
+		protected override string GetRequiredSDKString()
+		{
+			return GetDesiredVersion();
+		}
+
+		protected override String GetRequiredScriptVersionString()
+		{
+			return "Lumin_15";
+		}
+
+		// prefer auto sdk on android as correct 'manual' sdk detection isn't great at the moment.
+		protected override bool PreferAutoSDK()
+		{
+			return true;
+		}
+	}
+}
diff --git a/Engine/Source/Programs/UnrealBuildTool/Platform/Lumin/UEBuildLumin.cs b/Engine/Source/Programs/UnrealBuildTool/Platform/Lumin/UEBuildLumin.cs
--- a/Engine/Source/Programs/UnrealBuildTool/Platform/Lumin/UEBuildLumin.cs
+++ b/Engine/Source/Programs/UnrealBuildTool/Platform/Lumin/UEBuildLumin.cs
@@ -244,100 +244,6 @@ namespace UnrealBuildTool
 		}
 	}
 
-	class LuminPlatformSDK : UEBuildPlatformSDK
-	{
-
-		public override string GetDesiredVersion()
-		{
-			return "0.23";
-		}
-
-		// not needed just returning GetDesiredVersion()
-// 		public override void GetValidVersionRange(out string MinVersion, out string MaxVersion)
-// 		{
-// 			MinVersion = "0.23";
-// 			MaxVersion = "0.23";
-// 		}
-
-
-		public override string GetInstalledSDKVersion()
-		{
-			string EnvVarKey = "MLSDK";
-
-			string MLSDKPath = Environment.GetEnvironmentVariable(EnvVarKey);
-			if (String.IsNullOrEmpty(MLSDKPath))
-			{
-				return null;
-			}
-
-			String VersionFile = Path.Combine(MLSDKPath, "include/ml_version.h");
-			if (File.Exists(VersionFile))
-			{
-				string[] VersionText = File.ReadAllLines(VersionFile);
-
-				String MajorVersion = FindVersionNumber("MLSDK_VERSION_MAJOR", VersionText);
-				String MinorVersion = FindVersionNumber("MLSDK_VERSION_MINOR", VersionText);
-
-				return string.Format("{0}.{1}", MajorVersion, MinorVersion);
-			}
-
-			return null;
-		}
-
-
-		public override bool TryConvertVersionToInt(string StringValue, out UInt64 OutValue)
-		{
-			// make major high 16 bits, minor low 16
-			Match Result = Regex.Match(StringValue, @"^(\d+).(\d+)$");
-			if (Result.Success)
-			{
-				OutValue = UInt64.Parse(Result.Groups[1].Value) << 16 | UInt64.Parse(Result.Groups[2].Value) << 0;
-				return true;
-			}
-			OutValue = 0;
-			return false;
-		}
-
-
-
-
-		private string FindVersionNumber(string StringToFind, string[] AllLines)
-		{
-			string FoundVersion = "Unknown";
-			foreach (string CurrentLine in AllLines)
-			{
-				int Index = CurrentLine.IndexOf(StringToFind);
-				if (Index != -1)
-				{
-					FoundVersion = CurrentLine.Substring(Index + StringToFind.Length);
-					break;
-				}
-			}
-			return FoundVersion.Trim();
-		}
-
-
-		protected override bool PlatformSupportsAutoSDKs()
-		{
-			return true;
-		}
-
-		protected override string GetRequiredSDKString()
-		{
-			return GetDesiredVersion();
-		}
-
-		protected override String GetRequiredScriptVersionString()
-		{
-			return "Lumin_15";
-		}
-
-		// prefer auto sdk on android as correct 'manual' sdk detection isn't great at the moment.
-		protected override bool PreferAutoSDK()
-		{
-			return true;
-		}
-	}
 
 	class LuminPlatformFactory : UEBuildPlatformFactory
 	{
diff --git a/Engine/Source/Programs/UnrealBuildTool/Platform/Windows/UEBuildWindows.cs b/Engine/Source/Programs/UnrealBuildTool/Platform/Windows/UEBuildWindows.cs
--- a/Engine/Source/Programs/UnrealBuildTool/Platform/Windows/UEBuildWindows.cs
+++ b/Engine/Source/Programs/UnrealBuildTool/Platform/Windows/UEBuildWindows.cs
@@ -730,14 +730,14 @@ namespace UnrealBuildTool
 		/// </summary>
 		public static readonly bool bBuildLargeAddressAwareBinary = true;
 
-		WindowsPlatformSDK SDK;
+		MicrosoftPlatformSDK SDK;
 
 		/// <summary>
 		/// Constructor
 		/// </summary>
 		/// <param name="InPlatform">Creates a windows platform with the given enum value</param>
 		/// <param name="InSDK">The installed Windows SDK</param>
-		public WindowsPlatform(UnrealTargetPlatform InPlatform, WindowsPlatformSDK InSDK)
+		public WindowsPlatform(UnrealTargetPlatform InPlatform, MicrosoftPlatformSDK InSDK)
 			: base(InPlatform, InSDK)
 		{
 			SDK = InSDK;
@@ -2299,10 +2299,6 @@ namespace UnrealBuildTool
 		}
 	}
 
-	class WindowsPlatformSDK : MicrosoftPlatformSDK
-	{
-	}
-
 	class WindowsPlatformFactory : UEBuildPlatformFactory
 	{
 		public override UnrealTargetPlatform TargetPlatform
@@ -2315,7 +2311,7 @@ namespace UnrealBuildTool
 		/// </summary>
 		public override void RegisterBuildPlatforms()
 		{
-			WindowsPlatformSDK SDK = new WindowsPlatformSDK();
+			MicrosoftPlatformSDK SDK = new MicrosoftPlatformSDK();
 
 			// Register this build platform for Win64 (no more Win32)
 			UEBuildPlatform.RegisterBuildPlatform(new WindowsPlatform(UnrealTargetPlatform.Win64, SDK));
diff --git a/Engine/Source/Programs/UnrealBuildTool/UnrealBuildTool.csproj b/Engine/Source/Programs/UnrealBuildTool/UnrealBuildTool.csproj
--- a/Engine/Source/Programs/UnrealBuildTool/UnrealBuildTool.csproj
+++ b/Engine/Source/Programs/UnrealBuildTool/UnrealBuildTool.csproj
@@ -165,9 +165,12 @@
     <Compile Include="Executors\SNDBS.cs" />
     <Compile Include="Executors\XGE.cs" />
     <Compile Include="Modes\WriteMetadataMode.cs" />
+    <Compile Include="Platform\Android\AndroidPlatformSDK.cs" />
     <Compile Include="Platform\IOS\IOSPostBuildSyncMode.cs" />
     <Compile Include="Platform\IOS\MobileProvisionContents.cs" />
+    <Compile Include="Platform\Linux\LinuxPlatformSDK.cs" />
     <Compile Include="Platform\Lumin\LuminPackageManifest.cs" />
+    <Compile Include="Platform\Lumin\LuminPlatformSDK.cs" />
     <Compile Include="Platform\Lumin\UEDeployLumin.cs" />
     <Compile Include="Platform\Mac\ApplePlatformSDK.cs" />
     <Compile Include="Platform\Windows\MicrosoftPlatformSDK.cs" />
