- Workaround for capturing ps output on Mac UBT getting stuck inside ps's printf with an attached stdout pipe

diff --git a/Engine/Source/Programs/UnrealBuildTool/System/BuildHostPlatform.cs b/Engine/Source/Programs/UnrealBuildTool/System/BuildHostPlatform.cs
--- a/Engine/Source/Programs/UnrealBuildTool/System/BuildHostPlatform.cs
+++ b/Engine/Source/Programs/UnrealBuildTool/System/BuildHostPlatform.cs
@@ -288,12 +288,11 @@ namespace UnrealBuildTool
 		{
 			List<ProcessInfo> Result = new List<ProcessInfo>();
 
+			string TempFile = Path.Combine("/var/tmp", Path.GetTempFileName());
 			ProcessStartInfo StartInfo = new ProcessStartInfo();
-			StartInfo.FileName = "ps";
-			StartInfo.Arguments = "-eaw -o pid,comm";
+			StartInfo.FileName = "/bin/sh";
+			StartInfo.Arguments = "-c 'ps -eaw -o pid,comm > " + TempFile + "'";
 			StartInfo.CreateNoWindow = true;
-			StartInfo.UseShellExecute = false;
-			StartInfo.RedirectStandardOutput = true;
 
 			Process Proc = new Process();
 			Proc.StartInfo = StartInfo;
@@ -301,9 +300,9 @@ namespace UnrealBuildTool
 			{
 				Proc.Start();
 				Proc.WaitForExit();
-				for (string Line = Proc.StandardOutput.ReadLine(); Line != null; Line = Proc.StandardOutput.ReadLine())
+				foreach (string FileLine in File.ReadAllLines(TempFile))
 				{
-					Line = Line.Trim();
+					string Line = FileLine.Trim();
 					int PIDEnd = Line.IndexOf(' ');
 					string PIDString = Line.Substring(0, PIDEnd);
 					if (PIDString != "PID")
@@ -322,6 +321,7 @@ namespace UnrealBuildTool
 						catch { }
 					}
 				}
+				File.Delete(TempFile);
 
 			}
 			catch { }
