- Fixed AppleSDK version to int conversion routine

diff --git a/Engine/Source/Programs/UnrealBuildTool/Platform/Mac/ApplePlatformSDK.cs b/Engine/Source/Programs/UnrealBuildTool/Platform/Mac/ApplePlatformSDK.cs
--- a/Engine/Source/Programs/UnrealBuildTool/Platform/Mac/ApplePlatformSDK.cs
+++ b/Engine/Source/Programs/UnrealBuildTool/Platform/Mac/ApplePlatformSDK.cs
@@ -72,11 +72,11 @@ namespace UnrealBuildTool
 				OutValue = UInt64.Parse(Result.Groups[1].Value) << 24 | UInt64.Parse(Result.Groups[2].Value) << 16;
 				if (Result.Groups[4].Success)
 				{
-					OutValue = UInt64.Parse(Result.Groups[4].Value) << 8;
+					OutValue |= UInt64.Parse(Result.Groups[4].Value) << 8;
 				}
 				if (Result.Groups[6].Success)
 				{
-					OutValue = UInt64.Parse(Result.Groups[6].Value) << 0;
+					OutValue |= UInt64.Parse(Result.Groups[6].Value) << 0;
 				}
 				return true;
 			}
