- Allows the UnrealBuildTool.csprojh to open in Visual Studio Mac (and Linux as well assumedly). Also removed some stale conditional entries

diff --git a/Engine/Source/Programs/UnrealBuildTool/UnrealBuildTool.csproj b/Engine/Source/Programs/UnrealBuildTool/UnrealBuildTool.csproj
--- a/Engine/Source/Programs/UnrealBuildTool/UnrealBuildTool.csproj
+++ b/Engine/Source/Programs/UnrealBuildTool/UnrealBuildTool.csproj
@@ -82,39 +82,6 @@
     <Reference Include="System.Xml" />
     <Reference Include="System.Xml.Linq" />
   </ItemGroup>
-  <Choose>
-    <When Condition="Exists('Platform\XboxOne\UEBuildXboxOne.cs')">
-      <ItemGroup>
-        <Compile Include="Platform\XboxOne\UEBuildXboxOne.cs" />
-        <Compile Include="Platform\XboxOne\XboxOneToolChain.cs" />
-        <Compile Include="Platform\XboxOne\XboxOneDeploy.cs" />
-        <Compile Include="Platform\XboxOne\XboxOneExports.cs" />
-        <Compile Include="Platform\XboxOne\XboxOneProjectGenerator.cs" />
-        <Compile Include="Platform\XboxOne\XboxOneManifestGenerator.cs" />
-      </ItemGroup>
-    </When>
-  </Choose>
-  <Choose>
-    <When Condition="Exists('Platform\Switch\UEBuildSwitch.cs')">
-      <ItemGroup>
-        <Compile Include="Platform\Switch\SwitchExports.cs" />
-        <Compile Include="Platform\Switch\SwitchProjectGenerator.cs" />
-        <Compile Include="Platform\Switch\SwitchToolChain.cs" />
-        <Compile Include="Platform\Switch\UEBuildSwitch.cs" />
-        <Compile Include="Platform\Switch\UEDeploySwitch.cs" />
-      </ItemGroup>
-    </When>
-  </Choose>
-  <Choose>
-    <When Condition="Exists('Platform\PS4\UEBuildPS4.cs')">
-      <ItemGroup>
-        <Compile Include="Platform\PS4\PS4ProjectGenerator.cs" />
-        <Compile Include="Platform\PS4\PS4ToolChain.cs" />
-        <Compile Include="Platform\PS4\UEBuildPS4.cs" />
-        <Compile Include="Platform\PS4\UEDeployPS4.cs" />
-      </ItemGroup>
-    </When>
-  </Choose>
   <Choose>
     <When Condition="Exists('Platform\Lumin\UEBuildLumin.cs')" />
   </Choose>
@@ -130,15 +97,6 @@
       </ItemGroup>
     </When>
   </Choose>
-  <Choose>
-    <When Condition="Exists('Platform\Quail\UEBuildQuail.cs')">
-      <ItemGroup>
-        <Compile Include="Platform\Quail\QuailProjectGenerator.cs" />
-        <Compile Include="Platform\Quail\QuailToolChain.cs" />
-        <Compile Include="Platform\Quail\UEBuildQuail.cs" />
-      </ItemGroup>
-    </When>
-  </Choose>
   <ItemGroup>
     <Compile Include="Configuration\ModuleRulesContext.cs" />
     <Compile Include="Executors\HybridExecutor.cs" />
@@ -426,21 +384,37 @@
   </ItemGroup>
   <Import Project="$(MSBuildBinPath)\Microsoft.CSharp.targets" />
   <!-- The following will glob all platform extensions' UnrealBuildTool source
-  files and include them in the build. They will also appear in VisualStudio -->
-  <ItemGroup>
-    <Compile Include="../../../Platforms/*/Source/Programs/UnrealBuildTool/**/*.cs">
-      <!-- RecursiveDir metadata expands to [PlatformName]/Source/Programs/[etc]. The Replace() will
-      replace everything from the first slash to the end leaving just the platform name. Should the
-      Replace() fail then RecursiveDir is still an agreeable value -->
-      <Link>Platform/$([System.Text.RegularExpressions.Regex]::Replace(%(Compile.RecursiveDir), [\\/].+$, ``))/_</Link>
-    </Compile>
-    <Compile Include="../../../Restricted/*/Source/Programs/UnrealBuildTool/**/*.cs">
-      <!-- Similar logic to above for restricted folders, but just remove the first four parts of the
-	  path (eg. NoRedist/Source/Programs/UnrealBuildTool) and allow the files to be based anywhere in the
-	  project from there -->
+  files and include them in the build. They will also appear in VisualStudio.
+  The Regex stuff doesn't work in Visual Studio for Mac, so we just dump out the files in that case -->
+  <Choose>
+  <When Condition="'$(OS)' == 'Windows_NT'">
+    <ItemGroup>
+      <Compile Include="../../../Platforms/*/Source/Programs/UnrealBuildTool/**/*.cs">
+        <!-- RecursiveDir metadata expands to [PlatformName]/Source/Programs/[etc]. The Replace() will
+        replace everything from the first slash to the end leaving just the platform name. Should the
+        Replace() fail then RecursiveDir is still an agreeable value -->
+        <Link>Platform/$([System.Text.RegularExpressions.Regex]::Replace(%(Compile.RecursiveDir), [\\/].+$, ``))/_</Link>
+      </Compile>
+      <Compile Include="../../../Restricted/*/Source/Programs/UnrealBuildTool/**/*.cs">
+        <!-- Similar logic to above for restricted folders, but just remove the first four parts of the
+      path (eg. NoRedist/Source/Programs/UnrealBuildTool) and allow the files to be based anywhere in the
+      project from there 
+      -->
       <Link>$([System.Text.RegularExpressions.Regex]::Replace(%(Compile.RecursiveDir), ^([^\\/]+[\\/]){4}, ``))/_</Link>
-    </Compile>
-  </ItemGroup>
+      </Compile>
+    </ItemGroup>
+  </When>
+  <When Condition="'$(OS)' != 'Windows_NT'">
+    <ItemGroup>
+      <Compile Include="../../../Platforms/*/Source/Programs/UnrealBuildTool/**/*.cs">
+        <Link>Platform/%(Filename).cs</Link>
+      </Compile>
+      <Compile Include="../../../Restricted/*/Source/Programs/UnrealBuildTool/**/*.cs">
+        <Link>Restricted/%(Filename).cs</Link>
+      </Compile>
+    </ItemGroup>
+  </When>
+  </Choose>
   <!-- To modify your build process, add your task inside one of the targets below and uncomment it. 
        Other similar extension points exist, see Microsoft.Common.targets.
   <Target Name="BeforeBuild">
