- Much fixed UBT.csproj on VS Mac (and added support Restricted NDA platform files), now works on all platforms with the same regex stuff

diff --git a/Engine/Source/Programs/UnrealBuildTool/UnrealBuildTool.csproj b/Engine/Source/Programs/UnrealBuildTool/UnrealBuildTool.csproj
--- a/Engine/Source/Programs/UnrealBuildTool/UnrealBuildTool.csproj
+++ b/Engine/Source/Programs/UnrealBuildTool/UnrealBuildTool.csproj
@@ -385,36 +385,24 @@
   <Import Project="$(MSBuildBinPath)\Microsoft.CSharp.targets" />
   <!-- The following will glob all platform extensions' UnrealBuildTool source
   files and include them in the build. They will also appear in VisualStudio.
-  The Regex stuff doesn't work in Visual Studio for Mac, so we just dump out the files in that case -->
-  <Choose>
-  <When Condition="'$(OS)' == 'Windows_NT'">
-    <ItemGroup>
-      <Compile Include="../../../Platforms/*/Source/Programs/UnrealBuildTool/**/*.cs">
-        <!-- RecursiveDir metadata expands to [PlatformName]/Source/Programs/[etc]. The Replace() will
-        replace everything from the first slash to the end leaving just the platform name. Should the
-        Replace() fail then RecursiveDir is still an agreeable value -->
-        <Link>Platform/$([System.Text.RegularExpressions.Regex]::Replace(%(Compile.RecursiveDir), [\\/].+$, ``))/_</Link>
-      </Compile>
-      <Compile Include="../../../Restricted/*/Source/Programs/UnrealBuildTool/**/*.cs">
-        <!-- Similar logic to above for restricted folders, but just remove the first four parts of the
-      path (eg. NoRedist/Source/Programs/UnrealBuildTool) and allow the files to be based anywhere in the
-      project from there 
-      -->
-      <Link>$([System.Text.RegularExpressions.Regex]::Replace(%(Compile.RecursiveDir), ^([^\\/]+[\\/]){4}, ``))/_</Link>
-      </Compile>
-    </ItemGroup>
-  </When>
-  <When Condition="'$(OS)' != 'Windows_NT'">
-    <ItemGroup>
-      <Compile Include="../../../Platforms/*/Source/Programs/UnrealBuildTool/**/*.cs">
-        <Link>Platform/%(Filename).cs</Link>
-      </Compile>
-      <Compile Include="../../../Restricted/*/Source/Programs/UnrealBuildTool/**/*.cs">
-        <Link>Restricted/%(Filename).cs</Link>
-      </Compile>
-    </ItemGroup>
-  </When>
-  </Choose>
+  RecursiveDir does not work in VS Mac (and maybe Linux?), so only use it in Windows. 
+  See: https://developercommunity.visualstudio.com/content/problem/527478/recursivedir-causes-infinitely-nested-subdirectori.html -->
+  <ItemGroup>
+    <Compile Include="../../../Platforms/*/Source/Programs/UnrealBuildTool/**/*.cs">
+      <!-- Since RecursiveDir is bugged on VS Mac, we use the FullPath, and use Regex to pull out portions of the path to get Platform name -->
+      <Link>Platform/$([System.Text.RegularExpressions.Regex]::Replace(%(FullPath), '^(.*[\\/]Platforms[\\/])(.+?)[\\/](.*)$', '$2'))/%(Filename)%(Extension)</Link>
+    </Compile>
+    <Compile Include="../../../Restricted/*/Platforms/*/Source/Programs/UnrealBuildTool/**/*.cs">
+      <!-- Restricted may or may not be in a Platforms dir, so just use the subdir directly, even tho it makes for a
+       messier directory structure for these few files -->
+      <Link>Platform/$([System.Text.RegularExpressions.Regex]::Replace(%(FullPath), '^(.*[\\/]Platforms[\\/])(.+?)[\\/](.*)$', '$2'))/%(Filename)%(Extension)</Link>
+    </Compile>
+    <Compile Include="../../../Restricted/*/Source/Programs/UnrealBuildTool/**/*.cs">
+      <!-- Restricted may or may not be in a Platforms dir, so just use the subdir directly, even tho it makes for a
+       messier directory structure for these few files -->
+      <Link>$([System.Text.RegularExpressions.Regex]::Replace(%(FullPath), '^.+?[\\/]UnrealBuildTool[\\/]', ''))</Link>
+    </Compile>
+  </ItemGroup>
   <!-- To modify your build process, add your task inside one of the targets below and uncomment it. 
        Other similar extension points exist, see Microsoft.Common.targets.
   <Target Name="BeforeBuild">
