clang v16 is the desired version not v15

diff --git a/Engine/Source/Programs/UnrealBuildTool/Platform/Linux/LinuxPlatformSDK.cs b/Engine/Source/Programs/UnrealBuildTool/Platform/Linux/LinuxPlatformSDK.cs
--- a/Engine/Source/Programs/UnrealBuildTool/Platform/Linux/LinuxPlatformSDK.cs
+++ b/Engine/Source/Programs/UnrealBuildTool/Platform/Linux/LinuxPlatformSDK.cs
@@ -12,8 +12,7 @@ namespace UnrealBuildTool
 	{
 		public override string GetDesiredVersion()
 		{
-			return "v15_clang-8.0.1-centos7";
-			//			return "v16_clang-9.0.1-centos7";
+			return "v16_clang-9.0.1-centos7";
 		}
 
 		public override void GetValidVersionRange(out string MinVersion, out string MaxVersion)
