Fix IsHostOnVpn() always returning false.

- Without that, it proceeds using XGE even on a non-local coordinator.

(Edigrating 14066457).

diff --git a/Engine/Source/Programs/UnrealBuildTool/Executors/XGE.cs b/Engine/Source/Programs/UnrealBuildTool/Executors/XGE.cs
--- a/Engine/Source/Programs/UnrealBuildTool/Executors/XGE.cs
+++ b/Engine/Source/Programs/UnrealBuildTool/Executors/XGE.cs
@@ -245,7 +245,7 @@ namespace UnrealBuildTool
 							if (Subnet.Contains(AddressBytes))
 							{
 								Log.TraceInformationOnce("XGE coordinator {0} will be not be used over VPN (adapter '{1}' with IP {2} is in subnet {3}). Set <XGE><bAllowOverVpn>true</bAllowOverVpn></XGE> in BuildConfiguration.xml to override.", HostName, Interface.Description, UnicastAddressInfo.Address, Subnet);
-								return false;
+								return true;
 							}
 						}
 					}
