Allow UsingLld to be overriden

diff --git a/Engine/Source/Programs/UnrealBuildTool/Platform/Linux/LinuxToolChain.cs b/Engine/Source/Programs/UnrealBuildTool/Platform/Linux/LinuxToolChain.cs
--- a/Engine/Source/Programs/UnrealBuildTool/Platform/Linux/LinuxToolChain.cs
+++ b/Engine/Source/Programs/UnrealBuildTool/Platform/Linux/LinuxToolChain.cs
@@ -1437,7 +1437,7 @@ namespace UnrealBuildTool
 			return Result;
 		}
 
-		bool UsingLld(string Architecture)
+		protected virtual bool UsingLld(string Architecture)
 		{
 			return bUseLld && (Architecture.StartsWith("x86_64") || (CompilerVersionMajor >= 9));
 		}
