Fixed retry in Linux linking batch files

Batch files require Windows line endings for goto to work properly. Windows performs line ending conversion automatically in some contexts, but not when running the script with cmd.exe /c as UBT does.

diff --git a/Engine/Source/Programs/UnrealBuildTool/Platform/Linux/LinuxToolChain.cs b/Engine/Source/Programs/UnrealBuildTool/Platform/Linux/LinuxToolChain.cs
--- a/Engine/Source/Programs/UnrealBuildTool/Platform/Linux/LinuxToolChain.cs
+++ b/Engine/Source/Programs/UnrealBuildTool/Platform/Linux/LinuxToolChain.cs
@@ -255,7 +255,7 @@ namespace UnrealBuildTool
 		internal string GetDumpEncodeDebugCommand(LinkEnvironment LinkEnvironment, FileItem OutputFile)
 		{
 			bool bUseCmdExe = BuildHostPlatform.Current.Platform == UnrealTargetPlatform.Win64 || BuildHostPlatform.Current.Platform == UnrealTargetPlatform.Win32;
-			string DumpCommand = bUseCmdExe ? "\"{0}\" \"{1}\" \"{2}\" 2>NUL\n" : "\"{0}\" -c -o \"{2}\" \"{1}\"\n";
+			string DumpCommand = bUseCmdExe ? "\"{0}\" \"{1}\" \"{2}\" 2>NUL" : "\"{0}\" -c -o \"{2}\" \"{1}\"";
 			FileItem EncodedBinarySymbolsFile = FileItem.GetItemByPath(Path.Combine(LinkEnvironment.OutputDirectory.FullName, OutputFile.Location.GetFileNameWithoutExtension() + ".sym"));
 			FileItem SymbolsFile  = FileItem.GetItemByPath(Path.Combine(LinkEnvironment.LocalShadowDirectory.FullName, OutputFile.Location.GetFileName() + ".psym"));
 			FileItem StrippedFile = FileItem.GetItemByPath(Path.Combine(LinkEnvironment.LocalShadowDirectory.FullName, OutputFile.Location.GetFileName() + "_nodebug"));
@@ -266,17 +266,18 @@ namespace UnrealBuildTool
 				SymbolsFile = FileItem.GetItemByPath(Path.Combine(LinkEnvironment.OutputDirectory.FullName, OutputFile.Location.GetFileNameWithoutExtension() + ".psym"));
 			}
 
-			string Out = "";
+			StringWriter Out = new StringWriter();
+			Out.NewLine = bUseCmdExe ? "\r\n" : "\n";
 
 			// dump_syms
-			Out += string.Format(DumpCommand,
+			Out.WriteLine(DumpCommand,
 				DumpSymsPath,
 				OutputFile.AbsolutePath,
 				SymbolsFile.AbsolutePath
 			);
 
 			// encode breakpad symbols
-			Out += string.Format("\"{0}\" \"{1}\" \"{2}\"\n",
+			Out.WriteLine("\"{0}\" \"{1}\" \"{2}\"",
 				BreakpadEncoderPath,
 				SymbolsFile.AbsolutePath,
 				EncodedBinarySymbolsFile.AbsolutePath
@@ -286,29 +287,29 @@ namespace UnrealBuildTool
 			{
 				if (MaxBinarySizeOverrideForObjcopy > 0 && bUseCmdExe)
 				{
-					Out += string.Format("for /F \"tokens=*\" %%F in (\"{0}\") DO set size=%%~zF\n",
+					Out.WriteLine("for /F \"tokens=*\" %%F in (\"{0}\") DO set size=%%~zF",
 						OutputFile.AbsolutePath
 					);
 
-					Out += string.Format("if %size% LSS {0} (\n", MaxBinarySizeOverrideForObjcopy);
+					Out.WriteLine("if %size% LSS {0} (", MaxBinarySizeOverrideForObjcopy);
 				}
 
 				// objcopy stripped file
-				Out += string.Format("\"{0}\" --strip-all \"{1}\" \"{2}\"\n",
+				Out.WriteLine("\"{0}\" --strip-all \"{1}\" \"{2}\"",
 					GetObjcopyPath(LinkEnvironment.Architecture),
 					OutputFile.AbsolutePath,
 					StrippedFile.AbsolutePath
 				);
 
 				// objcopy debug file
-				Out += string.Format("\"{0}\" --only-keep-debug \"{1}\" \"{2}\"\n",
+				Out.WriteLine("\"{0}\" --only-keep-debug \"{1}\" \"{2}\"",
 					GetObjcopyPath(LinkEnvironment.Architecture),
 					OutputFile.AbsolutePath,
 					DebugFile.AbsolutePath
 				);
 
 				// objcopy link debug file to final so
-				Out += string.Format("\"{0}\" --add-gnu-debuglink=\"{1}\" \"{2}\" \"{3}.temp\"\n",
+				Out.WriteLine("\"{0}\" --add-gnu-debuglink=\"{1}\" \"{2}\" \"{3}.temp\"",
 					GetObjcopyPath(LinkEnvironment.Architecture),
 					DebugFile.AbsolutePath,
 					StrippedFile.AbsolutePath,
@@ -318,7 +319,7 @@ namespace UnrealBuildTool
 				if (bUseCmdExe)
 				{
 					// Only move the temp final elf file once its done being linked by objcopy
-					Out += string.Format("move /Y \"{0}.temp\" \"{1}\"\n",
+					Out.WriteLine("move /Y \"{0}.temp\" \"{1}\"",
 						OutputFile.AbsolutePath,
 						OutputFile.AbsolutePath
 					);
@@ -326,21 +327,21 @@ namespace UnrealBuildTool
 					if (MaxBinarySizeOverrideForObjcopy > 0)
 					{
 						// If we have an override size, then we need to create a dummy file if that size is exceeded
-						Out += string.Format(") ELSE (\necho DummyDebug >> \"{0}\"\n)\n",
-							DebugFile.AbsolutePath
-						);
+						Out.WriteLine(") ELSE (");
+						Out.WriteLine("echo DummyDebug >> \"{0}\"", DebugFile.AbsolutePath);
+						Out.WriteLine();
 					}
 				}
 				else
 				{
 					// Only move the temp final elf file once its done being linked by objcopy
-					Out += string.Format("mv \"{0}.temp\" \"{1}\"\n",
+					Out.WriteLine("mv \"{0}.temp\" \"{1}\"",
 						OutputFile.AbsolutePath,
 						OutputFile.AbsolutePath
 					);
 
 					// Change the debug file to normal permissions. It was taking on the +x rights from the output file
-					Out += string.Format("chmod 644 \"{0}\"\n",
+					Out.WriteLine("chmod 644 \"{0}\"",
 						DebugFile.AbsolutePath
 					);
 				}
@@ -348,12 +349,12 @@ namespace UnrealBuildTool
 			else
 			{
 				// If we have disabled objcopy then we need to create a dummy debug file
-				Out += string.Format("echo DummyDebug >> {0}\n",
+				Out.WriteLine("echo DummyDebug >> \"{0}\"",
 					DebugFile.AbsolutePath
 				);
 			}
 
-			return Out;
+			return Out.ToString();
 		}
 
 		/// <summary>
@@ -1902,32 +1903,36 @@ namespace UnrealBuildTool
 			{
 				if (bUseCmdExe)
 				{
-					LinkWriter.Write("@echo off\n");
-					LinkWriter.Write("rem Automatically generated by UnrealBuildTool\n");
-					LinkWriter.Write("rem *DO NOT EDIT*\n\n");
-					LinkWriter.Write("set Retries=0\n");
-					LinkWriter.Write(":linkloop\n");
-					LinkWriter.Write("if %Retries% GEQ 10 goto failedtorelink\n");
-					LinkWriter.Write(LinkCommandString + "\n");
-					LinkWriter.Write("if %errorlevel% neq 0 goto sleepandretry\n");
-					LinkWriter.Write(GetDumpEncodeDebugCommand(LinkEnvironment, OutputFile) + "\n");
-					LinkWriter.Write("exit 0\n");
-					LinkWriter.Write(":sleepandretry\n");
-					LinkWriter.Write("ping 127.0.0.1 -n 1 -w 5000 >NUL 2>NUL\n");     // timeout complains about lack of redirection
-					LinkWriter.Write("set /a Retries+=1\n");
-					LinkWriter.Write("goto linkloop\n");
-					LinkWriter.Write(":failedtorelink\n");
-					LinkWriter.Write("echo Failed to link {0} after %Retries% retries\n", OutputFile.AbsolutePath);
-					LinkWriter.Write("exit 1\n");
+					LinkWriter.NewLine = "\r\n";
+					LinkWriter.WriteLine("@echo off");
+					LinkWriter.WriteLine("rem Automatically generated by UnrealBuildTool");
+					LinkWriter.WriteLine("rem *DO NOT EDIT*");
+					LinkWriter.WriteLine();
+					LinkWriter.WriteLine("set Retries=0");
+					LinkWriter.WriteLine(":linkloop");
+					LinkWriter.WriteLine("if %Retries% GEQ 10 goto failedtorelink");
+					LinkWriter.WriteLine(LinkCommandString);
+					LinkWriter.WriteLine("if %errorlevel% neq 0 goto sleepandretry");
+					LinkWriter.WriteLine(GetDumpEncodeDebugCommand(LinkEnvironment, OutputFile));
+					LinkWriter.WriteLine("exit 0");
+					LinkWriter.WriteLine(":sleepandretry");
+					LinkWriter.WriteLine("ping 127.0.0.1 -n 1 -w 5000 >NUL 2>NUL");     // timeout complains about lack of redirection
+					LinkWriter.WriteLine("set /a Retries+=1");
+					LinkWriter.WriteLine("goto linkloop");
+					LinkWriter.WriteLine(":failedtorelink");
+					LinkWriter.WriteLine("echo Failed to link {0} after %Retries% retries", OutputFile.AbsolutePath);
+					LinkWriter.WriteLine("exit 1");
 				}
 				else
 				{
-					LinkWriter.Write("#!/bin/sh\n");
-					LinkWriter.Write("# Automatically generated by UnrealBuildTool\n");
-					LinkWriter.Write("# *DO NOT EDIT*\n\n");
-					LinkWriter.Write("set -o errexit\n");
-					LinkWriter.Write(LinkCommandString + "\n");
-					LinkWriter.Write(GetDumpEncodeDebugCommand(LinkEnvironment, OutputFile) + "\n");
+					LinkWriter.NewLine = "\n";
+					LinkWriter.WriteLine("#!/bin/sh");
+					LinkWriter.WriteLine("# Automatically generated by UnrealBuildTool");
+					LinkWriter.WriteLine("# *DO NOT EDIT*");
+					LinkWriter.WriteLine();
+					LinkWriter.WriteLine("set -o errexit");
+					LinkWriter.WriteLine(LinkCommandString);
+					LinkWriter.WriteLine(GetDumpEncodeDebugCommand(LinkEnvironment, OutputFile));
 				}
 			};
 
@@ -1968,21 +1973,26 @@ namespace UnrealBuildTool
 						{
 						if (bUseCmdExe)
 						{
-							Writer.Write("@echo off\n");
-							Writer.Write("rem Automatically generated by UnrealBuildTool\n");
-							Writer.Write("rem *DO NOT EDIT*\n\n");
+							Writer.NewLine = "\r\n";
+							Writer.WriteLine("@echo off");
+							Writer.WriteLine("rem Automatically generated by UnrealBuildTool");
+							Writer.WriteLine("rem *DO NOT EDIT*");
+							Writer.WriteLine();
 						}
 						else
 						{
-							Writer.Write("#!/bin/sh\n");
-							Writer.Write("# Automatically generated by UnrealBuildTool\n");
-							Writer.Write("# *DO NOT EDIT*\n\n");
-							Writer.Write("set -o errexit\n");
+							Writer.NewLine = "\n";
+							Writer.WriteLine("#!/bin/sh");
+							Writer.WriteLine("# Automatically generated by UnrealBuildTool");
+							Writer.WriteLine("# *DO NOT EDIT*");
+							Writer.WriteLine();
+							Writer.WriteLine("set -o errexit");
 						}
 						}
 					}
 
 					StreamWriter FixDepsScript = File.AppendText(FixDepsScriptPath);
+					FixDepsScript.NewLine = bUseCmdExe ? "\r\n" : "\n";
 
 					string EngineAndGameLibrariesString = "";
 					foreach (string Library in EngineAndGameLibrariesLinkFlags)
@@ -1990,10 +2000,10 @@ namespace UnrealBuildTool
 						EngineAndGameLibrariesString += Library;
 					}
 
-					FixDepsScript.Write(string.Format("echo Fixing {0}\n", Path.GetFileName(OutputFile.AbsolutePath)));
+					FixDepsScript.WriteLine("echo Fixing {0}", Path.GetFileName(OutputFile.AbsolutePath));
 					if (!bUseCmdExe)
 					{
-						FixDepsScript.Write(string.Format("TIMESTAMP=`stat --format %y \"{0}\"`\n", OutputFile.AbsolutePath));
+						FixDepsScript.WriteLine("TIMESTAMP=`stat --format %y \"{0}\"`", OutputFile.AbsolutePath);
 					}
 					string FixDepsLine = LinkCommandString;
 					string Replace = "-Wl,--allow-shlib-undefined";
@@ -2002,15 +2012,16 @@ namespace UnrealBuildTool
 					string OutputFileForwardSlashes = OutputFile.AbsolutePath.Replace("\\", "/");
 					FixDepsLine = FixDepsLine.Replace(OutputFileForwardSlashes, OutputFileForwardSlashes + ".fixed");
 					FixDepsLine = FixDepsLine.Replace("$", "\\$");
-					FixDepsScript.Write(FixDepsLine + "\n");
+					FixDepsScript.WriteLine(FixDepsLine);
 					if (bUseCmdExe)
 					{
-						FixDepsScript.Write(string.Format("move /Y \"{0}.fixed\" \"{0}\"\n", OutputFile.AbsolutePath));
+						FixDepsScript.WriteLine("move /Y \"{0}.fixed\" \"{0}\"", OutputFile.AbsolutePath);
 					}
 					else
 					{
-						FixDepsScript.Write(string.Format("mv \"{0}.fixed\" \"{0}\"\n", OutputFile.AbsolutePath));
-						FixDepsScript.Write(string.Format("touch -d \"$TIMESTAMP\" \"{0}\"\n\n", OutputFile.AbsolutePath));
+						FixDepsScript.WriteLine("mv \"{0}.fixed\" \"{0}\"", OutputFile.AbsolutePath);
+						FixDepsScript.WriteLine("touch -d \"$TIMESTAMP\" \"{0}\"", OutputFile.AbsolutePath);
+						FixDepsScript.WriteLine();
 					}
 					FixDepsScript.Close();
 				}
@@ -2061,43 +2072,47 @@ namespace UnrealBuildTool
 
 					if (bUseCmdExe)
 					{
-						RelinkWriter.Write("@echo off\n");
-						RelinkWriter.Write("rem Automatically generated by UnrealBuildTool\n");
-						RelinkWriter.Write("rem *DO NOT EDIT*\n\n");
-						RelinkWriter.Write("set Retries=0\n");
-						RelinkWriter.Write(":relinkloop\n");
-						RelinkWriter.Write("if %Retries% GEQ 10 goto failedtorelink\n");
-						RelinkWriter.Write(RelinkInvocation + "\n");
-						RelinkWriter.Write("if %errorlevel% neq 0 goto sleepandretry\n");
-						RelinkWriter.Write("copy /B \"{0}\" \"{1}.temp\" >NUL 2>NUL\n", RelinkedFileForwardSlashes, OutputFile.AbsolutePath);
-						RelinkWriter.Write("if %errorlevel% neq 0 goto sleepandretry\n");
-						RelinkWriter.Write("move /Y \"{0}.temp\" \"{1}\" >NUL 2>NUL\n", OutputFile.AbsolutePath, OutputFile.AbsolutePath);
-						RelinkWriter.Write("if %errorlevel% neq 0 goto sleepandretry\n");
-						RelinkWriter.Write(GetDumpEncodeDebugCommand(LinkEnvironment, OutputFile) + "\n");
-						RelinkWriter.Write(string.Format("echo \"Dummy\" >> \"{0}\" && copy /b \"{0}\" +,,\n", RelinkActionDummyProductRef.FullName));
-						RelinkWriter.Write("echo Relinked {0} successfully after %Retries% retries\n", OutputFile.AbsolutePath);
-						RelinkWriter.Write("exit 0\n");
-						RelinkWriter.Write(":sleepandretry\n");
-						RelinkWriter.Write("ping 127.0.0.1 -n 1 -w 5000 >NUL 2>NUL\n");     // timeout complains about lack of redirection
-						RelinkWriter.Write("set /a Retries+=1\n");
-						RelinkWriter.Write("goto relinkloop\n");
-						RelinkWriter.Write(":failedtorelink\n");
-						RelinkWriter.Write("echo Failed to relink {0} after %Retries% retries\n", OutputFile.AbsolutePath);
-						RelinkWriter.Write("exit 1\n");
+						RelinkWriter.WriteLine("@echo off");
+						RelinkWriter.WriteLine("rem Automatically generated by UnrealBuildTool");
+						RelinkWriter.WriteLine("rem *DO NOT EDIT*");
+						RelinkWriter.WriteLine();
+						RelinkWriter.WriteLine("set Retries=0");
+						RelinkWriter.WriteLine(":relinkloop");
+						RelinkWriter.WriteLine("if %Retries% GEQ 10 goto failedtorelink");
+						RelinkWriter.WriteLine(RelinkInvocation);
+						RelinkWriter.WriteLine("if %errorlevel% neq 0 goto sleepandretry");
+						RelinkWriter.WriteLine("copy /B \"{0}\" \"{1}.temp\" >NUL 2>NUL", RelinkedFileForwardSlashes, OutputFile.AbsolutePath);
+						RelinkWriter.WriteLine("if %errorlevel% neq 0 goto sleepandretry");
+						RelinkWriter.WriteLine("move /Y \"{0}.temp\" \"{1}\" >NUL 2>NUL", OutputFile.AbsolutePath, OutputFile.AbsolutePath);
+						RelinkWriter.WriteLine("if %errorlevel% neq 0 goto sleepandretry");
+						RelinkWriter.WriteLine(GetDumpEncodeDebugCommand(LinkEnvironment, OutputFile));
+						RelinkWriter.WriteLine("echo \"Dummy\" >> \"{0}\" && copy /b \"{0}\" +,,", RelinkActionDummyProductRef.FullName);
+						RelinkWriter.WriteLine("echo Relinked {0} successfully after %Retries% retries", OutputFile.AbsolutePath);
+						RelinkWriter.WriteLine("exit 0");
+						RelinkWriter.WriteLine(":sleepandretry");
+						RelinkWriter.WriteLine("ping 127.0.0.1 -n 1 -w 5000 >NUL 2>NUL");     // timeout complains about lack of redirection
+						RelinkWriter.WriteLine("set /a Retries+=1");
+						RelinkWriter.WriteLine("goto relinkloop");
+						RelinkWriter.WriteLine(":failedtorelink");
+						RelinkWriter.WriteLine("echo Failed to relink {0} after %Retries% retries", OutputFile.AbsolutePath);
+						RelinkWriter.WriteLine("exit 1");
 					}
 					else
 					{
-						RelinkWriter.Write("#!/bin/sh\n");
-						RelinkWriter.Write("# Automatically generated by UnrealBuildTool\n");
-						RelinkWriter.Write("# *DO NOT EDIT*\n\n");
-						RelinkWriter.Write("set -o errexit\n");
-						RelinkWriter.Write(RelinkInvocation + "\n");
-						RelinkWriter.Write(string.Format("TIMESTAMP=`stat --format %y \"{0}\"`\n", OutputFile.AbsolutePath));
-						RelinkWriter.Write("cp \"{0}\" \"{1}.temp\"\n", RelinkedFileForwardSlashes, OutputFile.AbsolutePath);
-						RelinkWriter.Write("mv \"{0}.temp\" \"{1}\"\n", OutputFile.AbsolutePath, OutputFile.AbsolutePath);
-						RelinkWriter.Write(GetDumpEncodeDebugCommand(LinkEnvironment, OutputFile) + "\n");
-						RelinkWriter.Write(string.Format("touch -d \"$TIMESTAMP\" \"{0}\"\n\n", OutputFile.AbsolutePath));
-						RelinkWriter.Write(string.Format("echo \"Dummy\" >> \"{0}\"", RelinkActionDummyProductRef.FullName));
+						RelinkWriter.NewLine = "\n";
+						RelinkWriter.WriteLine("#!/bin/sh");
+						RelinkWriter.WriteLine("# Automatically generated by UnrealBuildTool");
+						RelinkWriter.WriteLine("# *DO NOT EDIT*");
+						RelinkWriter.WriteLine();
+						RelinkWriter.WriteLine("set -o errexit");
+						RelinkWriter.WriteLine(RelinkInvocation);
+						RelinkWriter.WriteLine("TIMESTAMP=`stat --format %y \"{0}\"`", OutputFile.AbsolutePath);
+						RelinkWriter.WriteLine("cp \"{0}\" \"{1}.temp\"", RelinkedFileForwardSlashes, OutputFile.AbsolutePath);
+						RelinkWriter.WriteLine("mv \"{0}.temp\" \"{1}\"", OutputFile.AbsolutePath, OutputFile.AbsolutePath);
+						RelinkWriter.WriteLine(GetDumpEncodeDebugCommand(LinkEnvironment, OutputFile));
+						RelinkWriter.WriteLine("touch -d \"$TIMESTAMP\" \"{0}\"", OutputFile.AbsolutePath);
+						RelinkWriter.WriteLine();
+						RelinkWriter.WriteLine("echo \"Dummy\" >> \"{0}\"", RelinkActionDummyProductRef.FullName);
 					}
 					}
 
