Add support for generating an accurate timing file for CompileTimeAnalyzer on the farm.

#rnx

diff --git a/Engine/Source/Programs/UnrealBuildTool/Configuration/BuildConfiguration.cs b/Engine/Source/Programs/UnrealBuildTool/Configuration/BuildConfiguration.cs
--- a/Engine/Source/Programs/UnrealBuildTool/Configuration/BuildConfiguration.cs
+++ b/Engine/Source/Programs/UnrealBuildTool/Configuration/BuildConfiguration.cs
@@ -33,6 +33,7 @@ namespace UnrealBuildTool
 		/// Whether debug info should be written to the console.
 		/// </summary>
 		[XmlConfigFile]
+		[CommandLine("-PrintDebugInfo", Value = "true")]
 		public bool bPrintDebugInfo = false;
 
 		/// <summary>
diff --git a/Engine/Source/Programs/UnrealBuildTool/Configuration/TargetRules.cs b/Engine/Source/Programs/UnrealBuildTool/Configuration/TargetRules.cs
--- a/Engine/Source/Programs/UnrealBuildTool/Configuration/TargetRules.cs
+++ b/Engine/Source/Programs/UnrealBuildTool/Configuration/TargetRules.cs
@@ -861,6 +861,7 @@ namespace UnrealBuildTool
 		/// incremental compile times. The current implementation uses the read-only flag to distinguish the working set, assuming that files will
 		/// be made writable by the source control system if they are being modified. This is true for Perforce, but not for Git.
 		/// </summary>
+		[CommandLine("-DisableAdaptiveUnity", Value = "false")]
 		[XmlConfigFile(Category = "BuildConfiguration")]
 		public bool bUseAdaptiveUnityBuild = true;
 
diff --git a/Engine/Source/Programs/UnrealBuildTool/Modes/AggregateParsedTimingInfo.cs b/Engine/Source/Programs/UnrealBuildTool/Modes/AggregateParsedTimingInfo.cs
--- a/Engine/Source/Programs/UnrealBuildTool/Modes/AggregateParsedTimingInfo.cs
+++ b/Engine/Source/Programs/UnrealBuildTool/Modes/AggregateParsedTimingInfo.cs
@@ -47,7 +47,8 @@ namespace UnrealBuildTool
 			GroupTimingDataOnName(AggregateData, "Function Timings", AggregateFunctions);
 
 			// Write out aggregate summary.
-			string OutputFile = Path.Combine(ManifestFile.Directory.FullName, String.Format("{0}.timing.bin", AggregateName));
+			string OutputFile = Path.Combine(ManifestFile.Directory.FullName, String.Format("{0}.cta", AggregateName));
+			Log.TraceLog("Writing {0}", OutputFile);
 			using (BinaryWriter Writer = new BinaryWriter(File.Open(OutputFile, FileMode.Create)))
 			{
 				// Write out the aggregate data.
@@ -104,6 +105,7 @@ namespace UnrealBuildTool
 
 		private void ParseTimingDataFile(string ParsedFileName)
 		{
+			Log.TraceLog("Parsing {0}", ParsedFileName);
 			// Convert input file back into summary objects.
 			using (BinaryReader Reader = new BinaryReader(File.Open(ParsedFileName, FileMode.Open, FileAccess.Read)))
 			{
diff --git a/Engine/Source/Programs/UnrealBuildTool/Modes/ParseMsvcTimingInfo.cs b/Engine/Source/Programs/UnrealBuildTool/Modes/ParseMsvcTimingInfo.cs
--- a/Engine/Source/Programs/UnrealBuildTool/Modes/ParseMsvcTimingInfo.cs
+++ b/Engine/Source/Programs/UnrealBuildTool/Modes/ParseMsvcTimingInfo.cs
@@ -105,7 +105,7 @@ namespace UnrealBuildTool
 			Summary.AddChild(SummarizeParsedTimingData("FunctionTimings", TimingDataType.Function, Functions));
 
 			// Write out the timing binary file.
-			using (BinaryWriter Writer = new BinaryWriter(File.Open(InputFile.ChangeExtension(".timing.bin").FullName, FileMode.Create)))
+			using (BinaryWriter Writer = new BinaryWriter(File.Open(InputFile.ChangeExtension(".cta").FullName, FileMode.Create)))
 			{
 				Writer.Write(Summary);
 			}
diff --git a/Engine/Source/Programs/UnrealBuildTool/Platform/Windows/VCToolChain.cs b/Engine/Source/Programs/UnrealBuildTool/Platform/Windows/VCToolChain.cs
--- a/Engine/Source/Programs/UnrealBuildTool/Platform/Windows/VCToolChain.cs
+++ b/Engine/Source/Programs/UnrealBuildTool/Platform/Windows/VCToolChain.cs
@@ -1235,7 +1235,7 @@ namespace UnrealBuildTool
 
 		private Action GenerateParseTimingInfoAction(FileItem SourceFile, Action CompileAction, CPPOutput Result, IActionGraphBuilder Graph)
 		{
-			FileItem TimingJsonFile = FileItem.GetItemByPath(Path.ChangeExtension(CompileAction.TimingFile.AbsolutePath, ".timing.bin"));
+			FileItem TimingJsonFile = FileItem.GetItemByPath(Path.ChangeExtension(CompileAction.TimingFile.AbsolutePath, ".cta"));
 			Result.DebugDataFiles.Add(TimingJsonFile);
 			Result.DebugDataFiles.Add(CompileAction.TimingFile);
 
@@ -1265,9 +1265,9 @@ namespace UnrealBuildTool
 
 			if (Target.bPrintToolChainTimingInfo || Target.WindowsPlatform.bCompilerTrace)
 			{
-				CompileAction.TimingFile = FileItem.GetItemByFileReference(FileReference.Combine(OutputDir, String.Format("{0}.timing.txt", SourceFile.Location.GetFileName())));
+				CompileAction.TimingFile = FileItem.GetItemByFileReference(FileReference.Combine(OutputDir, String.Format("{0}.timing", SourceFile.Location.GetFileName())));
 				CompileAction.ProducedItems.Add(CompileAction.TimingFile);
-				/*Action ParseTimingInfoAction = */GenerateParseTimingInfoAction(SourceFile, CompileAction, Result, Graph);
+				GenerateParseTimingInfoAction(SourceFile, CompileAction, Result, Graph);
 				CommandArguments.Add(string.Format("-timing={0}", Utils.MakePathSafeToUseWithCommandLine(CompileAction.TimingFile.Location)));
 			}
 
@@ -1290,7 +1290,7 @@ namespace UnrealBuildTool
 			if (Target.bPrintToolChainTimingInfo || Target.WindowsPlatform.bCompilerTrace)
 			{
 				List<Action> ParseTimingActions = Makefile.Actions.Where(x => x.ActionType == ActionType.ParseTimingInfo).ToList();
-				List<FileItem> TimingJsonFiles = ParseTimingActions.SelectMany(a => a.ProducedItems.Where(i => i.HasExtension(".timing.bin"))).ToList();
+				List<FileItem> TimingJsonFiles = ParseTimingActions.SelectMany(a => a.ProducedItems.Where(i => i.HasExtension(".cta"))).ToList();
 				Makefile.OutputItems.AddRange(TimingJsonFiles);
 
 				// Handing generating aggregate timing information if we compiled more than one file.
@@ -1319,7 +1319,7 @@ namespace UnrealBuildTool
 					AggregateTimingInfoAction.bCanExecuteRemotelyWithSNDBS = false;
 					AggregateTimingInfoAction.PrerequisiteItems.AddRange(TimingJsonFiles);
 
-					FileItem AggregateOutputFile = FileItem.GetItemByFileReference(FileReference.Combine(Makefile.ProjectIntermediateDirectory, $"{Target.Name}.timing.bin"));
+					FileItem AggregateOutputFile = FileItem.GetItemByFileReference(FileReference.Combine(Makefile.ProjectIntermediateDirectory, $"{Target.Name}.cta"));
 					AggregateTimingInfoAction.ProducedItems.Add(AggregateOutputFile);
 					Makefile.OutputItems.Add(AggregateOutputFile);
 				}
