- Quick fix for Android SDK detection before bigger re-write in Main

diff --git a/Engine/Source/Programs/AutomationTool/Turnkey/SdkInfo.cs b/Engine/Source/Programs/AutomationTool/Turnkey/SdkInfo.cs
--- a/Engine/Source/Programs/AutomationTool/Turnkey/SdkInfo.cs
+++ b/Engine/Source/Programs/AutomationTool/Turnkey/SdkInfo.cs
@@ -61,7 +61,7 @@ namespace Turnkey
 					foreach (DirectoryInfo Dir in PlatformDir.EnumerateDirectories())
 					{
 						// check for valid version number (auto SDK must match version exactly)
-						if (string.Compare(Dir.Name, SDK.GetDesiredVersion()) == 0)
+						if (string.Compare(Dir.Name, SDK.GetRequiredSDKString()) == 0)
 						{
 							// make sure it actually has buits in it
 							if (File.Exists(Path.Combine(Dir.FullName, "setup.bat")))
diff --git a/Engine/Source/Programs/UnrealBuildTool/Platform/Android/AndroidPlatformSDK.cs b/Engine/Source/Programs/UnrealBuildTool/Platform/Android/AndroidPlatformSDK.cs
--- a/Engine/Source/Programs/UnrealBuildTool/Platform/Android/AndroidPlatformSDK.cs
+++ b/Engine/Source/Programs/UnrealBuildTool/Platform/Android/AndroidPlatformSDK.cs
@@ -11,10 +11,14 @@ namespace UnrealBuildTool
 {
 	class AndroidPlatformSDK : UEBuildPlatformSDK
 	{
-		public override string GetDesiredVersion()
+		public override string GetRequiredSDKString()
 		{
 			return "-23";
 		}
+		public override string GetDesiredVersion()
+		{
+			return "r21b";
+		}
 
 		public override void GetValidVersionRange(out string MinVersion, out string MaxVersion)
 		{
@@ -116,11 +120,6 @@ namespace UnrealBuildTool
 			return true;
 		}
 
-		public override string GetRequiredSDKString()
-		{
-			return GetDesiredVersion();
-		}
-
 		protected override String GetRequiredScriptVersionString()
 		{
 			return "3.6";
