#UE4 Re-enable undefined macros as errors for all VCToolChain compilers

#rnx

diff --git a/Engine/Source/Programs/UnrealBuildTool/Platform/Windows/VCToolChain.cs b/Engine/Source/Programs/UnrealBuildTool/Platform/Windows/VCToolChain.cs
--- a/Engine/Source/Programs/UnrealBuildTool/Platform/Windows/VCToolChain.cs
+++ b/Engine/Source/Programs/UnrealBuildTool/Platform/Windows/VCToolChain.cs
@@ -504,14 +504,7 @@ namespace UnrealBuildTool
 			{
 				if (CompileEnvironment.bUndefinedIdentifierWarningsAsErrors)
 				{
-					if (Target.WindowsPlatform.Compiler == WindowsCompiler.VisualStudio2015_DEPRECATED)
-					{
-						Arguments.Add("/we4668");
-					}
-					else if (Target.WindowsPlatform.Compiler == WindowsCompiler.VisualStudio2017)
-					{
-						Arguments.Add("/wd4668");
-					}
+					Arguments.Add("/we4668");
 				}
 				else
 				{
