- Fixed rare crash when no mobileprovision has ever been installed on a Mac (CreateDirectory before calling File.Copy)

diff --git a/Engine/Source/Programs/UnrealBuildTool/Platform/IOS/UEDeployIOS.cs b/Engine/Source/Programs/UnrealBuildTool/Platform/IOS/UEDeployIOS.cs
--- a/Engine/Source/Programs/UnrealBuildTool/Platform/IOS/UEDeployIOS.cs
+++ b/Engine/Source/Programs/UnrealBuildTool/Platform/IOS/UEDeployIOS.cs
@@ -916,6 +916,7 @@ namespace UnrealBuildTool
 			Directory.CreateDirectory(PayloadDirectory);
 			Directory.CreateDirectory(AppDirectory);
 			Directory.CreateDirectory(BuildDirectory);
+			Directory.CreateDirectory(Environment.GetEnvironmentVariable("HOME") + "/Library/MobileDevice/Provisioning Profiles");
 
 			// create the entitlements file
 
