- Couple of Mac icns file fixes

diff --git a/Engine/Source/Programs/UnrealBuildTool/Platform/Mac/MacToolChain.cs b/Engine/Source/Programs/UnrealBuildTool/Platform/Mac/MacToolChain.cs
--- a/Engine/Source/Programs/UnrealBuildTool/Platform/Mac/MacToolChain.cs
+++ b/Engine/Source/Programs/UnrealBuildTool/Platform/Mac/MacToolChain.cs
@@ -1004,7 +1004,7 @@ namespace UnrealBuildTool
 					AppendMacLine(FinalizeAppBundleScript, "mkdir -p \"{0}.app/Contents/MacOS\"", ExeName);
 					AppendMacLine(FinalizeAppBundleScript, "mkdir -p \"{0}.app/Contents/Resources\"", ExeName);
 
-					string IconName = "UE4";
+					string IconName = "UnrealEngine";
 					string EngineSourcePath = Directory.GetCurrentDirectory().Replace("$", "\\$");
 					string CustomResourcesPath = "";
 					string CustomBuildPath = "";
@@ -1340,7 +1340,7 @@ namespace UnrealBuildTool
 
 				if (Target.Type == TargetType.Editor)
 				{
-					BuildProducts.Add(FileReference.Combine(BundleContentsDirectory, "Resources/UE4Editor.icns"), BuildProductType.RequiredResource);
+					BuildProducts.Add(FileReference.Combine(BundleContentsDirectory, "Resources/UnrealEditor.icns"), BuildProductType.RequiredResource);
 					BuildProducts.Add(FileReference.Combine(BundleContentsDirectory, "Resources/UProject.icns"), BuildProductType.RequiredResource);
 				}
 				else
