- Updating allowed Xcode version to 12

diff --git a/Engine/Source/Programs/UnrealBuildTool/Platform/Mac/ApplePlatformSDK.cs b/Engine/Source/Programs/UnrealBuildTool/Platform/Mac/ApplePlatformSDK.cs
--- a/Engine/Source/Programs/UnrealBuildTool/Platform/Mac/ApplePlatformSDK.cs
+++ b/Engine/Source/Programs/UnrealBuildTool/Platform/Mac/ApplePlatformSDK.cs
@@ -21,14 +21,14 @@ namespace UnrealBuildTool
 		{
 			if (IsMac())
 			{
-				MinVersion = "9.4.0";
-				MaxVersion = "11.9.9";
+				MinVersion = "11.0.0";
+				MaxVersion = "12.9.9";
 			}
 			else
 			{
 				// @todo turnkey: these are MobileDevice .dll versions in Windows - to get the iTunes app version (12.3.4.1 etc) would need to hunt down the .exe
 				MinVersion = "1100.0.0.0";
-				MaxVersion = "1199.0";
+				MaxVersion = "1299.0";
 			}
 		}
 
