Initial support for C++20 modules.

Requires Visual Studio [REDACTED_PHONE] Preview 3, and targets to specify bEnableCppModules = true. Unsupported on other compilers.

Notes:
- Module interfaces need to be placed in .IXX files.
- Module/module dependencies are determined by a post-makefile/pre-build scan of all module interfaces using the compiler, and cause additional prerequisites to be inserted into the dependency graph before execution.
- Available module IFC search paths use normal UE module dependency rules.
- Header units are not currently supported.
- Regular C++ files cannot currently implement functions defined in module interfaces, due to the macros force included by UBT. These macros need to be placed in the global module fragment, but declaring a global module fragment must be done from the translation unit being compiled (making force-included headers invalid).

diff --git a/Engine/Source/Programs/UnrealBuildTool/Configuration/TargetRules.cs b/Engine/Source/Programs/UnrealBuildTool/Configuration/TargetRules.cs
--- a/Engine/Source/Programs/UnrealBuildTool/Configuration/TargetRules.cs
+++ b/Engine/Source/Programs/UnrealBuildTool/Configuration/TargetRules.cs
@@ -1077,6 +1077,11 @@ namespace UnrealBuildTool
 		[XmlConfigFile(Category = "BuildConfiguration")]
 		public bool bOmitFramePointers = true;
 
+		/// <summary>
+		/// Whether to enable support for C++20 modules
+		/// </summary>
+		public bool bEnableCppModules = false;
+
 		/// <summary>
 		/// Whether to strip iOS symbols or not (implied by bGeneratedSYMFile).
 		/// </summary>
@@ -2451,6 +2456,11 @@ namespace UnrealBuildTool
 			get { return Inner.bOmitFramePointers; }
 		}
 
+		public bool bEnableCppModules
+		{
+			get { return Inner.bEnableCppModules; }
+		}
+
 		[Obsolete("bStripSymbolsOnIOS has been deprecated. Use IOSPlatform.bStripSymbols instead.")]
 		public bool bStripSymbolsOnIOS
 		{
diff --git a/Engine/Source/Programs/UnrealBuildTool/Configuration/UEBuildModule.cs b/Engine/Source/Programs/UnrealBuildTool/Configuration/UEBuildModule.cs
--- a/Engine/Source/Programs/UnrealBuildTool/Configuration/UEBuildModule.cs
+++ b/Engine/Source/Programs/UnrealBuildTool/Configuration/UEBuildModule.cs
@@ -503,6 +503,7 @@ namespace UnrealBuildTool
 			UEBuildBinary SourceBinary,
 			HashSet<DirectoryReference> IncludePaths,
 			HashSet<DirectoryReference> SystemIncludePaths,
+			HashSet<DirectoryReference> ModuleInterfacePaths,
 			List<string> Definitions,
 			List<UEBuildFramework> AdditionalFrameworks,
 			List<FileItem> AdditionalPrerequisites,
@@ -569,6 +570,7 @@ namespace UnrealBuildTool
 		protected virtual void SetupPrivateCompileEnvironment(
 			HashSet<DirectoryReference> IncludePaths,
 			HashSet<DirectoryReference> SystemIncludePaths,
+			HashSet<DirectoryReference> ModuleInterfacePaths,
 			List<string> Definitions,
 			List<UEBuildFramework> AdditionalFrameworks,
 			List<FileItem> AdditionalPrerequisites,
@@ -591,7 +593,7 @@ namespace UnrealBuildTool
 			// Now set up the compile environment for the modules in the original order that we encountered them
 			foreach (UEBuildModule Module in ModuleToIncludePathsOnlyFlag.Keys)
 			{
-				Module.AddModuleToCompileEnvironment(Binary, IncludePaths, SystemIncludePaths, Definitions, AdditionalFrameworks, AdditionalPrerequisites, bWithLegacyPublicIncludePaths);
+				Module.AddModuleToCompileEnvironment(Binary, IncludePaths, SystemIncludePaths, ModuleInterfacePaths, Definitions, AdditionalFrameworks, AdditionalPrerequisites, bWithLegacyPublicIncludePaths);
 			}
 		}
 
diff --git a/Engine/Source/Programs/UnrealBuildTool/Configuration/UEBuildModuleCPP.cs b/Engine/Source/Programs/UnrealBuildTool/Configuration/UEBuildModuleCPP.cs
--- a/Engine/Source/Programs/UnrealBuildTool/Configuration/UEBuildModuleCPP.cs
+++ b/Engine/Source/Programs/UnrealBuildTool/Configuration/UEBuildModuleCPP.cs
@@ -24,6 +24,7 @@ namespace UnrealBuildTool
 			public readonly List<FileItem> HeaderFiles = new List<FileItem>();
 			public readonly List<FileItem> ISPCHeaderFiles = new List<FileItem>();
 
+			public readonly List<FileItem> IXXFiles = new List<FileItem>();
 			public readonly List<FileItem> CPPFiles = new List<FileItem>();
 			public readonly List<FileItem> CFiles = new List<FileItem>();
 			public readonly List<FileItem> CCFiles = new List<FileItem>();
@@ -285,6 +286,7 @@ namespace UnrealBuildTool
 			UEBuildBinary SourceBinary,
 			HashSet<DirectoryReference> IncludePaths,
 			HashSet<DirectoryReference> SystemIncludePaths,
+			HashSet<DirectoryReference> ModuleInterfacePaths,
 			List<string> Definitions,
 			List<UEBuildFramework> AdditionalFrameworks,
 			List<FileItem> AdditionalPrerequisites,
@@ -309,7 +311,9 @@ namespace UnrealBuildTool
 				}
 			}
 
-			base.AddModuleToCompileEnvironment(SourceBinary, IncludePaths, SystemIncludePaths, Definitions, AdditionalFrameworks, AdditionalPrerequisites, bLegacyPublicIncludePaths);
+			ModuleInterfacePaths.Add(UEToolChain.GetModuleInterfaceDir(IntermediateDirectory));
+
+			base.AddModuleToCompileEnvironment(SourceBinary, IncludePaths, SystemIncludePaths, ModuleInterfacePaths, Definitions, AdditionalFrameworks, AdditionalPrerequisites, bLegacyPublicIncludePaths);
 		}
 
 		// UEBuildModule interface.
@@ -448,6 +452,13 @@ namespace UnrealBuildTool
 				CreateHeadersForISPC(ToolChain, CompileEnvironment, InputFiles.ISPCFiles, IntermediateDirectory, Graph);
 			}
 
+			// Compile any module interfaces
+			if (InputFiles.IXXFiles.Count > 0 && Target.bEnableCppModules)
+			{
+				CPPOutput ModuleOutput = ToolChain.CompileCPPFiles(CompileEnvironment, InputFiles.IXXFiles, IntermediateDirectory, Name, Graph);
+				LinkInputFiles.AddRange(ModuleOutput.ObjectFiles);
+			}
+
 			// Configure the precompiled headers for this module
 			CompileEnvironment = SetupPrecompiledHeaders(Target, ToolChain, CompileEnvironment, LinkInputFiles, Graph);
 
@@ -1306,10 +1317,14 @@ namespace UnrealBuildTool
 			Result.bEnableUndefinedIdentifierWarnings = Rules.bEnableUndefinedIdentifierWarnings;
 
 			// If the module overrides the C++ language version, override it on the compile environment
-			if(Rules.CppStandard != CppStandardVersion.Default)
+			if (Rules.CppStandard != CppStandardVersion.Default)
 			{
 				Result.CppStandard = Rules.CppStandard;
 			}
+			if (Target.bEnableCppModules && Result.CppStandard == CppStandardVersion.Default)
+			{
+				Result.CppStandard = CppStandardVersion.Cpp20;
+			}
 
 			// Set the macro used to check whether monolithic headers can be used
 			if (Rules.bTreatAsEngineModule && (!Rules.bEnforceIWYU || !Target.bEnforceIWYU))
@@ -1350,7 +1365,7 @@ namespace UnrealBuildTool
 			}
 
 			// Setup the compile environment for the module.
-			SetupPrivateCompileEnvironment(Result.UserIncludePaths, Result.SystemIncludePaths, Result.Definitions, Result.AdditionalFrameworks, Result.AdditionalPrerequisites, Rules.bLegacyPublicIncludePaths);
+			SetupPrivateCompileEnvironment(Result.UserIncludePaths, Result.SystemIncludePaths, Result.ModuleInterfacePaths, Result.Definitions, Result.AdditionalFrameworks, Result.AdditionalPrerequisites, Rules.bLegacyPublicIncludePaths);
 
 			return Result;
 		}
@@ -1392,7 +1407,7 @@ namespace UnrealBuildTool
 			// Now set up the compile environment for the modules in the original order that we encountered them
 			foreach (UEBuildModule Module in ModuleToIncludePathsOnlyFlag.Keys)
 			{
-				Module.AddModuleToCompileEnvironment(null, CompileEnvironment.UserIncludePaths, CompileEnvironment.SystemIncludePaths, CompileEnvironment.Definitions, CompileEnvironment.AdditionalFrameworks, CompileEnvironment.AdditionalPrerequisites, Rules.bLegacyPublicIncludePaths);
+				Module.AddModuleToCompileEnvironment(null, CompileEnvironment.UserIncludePaths, CompileEnvironment.SystemIncludePaths, CompileEnvironment.ModuleInterfacePaths, CompileEnvironment.Definitions, CompileEnvironment.AdditionalFrameworks, CompileEnvironment.AdditionalPrerequisites, Rules.bLegacyPublicIncludePaths);
 			}
 			return CompileEnvironment;
 		}
@@ -1496,11 +1511,16 @@ namespace UnrealBuildTool
 				{
 					InputFiles.ISPCHeaderFiles.Add(InputFile);
 				}
-				if (InputFile.HasExtension(".cpp"))
+				else if (InputFile.HasExtension(".cpp"))
 				{
 					SourceFiles.Add(InputFile);
 					InputFiles.CPPFiles.Add(InputFile);
 				}
+				else if (InputFile.HasExtension(".ixx"))
+				{
+					SourceFiles.Add(InputFile);
+					InputFiles.IXXFiles.Add(InputFile);
+				}
 				else if (InputFile.HasExtension(".c"))
 				{
 					SourceFiles.Add(InputFile);
diff --git a/Engine/Source/Programs/UnrealBuildTool/Configuration/UEBuildTarget.cs b/Engine/Source/Programs/UnrealBuildTool/Configuration/UEBuildTarget.cs
--- a/Engine/Source/Programs/UnrealBuildTool/Configuration/UEBuildTarget.cs
+++ b/Engine/Source/Programs/UnrealBuildTool/Configuration/UEBuildTarget.cs
@@ -2231,7 +2231,7 @@ namespace UnrealBuildTool
 			List<string> Definitions = new List<string>(GlobalCompileEnvironment.Definitions);
 			foreach(UEBuildModule Module in Binary.Modules)
 			{
-				Module.AddModuleToCompileEnvironment(null, new HashSet<DirectoryReference>(), new HashSet<DirectoryReference>(), Definitions, new List<UEBuildFramework>(), new List<FileItem>(), false);
+				Module.AddModuleToCompileEnvironment(null, new HashSet<DirectoryReference>(), new HashSet<DirectoryReference>(), new HashSet<DirectoryReference>(), Definitions, new List<UEBuildFramework>(), new List<FileItem>(), false);
 			}
 
 			// Write the header
@@ -3723,6 +3723,16 @@ namespace UnrealBuildTool
 				GlobalCompileEnvironment.Definitions.Add("WITH_XGE_CONTROLLER=0");
 			}
 
+			// Whether C++20 modules are enabled
+			if (Rules.bEnableCppModules)
+			{
+				GlobalCompileEnvironment.Definitions.Add("WITH_CPP_MODULES=1");
+			}
+			else
+			{
+				GlobalCompileEnvironment.Definitions.Add("WITH_CPP_MODULES=0");
+			}
+
 			// Compile in the names of the module manifests
 			GlobalCompileEnvironment.Definitions.Add(String.Format("UBT_MODULE_MANIFEST=\"{0}\"", ModuleManifest.GetStandardFileName(AppName, Platform, Configuration, Architecture, false)));
 			GlobalCompileEnvironment.Definitions.Add(String.Format("UBT_MODULE_MANIFEST_DEBUGGAME=\"{0}\"", ModuleManifest.GetStandardFileName(AppName, Platform, UnrealTargetConfiguration.DebugGame, Architecture, true)));
diff --git a/Engine/Source/Programs/UnrealBuildTool/Modes/BuildMode.cs b/Engine/Source/Programs/UnrealBuildTool/Modes/BuildMode.cs
--- a/Engine/Source/Programs/UnrealBuildTool/Modes/BuildMode.cs
+++ b/Engine/Source/Programs/UnrealBuildTool/Modes/BuildMode.cs
@@ -352,6 +352,26 @@ namespace UnrealBuildTool
 					}
 				}
 
+				// Pre-process module interfaces to generate dependency files
+				List<Action> ModuleDependencyActions = PrerequisiteActions.Where(x => x.ActionType == ActionType.GatherModuleDependencies).ToList();
+				if (ModuleDependencyActions.Count > 0)
+				{
+					Dictionary<Action, bool> PreprocessActionToOutdatedFlag = new Dictionary<Action, bool>();
+					ActionGraph.GatherAllOutdatedActions(ModuleDependencyActions, History, PreprocessActionToOutdatedFlag, null, BuildConfiguration.bIgnoreOutdatedImportLibraries);
+
+					List<Action> PreprocessActions = PreprocessActionToOutdatedFlag.Where(x => x.Value).Select(x => x.Key).ToList();
+					if (PreprocessActions.Count > 0)
+					{
+						Log.TraceInformation("Updating module dependencies...");
+						ActionGraph.ExecuteActions(BuildConfiguration, PreprocessActions);
+
+						foreach (FileItem ProducedItem in PreprocessActions.SelectMany(x => x.ProducedItems))
+						{
+							ProducedItem.ResetCachedInfo();
+						}
+					}
+				}
+
 				// Figure out which actions need to be built
 				Dictionary<Action, bool> ActionToOutdatedFlag = new Dictionary<Action, bool>();
 				for (int TargetIdx = 0; TargetIdx < TargetDescriptors.Count; TargetIdx++)
@@ -365,6 +385,54 @@ namespace UnrealBuildTool
 						CppDependencies = CppDependencyCache.CreateHierarchy(TargetDescriptor.ProjectFile, TargetDescriptor.Name, TargetDescriptor.Platform, TargetDescriptor.Configuration, Makefiles[TargetIdx].TargetType, TargetDescriptor.Architecture);
 					}
 
+					// Update the module dependencies
+					Dictionary<string, FileItem> ModuleOutputs = new Dictionary<string, FileItem>(StringComparer.Ordinal);
+					if (ModuleDependencyActions.Count > 0)
+					{
+						List<FileItem> CompiledModuleInterfaces = new List<FileItem>();
+						foreach(Action PrerequisiteAction in PrerequisiteActions)
+						{
+							if (PrerequisiteAction.ActionType == ActionType.CompileModuleInterface && PrerequisiteAction.CompiledModuleInterfaceFile != null)
+							{
+								string ProducedModule;
+								if (CppDependencies.TryGetProducedModule(PrerequisiteAction.DependencyListFile, out ProducedModule))
+								{
+									ModuleOutputs[ProducedModule] = PrerequisiteAction.CompiledModuleInterfaceFile;
+								}
+								CompiledModuleInterfaces.Add(PrerequisiteAction.CompiledModuleInterfaceFile);
+							}
+						}
+						foreach (Action PrerequisiteAction in PrerequisiteActions)
+						{
+							if (PrerequisiteAction.ActionType == ActionType.CompileModuleInterface)
+							{
+								List<string> ImportedModules;
+								if(CppDependencies.TryGetImportedModules(PrerequisiteAction.DependencyListFile, out ImportedModules))
+								{
+									foreach (string ImportedModule in ImportedModules)
+									{
+										FileItem ModuleOutput;
+										if (ModuleOutputs.TryGetValue(ImportedModule, out ModuleOutput))
+										{
+											PrerequisiteAction.PrerequisiteItems.Add(ModuleOutput);
+										}
+										else
+										{
+											throw new BuildException("Unable to find interface for module '{0}'", ImportedModule);
+										}
+									}
+								}
+							}
+							else
+							{
+								if (PrerequisiteAction.ActionType != ActionType.GatherModuleDependencies)
+								{
+									PrerequisiteAction.PrerequisiteItems.AddRange(CompiledModuleInterfaces);
+								}
+							}
+						}
+					}
+
 					// Plan the actions to execute for the build. For single file compiles, always rebuild the source file regardless of whether it's out of date.
 					if (TargetDescriptor.SpecificFilesToCompile.Count == 0)
 					{
diff --git a/Engine/Source/Programs/UnrealBuildTool/Platform/Windows/UEBuildWindows.cs b/Engine/Source/Programs/UnrealBuildTool/Platform/Windows/UEBuildWindows.cs
--- a/Engine/Source/Programs/UnrealBuildTool/Platform/Windows/UEBuildWindows.cs
+++ b/Engine/Source/Programs/UnrealBuildTool/Platform/Windows/UEBuildWindows.cs
@@ -861,6 +861,12 @@ namespace UnrealBuildTool
 			Target.WindowsPlatform.CompilerVersion = Target.WindowsPlatform.Environment.CompilerVersion.ToString();
 			Target.WindowsPlatform.WindowsSdkVersion = Target.WindowsPlatform.Environment.WindowsSdkVersion.ToString();
 
+			// If we're enabling support for C++ modules, make sure the compiler supports it. VS 16.8 changed which command line arguments are used to enable modules support.
+			if (Target.bEnableCppModules && !ProjectFileGenerator.bGenerateProjectFiles && Target.WindowsPlatform.Environment.CompilerVersion < new VersionNumber(14, 28, 29304))
+			{
+				throw new BuildException("Support for C++20 modules requires Visual Studio 2019 16.8 preview 3 or later.");
+			}
+
 //			@Todo: Still getting reports of frequent OOM issues with this enabled as of 15.7.
 //			// Enable fast PDB linking if we're on VS2017 15.7 or later. Previous versions have OOM issues with large projects.
 //			if(!Target.bFormalBuild && !Target.bUseFastPDBLinking.HasValue && Target.WindowsPlatform.Compiler >= WindowsCompiler.VisualStudio2017)
diff --git a/Engine/Source/Programs/UnrealBuildTool/Platform/Windows/VCToolChain.cs b/Engine/Source/Programs/UnrealBuildTool/Platform/Windows/VCToolChain.cs
--- a/Engine/Source/Programs/UnrealBuildTool/Platform/Windows/VCToolChain.cs
+++ b/Engine/Source/Programs/UnrealBuildTool/Platform/Windows/VCToolChain.cs
@@ -981,6 +981,15 @@ namespace UnrealBuildTool
 				AddDefinition(SharedArguments, DefinitionArgument);
 			}
 
+			// Module interface paths.
+			if (Target.bEnableCppModules)
+			{
+				foreach (DirectoryReference ModuleInterfacePath in CompileEnvironment.ModuleInterfacePaths)
+				{
+					SharedArguments.Add(String.Format("/ifcSearchDir \"{0}\"", ModuleInterfacePath));
+				}
+			}
+
 			// Create a compile action for each source file.
 			CPPOutput Result = new CPPOutput();
 			foreach (FileItem SourceFile in InputFiles)
@@ -1190,7 +1199,32 @@ namespace UnrealBuildTool
 					CompileAction.CommandArguments = String.Join(" ", SharedArguments.Concat(FileArguments).Concat(AdditionalArguments));
 				}
 
-				if(CompileEnvironment.bGenerateDependenciesFile)
+				if (SourceFile.HasExtension(".ixx"))
+				{
+					FileItem IfcDepsFile = FileItem.GetItemByFileReference(FileReference.Combine(OutputDir, SourceFile.Location.GetFileName() + ".md.json"));
+
+					Action CompileDepsAction = Graph.CreateAction(ActionType.GatherModuleDependencies);
+					CompileDepsAction.StatusDescription = SourceFile.Location.GetFileName();
+					CompileDepsAction.WorkingDirectory = CompileAction.WorkingDirectory;
+					CompileDepsAction.CommandPath = CompileAction.CommandPath;
+					CompileDepsAction.CommandArguments = String.Format("{0} /sourceDependencies:directives \"{1}\"", CompileAction.CommandArguments, IfcDepsFile.Location);
+					CompileDepsAction.PrerequisiteItems.Add(SourceFile);
+					CompileDepsAction.PrerequisiteItems.AddRange(CompileEnvironment.ForceIncludeFiles);
+					CompileDepsAction.PrerequisiteItems.AddRange(CompileEnvironment.AdditionalPrerequisites);
+					CompileDepsAction.ProducedItems.Add(IfcDepsFile);
+					CompileDepsAction.bShouldOutputStatusDescription = false;
+
+					FileItem IfcFile = FileItem.GetItemByFileReference(FileReference.Combine(GetModuleInterfaceDir(OutputDir), SourceFile.Location.ChangeExtension(".ifc").GetFileName()));
+					CompileAction.ActionType = ActionType.CompileModuleInterface;
+					CompileAction.CommandArguments += String.Format(" /interface /ifcOutput \"{0}\"", IfcFile.Location);
+					CompileAction.PrerequisiteItems.Add(IfcDepsFile); // Force the dependencies file into the action graph
+					CompileAction.ProducedItems.Add(IfcFile);
+					CompileAction.DependencyListFile = IfcDepsFile;
+					CompileAction.CompiledModuleInterfaceFile = IfcFile;
+
+					Result.CompiledModuleInterfaces.Add(IfcFile);
+				}
+				else if (CompileEnvironment.bGenerateDependenciesFile)
 				{
 					GenerateDependenciesFile(Target, OutputDir, Result, SourceFile, CompileAction, Graph);
 				}
diff --git a/Engine/Source/Programs/UnrealBuildTool/System/Action.cs b/Engine/Source/Programs/UnrealBuildTool/System/Action.cs
--- a/Engine/Source/Programs/UnrealBuildTool/System/Action.cs
+++ b/Engine/Source/Programs/UnrealBuildTool/System/Action.cs
@@ -18,6 +18,10 @@ namespace UnrealBuildTool
 	{
 		BuildProject,
 
+		GatherModuleDependencies,
+
+		CompileModuleInterface,
+
 		Compile,
 
 		CreateAppBundle,
@@ -45,7 +49,7 @@ namespace UnrealBuildTool
 		/// <summary>
 		/// The type of this action (for debugging purposes).
 		/// </summary>
-		public readonly ActionType ActionType;
+		public ActionType ActionType;
 
 		/// <summary>
 		/// Every file this action depends on.  These files need to exist and be up to date in order for this action to even be considered
@@ -67,6 +71,11 @@ namespace UnrealBuildTool
 		/// </summary>
 		public FileItem DependencyListFile;
 
+		/// <summary>
+		/// The compiled C++ module interface (IFC) file produced by this action
+		/// </summary>
+		public FileItem CompiledModuleInterfaceFile;
+
 		/// <summary>
 		/// For C++ source files, specifies a timing file used to track timing information.
 		/// </summary>
@@ -194,6 +203,7 @@ namespace UnrealBuildTool
 			ProducedItems = Reader.ReadList(() => Reader.ReadFileItem());
 			DeleteItems = Reader.ReadList(() => Reader.ReadFileItem());
 			DependencyListFile = Reader.ReadFileItem();
+			CompiledModuleInterfaceFile = Reader.ReadFileItem();
 		}
 
 		/// <summary>
@@ -217,6 +227,7 @@ namespace UnrealBuildTool
 			Writer.WriteList(ProducedItems, Item => Writer.WriteFileItem(Item));
 			Writer.WriteList(DeleteItems, Item => Writer.WriteFileItem(Item));
 			Writer.WriteFileItem(DependencyListFile);
+			Writer.WriteFileItem(CompiledModuleInterfaceFile);
 		}
 
 		/// <summary>
diff --git a/Engine/Source/Programs/UnrealBuildTool/System/CppCompileEnvironment.cs b/Engine/Source/Programs/UnrealBuildTool/System/CppCompileEnvironment.cs
--- a/Engine/Source/Programs/UnrealBuildTool/System/CppCompileEnvironment.cs
+++ b/Engine/Source/Programs/UnrealBuildTool/System/CppCompileEnvironment.cs
@@ -42,6 +42,11 @@ namespace UnrealBuildTool
 		/// </summary>
 		Cpp17,
 
+		/// <summary>
+		/// Supports C++20
+		/// </summary>
+		Cpp20,
+
 		/// <summary>
 		/// Latest standard supported by the compiler
 		/// </summary>
@@ -74,6 +79,7 @@ namespace UnrealBuildTool
 	{
 		public List<FileItem> ObjectFiles = new List<FileItem>();
 		public List<FileItem> DebugDataFiles = new List<FileItem>();
+		public List<FileItem> CompiledModuleInterfaces = new List<FileItem>();
 		public List<FileItem> GeneratedHeaderFiles = new List<FileItem>();
 		public FileItem PrecompiledHeaderFile = null;
 	}
@@ -317,6 +323,11 @@ namespace UnrealBuildTool
 		/// </summary>
 		public HashSet<DirectoryReference> SystemIncludePaths;
 
+		/// <summary>
+		/// List of paths to search for compiled module interface (*.ifc) files
+		/// </summary>
+		public HashSet<DirectoryReference> ModuleInterfacePaths;
+
 		/// <summary>
 		/// Whether headers in system paths should be checked for modification when determining outdated actions.
 		/// </summary>
@@ -379,6 +390,7 @@ namespace UnrealBuildTool
 			this.SharedPCHs = new List<PrecompiledHeaderTemplate>();
 			this.UserIncludePaths = new HashSet<DirectoryReference>();
 			this.SystemIncludePaths = new HashSet<DirectoryReference>();
+			this.ModuleInterfacePaths = new HashSet<DirectoryReference>();
 		}
 
 		/// <summary>
@@ -432,6 +444,7 @@ namespace UnrealBuildTool
 			bAllowRemotelyCompiledPCHs = Other.bAllowRemotelyCompiledPCHs;
 			UserIncludePaths = new HashSet<DirectoryReference>(Other.UserIncludePaths);
 			SystemIncludePaths = new HashSet<DirectoryReference>(Other.SystemIncludePaths);
+			ModuleInterfacePaths = new HashSet<DirectoryReference>(Other.ModuleInterfacePaths);
 			bCheckSystemHeadersForModification = Other.bCheckSystemHeadersForModification;
 			ForceIncludeFiles.AddRange(Other.ForceIncludeFiles);
 			AdditionalPrerequisites.AddRange(Other.AdditionalPrerequisites);
diff --git a/Engine/Source/Programs/UnrealBuildTool/System/CppDependencyCache.cs b/Engine/Source/Programs/UnrealBuildTool/System/CppDependencyCache.cs
--- a/Engine/Source/Programs/UnrealBuildTool/System/CppDependencyCache.cs
+++ b/Engine/Source/Programs/UnrealBuildTool/System/CppDependencyCache.cs
@@ -23,25 +23,32 @@ namespace UnrealBuildTool
 		class DependencyInfo
 		{
 			public long LastWriteTimeUtc;
+			public string ProducedModule;
+			public List<string> ImportedModules;
 			public List<FileItem> Files;
 
-			public DependencyInfo(long LastWriteTimeUtc, List<FileItem> Files)
+			public DependencyInfo(long LastWriteTimeUtc, string ProducedModule, List<string> ImportedModules, List<FileItem> Files)
 			{
 				this.LastWriteTimeUtc = LastWriteTimeUtc;
+				this.ProducedModule = ProducedModule;
+				this.ImportedModules = ImportedModules;
 				this.Files = Files;
 			}
 
 			public static DependencyInfo Read(BinaryArchiveReader Reader)
 			{
 				long LastWriteTimeUtc = Reader.ReadLong();
+				string ProducedModule = Reader.ReadString();
+				List<string> ImportedModules = Reader.ReadList(() => Reader.ReadString());
 				List<FileItem> Files = Reader.ReadList(() => Reader.ReadCompactFileItem());
-
-				return new DependencyInfo(LastWriteTimeUtc, Files);
+				return new DependencyInfo(LastWriteTimeUtc, ProducedModule, ImportedModules, Files);
 			}
 
 			public void Write(BinaryArchiveWriter Writer)
 			{
 				Writer.WriteLong(LastWriteTimeUtc);
+				Writer.WriteString(ProducedModule);
+				Writer.WriteList(ImportedModules, Module => Writer.WriteString(Module));
 				Writer.WriteList<FileItem>(Files, File => Writer.WriteCompactFileItem(File));
 			}
 		}
@@ -49,7 +56,7 @@ namespace UnrealBuildTool
 		/// <summary>
 		/// The current file version
 		/// </summary>
-		public const int CurrentVersion = 2;
+		public const int CurrentVersion = 3;
 
 		/// <summary>
 		/// Location of this dependency cache
@@ -99,6 +106,48 @@ namespace UnrealBuildTool
 			}
 		}
 
+		/// <summary>
+		/// Gets the produced module from a dependencies file
+		/// </summary>
+		/// <param name="InputFile">The dependencies file</param>
+		/// <param name="OutModule">The produced module name</param>
+		/// <returns>True if a produced module was found</returns>
+		public bool TryGetProducedModule(FileItem InputFile, out string OutModule)
+		{
+			DependencyInfo Info;
+			if (TryGetDependencyInfo(InputFile, out Info) && Info.ProducedModule != null)
+			{
+				OutModule = Info.ProducedModule;
+				return true;
+			}
+			else
+			{
+				OutModule = null;
+				return false;
+			}
+		}
+
+		/// <summary>
+		/// Attempts to get a list of imported modules for the given file
+		/// </summary>
+		/// <param name="InputFile">The dependency file to query</param>
+		/// <param name="OutImportedModules">List of imported modules</param>
+		/// <returns>True if a list of imported modules was obtained</returns>
+		public bool TryGetImportedModules(FileItem InputFile, out List<string> OutImportedModules)
+		{
+			DependencyInfo Info;
+			if (TryGetDependencyInfo(InputFile, out Info))
+			{
+				OutImportedModules = Info.ImportedModules;
+				return true;
+			}
+			else
+			{
+				OutImportedModules = null;
+				return false;
+			}
+		}
+
 		/// <summary>
 		/// Attempts to read the dependencies from the given input file
 		/// </summary>
@@ -107,20 +156,41 @@ namespace UnrealBuildTool
 		/// <returns>True if the input file exists and the dependencies were read</returns>
 		public bool TryGetDependencies(FileItem InputFile, out List<FileItem> OutDependencyItems)
 		{
-			if (!InputFile.Exists)
+			DependencyInfo Info;
+			if (TryGetDependencyInfo(InputFile, out Info))
+			{
+				OutDependencyItems = Info.Files;
+				return true;
+			}
+			else
 			{
 				OutDependencyItems = null;
 				return false;
 			}
+		}
+
+		/// <summary>
+		/// Attempts to read the dependencies from the given input file
+		/// </summary>
+		/// <param name="InputFile">File to be read</param>
+		/// <param name="OutInfo">The dependency info</param>
+		/// <returns>True if the input file exists and the dependencies were read</returns>
+		private bool TryGetDependencyInfo(FileItem InputFile, out DependencyInfo OutInfo)
+		{
+			if (!InputFile.Exists)
+			{
+				OutInfo = null;
+				return false;
+			}
 
 			try
 			{
-				return TryGetDependenciesInternal(InputFile, out OutDependencyItems);
+				return TryGetDependencyInfoInternal(InputFile, out OutInfo);
 			}
 			catch (Exception Ex)
 			{
 				Log.TraceLog("Unable to read {0}:\n{1}", InputFile, ExceptionUtils.FormatExceptionDetails(Ex));
-				OutDependencyItems = null;
+				OutInfo = null;
 				return false;
 			}
 		}
@@ -129,28 +199,25 @@ namespace UnrealBuildTool
 		/// Attempts to read dependencies from the given file.
 		/// </summary>
 		/// <param name="InputFile">File to be read</param>
-		/// <param name="OutDependencyItems">Receives a list of output items</param>
+		/// <param name="OutInfo">The dependency info</param>
 		/// <returns>True if the input file exists and the dependencies were read</returns>
-		private bool TryGetDependenciesInternal(FileItem InputFile, out List<FileItem> OutDependencyItems)
+		private bool TryGetDependencyInfoInternal(FileItem InputFile, out DependencyInfo OutInfo)
 		{
 			if (Parent != null && !InputFile.Location.IsUnderDirectory(BaseDir))
 			{
-				return Parent.TryGetDependencies(InputFile, out OutDependencyItems);
+				return Parent.TryGetDependencyInfoInternal(InputFile, out OutInfo);
 			}
 			else
 			{
 				DependencyInfo Info;
-				if (DependencyFileToInfo.TryGetValue(InputFile, out Info) && InputFile.LastWriteTimeUtc.Ticks <= Info.LastWriteTimeUtc)
+				if (!DependencyFileToInfo.TryGetValue(InputFile, out Info) || InputFile.LastWriteTimeUtc.Ticks > Info.LastWriteTimeUtc)
 				{
-					OutDependencyItems = Info.Files;
-					return true;
+					Info = ReadDependencyInfo(InputFile);
+					DependencyFileToInfo.TryAdd(InputFile, Info);
+					bModified = true;
 				}
 
-				List<FileItem> DependencyItems = ReadDependenciesFile(InputFile.Location);
-				DependencyFileToInfo.TryAdd(InputFile, new DependencyInfo(InputFile.LastWriteTimeUtc.Ticks, DependencyItems));
-				bModified = true;
-
-				OutDependencyItems = DependencyItems;
+				OutInfo = Info;
 				return true;
 			}
 		}
@@ -287,11 +354,11 @@ namespace UnrealBuildTool
 		/// </summary>
 		/// <param name="InputFile">The file to read from</param>
 		/// <returns>List of included dependencies</returns>
-		static List<FileItem> ReadDependenciesFile(FileReference InputFile)
+		static DependencyInfo ReadDependencyInfo(FileItem InputFile)
 		{
 			if (InputFile.HasExtension(".d"))
 			{
-				string Text = FileReference.ReadAllText(InputFile);
+				string Text = FileReference.ReadAllText(InputFile.Location);
 
 				List<string> Tokens = new List<string>();
 
@@ -330,11 +397,11 @@ namespace UnrealBuildTool
 					throw new BuildException("Unable to parse dependency file");
 				}
 
-				return NewDependencyFiles;
+				return new DependencyInfo(InputFile.LastWriteTimeUtc.Ticks, null, null, NewDependencyFiles);
 			}
 			else if (InputFile.HasExtension(".txt"))
 			{
-				string[] Lines = FileReference.ReadAllLines(InputFile);
+				string[] Lines = FileReference.ReadAllLines(InputFile.Location);
 
 				HashSet<FileItem> DependencyItems = new HashSet<FileItem>();
 				foreach (string Line in Lines)
@@ -349,7 +416,37 @@ namespace UnrealBuildTool
 						}
 					}
 				}
-				return DependencyItems.ToList();
+				return new DependencyInfo(InputFile.LastWriteTimeUtc.Ticks, null, null, DependencyItems.ToList());
+			}
+			else if (InputFile.HasExtension(".json"))
+			{
+				JsonObject Object = JsonObject.Read(InputFile.Location);
+
+				JsonObject Data;
+				if (!Object.TryGetObjectField("Data", out Data))
+				{
+					throw new BuildException("Missing 'Data' field in {0}", InputFile);
+				}
+
+				string ProducedModule;
+				Data.TryGetStringField("ProvidedModule", out ProducedModule);
+
+				string[] ImportedModules;
+				Data.TryGetStringArrayField("ImportedModules", out ImportedModules);
+
+				string[] Includes;
+				Data.TryGetStringArrayField("Includes", out Includes);
+
+				List<FileItem> Files = new List<FileItem>();
+				if (Includes != null)
+				{
+					foreach (string Include in Includes)
+					{
+						Files.Add(FileItem.GetItemByPath(Include));
+					}
+				}
+
+				return new DependencyInfo(InputFile.LastWriteTimeUtc.Ticks, ProducedModule, ImportedModules.ToList(), Files);
 			}
 			else
 			{
diff --git a/Engine/Source/Programs/UnrealBuildTool/System/TargetMakefile.cs b/Engine/Source/Programs/UnrealBuildTool/System/TargetMakefile.cs
--- a/Engine/Source/Programs/UnrealBuildTool/System/TargetMakefile.cs
+++ b/Engine/Source/Programs/UnrealBuildTool/System/TargetMakefile.cs
@@ -21,7 +21,7 @@ namespace UnrealBuildTool
 		/// <summary>
 		/// The version number to write
 		/// </summary>
-		public const int CurrentVersion = 19;
+		public const int CurrentVersion = 20;
 
 		/// <summary>
 		/// The time at which the makefile was created
