Remove references to Enterprise from UBT.

diff --git a/Engine/Source/Programs/UnrealBuildTool/Configuration/UEBuildTarget.cs b/Engine/Source/Programs/UnrealBuildTool/Configuration/UEBuildTarget.cs
--- a/Engine/Source/Programs/UnrealBuildTool/Configuration/UEBuildTarget.cs
+++ b/Engine/Source/Programs/UnrealBuildTool/Configuration/UEBuildTarget.cs
@@ -653,7 +653,7 @@ namespace UnrealBuildTool
 		/// </summary>
 		/// <param name="Descriptor">Information about the target</param>
 		/// <param name="bSkipRulesCompile">Whether to skip compiling any rules assemblies</param>
-		/// <param name="bUsePrecompiled">Whether to use a precompiled engine/enterprise build</param>
+		/// <param name="bUsePrecompiled">Whether to use a precompiled engine build</param>
 		/// <returns>The build target object for the specified build rules source file</returns>
 		public static UEBuildTarget Create(TargetDescriptor Descriptor, bool bSkipRulesCompile, bool bUsePrecompiled)
 		{
@@ -1085,10 +1085,6 @@ namespace UnrealBuildTool
 			{
 				ProjectDirectory = ProjectFile.Directory;
 			}
-			else if (Rules.File.IsUnderDirectory(UnrealBuildTool.EnterpriseDirectory))
-			{
-				ProjectDirectory = UnrealBuildTool.EnterpriseDirectory;
-			}
 			else
 			{
 				ProjectDirectory = UnrealBuildTool.EngineDirectory;
@@ -2496,7 +2492,6 @@ namespace UnrealBuildTool
 			Dictionary<string, string> Variables = new Dictionary<string,string>();
 			Variables.Add("RootDir", UnrealBuildTool.RootDirectory.FullName);
 			Variables.Add("EngineDir", UnrealBuildTool.EngineDirectory.FullName);
-			Variables.Add("EnterpriseDir", UnrealBuildTool.EnterpriseDirectory.FullName);
 			Variables.Add("ProjectDir", ProjectDirectory.FullName);
 			Variables.Add("TargetName", TargetName);
 			Variables.Add("TargetPlatform", Platform.ToString());
@@ -2700,7 +2695,6 @@ namespace UnrealBuildTool
 				if(Rules.bBuildDeveloperTools)
 				{
 					Directories.AddRange(UnrealBuildTool.GetExtensionDirs(UnrealBuildTool.EngineDirectory, "Source/Developer"));
-					Directories.Add(DirectoryReference.Combine(UnrealBuildTool.EnterpriseSourceDirectory, "Developer"));
 				}
 
 				// Find all the modules that are not part of the standard set
diff --git a/Engine/Source/Programs/UnrealBuildTool/Modes/CleanMode.cs b/Engine/Source/Programs/UnrealBuildTool/Modes/CleanMode.cs
--- a/Engine/Source/Programs/UnrealBuildTool/Modes/CleanMode.cs
+++ b/Engine/Source/Programs/UnrealBuildTool/Modes/CleanMode.cs
@@ -117,7 +117,6 @@ namespace UnrealBuildTool
 				// Find the base folders that can contain binaries
 				List<DirectoryReference> BaseDirs = new List<DirectoryReference>();
 				BaseDirs.Add(UnrealBuildTool.EngineDirectory);
-				BaseDirs.Add(UnrealBuildTool.EnterpriseDirectory);
 				foreach (FileReference Plugin in Plugins.EnumeratePlugins(Target.ProjectFile))
 				{
 					BaseDirs.Add(Plugin.Directory);
diff --git a/Engine/Source/Programs/UnrealBuildTool/Modes/QueryTargetsMode.cs b/Engine/Source/Programs/UnrealBuildTool/Modes/QueryTargetsMode.cs
--- a/Engine/Source/Programs/UnrealBuildTool/Modes/QueryTargetsMode.cs
+++ b/Engine/Source/Programs/UnrealBuildTool/Modes/QueryTargetsMode.cs
@@ -53,10 +53,6 @@ namespace UnrealBuildTool
 			{
 				Assembly = RulesCompiler.CreateProjectRulesAssembly(ProjectFile, BuildConfiguration.bUsePrecompiled, BuildConfiguration.bSkipRulesCompile);
 			}
-			else if(DirectoryReference.Exists(UnrealBuildTool.EnterpriseDirectory))
-			{
-				Assembly = RulesCompiler.CreateEnterpriseRulesAssembly(BuildConfiguration.bUsePrecompiled, BuildConfiguration.bSkipRulesCompile);
-			}
 			else
 			{
 				Assembly = RulesCompiler.CreateEngineRulesAssembly(BuildConfiguration.bUsePrecompiled, BuildConfiguration.bSkipRulesCompile);
diff --git a/Engine/Source/Programs/UnrealBuildTool/ProjectFiles/ProjectFileGenerator.cs b/Engine/Source/Programs/UnrealBuildTool/ProjectFiles/ProjectFileGenerator.cs
--- a/Engine/Source/Programs/UnrealBuildTool/ProjectFiles/ProjectFileGenerator.cs
+++ b/Engine/Source/Programs/UnrealBuildTool/ProjectFiles/ProjectFileGenerator.cs
@@ -242,12 +242,6 @@ namespace UnrealBuildTool
 		[XmlConfigFile]
 		protected bool bIncludeEngineSource = true;
 
-		/// <summary>
-		/// Whether to include enterprise source in the generated solution.
-		/// </summary>
-		[XmlConfigFile]
-		protected bool bIncludeEnterpriseSource = true;
-
 		/// <summary>
 		/// True if shader source files should be included in generated projects.
 		/// </summary>
@@ -326,11 +320,6 @@ namespace UnrealBuildTool
 		/// </summary>
 		public const string EngineProjectFileNameBase = "UE5";
 
-		/// <summary>
-		/// Name of the UE4 enterprise project that contains all of the enterprise code, config files and other files
-		/// </summary>
-		public const string EnterpriseProjectFileNameBase = "Studio";
-
 		/// <summary>
 		/// When ProjectsAreIntermediate is true, this is the directory to store generated project files
 		/// @todo projectfiles: Ideally, projects for game modules/targets would be created in the game's Intermediate folder!
@@ -741,13 +730,12 @@ namespace UnrealBuildTool
 			List<FileReference> AllModuleFiles = DiscoverModules(AllGameProjects);
 
 			ProjectFile EngineProject;
-			ProjectFile EnterpriseProject;
 			List<ProjectFile> GameProjects;
 			List<ProjectFile> ModProjects;
 			Dictionary<FileReference, ProjectFile> ProgramProjects;
 			{
 				// Setup buildable projects for all targets
-				AddProjectsForAllTargets( PlatformProjectGenerators, AllGameProjects, AllTargetFiles, Arguments, out EngineProject, out EnterpriseProject, out GameProjects, out ProgramProjects );
+				AddProjectsForAllTargets( PlatformProjectGenerators, AllGameProjects, AllTargetFiles, Arguments, out EngineProject, out GameProjects, out ProgramProjects );
 
 				// Add projects for mods
 				AddProjectsForMods(GameProjects, out ModProjects);
@@ -845,11 +833,6 @@ namespace UnrealBuildTool
 						}
 					}
 
-					if (EnterpriseProject != null)
-					{
-						RootFolder.AddSubFolder(UnrealBuildTool.EnterpriseDirectory.GetDirectoryName()).ChildProjects.Add(EnterpriseProject);
-					}
-
 					foreach( ProjectFile CurModProject in ModProjects )
 					{
 						RootFolder.AddSubFolder("Mods").ChildProjects.Add(CurModProject);
@@ -857,36 +840,20 @@ namespace UnrealBuildTool
 
 					DirectoryReference TemplatesDirectory = DirectoryReference.Combine(UnrealBuildTool.RootDirectory, "Templates");
 					DirectoryReference SamplesDirectory = DirectoryReference.Combine(UnrealBuildTool.RootDirectory, "Samples");
-					DirectoryReference EnterpriseTemplatesDirectory = DirectoryReference.Combine(UnrealBuildTool.EnterpriseDirectory, "Templates");
-					DirectoryReference EnterpriseSamplesDirectory = DirectoryReference.Combine(UnrealBuildTool.EnterpriseDirectory, "Samples");
 
 					foreach( ProjectFile CurGameProject in GameProjects )
 					{
 						// Templates go under a different solution folder than games
 						FileReference UnrealProjectFile = CurGameProject.ProjectTargets.First().UnrealProjectFilePath;
-						if (UnrealProjectFile.IsUnderDirectory(TemplatesDirectory) || UnrealProjectFile.IsUnderDirectory(EnterpriseTemplatesDirectory))
+						if (UnrealProjectFile.IsUnderDirectory(TemplatesDirectory))
 						{
 							DirectoryReference TemplateGameDirectory = CurGameProject.BaseDir;
-							if( TemplateGameDirectory.IsUnderDirectory((UnrealBuildTool.EnterpriseDirectory)) )
-							{
-								RootFolder.AddSubFolder(UnrealBuildTool.EnterpriseDirectory.GetDirectoryName() + Path.DirectorySeparatorChar + "Templates").ChildProjects.Add(CurGameProject);
-							}
-							else
-							{
-								RootFolder.AddSubFolder("Templates").ChildProjects.Add(CurGameProject);
-							}
+							RootFolder.AddSubFolder("Templates").ChildProjects.Add(CurGameProject);
 						}
-						else if (UnrealProjectFile.IsUnderDirectory(SamplesDirectory) || UnrealProjectFile.IsUnderDirectory(EnterpriseSamplesDirectory))
+						else if (UnrealProjectFile.IsUnderDirectory(SamplesDirectory))
 						{
 							DirectoryReference SampleGameDirectory = CurGameProject.BaseDir;
-							if( SampleGameDirectory.IsUnderDirectory((UnrealBuildTool.EnterpriseDirectory)) )
-                            {
-								RootFolder.AddSubFolder(UnrealBuildTool.EnterpriseDirectory.GetDirectoryName() + Path.DirectorySeparatorChar + "Samples").ChildProjects.Add(CurGameProject);
-							}
-							else
-							{
-								RootFolder.AddSubFolder("Samples").ChildProjects.Add(CurGameProject);
-							}
+							RootFolder.AddSubFolder("Samples").ChildProjects.Add(CurGameProject);
 						}
 						else
 						{
@@ -922,14 +889,7 @@ namespace UnrealBuildTool
                         }
                         else
                         {
-							if (CurProgramProject.Key.IsUnderDirectory(UnrealBuildTool.EnterpriseDirectory))
-							{
-								RootFolder.AddSubFolder(UnrealBuildTool.EnterpriseDirectory.GetDirectoryName() + Path.DirectorySeparatorChar + "Programs").ChildProjects.Add(CurProgramProject.Value);
-							}
-							else
-							{
-								RootFolder.AddSubFolder( "Programs" ).ChildProjects.Add( CurProgramProject.Value );
-							}
+							RootFolder.AddSubFolder( "Programs" ).ChildProjects.Add( CurProgramProject.Value );
                         }
 					}
 
@@ -986,23 +946,6 @@ namespace UnrealBuildTool
 						}
 					}
 
-					// Enterprise targets
-					if (EnterpriseProject != null)
-					{
-						foreach (ProjectTarget ProjectTarget in EnterpriseProject.ProjectTargets)
-						{
-							if (ProjectTarget.TargetFilePath != null)
-							{
-								// Only bother with the editor target.  We want to make sure that definitions are setup to be as inclusive as possible
-								// for good quality IntelliSense.  For example, we want WITH_EDITORONLY_DATA=1, so using the editor targets works well.
-								if (ProjectTarget.TargetRules.Type == TargetType.Editor)
-								{
-									IntelliSenseTargetFiles.Add(Tuple.Create(EnterpriseProject, ProjectTarget));
-								}
-							}
-						}
-					}
-
 					// Program targets
 					foreach( ProjectFile ProgramProject in ProgramProjects.Values )
 					{
@@ -1235,16 +1178,6 @@ namespace UnrealBuildTool
 			{
 				// At least one extra argument was specified, but we weren't expected it.  Ignored.
 			}
-
-			// If we're generating a solution for only one project, only include the enterprise folder if it's an enterprise project
-			if(OnlyGameProject == null)
-			{
-				bIncludeEnterpriseSource = bIncludeEngineSource;
-			}
-			else
-			{
-				bIncludeEnterpriseSource = bIncludeEngineSource && ProjectDescriptor.FromFile(OnlyGameProject).IsEnterpriseProject;
-			}
 		}
 
 
@@ -1789,15 +1722,7 @@ namespace UnrealBuildTool
 					// @todo projectfiles: Fragile assumption here about Programs always being under /Engine/Programs
 
 					DirectoryReference ProgramDirectory;
-
-                    if ( FileAndProject.Key.IsUnderDirectory(UnrealBuildTool.EnterpriseDirectory) )
-					{
-						ProgramDirectory = DirectoryReference.Combine( UnrealBuildTool.EnterpriseDirectory, "Programs", ProgramName );
-					}
-					else
-					{
-						ProgramDirectory = DirectoryReference.Combine( UnrealBuildTool.EngineDirectory, "Programs", ProgramName );
-					}
+					ProgramDirectory = DirectoryReference.Combine( UnrealBuildTool.EngineDirectory, "Programs", ProgramName );
 
 					DirectoryReference ProgramConfigDirectory = DirectoryReference.Combine( ProgramDirectory, "Config" );
 					if( DirectoryReference.Exists(ProgramConfigDirectory) )
@@ -2048,12 +1973,6 @@ namespace UnrealBuildTool
 					}
 				}
 
-				if( CurModuleFile.IsUnderDirectory(UnrealBuildTool.EnterpriseDirectory) && !bIncludeEnterpriseSource )
-				{
-					// We were asked to exclude enterprise modules from the generated projects
-					WantProjectFileForModule = false;
-				}
-
 				if( WantProjectFileForModule )
 				{
 					DirectoryReference BaseFolder;
@@ -2189,11 +2108,6 @@ namespace UnrealBuildTool
 						return FindOrAddProjectHelper(EngineProjectFileNameBase, BaseFolder);
 					}
 				}
-				if (Path == UnrealBuildTool.EnterpriseDirectory)
-				{
-					BaseFolder = UnrealBuildTool.EnterpriseDirectory;
-					return FindOrAddProjectHelper(EnterpriseProjectFileNameBase, BaseFolder);
-				}
 
 				// no match found, lets search the parent directory
 				Path = Path.ParentDirectory;
@@ -2210,7 +2124,6 @@ namespace UnrealBuildTool
 		/// <param name="AllTargetFiles">All the target files to add</param>
 		/// <param name="Arguments">The commandline arguments used</param>
 		/// <param name="EngineProject">The engine project we created</param>
-		/// <param name="EnterpriseProject">The enterprise project we created</param>
 		/// <param name="GameProjects">Map of game folder name to all of the game projects we created</param>
 		/// <param name="ProgramProjects">Map of program names to all of the program projects we created</param>
 		private void AddProjectsForAllTargets(
@@ -2219,13 +2132,11 @@ namespace UnrealBuildTool
 			List<FileReference> AllTargetFiles,
 			String[] Arguments,
 			out ProjectFile EngineProject, 
-			out ProjectFile EnterpriseProject,
 			out List<ProjectFile> GameProjects, 
 			out Dictionary<FileReference, ProjectFile> ProgramProjects)
 		{
 			// As we're creating project files, we'll also keep track of whether we created an "engine" project and return that if we have one
 			EngineProject = null;
-			EnterpriseProject = null;
 			GameProjects = new List<ProjectFile>();
 			ProgramProjects = new Dictionary<FileReference, ProjectFile>();
 
@@ -2250,7 +2161,6 @@ namespace UnrealBuildTool
 
 			// Get some standard directories
 			//DirectoryReference EngineSourceProgramsDirectory = DirectoryReference.Combine(UnrealBuildTool.EngineSourceDirectory, "Programs");
-			DirectoryReference EnterpriseSourceProgramsDirectory = DirectoryReference.Combine(UnrealBuildTool.EnterpriseSourceDirectory, "Programs");
 
 			// Keep a cache of the default editor target for each project so that we don't have to interrogate the configs for each target
 			Dictionary<string, string> DefaultProjectEditorTargetCache = new Dictionary<string, string>();
@@ -2261,7 +2171,6 @@ namespace UnrealBuildTool
 
 				// Check to see if this is an Engine target.  That is, the target is located under the "Engine" folder
 				bool IsEngineTarget = false;
-				bool IsEnterpriseTarget = false;
 				bool WantProjectFileForTarget = true;
 				if(TargetFilePath.IsUnderDirectory(UnrealBuildTool.EngineDirectory))
 				{
@@ -2277,20 +2186,6 @@ namespace UnrealBuildTool
 						WantProjectFileForTarget = bIncludeEngineSource;
 					}
 				}
-				else if (TargetFilePath.IsUnderDirectory(UnrealBuildTool.EnterpriseSourceDirectory))
-				{
-					// This is an enterprise target
-					IsEnterpriseTarget = true;
-
-					if(TargetFilePath.IsUnderDirectory(EnterpriseSourceProgramsDirectory))
-					{
-						WantProjectFileForTarget = bIncludeEnterpriseSource && IncludeEnginePrograms;
-					}
-					else
-					{
-						WantProjectFileForTarget = bIncludeEnterpriseSource;
-					}
-				}
 				
 				if (WantProjectFileForTarget)
 				{
@@ -2299,14 +2194,7 @@ namespace UnrealBuildTool
 					FileReference CheckProjectFile = AllGames.FirstOrDefault(x => TargetFilePath.IsUnderDirectory(x.Directory));
 					if(CheckProjectFile == null)
 					{
-						if(TargetFilePath.IsUnderDirectory(UnrealBuildTool.EnterpriseDirectory))
-						{
-							RulesAssembly = RulesCompiler.CreateEnterpriseRulesAssembly(false, false);
-						}
-						else
-						{
-							RulesAssembly = RulesCompiler.CreateEngineRulesAssembly(false, false);
-						}
+						RulesAssembly = RulesCompiler.CreateEngineRulesAssembly(false, false);
 					}
 					else
 					{
@@ -2330,10 +2218,6 @@ namespace UnrealBuildTool
 					{
 						ProjectFileNameBase = EngineProjectFileNameBase;
 					}
-					else if (IsEnterpriseTarget)
-					{
-						ProjectFileNameBase = EnterpriseProjectFileNameBase;
-					}
 					else
 					{
 						// Figure out which game project this target belongs to
@@ -2416,10 +2300,6 @@ namespace UnrealBuildTool
 					{
 						BaseFolder = UnrealBuildTool.EngineDirectory;
 					}
-					else if (IsEnterpriseTarget)
-					{
-						BaseFolder = UnrealBuildTool.EnterpriseDirectory;
-					}
 					else
 					{
 						BaseFolder = GameFolder;
@@ -2451,17 +2331,6 @@ namespace UnrealBuildTool
 							EngineProject.IsStubProject = true;
 						}
 					}
-					else if (IsEnterpriseTarget)
-					{
-						EnterpriseProject = ProjectFile;
-						if (UnrealBuildTool.IsEnterpriseInstalled())
-						{
-							// Allow enterprise projects to be created but not built for Installed Engine builds
-							EnterpriseProject.IsForeignProject = false;
-							EnterpriseProject.IsGeneratedProject = true;
-							EnterpriseProject.IsStubProject = true;
-						}
-					}
 					else
 					{
 						if (!bProjectAlreadyExisted)
diff --git a/Engine/Source/Programs/UnrealBuildTool/ProjectFiles/Rider/RiderProjectFileGenerator.cs b/Engine/Source/Programs/UnrealBuildTool/ProjectFiles/Rider/RiderProjectFileGenerator.cs
--- a/Engine/Source/Programs/UnrealBuildTool/ProjectFiles/Rider/RiderProjectFileGenerator.cs
+++ b/Engine/Source/Programs/UnrealBuildTool/ProjectFiles/Rider/RiderProjectFileGenerator.cs
@@ -23,9 +23,6 @@ namespace UnrealBuildTool
 		private readonly DirectoryReference EngineSourceProgramsDirectory =
 			DirectoryReference.Combine(UnrealBuildTool.EngineSourceDirectory, "Programs");
 
-		private readonly DirectoryReference EnterpriseSourceProgramsDirectory =
-			DirectoryReference.Combine(UnrealBuildTool.EnterpriseSourceDirectory, "Programs");
-
 		private readonly CommandLineArguments Arguments;
 
 
@@ -185,7 +182,6 @@ namespace UnrealBuildTool
 
 				// Check to see if this is an Engine target.  That is, the target is located under the "Engine" folder
 				bool IsEngineTarget = false;
-				bool IsEnterpriseTarget = false;
 				bool WantProjectFileForTarget = true;
 				if (TargetFilePath.IsUnderDirectory(UnrealBuildTool.EngineDirectory))
 				{
@@ -201,20 +197,6 @@ namespace UnrealBuildTool
 						WantProjectFileForTarget = bIncludeEngineSource;
 					}
 				}
-				else if (TargetFilePath.IsUnderDirectory(UnrealBuildTool.EnterpriseSourceDirectory))
-				{
-					// This is an enterprise target
-					IsEnterpriseTarget = true;
-
-					if (TargetFilePath.IsUnderDirectory(EnterpriseSourceProgramsDirectory))
-					{
-						WantProjectFileForTarget = bIncludeEnterpriseSource && IncludeEnginePrograms;
-					}
-					else
-					{
-						WantProjectFileForTarget = bIncludeEnterpriseSource;
-					}
-				}
 
 				if (WantProjectFileForTarget)
 				{
@@ -224,14 +206,7 @@ namespace UnrealBuildTool
 						AllGames.FirstOrDefault(x => TargetFilePath.IsUnderDirectory(x.Directory));
 					if (CheckProjectFile == null)
 					{
-						if (TargetFilePath.IsUnderDirectory(UnrealBuildTool.EnterpriseDirectory))
-						{
-							RulesAssembly = RulesCompiler.CreateEnterpriseRulesAssembly(false, false);
-						}
-						else
-						{
-							RulesAssembly = RulesCompiler.CreateEngineRulesAssembly(false, false);
-						}
+						RulesAssembly = RulesCompiler.CreateEngineRulesAssembly(false, false);
 					}
 					else
 					{
@@ -257,10 +232,6 @@ namespace UnrealBuildTool
 					{
 						ProjectFileNameBase = EngineProjectFileNameBase;
 					}
-					else if (IsEnterpriseTarget)
-					{
-						ProjectFileNameBase = EnterpriseProjectFileNameBase;
-					}
 					else
 					{
 						// Figure out which game project this target belongs to
@@ -314,10 +285,6 @@ namespace UnrealBuildTool
 					{
 						BaseFolder = UnrealBuildTool.EngineDirectory;
 					}
-					else if (IsEnterpriseTarget)
-					{
-						BaseFolder = UnrealBuildTool.EnterpriseDirectory;
-					}
 					else
 					{
 						BaseFolder = GameFolder;
@@ -352,17 +319,6 @@ namespace UnrealBuildTool
 							EngineProject.IsStubProject = true;
 						}
 					}
-					else if (IsEnterpriseTarget)
-					{
-						ProjectFile EnterpriseProject = ProjectFile;
-						if (UnrealBuildTool.IsEnterpriseInstalled())
-						{
-							// Allow enterprise projects to be created but not built for Installed Engine builds
-							EnterpriseProject.IsForeignProject = false;
-							EnterpriseProject.IsGeneratedProject = true;
-							EnterpriseProject.IsStubProject = true;
-						}
-					}
 					else
 					{
 						if (!bProjectAlreadyExisted)
diff --git a/Engine/Source/Programs/UnrealBuildTool/System/Plugins.cs b/Engine/Source/Programs/UnrealBuildTool/System/Plugins.cs
--- a/Engine/Source/Programs/UnrealBuildTool/System/Plugins.cs
+++ b/Engine/Source/Programs/UnrealBuildTool/System/Plugins.cs
@@ -35,11 +35,6 @@ namespace UnrealBuildTool
 		/// </summary>
 		Engine,
 
-		/// <summary>
-		/// Enterprise plugin
-		/// </summary>
-		Enterprise,
-
 		/// <summary>
 		/// Project-specific plugin, stored within a game project directory
 		/// </summary>
@@ -135,7 +130,7 @@ namespace UnrealBuildTool
 		{
 			get
 			{
-				if(Type == PluginType.Engine || Type == PluginType.Enterprise)
+				if(Type == PluginType.Engine)
 				{
 					return PluginLoadedFrom.Engine;
 				}
@@ -246,12 +241,6 @@ namespace UnrealBuildTool
 				yield return PluginFile;
 			}
 
-			DirectoryReference EnterprisePluginsDir = DirectoryReference.Combine(UnrealBuildTool.EnterpriseDirectory, "Plugins");
-			foreach(FileReference PluginFile in EnumeratePlugins(EnterprisePluginsDir))
-			{
-				yield return PluginFile;
-			}
-
 			if(ProjectFile != null)
 			{
 				DirectoryReference ProjectPluginsDir = DirectoryReference.Combine(ProjectFile.Directory, "Plugins");
@@ -278,16 +267,6 @@ namespace UnrealBuildTool
 			return ReadPluginsFromDirectory(EngineDirectory, "Plugins", PluginType.Engine);
 		}
 
-		/// <summary>
-		/// Read all the plugin descriptors under the given enterprise directory
-		/// </summary>
-		/// <param name="EnterpriseDirectory">The parent directory to look in.</param>
-		/// <returns>Sequence of the found PluginInfo object.</returns>
-		public static IReadOnlyList<PluginInfo> ReadEnterprisePlugins(DirectoryReference EnterpriseDirectory)
-		{
-			return ReadPluginsFromDirectory(EnterpriseDirectory, "Plugins", PluginType.Enterprise);
-		}
-
 		/// <summary>
 		/// Read all the plugin descriptors under the given project directory
 		/// </summary>
diff --git a/Engine/Source/Programs/UnrealBuildTool/System/RulesCompiler.cs b/Engine/Source/Programs/UnrealBuildTool/System/RulesCompiler.cs
--- a/Engine/Source/Programs/UnrealBuildTool/System/RulesCompiler.cs
+++ b/Engine/Source/Programs/UnrealBuildTool/System/RulesCompiler.cs
@@ -66,10 +66,9 @@ namespace UnrealBuildTool
 		/// <param name="ForeignPlugins"></param>
 		/// <param name="AdditionalSearchPaths"></param>
 		/// <param name="bIncludeEngine"></param>
-		/// <param name="bIncludeEnterprise"></param>
 		/// <param name="bIncludeTempTargets">Whether to include targets generated by UAT to accomodate content-only projects that need to be compiled to include plugins</param>
 		/// <returns></returns>
-		public static List<FileReference> FindAllRulesSourceFiles(RulesFileType RulesFileType, List<DirectoryReference> GameFolders, List<FileReference> ForeignPlugins, List<DirectoryReference> AdditionalSearchPaths, bool bIncludeEngine = true, bool bIncludeEnterprise = true, bool bIncludeTempTargets = true)
+		public static List<FileReference> FindAllRulesSourceFiles(RulesFileType RulesFileType, List<DirectoryReference> GameFolders, List<FileReference> ForeignPlugins, List<DirectoryReference> AdditionalSearchPaths, bool bIncludeEngine = true, bool bIncludeTempTargets = true)
 		{
 			List<DirectoryReference> Folders = new List<DirectoryReference>();
 
@@ -78,10 +77,6 @@ namespace UnrealBuildTool
 			{
 				Folders.AddRange(UnrealBuildTool.GetExtensionDirs(UnrealBuildTool.EngineDirectory, "Source"));
 			}
-			if(bIncludeEnterprise)
-			{
-				Folders.Add(UnrealBuildTool.EnterpriseSourceDirectory);
-			}
 
 			// @todo plugin: Disallow modules from including plugin modules as dependency modules? (except when the module is part of that plugin)
 
@@ -91,10 +86,6 @@ namespace UnrealBuildTool
 			{
 				RootFolders.AddRange(UnrealBuildTool.GetExtensionDirs(UnrealBuildTool.EngineDirectory));
 			}
-			if(bIncludeEnterprise)
-			{
-				RootFolders.Add(UnrealBuildTool.EnterpriseDirectory);
-			}
 			if (GameFolders != null)
 			{
 				RootFolders.AddRange(GameFolders.SelectMany(x => UnrealBuildTool.GetExtensionDirs(x)));
@@ -345,11 +336,6 @@ namespace UnrealBuildTool
 		/// </summary>
 		private static RulesAssembly EngineRulesAssembly;
 
-		/// <summary>
-		/// The cached rules assembly for enterprise modules and targets.
-		/// </summary>
-		private static RulesAssembly EnterpriseRulesAssembly;
-
 		/// <summary>
 		/// Map of assembly names we've already compiled and loaded to their Assembly and list of game folders.  This is used to prevent
 		/// trying to recompile the same assembly when ping-ponging between different types of targets
@@ -373,43 +359,11 @@ namespace UnrealBuildTool
 
 				RulesScope EngineScope = new RulesScope("Engine", null);
 
-				EngineRulesAssembly = CreateEngineOrEnterpriseRulesAssembly(EngineScope, UnrealBuildTool.GetExtensionDirs(UnrealBuildTool.EngineDirectory), ProjectFileGenerator.EngineProjectFileNameBase, IncludedPlugins, UnrealBuildTool.IsEngineInstalled() || bUsePrecompiled, bSkipCompile, null);
+				EngineRulesAssembly = CreateEngineRulesAssemblyInternal(EngineScope, UnrealBuildTool.GetExtensionDirs(UnrealBuildTool.EngineDirectory), ProjectFileGenerator.EngineProjectFileNameBase, IncludedPlugins, UnrealBuildTool.IsEngineInstalled() || bUsePrecompiled, bSkipCompile, null);
 			}
 			return EngineRulesAssembly;
 		}
 
-		/// <summary>
-		/// Creates the enterprise rules assembly
-		/// </summary>
-		/// <param name="bUsePrecompiled">Whether to use a precompiled enterprise and engine folder</param>
-		/// <param name="bSkipCompile">Whether to skip compilation for this assembly</param>
-		/// <returns>New rules assembly. Returns null if the enterprise directory is unavailable.</returns>
-		public static RulesAssembly CreateEnterpriseRulesAssembly(bool bUsePrecompiled, bool bSkipCompile)
-		{
-			if (EnterpriseRulesAssembly == null)
-			{
-				RulesAssembly EngineAssembly = CreateEngineRulesAssembly(bUsePrecompiled, bSkipCompile);
-				if (DirectoryReference.Exists(UnrealBuildTool.EnterpriseDirectory))
-				{
-					RulesScope EnterpriseScope = new RulesScope("Enterprise", EngineAssembly.Scope);
-
-					//List<DirectoryReference> EnterpriseDirectories = new List<DirectoryReference>() { UnrealBuildTool.EnterpriseDirectory };
-
-					IReadOnlyList<PluginInfo> IncludedPlugins = Plugins.ReadEnterprisePlugins(UnrealBuildTool.EnterpriseDirectory);
-					EnterpriseRulesAssembly = CreateEngineOrEnterpriseRulesAssembly(EnterpriseScope, new List<DirectoryReference>() { UnrealBuildTool.EnterpriseDirectory }, 
-						ProjectFileGenerator.EnterpriseProjectFileNameBase, IncludedPlugins, UnrealBuildTool.IsEnterpriseInstalled() || bUsePrecompiled, bSkipCompile, EngineAssembly);
-				}
-				else
-				{
-					// If we're asked for the enterprise rules assembly but the enterprise directory is missing, fallback on the engine rules assembly
-					Log.TraceWarning("Trying to build an enterprise target but the enterprise directory is missing. Falling back on engine components only.");
-					return EngineAssembly;
-				}
-			}
-
-			return EnterpriseRulesAssembly;
-		}
-
 		/// <summary>
 		/// Creates a rules assembly
 		/// </summary>
@@ -421,7 +375,7 @@ namespace UnrealBuildTool
 		/// <param name="bSkipCompile">Whether to skip compilation for this assembly</param>
 		/// <param name="Parent">The parent rules assembly</param>
 		/// <returns>New rules assembly</returns>
-		private static RulesAssembly CreateEngineOrEnterpriseRulesAssembly(RulesScope Scope, List<DirectoryReference> RootDirectories, string AssemblyPrefix, IReadOnlyList<PluginInfo> Plugins, bool bReadOnly, bool bSkipCompile, RulesAssembly Parent)
+		private static RulesAssembly CreateEngineRulesAssemblyInternal(RulesScope Scope, List<DirectoryReference> RootDirectories, string AssemblyPrefix, IReadOnlyList<PluginInfo> Plugins, bool bReadOnly, bool bSkipCompile, RulesAssembly Parent)
 		{
 			// Scope hierarchy
 			RulesScope PluginsScope = new RulesScope(Scope.Name + " Plugins", Scope);
@@ -503,15 +457,7 @@ namespace UnrealBuildTool
 				ProjectDescriptor Project = ProjectDescriptor.FromFile(ProjectFileName);
 
 				// Create the parent assembly
-				RulesAssembly Parent;
-				if (Project.IsEnterpriseProject)
-				{
-					Parent = CreateEnterpriseRulesAssembly(bUsePrecompiled, bSkipCompile);
-				}
-				else
-				{
-					Parent = CreateEngineRulesAssembly(bUsePrecompiled, bSkipCompile);
-				}
+				RulesAssembly Parent = CreateEngineRulesAssembly(bUsePrecompiled, bSkipCompile);
 
 				DirectoryReference MainProjectDirectory = ProjectFileName.Directory;
 				//DirectoryReference MainProjectSourceDirectory = DirectoryReference.Combine(MainProjectDirectory, "Source");
@@ -632,7 +578,7 @@ namespace UnrealBuildTool
 		/// <param name="ProjectFile">The project file being compiled</param>
 		/// <param name="TargetName">The target being built</param>
 		/// <param name="bSkipRulesCompile">Whether to skip compiling any rules assemblies</param>
-		/// <param name="bUsePrecompiled">Whether to use a precompiled engine/enterprise build</param>
+		/// <param name="bUsePrecompiled">Whether to use a precompiled engine build</param>
 		/// <param name="ForeignPlugin">Foreign plugin to be compiled</param>
 		/// <returns>The compiled rules assembly</returns>
 		public static RulesAssembly CreateTargetRulesAssembly(FileReference ProjectFile, string TargetName, bool bSkipRulesCompile, bool bUsePrecompiled, FileReference ForeignPlugin)
@@ -645,12 +591,6 @@ namespace UnrealBuildTool
 			else
 			{
 				RulesAssembly = CreateEngineRulesAssembly(bUsePrecompiled, bSkipRulesCompile);
-
-				if (RulesAssembly.GetTargetFileName(TargetName) == null && DirectoryReference.Exists(UnrealBuildTool.EnterpriseDirectory))
-				{
-					// Target isn't part of the engine assembly, try the enterprise assembly
-					RulesAssembly = CreateEnterpriseRulesAssembly(bUsePrecompiled, bSkipRulesCompile);
-				}
 			}
 			if (ForeignPlugin != null)
 			{
@@ -699,10 +639,6 @@ namespace UnrealBuildTool
 			{
 				return FileName.FullName;
 			}
-			else if (EnterpriseRulesAssembly != null && EnterpriseRulesAssembly.TryGetFileNameFromType(ExistingType, out FileName))
-			{
-				return FileName.FullName;
-			}
 
 			foreach (RulesAssembly RulesAssembly in LoadedAssemblyMap.Values)
 			{
diff --git a/Engine/Source/Programs/UnrealBuildTool/UnrealBuildTool.cs b/Engine/Source/Programs/UnrealBuildTool/UnrealBuildTool.cs
--- a/Engine/Source/Programs/UnrealBuildTool/UnrealBuildTool.cs
+++ b/Engine/Source/Programs/UnrealBuildTool/UnrealBuildTool.cs
@@ -30,11 +30,6 @@ namespace UnrealBuildTool
 		/// </summary>
 		static private bool? bIsEngineInstalled;
 
-		/// <summary>
-		/// Whether we're running with enterprise installed
-		/// </summary>
-		static private bool? bIsEnterpriseInstalled;
-
 		/// <summary>
 		/// Whether we're running with an installed project
 		/// </summary>
@@ -100,26 +95,6 @@ namespace UnrealBuildTool
 		[Obsolete("Please use UnrealBuildTool.GetExtensionDirs(UnrealBuildTool.EngineDirectory, \"Source/ThirdParty\") instead.")]
 		public static readonly DirectoryReference EngineSourceThirdPartyDirectory = DirectoryReference.Combine(EngineSourceDirectory, "ThirdParty");
 
-		/// <summary>
-		/// The full name of the Enterprise directory
-		/// </summary>
-		public static readonly DirectoryReference EnterpriseDirectory = DirectoryReference.Combine(RootDirectory, "Enterprise");
-
-		/// <summary>
-		/// The full name of the Enterprise/Source directory
-		/// </summary>
-		public static readonly DirectoryReference EnterpriseSourceDirectory = DirectoryReference.Combine(EnterpriseDirectory, "Source");
-
-		/// <summary>
-		/// The full name of the Enterprise/Plugins directory
-		/// </summary>
-		public static readonly DirectoryReference EnterprisePluginsDirectory = DirectoryReference.Combine(EnterpriseDirectory, "Plugins");
-
-		/// <summary>
-		/// The full name of the Enterprise/Intermediate directory
-		/// </summary>
-		public static readonly DirectoryReference EnterpriseIntermediateDirectory = DirectoryReference.Combine(EnterpriseDirectory, "Intermediate");
-
 		/// <summary>
 		/// The engine programs directory
 		/// </summary>
@@ -235,19 +210,6 @@ namespace UnrealBuildTool
 			return bIsEngineInstalled.Value;
 		}
 
-		/// <summary>
-		/// Returns true if UnrealBuildTool is running using installed Enterprise components
-		/// </summary>
-		/// <returns>True if running using installed Enterprise components</returns>
-		static public bool IsEnterpriseInstalled()
-		{
-			if (!bIsEnterpriseInstalled.HasValue)
-			{
-				bIsEnterpriseInstalled = FileReference.Exists(FileReference.Combine(EnterpriseDirectory, "Build", "InstalledBuild.txt"));
-			}
-			return bIsEnterpriseInstalled.Value;
-		}
-
 		/// <summary>
 		/// Returns true if UnrealBuildTool is running using an installed project (ie. a mod kit)
 		/// </summary>
@@ -298,10 +260,6 @@ namespace UnrealBuildTool
 			{
 				return true;
 			}
-			if(IsEnterpriseInstalled() && File.IsUnderDirectory(EnterpriseDirectory))
-			{
-				return true;
-			}
 			if(IsProjectInstalled() && File.IsUnderDirectory(InstalledProjectFile.Directory))
 			{
 				return true;
