Improvements to C++20 modules support.

- PCHs and unity builds are now automatically disabled when module support is enabled.
- C++20 mode is automatically enabled if module support is enabled.
- Added test for using a module implementation unit. Seems to work with force-included headers after all (but breaks if unity builds are on)

diff --git a/Engine/Source/Programs/UnrealBuildTool/Configuration/ModuleRules.cs b/Engine/Source/Programs/UnrealBuildTool/Configuration/ModuleRules.cs
--- a/Engine/Source/Programs/UnrealBuildTool/Configuration/ModuleRules.cs
+++ b/Engine/Source/Programs/UnrealBuildTool/Configuration/ModuleRules.cs
@@ -651,6 +651,10 @@ namespace UnrealBuildTool
 					// Use the override
 					return PCHUsagePrivate.Value;
 				}
+				else if(Target.bEnableCppModules)
+				{
+					return PCHUsageMode.NoPCHs;
+				}
 				else if(Target.bIWYU || DefaultBuildSettings >= BuildSettingsVersion.V2)
 				{
 					// Use shared or explicit PCHs, and enable IWYU
diff --git a/Engine/Source/Programs/UnrealBuildTool/Configuration/TargetRules.cs b/Engine/Source/Programs/UnrealBuildTool/Configuration/TargetRules.cs
--- a/Engine/Source/Programs/UnrealBuildTool/Configuration/TargetRules.cs
+++ b/Engine/Source/Programs/UnrealBuildTool/Configuration/TargetRules.cs
@@ -845,9 +845,15 @@ namespace UnrealBuildTool
 		/// <summary>
 		/// Whether to unify C++ code into larger files for faster compilation.
 		/// </summary>
+		public bool bUseUnityBuild
+		{
+			get { return bUseUnityBuildOverride ?? !bEnableCppModules; }
+			set { bUseUnityBuildOverride = value; }
+		}
+
 		[CommandLine("-DisableUnity", Value = "false")]
-		[XmlConfigFile(Category = "BuildConfiguration")]
-		public bool bUseUnityBuild = true;
+		[XmlConfigFile(Category = "BuildConfiguration", Name = nameof(bUseUnityBuild))]
+		bool? bUseUnityBuildOverride = null;
 
 		/// <summary>
 		/// Whether to force C++ source files to be combined into larger files for faster compilation.
diff --git a/Engine/Source/Programs/UnrealBuildTool/Platform/Windows/VCToolChain.cs b/Engine/Source/Programs/UnrealBuildTool/Platform/Windows/VCToolChain.cs
--- a/Engine/Source/Programs/UnrealBuildTool/Platform/Windows/VCToolChain.cs
+++ b/Engine/Source/Programs/UnrealBuildTool/Platform/Windows/VCToolChain.cs
@@ -596,7 +596,7 @@ namespace UnrealBuildTool
 			}
 			else
 			{
-				if(CompileEnvironment.CppStandard >= CppStandardVersion.Latest)
+				if(CompileEnvironment.CppStandard >= CppStandardVersion.Cpp20)
 				{
 					Arguments.Add("/std:c++latest");
 				}
