Missing file.

diff --git a/Engine/Source/Programs/UnrealBuildTool/ToolChain/UEToolChain.cs b/Engine/Source/Programs/UnrealBuildTool/ToolChain/UEToolChain.cs
--- a/Engine/Source/Programs/UnrealBuildTool/ToolChain/UEToolChain.cs
+++ b/Engine/Source/Programs/UnrealBuildTool/ToolChain/UEToolChain.cs
@@ -23,6 +23,11 @@ namespace UnrealBuildTool
 		{
 		}
 
+		public static DirectoryReference GetModuleInterfaceDir(DirectoryReference OutputDir)
+		{
+			return DirectoryReference.Combine(OutputDir, "Ifc");
+		}
+
 		public abstract CPPOutput CompileCPPFiles(CppCompileEnvironment CompileEnvironment, List<FileItem> InputFiles, DirectoryReference OutputDir, string ModuleName, IActionGraphBuilder Graph);
 
 		public virtual CPPOutput CompileRCFiles(CppCompileEnvironment Environment, List<FileItem> InputFiles, DirectoryReference OutputDir, IActionGraphBuilder Graph)
