Add exception handler for building without visual studio installed on windows (clang).

diff --git a/Engine/Source/Programs/UnrealBuildTool/Platform/Windows/UEBuildWindows.cs b/Engine/Source/Programs/UnrealBuildTool/Platform/Windows/UEBuildWindows.cs
--- a/Engine/Source/Programs/UnrealBuildTool/Platform/Windows/UEBuildWindows.cs
+++ b/Engine/Source/Programs/UnrealBuildTool/Platform/Windows/UEBuildWindows.cs
@@ -359,7 +359,17 @@ namespace UnrealBuildTool
 		/// </summary>
 		public string IDEDir
 		{
-			get { return WindowsPlatform.FindVisualStudioInstallations(Environment.Compiler).Select(x => x.BaseDir.FullName).FirstOrDefault(); }
+			get
+			{
+				try
+				{
+					return WindowsPlatform.FindVisualStudioInstallations(Environment.Compiler).Select(x => x.BaseDir.FullName).FirstOrDefault();
+				}
+				catch(Exception) // Find function will throw if there is no visual studio installed! This can happen w/ clang builds
+				{
+					return string.Empty;
+				}
+			}
 		}
 
 		/// <summary>
