Added UnrealEngineGeneratedSolutionVersion tag to generated solution file which is used by UnrealVS to detect that this is a unreal solution. Should fix failure to detect renamed solution in UnrealVS.

diff --git a/Engine/Source/Programs/UnrealBuildTool/ProjectFiles/VisualStudio/VCProjectFileGenerator.cs b/Engine/Source/Programs/UnrealBuildTool/ProjectFiles/VisualStudio/VCProjectFileGenerator.cs
--- a/Engine/Source/Programs/UnrealBuildTool/ProjectFiles/VisualStudio/VCProjectFileGenerator.cs
+++ b/Engine/Source/Programs/UnrealBuildTool/ProjectFiles/VisualStudio/VCProjectFileGenerator.cs
@@ -532,6 +532,9 @@ namespace UnrealBuildTool
 				throw new BuildException("Unexpected ProjectFileFormat");
 			}
 
+			// Embed custom comment used by UnrealVS to determine that this is a Unreal solution
+			VCSolutionFileContent.AppendLine("# UnrealEngineGeneratedSolutionVersion=" + AssemblyUtils.ExecutableVersion.ProductVersion);
+
 			IDictionary<MasterProjectFolder, Guid> ProjectFolderGuids = GenerateProjectFolderGuids(RootFolder);
 
 			// Solution folders, files and project entries
