UnrealBuildTool now supports running as a net core project.
Added a set of netcore csprojs to BuildUtilities and DotnetUtilities that build to a seperate output folder. This allows other tools to still target .net framework (like UAT for instance).
UBT is still by default targeted as .net framework.

Note that UBT built for net core has a different output directory Engine/Binaries/DotNet/UnrealBuildTool/UnrealBuildTool.exe - this is due to how a netcore project output looks with signficantly more files that are related to that application (that would be overwritten if using a shared directory).

To opt in to this set UE_USE_DOTNET=1 environment variable.
Please note that due to the changed output directory a lot of tooling will likely break at this point.

diff --git a/Engine/Build/BatchFiles/Build.bat b/Engine/Build/BatchFiles/Build.bat
--- a/Engine/Build/BatchFiles/Build.bat
+++ b/Engine/Build/BatchFiles/Build.bat
@@ -10,8 +10,13 @@ REM %1 is the game name
 REM %2 is the platform name
 REM %3 is the configuration name
 
-IF EXIST ..\..\Engine\Binaries\DotNET\UnrealBuildTool.exe (
-        ..\..\Engine\Binaries\DotNET\UnrealBuildTool.exe %*
+ set UBTPath = "..\..\Engine\Binaries\DotNET\UnrealBuildTool.exe"
+if "%UE_USE_DOTNET%" == "1" (
+    set UBTPath = "..\..\Engine\Binaries\DotNET\UnrealBuildTool\UnrealBuildTool.exe"
+)
+
+IF EXIST %UBTPath% (
+        %UBTPath% %*
 		popd
 
 		REM Ignore exit codes of 2 ("ECompilationResult.UpToDate") from UBT; it's not a failure.
@@ -21,7 +26,7 @@ IF EXIST ..\..\Engine\Binaries\DotNET\UnrealBuildTool.exe (
 		 
 		EXIT /B !ERRORLEVEL!
 ) ELSE (
-	ECHO UnrealBuildTool.exe not found in ..\..\Engine\Binaries\DotNET\UnrealBuildTool.exe 
+	ECHO UnrealBuildTool.exe not found in %UBTPath%
 	popd
 	EXIT /B 999
 )
diff --git a/Engine/Build/BatchFiles/Clean.bat b/Engine/Build/BatchFiles/Clean.bat
--- a/Engine/Build/BatchFiles/Clean.bat
+++ b/Engine/Build/BatchFiles/Clean.bat
@@ -30,16 +30,27 @@ if errorlevel 1 goto Error_NoVisualStudioEnvironment
 
 
 rem ## Compile UBT if the project file exists
-if not exist Programs\UnrealBuildTool\UnrealBuildTool.csproj goto NoProjectFile
-%MSBUILD_EXE% /nologo /verbosity:quiet Programs\UnrealBuildTool\UnrealBuildTool.csproj /property:Configuration=Development /property:Platform=AnyCPU
+
+SET ProjectFile = "Programs\UnrealBuildTool\UnrealBuildTool.csproj"
+set UBTPath = "..\..\Engine\Binaries\DotNET\UnrealBuildTool.exe"
+if "%UE_USE_DOTNET%" == "1" (
+    set ProjectFile = "Programs\UnrealBuildTool\UnrealBuildToolCore.csproj"
+    set UBTPath = "..\..\Engine\Binaries\DotNET\UnrealBuildTool\UnrealBuildTool.exe"
+)
+if not exist %ProjectFile% goto NoProjectFile
+if "%UE_USE_DOTNET%" == "1" (
+    dotnet build Programs\UnrealBuildTool\UnrealBuildToolCore.csproj -c Development
+) else (
+    %MSBUILD_EXE% /nologo /verbosity:quiet Programs\UnrealBuildTool\UnrealBuildTool.csproj /property:Configuration=Development /property:Platform=AnyCPU
+)
 if errorlevel 1 goto Error_UBTCompileFailed
 :NoProjectFile
 
 
 rem ## Execute UBT
 :ReadyToClean
-if not exist ..\..\Engine\Binaries\DotNET\UnrealBuildTool.exe goto Error_UBTMissing
-..\..\Engine\Binaries\DotNET\UnrealBuildTool.exe %* -clean
+if not exist %UBTPath% goto Error_UBTMissing
+%UBTPath% %* -clean
 goto Exit
 
 
@@ -59,7 +70,11 @@ pause
 goto Exit
 
 :Error_UBTMissing
-echo Clean ERROR: UnrealBuildTool.exe not found in ..\..\Engine\Binaries\DotNET\UnrealBuildTool.exe 
+if "%UE_USE_DOTNET%" == "1" (
+    echo Clean ERROR: UnrealBuildTool.exe not found in ..\..\Engine\Binaries\DotNET\UnrealBuildTool\UnrealBuildTool.exe 
+) else (
+    echo Clean ERROR: UnrealBuildTool.exe not found in ..\..\Engine\Binaries\DotNET\UnrealBuildTool.exe 
+)
 pause
 goto Exit
 
diff --git a/Engine/Build/BatchFiles/GenerateProjectFiles.bat b/Engine/Build/BatchFiles/GenerateProjectFiles.bat
--- a/Engine/Build/BatchFiles/GenerateProjectFiles.bat
+++ b/Engine/Build/BatchFiles/GenerateProjectFiles.bat
@@ -24,6 +24,16 @@ rem ## Check to make sure that we have a Binaries directory with at least one de
 rem ## in order to run.  It's possible the user acquired source but did not download and unpack the other prerequiste binaries.
 if not exist ..\Build\BinaryPrerequisitesMarker.dat goto Error_MissingBinaryPrerequisites
 
+if "%UE_USE_DOTNET%" == "1" (
+	echo Using Net Core for UBT
+
+	rem ## Verify that dotnet is present
+	call "%~dp0GetDotnetPath.bat"
+	if errorlevel 1 goto Error_NoDotnetSDK
+	REM ## Skip msbuild detection if using dotnet as this is done for us by dotnet-cli
+	goto :NoVsWhere
+)
+
 rem ## Get the path to MSBuild
 call "%~dp0GetMSBuildPath.bat"
 if errorlevel 1 goto Error_NoVisualStudioEnvironment
@@ -74,14 +84,26 @@ for /d %%D in (..\Restricted\*) do (
 fc /b ..\Intermediate\Build\UnrealBuildToolFiles.txt ..\Intermediate\Build\UnrealBuildToolPrevFiles.txt >nul 2>nul
 if not errorlevel 1 goto SkipClean
 copy /y ..\Intermediate\Build\UnrealBuildToolFiles.txt ..\Intermediate\Build\UnrealBuildToolPrevFiles.txt >nul
-%MSBUILD_EXE% /nologo /verbosity:quiet Programs\UnrealBuildTool\UnrealBuildTool.csproj /property:Configuration=Development /property:Platform=AnyCPU /target:Clean
+if "%UE_USE_DOTNET%" == "1" (
+	dotnet clean Programs\UnrealBuildTool\UnrealBuildToolCore.csproj -c Development
+) else (
+	%MSBUILD_EXE% /nologo /verbosity:quiet Programs\UnrealBuildTool\UnrealBuildTool.csproj /property:Configuration=Development /property:Platform=AnyCPU /target:Clean
+)
 :SkipClean
-%MSBUILD_EXE% /nologo /verbosity:quiet Programs\UnrealBuildTool\UnrealBuildTool.csproj /property:Configuration=Development /property:Platform=AnyCPU /target:Build
+if "%UE_USE_DOTNET%" == "1" (
+	dotnet build Programs\UnrealBuildTool\UnrealBuildToolCore.csproj -c Development 
+) else (
+	%MSBUILD_EXE% /nologo /verbosity:quiet Programs\UnrealBuildTool\UnrealBuildTool.csproj /property:Configuration=Development /property:Platform=AnyCPU /target:Build
+)
 if errorlevel 1 goto Error_UBTCompileFailed
 
 rem ## Run UnrealBuildTool to generate Visual Studio solution and project files
 rem ## NOTE: We also pass along any arguments to the GenerateProjectFiles.bat here
-..\Binaries\DotNET\UnrealBuildTool.exe -ProjectFiles %*
+if "%UE_USE_DOTNET%" == "1" (
+	..\Binaries\DotNET\UnrealBuildTool\UnrealBuildTool.exe -ProjectFiles %*
+) else (
+	..\Binaries\DotNET\UnrealBuildTool.exe -ProjectFiles %*
+)
 if errorlevel 1 goto Error_ProjectGenerationFailed
 
 rem ## Success!
@@ -118,6 +140,13 @@ echo.
 pause
 goto Exit
 
+:Error_NoDotnetSDK
+echo.
+echo GenerateProjectFiles ERROR: Unable to find a install of Dotnet SDK.  Please make sure you have it installed and that `dotnet` is a globally available command.
+echo.
+pause
+goto Exit
+
 :Error_UBTCompileFailed
 echo.
 echo GenerateProjectFiles ERROR: UnrealBuildTool failed to compile.
diff --git a/Engine/Build/BatchFiles/GetDotnetPath.bat b/Engine/Build/BatchFiles/GetDotnetPath.bat
new file mode 100755
--- /dev/null
+++ b/Engine/Build/BatchFiles/GetDotnetPath.bat
@@ -0,0 +1,58 @@
+@echo off
+setlocal
+
+rem ## Unreal Engine 5 utility script
+rem ## Copyright Epic Games, Inc. All Rights Reserved.
+rem ##
+rem ## This script verifies that dotnet sdk is installed and a new enough SDK is present.
+
+for /f "delims=" %%i in ('where dotnet') do (
+	REM Dotnet exists
+	goto find_sdks
+)
+
+REM Dotnet command did not exist
+exit /B 1
+
+:find_sdks
+
+set REQUIRED_MAJOR_VERSION=3
+set REQUIRED_MINOR_VERSION=1
+
+set FOUND_MAJOR=
+set FOUND_MINOR=
+REM Unfortunately dotnet lists the sdks in oldest version first, thus we will pick the oldest version that matches our criteria as valid.
+REM This does not really matter as we are just trying to verify that a new enough SDK is actually present 
+for /f "tokens=1,* delims= " %%I in ('dotnet --list-sdks') do (
+
+	for /f "tokens=1,2,3 delims=." %%X in ("%%I") do (
+		REM We can check the patch version for preview versions and ignore those, but it slowed down this batch to much so accepting those for now. 
+		REM We do not actually use the determined version for anything so usually the newest SDK installed is used anyway
+		if %%X EQU %REQUIRED_MAJOR_VERSION% (
+			REM If the major version is the same as we require we check the minor version
+			if %%Y GEQ %REQUIRED_MINOR_VERSION% (
+
+				set FOUND_MAJOR=%%X
+				set FOUND_MINOR=%%Y
+				goto Succeeded
+			)
+		)
+
+		if %%X GTR %REQUIRED_MAJOR_VERSION% (
+			REM If the major version is greater then what we require then this sdk is good enough
+			set FOUND_MAJOR=%%X
+			set FOUND_MINOR=%%Y
+			goto Succeeded
+		)
+
+	)
+
+)
+
+REM Dotnet is installed but the sdk present is to old
+exit /B 1
+
+:Succeeded
+ECHO Found Dotnet SDK version: %FOUND_MAJOR%.%FOUND_MINOR%
+exit /B 0
+	
\ No newline at end of file
diff --git a/Engine/Source/Programs/DotNETCommon/BuildUtilities/BuildUtilitiesCore.csproj b/Engine/Source/Programs/DotNETCommon/BuildUtilities/BuildUtilitiesCore.csproj
new file mode 100644
--- /dev/null
+++ b/Engine/Source/Programs/DotNETCommon/BuildUtilities/BuildUtilitiesCore.csproj
@@ -0,0 +1,57 @@
+﻿<Project>
+  <PropertyGroup>
+    <!-- Set a custom temp directory so we can live side by side with the .net framework output -->
+    <BaseIntermediateOutputPath>$(MSBuildThisFileDirectory)obj-core\$(MSBuildProjectName)\</BaseIntermediateOutputPath>
+    <BaseOutputPath>$(MSBuildThisFileDirectory)bin-core\$(MSBuildProjectName)\</BaseOutputPath>
+  </PropertyGroup>
+  <Import Project="Sdk.props" Sdk="Microsoft.NET.Sdk" />
+  <PropertyGroup>
+    <!--<TargetFrameworks>net472;netstandard2.0</TargetFrameworks>-->
+    <TargetFramework>netstandard2.0</TargetFramework>
+    <Configuration Condition=" '$(Configuration)' == '' ">Development</Configuration>
+    <OutputType>Library</OutputType>
+    <RootNamespace>Tools.DotNETCommon</RootNamespace>
+    <GenerateAssemblyInfo>false</GenerateAssemblyInfo>
+    <GenerateTargetFrameworkAttribute>false</GenerateTargetFrameworkAttribute>
+    <FileUpgradeFlags>
+    </FileUpgradeFlags>
+    <UpgradeBackupLocation>
+    </UpgradeBackupLocation>
+    <OldToolsVersion>2.0</OldToolsVersion>
+    <Configurations>Debug;Release;Development</Configurations>
+    <!-- Set a custom temp directory so we can live side by side with the .net framework output -->
+    <BaseIntermediateOutputPath>$(MSBuildThisFileDirectory)obj-core\$(MSBuildProjectName)\</BaseIntermediateOutputPath>
+    <BaseOutputPath>$(MSBuildThisFileDirectory)bin-core\$(MSBuildProjectName)\</BaseOutputPath>
+  </PropertyGroup>
+  <Import Project="Sdk.targets" Sdk="Microsoft.NET.Sdk" />
+  <PropertyGroup Condition=" '$(Configuration)|$(Platform)' == 'Debug|AnyCPU' ">
+    <DocumentationFile>
+    </DocumentationFile>
+  </PropertyGroup>
+  <PropertyGroup Condition="'$(Configuration)|$(Platform)' == 'Development|AnyCPU'">
+    <DocumentationFile>
+    </DocumentationFile>
+    <Optimize>true</Optimize>
+  </PropertyGroup>
+  <PropertyGroup Condition="'$(Configuration)|$(Platform)' == 'Release|AnyCPU'">
+  </PropertyGroup>
+  
+  <ItemGroup>
+    <Compile Remove="obj\**" />
+    <EmbeddedResource Remove="obj\**" />
+    <None Remove="obj\**" />
+  </ItemGroup>
+
+  <ItemGroup>
+    <Compile Include="..\MetaData.cs">
+      <Link>Properties\MetaData.cs</Link>
+    </Compile>
+  </ItemGroup>
+  <ItemGroup>
+    <ProjectReference Include="..\DotNETUtilities\DotNETUtilitiesCore.csproj" />
+  </ItemGroup>
+  <ItemGroup>
+    <PackageReference Include="Microsoft.CSharp" Version="4.7.0" />
+    <PackageReference Include="System.Data.DataSetExtensions" Version="4.5.0" />
+  </ItemGroup>
+</Project>
\ No newline at end of file
diff --git a/Engine/Source/Programs/DotNETCommon/DotNETUtilities/DotNETUtilitiesCore.csproj b/Engine/Source/Programs/DotNETCommon/DotNETUtilities/DotNETUtilitiesCore.csproj
new file mode 100644
--- /dev/null
+++ b/Engine/Source/Programs/DotNETCommon/DotNETUtilities/DotNETUtilitiesCore.csproj
@@ -0,0 +1,59 @@
+﻿<Project>
+  <PropertyGroup>
+    <!-- Set a custom temp directory so we can live side by side with the .net framework output -->
+    <BaseIntermediateOutputPath>$(MSBuildThisFileDirectory)obj-core\$(MSBuildProjectName)\</BaseIntermediateOutputPath>
+    <BaseOutputPath>$(MSBuildThisFileDirectory)bin-core\$(MSBuildProjectName)\</BaseOutputPath>
+  </PropertyGroup>
+  <Import Project="Sdk.props" Sdk="Microsoft.NET.Sdk" />
+  <PropertyGroup>
+    <!--<TargetFrameworks>net472;netstandard2.0</TargetFrameworks>-->
+    <TargetFramework>netstandard2.0</TargetFramework>
+    <Configuration Condition=" '$(Configuration)' == '' ">Development</Configuration>
+    <OutputType>Library</OutputType>
+    <RootNamespace>Tools.DotNETCommon</RootNamespace>
+    <GenerateAssemblyInfo>false</GenerateAssemblyInfo>
+    <GenerateTargetFrameworkAttribute>false</GenerateTargetFrameworkAttribute>
+    <FileUpgradeFlags>
+    </FileUpgradeFlags>
+    <UpgradeBackupLocation>
+    </UpgradeBackupLocation>
+    <OldToolsVersion>2.0</OldToolsVersion>
+    <Configurations>Debug;Release;Development</Configurations>
+  </PropertyGroup>
+  <Import Project="Sdk.targets" Sdk="Microsoft.NET.Sdk" />
+  <!-- We do not output these binaries to Binaries as it targets multiple frameworks. They would have overloaded names for each framework, let the application copy it to its output directory instead -->
+  <PropertyGroup Condition=" '$(Configuration)|$(Platform)' == 'Debug|AnyCPU' ">
+    <DocumentationFile>
+    </DocumentationFile>
+    <OutputPath></OutputPath>
+  </PropertyGroup>
+  <PropertyGroup Condition="'$(Configuration)|$(Platform)' == 'Development|AnyCPU'">
+    <DocumentationFile>
+    </DocumentationFile>
+    <Optimize>true</Optimize>
+  </PropertyGroup>
+  <PropertyGroup Condition="'$(Configuration)|$(Platform)' == 'Release|AnyCPU'">
+    <DocumentationFile>
+    </DocumentationFile>
+    <Optimize>true</Optimize>
+  </PropertyGroup>
+  <ItemGroup>
+    <Compile Remove="obj\**" />
+    <EmbeddedResource Remove="obj\**" />
+    <None Remove="obj\**" />
+  </ItemGroup>
+  <ItemGroup>
+    <Compile Include="..\MetaData.cs">
+      <Link>Properties\MetaData.cs</Link>
+    </Compile>
+  </ItemGroup>
+  <ItemGroup>
+    <PackageReference Include="Microsoft.CSharp" Version="4.7.0" />
+    <PackageReference Include="System.Data.DataSetExtensions" Version="4.5.0" />
+    <PackageReference Include="System.Net.Http" Version="4.3.4" />
+    <!-- For FastJson-->
+    <PackageReference Include="System.Reflection.Emit.Lightweight">
+      <Version>4.3.0</Version>
+    </PackageReference>
+  </ItemGroup>
+</Project>
\ No newline at end of file
diff --git a/Engine/Source/Programs/UnrealBuildTool/DotNetCore/UnrealBuildTool.csproj b/Engine/Source/Programs/UnrealBuildTool/DotNetCore/UnrealBuildTool.csproj
deleted file mode 100644
diff --git a/Engine/Source/Programs/UnrealBuildTool/Platform/HoloLens/WinMDRegistrationInfo.cs b/Engine/Source/Programs/UnrealBuildTool/Platform/HoloLens/WinMDRegistrationInfo.cs
--- a/Engine/Source/Programs/UnrealBuildTool/Platform/HoloLens/WinMDRegistrationInfo.cs
+++ b/Engine/Source/Programs/UnrealBuildTool/Platform/HoloLens/WinMDRegistrationInfo.cs
@@ -101,6 +101,7 @@ namespace UnrealBuildTool
 			if (BuildHostPlatform.Current.Platform == UnrealTargetPlatform.Win64)
 			{
 				AppDomain.CurrentDomain.ReflectionOnlyAssemblyResolve += (Sender, EventArgs) => Assembly.ReflectionOnlyLoad(EventArgs.Name);
+#if !NET_CORE
 				WindowsRuntimeMetadata.ReflectionOnlyNamespaceResolve += (Sender, EventArgs) =>
 				{
 					string Path = WindowsRuntimeMetadata.ResolveNamespace(EventArgs.NamespaceName, ResolveSearchPaths).FirstOrDefault();
@@ -110,7 +111,7 @@ namespace UnrealBuildTool
 					}
 					EventArgs.ResolvedAssemblies.Add(Assembly.ReflectionOnlyLoadFrom(Path));
 				};
-
+#endif	
 			}
 		}
 
diff --git a/Engine/Source/Programs/UnrealBuildTool/Platform/Mac/ApplePlatformSDK.cs b/Engine/Source/Programs/UnrealBuildTool/Platform/Mac/ApplePlatformSDK.cs
--- a/Engine/Source/Programs/UnrealBuildTool/Platform/Mac/ApplePlatformSDK.cs
+++ b/Engine/Source/Programs/UnrealBuildTool/Platform/Mac/ApplePlatformSDK.cs
@@ -47,21 +47,24 @@ namespace UnrealBuildTool
 				return Result.Success ? Result.Groups[1].Value : "";
 			}
 
-			// otherwise, get iTunes "Version"
-			string DllPath = Registry.GetValue("HKEY_LOCAL_MACHINE\\SOFTWARE\\Wow6432Node\\Apple Inc.\\Apple Mobile Device Support\\Shared", "iTunesMobileDeviceDLL", null) as string;
-			if (string.IsNullOrEmpty(DllPath) || !File.Exists(DllPath))
+			if (IsWindows())
 			{
-				DllPath = Registry.GetValue("HKEY_LOCAL_MACHINE\\SOFTWARE\\Wow6432Node\\Apple Inc.\\Apple Mobile Device Support\\Shared", "MobileDeviceDLL", null) as string;
+				// otherwise, get iTunes "Version"
+				string DllPath = Registry.GetValue("HKEY_LOCAL_MACHINE\\SOFTWARE\\Wow6432Node\\Apple Inc.\\Apple Mobile Device Support\\Shared", "iTunesMobileDeviceDLL", null) as string;
 				if (string.IsNullOrEmpty(DllPath) || !File.Exists(DllPath))
 				{
-					DllPath = FindWindowsStoreITunesDLL();
+					DllPath = Registry.GetValue("HKEY_LOCAL_MACHINE\\SOFTWARE\\Wow6432Node\\Apple Inc.\\Apple Mobile Device Support\\Shared", "MobileDeviceDLL", null) as string;
+					if (string.IsNullOrEmpty(DllPath) || !File.Exists(DllPath))
+					{
+						DllPath = FindWindowsStoreITunesDLL();
 
+					}
 				}
-			}
 
-			if (!string.IsNullOrEmpty(DllPath) && File.Exists(DllPath))
-			{
-				return FileVersionInfo.GetVersionInfo(DllPath).FileVersion;
+				if (!string.IsNullOrEmpty(DllPath) && File.Exists(DllPath))
+				{
+					return FileVersionInfo.GetVersionInfo(DllPath).FileVersion;
+				}
 			}
 
 			return null;
diff --git a/Engine/Source/Programs/UnrealBuildTool/ProjectFiles/ProjectFileGenerator.cs b/Engine/Source/Programs/UnrealBuildTool/ProjectFiles/ProjectFileGenerator.cs
--- a/Engine/Source/Programs/UnrealBuildTool/ProjectFiles/ProjectFileGenerator.cs
+++ b/Engine/Source/Programs/UnrealBuildTool/ProjectFiles/ProjectFileGenerator.cs
@@ -1609,7 +1609,13 @@ namespace UnrealBuildTool
 				DirectoryReference ProjectDirectory = DirectoryReference.Combine(UnrealBuildTool.EngineSourceDirectory, "Programs", ProjectDirectoryName);
 				if (DirectoryReference.Exists(ProjectDirectory))
 				{
-					FileReference ProjectFileName = FileReference.Combine(ProjectDirectory, "UnrealBuildTool.csproj");
+					string ProjectName = "UnrealBuildTool.csproj";
+					if (AllowDotNetCoreProjects)
+					{
+						ProjectName = "UnrealBuildToolCore.csproj";
+					}
+
+					FileReference ProjectFileName = FileReference.Combine(ProjectDirectory, ProjectName);
 
 					if (FileReference.Exists(ProjectFileName))
 					{
@@ -2803,13 +2809,16 @@ namespace UnrealBuildTool
                 try
                 {
 
+	                bool configsFound = false;
                     // Parse the project and ensure both Development and Debug configurations are present
                     foreach (string Config in XElement.Load(InProject.ProjectFilePath.FullName).Elements("{http://schemas.microsoft.com/developer/msbuild/2003}PropertyGroup")
                                            .Where(node => node.Attribute("Condition") != null)
                                            .Select(node => node.Attribute("Condition").ToString())
                                            .ToList())
                     {
-                        if (Config.Contains("Development|"))
+	                    configsFound = true;
+
+						if (Config.Contains("Development|"))
                         {
                             bFoundDevelopmentConfig = true;
                         }
@@ -2818,6 +2827,28 @@ namespace UnrealBuildTool
                             bFoundDebugConfig = true;
                         }
                     }
+
+                    // for a net core project the PropertyGroup namespace changes, so lets see if we can find if its a project like that, net core also provides a hand semicolon separated list of configurations.
+                    // This is optional and defaults to Debug;Release if not present which still means Development is not present
+					// as such we we expect every proj to have configurations presents.
+					if (!configsFound)
+					{
+	                    foreach (string Config in XElement.Load(InProject.ProjectFilePath.FullName).Elements("PropertyGroup")
+		                    .Elements("Configurations")
+	                        .SelectMany(node => node.Value.Split(';'))
+	                        .ToList())
+	                    {
+		                    if (Config == "Development")
+		                    {
+								bFoundDevelopmentConfig = true;
+		                    }
+							else if (Config == "Debug")
+		                    {
+			                    bFoundDebugConfig = true;
+		                    }
+	                    }
+
+					}
                 }
                 catch
                 {
diff --git a/Engine/Source/Programs/UnrealBuildTool/ProjectFiles/VisualStudioCode/VSCodeProjectFileGenerator.cs b/Engine/Source/Programs/UnrealBuildTool/ProjectFiles/VisualStudioCode/VSCodeProjectFileGenerator.cs
--- a/Engine/Source/Programs/UnrealBuildTool/ProjectFiles/VisualStudioCode/VSCodeProjectFileGenerator.cs
+++ b/Engine/Source/Programs/UnrealBuildTool/ProjectFiles/VisualStudioCode/VSCodeProjectFileGenerator.cs
@@ -492,7 +492,7 @@ namespace UnrealBuildTool
 				{
 					VCSharpProjectFile VCSharpProject = Project as VCSharpProjectFile;
 
-					if (VCSharpProject.IsDotNETCoreProject() == bBuildingForDotNetCore)
+					if (bBuildingForDotNetCore || !VCSharpProject.IsDotNETCoreProject())
 					{
 						string ProjectName = Project.ProjectFilePath.GetFileNameWithoutExtension();
 
@@ -504,7 +504,7 @@ namespace UnrealBuildTool
 						{
 							CsProjectInfo Info = VCSharpProject.GetProjectInfo(Config);
 
-							if (!Info.IsDotNETCoreProject() && Info.Properties.ContainsKey("OutputPath"))
+							if (Info.Properties.ContainsKey("OutputPath"))
 							{
 								ProjectData.EOutputType OutputType;
 								string OutputTypeName;
@@ -681,35 +681,24 @@ namespace UnrealBuildTool
 
 							string CleanParam = Command == "Clean" ? "-clean" : null;
 
-							if (bBuildingForDotNetCore)
+							if (HostPlatform == UnrealTargetPlatform.Win64)
 							{
-								OutFile.AddField("command", "dotnet");
+								OutFile.AddField("command", MakePathString(FileReference.Combine(UE4ProjectRoot, "Engine", "Build", "BatchFiles", Command + ".bat")));
+								CleanParam = null;
 							}
 							else
 							{
-								if (HostPlatform == UnrealTargetPlatform.Win64)
-								{
-									OutFile.AddField("command", MakePathString(FileReference.Combine(UE4ProjectRoot, "Engine", "Build", "BatchFiles", Command + ".bat")));
-									CleanParam = null;
-								}
-								else
-								{
-									OutFile.AddField("command", MakePathString(FileReference.Combine(UE4ProjectRoot, "Engine", "Build", "BatchFiles", HostPlatform.ToString(), "Build.sh")));
+								OutFile.AddField("command", MakePathString(FileReference.Combine(UE4ProjectRoot, "Engine", "Build", "BatchFiles", HostPlatform.ToString(), "Build.sh")));
 
-									if (Command == "Clean")
-									{
-										CleanParam = "-clean";
-									}
+								if (Command == "Clean")
+								{
+									CleanParam = "-clean";
 								}
 							}
+							
 
 							OutFile.BeginArray("args");
 							{
-								if (bBuildingForDotNetCore)
-								{
-									OutFile.AddUnnamedField(MakeUnquotedPathString(FileReference.Combine(UE4ProjectRoot, "Engine", "Binaries", "DotNET", "UnrealBuildTool_NETCore.dll"), EPathType.Relative));
-								}
-
 								OutFile.AddUnnamedField(Target.Name);
 								OutFile.AddUnnamedField(BuildProduct.Platform.ToString());
 								OutFile.AddUnnamedField(BuildProduct.Config.ToString());
@@ -784,7 +773,14 @@ namespace UnrealBuildTool
 							OutFile.AddField("group", "build");
 							if (bIsDotNetCore)
 							{
-								OutFile.AddField("command", "dotnet");
+								if (Utils.IsRunningOnMono)
+								{
+									OutFile.AddField("command", MakePathString(FileReference.Combine(UE4ProjectRoot, "Engine", "Build", "BatchFiles", HostPlatform.ToString(), "RunDotnet.sh")));
+								}
+								else
+								{
+									OutFile.AddField("command", "dotnet");
+								}
 							}
 							else
 							{
@@ -802,24 +798,7 @@ namespace UnrealBuildTool
 								if (bIsDotNetCore)
 								{
 									OutFile.AddUnnamedField(Command.ToLower());
-								}
-								else
-								{
-									OutFile.AddUnnamedField("/t:" + Command.ToLower());
-								}
-								
-								DirectoryReference BuildRoot = HostPlatform == UnrealTargetPlatform.Win64 ? UE4ProjectRoot : DirectoryReference.Combine(UE4ProjectRoot, "Engine");
-								OutFile.AddUnnamedField(MakeUnquotedPathString(InProject.SourceProject.ProjectFilePath, EPathType.Relative, BuildRoot));
-
-								OutFile.AddUnnamedField("/p:GenerateFullPaths=true");
-								if (HostPlatform == UnrealTargetPlatform.Win64)
-								{
-									OutFile.AddUnnamedField("/p:DebugType=portable");
-								}
-								OutFile.AddUnnamedField("/verbosity:minimal");
 
-								if (bIsDotNetCore)
-								{
 									OutFile.AddUnnamedField("--configuration");
 									OutFile.AddUnnamedField(BuildProduct.Config.ToString());
 									OutFile.AddUnnamedField("--output");
@@ -827,6 +806,18 @@ namespace UnrealBuildTool
 								}
 								else
 								{
+									OutFile.AddUnnamedField("/t:" + Command.ToLower());
+								
+									DirectoryReference BuildRoot = HostPlatform == UnrealTargetPlatform.Win64 ? UE4ProjectRoot : DirectoryReference.Combine(UE4ProjectRoot, "Engine");
+									OutFile.AddUnnamedField(MakeUnquotedPathString(InProject.SourceProject.ProjectFilePath, EPathType.Relative, BuildRoot));
+
+									OutFile.AddUnnamedField("/p:GenerateFullPaths=true");
+									if (HostPlatform == UnrealTargetPlatform.Win64)
+									{
+										OutFile.AddUnnamedField("/p:DebugType=portable");
+									}
+									OutFile.AddUnnamedField("/verbosity:minimal");
+
 									OutFile.AddUnnamedField("/p:Configuration=" + BuildProduct.Config.ToString());
 								}
 							}
@@ -1046,11 +1037,9 @@ namespace UnrealBuildTool
 
 				if (bIsDotNetCore)
 				{
-					OutFile.AddField("program", "dotnet");
+					OutFile.AddField("program", MakeUnquotedPathString(InExecutable, EPathType.Absolute));
 					OutFile.BeginArray("args");
 					{
-						OutFile.AddUnnamedField(MakePathString(InExecutable));
-
 						if (InArgs != null)
 						{
 							foreach (string Arg in InArgs)
@@ -1139,6 +1128,10 @@ namespace UnrealBuildTool
 				List<string> Args = new List<string>();
 				Args.Add("-projectfiles");
 				Args.Add("-vscode");
+				if (bBuildingForDotNetCore)
+				{
+					Args.Add("-dotnetcore");
+				}
 
 				if (bForeignProject)
 				{
@@ -1151,11 +1144,14 @@ namespace UnrealBuildTool
 					PreLaunchTask = "UnrealBuildTool " + HostPlatform.ToString() + " Development Build";
 				}
 
+				FileReference UbtPath = bBuildingForDotNetCore
+					? FileReference.Combine(UE4ProjectRoot, "Engine", "Binaries", "DotNET", "UnrealBuildTool", "UnrealBuildTool")
+					: FileReference.Combine(UE4ProjectRoot, "Engine", "Binaries", "DotNET", "UnrealBuildTool.exe");
 				WriteSingleCSharpLaunchConfig(
 					OutFile,
 					"Generate Project Files",
 					PreLaunchTask,
-					FileReference.Combine(UE4ProjectRoot, "Engine", "Binaries", "DotNET", "UnrealBuildTool.exe"),
+					UbtPath,
 					Args.ToArray(),
 					bBuildingForDotNetCore
 				);
diff --git a/Engine/Source/Programs/UnrealBuildTool/System/AggregateParsedTimingInfo.cs b/Engine/Source/Programs/UnrealBuildTool/System/AggregateParsedTimingInfo.cs
deleted file mode 100644
diff --git a/Engine/Source/Programs/UnrealBuildTool/System/DynamicCompilation.cs b/Engine/Source/Programs/UnrealBuildTool/System/DynamicCompilation.cs
--- a/Engine/Source/Programs/UnrealBuildTool/System/DynamicCompilation.cs
+++ b/Engine/Source/Programs/UnrealBuildTool/System/DynamicCompilation.cs
@@ -160,7 +160,7 @@ namespace UnrealBuildTool
 			}
 		}
 
-		private static Assembly CompileAssembly(FileReference OutputAssemblyPath, List<FileReference> SourceFileNames, List<string> ReferencedAssembies, List<string> PreprocessorDefines = null, bool TreatWarningsAsErrors = false)
+		private static Assembly CompileAssembly(FileReference OutputAssemblyPath, HashSet<FileReference> SourceFileNames, List<string> ReferencedAssembies, List<string> PreprocessorDefines = null, bool TreatWarningsAsErrors = false)
 		{
 			CSharpParseOptions ParseOptions = new CSharpParseOptions(
 				languageVersion:LanguageVersion.Latest, 
@@ -176,7 +176,7 @@ namespace UnrealBuildTool
 				SyntaxTree Tree = CSharpSyntaxTree.ParseText(Source, ParseOptions, SourceFileName.FullName);
 
 				IEnumerable<Diagnostic> Diagnostics = Tree.GetDiagnostics();
-				if (Diagnostics.Count() > 0)
+				if (Diagnostics.Any())
 				{
 					Log.TraceWarning($"Errors generated while parsing '{SourceFileName.FullName}'");
 					LogDiagnostics(Tree.GetDiagnostics());
@@ -214,19 +214,38 @@ namespace UnrealBuildTool
 			MetadataReferences.Add(MetadataReference.CreateFromFile(Assembly.Load("System.Collections").Location));
 			MetadataReferences.Add(MetadataReference.CreateFromFile(Assembly.Load("System.IO").Location));
 			MetadataReferences.Add(MetadataReference.CreateFromFile(Assembly.Load("System.IO.FileSystem").Location));
+			MetadataReferences.Add(MetadataReference.CreateFromFile(Assembly.Load("System.Linq").Location));
 			MetadataReferences.Add(MetadataReference.CreateFromFile(Assembly.Load("System.Console").Location));
 			MetadataReferences.Add(MetadataReference.CreateFromFile(Assembly.Load("System.Runtime.Extensions").Location));
+			MetadataReferences.Add(MetadataReference.CreateFromFile(Assembly.Load("netstandard").Location));
+			
+			// process start dependencies
+			MetadataReferences.Add(MetadataReference.CreateFromFile(Assembly.Load("System.ComponentModel.Primitives").Location));
+			MetadataReferences.Add(MetadataReference.CreateFromFile(Assembly.Load("System.Diagnostics.Process").Location));
+			
+			// registry access
 			MetadataReferences.Add(MetadataReference.CreateFromFile(Assembly.Load("Microsoft.Win32.Registry").Location));
+
+			// RNGCryptoServiceProvider, used to generate random hex bytes
+			MetadataReferences.Add(MetadataReference.CreateFromFile(Assembly.Load("System.Security.Cryptography.Algorithms").Location));
+			MetadataReferences.Add(MetadataReference.CreateFromFile(Assembly.Load("System.Security.Cryptography.Csp").Location));
+
 			MetadataReferences.Add(MetadataReference.CreateFromFile(typeof(UnrealBuildTool).Assembly.Location));
 			MetadataReferences.Add(MetadataReference.CreateFromFile(typeof(FileReference).Assembly.Location));
+			MetadataReferences.Add(MetadataReference.CreateFromFile(typeof(UEBuildPlatformSDK).Assembly.Location));
 
 			CSharpCompilationOptions CompilationOptions = new CSharpCompilationOptions(
 				outputKind:OutputKind.DynamicallyLinkedLibrary,
-				optimizationLevel:OptimizationLevel.Release,
+#if DEBUG
+				optimizationLevel: OptimizationLevel.Debug,
+#else
+				// Optimize the managed code in Development
+				optimizationLevel: OptimizationLevel.Release,
+#endif
 				warningLevel:4,
 				assemblyIdentityComparer:DesktopAssemblyIdentityComparer.Default,
 				reportSuppressedDiagnostics:true
-				);
+			);
 
 			CSharpCompilation Compilation = CSharpCompilation.Create(
 				assemblyName:OutputAssemblyPath.GetFileNameWithoutAnyExtensions(),
diff --git a/Engine/Source/Programs/UnrealBuildTool/System/ParseMsvcTimingInfo.cs b/Engine/Source/Programs/UnrealBuildTool/System/ParseMsvcTimingInfo.cs
deleted file mode 100644
diff --git a/Engine/Source/Programs/UnrealBuildTool/System/Utils.cs b/Engine/Source/Programs/UnrealBuildTool/System/Utils.cs
--- a/Engine/Source/Programs/UnrealBuildTool/System/Utils.cs
+++ b/Engine/Source/Programs/UnrealBuildTool/System/Utils.cs
@@ -24,7 +24,7 @@ namespace UnrealBuildTool
 		/// Whether we are currently running on Mono platform.  We cache this statically because it is a bit slow to check.
 		/// </summary>
 #if NET_CORE
-		public static readonly bool IsRunningOnMono = true;
+		public static readonly bool IsRunningOnMono = !RuntimeInformation.IsOSPlatform(OSPlatform.Windows); // TODO: This check seems to mainly be used to determine if this is running on windows or not. Should be refactored.
 #else
 		public static readonly bool IsRunningOnMono = Type.GetType("Mono.Runtime") != null;
 #endif
@@ -94,8 +94,9 @@ namespace UnrealBuildTool
 		/// </summary>
 		/// <param name="InputString">String to search for variable names</param>
 		/// <param name="AdditionalVariables">Lookup of variable names to values</param>
+		/// <param name="bUseAdditionalVariablesOnly">If true, then Environment.GetEnvironmentVariable will not be used if the var is not found in AdditionalVariables</param>
 		/// <returns>String with all variables replaced</returns>
-		public static string ExpandVariables(string InputString, Dictionary<string, string> AdditionalVariables = null)
+		public static string ExpandVariables(string InputString, Dictionary<string, string> AdditionalVariables = null, bool bUseAdditionalVariablesOnly = false)
 		{
 			string Result = InputString;
 			for (int Idx = Result.IndexOf("$("); Idx != -1; Idx = Result.IndexOf("$(", Idx))
@@ -111,10 +112,13 @@ namespace UnrealBuildTool
 				string Name = Result.Substring(Idx + 2, EndIdx - (Idx + 2));
 
 				// Find the value for it, either from the dictionary or the environment block
-				string Value;
+				string Value = null;
 				if (AdditionalVariables == null || !AdditionalVariables.TryGetValue(Name, out Value))
 				{
-					Value = Environment.GetEnvironmentVariable(Name);
+					if (bUseAdditionalVariablesOnly == false)
+					{
+						Value = Environment.GetEnvironmentVariable(Name);
+					}
 					if (Value == null)
 					{
 						Idx = EndIdx + 1;
diff --git a/Engine/Source/Programs/UnrealBuildTool/UnrealBuildTool.cs b/Engine/Source/Programs/UnrealBuildTool/UnrealBuildTool.cs
--- a/Engine/Source/Programs/UnrealBuildTool/UnrealBuildTool.cs
+++ b/Engine/Source/Programs/UnrealBuildTool/UnrealBuildTool.cs
@@ -48,13 +48,20 @@ namespace UnrealBuildTool
 		/// <summary>
 		/// The path to UBT
 		/// </summary>
+#if NET_CORE
+		public static readonly FileReference UnrealBuildToolPath = FileReference.FindCorrectCase(new FileReference(Path.ChangeExtension(Assembly.GetExecutingAssembly().GetOriginalLocation(), "exe")));
+#else
 		public static readonly FileReference UnrealBuildToolPath = FileReference.FindCorrectCase(new FileReference(Assembly.GetExecutingAssembly().GetOriginalLocation()));
-
+#endif
 		/// <summary>
 		/// The full name of the Root UE4 directory
 		/// </summary>
+#if NET_CORE
+		// net core builds are in a subdirectory of DotNet binaries
+		public static readonly DirectoryReference RootDirectory = DirectoryReference.Combine(UnrealBuildToolPath.Directory, "..", "..", "..", "..");
+#else
 		public static readonly DirectoryReference RootDirectory = DirectoryReference.Combine(UnrealBuildToolPath.Directory, "..", "..", "..");
-
+#endif
 		/// <summary>
 		/// The full name of the Engine directory
 		/// </summary>
diff --git a/Engine/Source/Programs/UnrealBuildTool/UnrealBuildToolCore.csproj b/Engine/Source/Programs/UnrealBuildTool/UnrealBuildToolCore.csproj
new file mode 100644
--- /dev/null
+++ b/Engine/Source/Programs/UnrealBuildTool/UnrealBuildToolCore.csproj
@@ -0,0 +1,105 @@
+﻿<Project>
+  <PropertyGroup>
+    <!-- Set a custom temp directory so we can live side by side with the .net framework output -->
+    <BaseIntermediateOutputPath>$(MSBuildThisFileDirectory)obj-core\$(MSBuildProjectName)\</BaseIntermediateOutputPath>
+    <BaseOutputPath>$(MSBuildThisFileDirectory)bin-core\$(MSBuildProjectName)\</BaseOutputPath>
+  </PropertyGroup>
+  <Import Project="Sdk.props" Sdk="Microsoft.NET.Sdk" />
+  <PropertyGroup>
+    <TargetFramework>netcoreapp3.1</TargetFramework>
+    <Configuration Condition=" '$(Configuration)' == '' ">Development</Configuration>
+    <OutputType>Exe</OutputType>
+    <AppendTargetFrameworkToOutputPath>false</AppendTargetFrameworkToOutputPath>
+    <GenerateAssemblyInfo>false</GenerateAssemblyInfo>
+    <GenerateTargetFrameworkAttribute>false</GenerateTargetFrameworkAttribute>
+    <DefineConstants>NET_CORE</DefineConstants>
+    <Configurations>Debug;Release;Development</Configurations>
+    <ApplicationManifest>app.manifest</ApplicationManifest>
+    <AssemblyName>UnrealBuildTool</AssemblyName>
+    <RootNamespace>UnrealBuildTool</RootNamespace>
+  </PropertyGroup>
+  <PropertyGroup Condition=" '$(Configuration)|$(Platform)' == 'Development|AnyCPU' ">
+    <Optimize>true</Optimize>
+    <OutputPath>..\..\..\Binaries\DotNET\UnrealBuildTool</OutputPath>
+    <DocumentationFile>..\..\..\Binaries\DotNET\UnrealBuildTool\UnrealBuildTool.xml</DocumentationFile>
+  </PropertyGroup>
+  <PropertyGroup Condition=" '$(Configuration)|$(Platform)' == 'Release|AnyCPU' ">
+    <Optimize>true</Optimize>
+    <OutputPath>..\..\..\Binaries\DotNET\UnrealBuildTool</OutputPath>
+    <DocumentationFile>..\..\..\Binaries\DotNET\UnrealBuildTool\UnrealBuildTool.xml</DocumentationFile>
+  </PropertyGroup>
+  <PropertyGroup Condition=" '$(Configuration)|$(Platform)' == 'Debug|AnyCPU' ">
+    <OutputPath>..\..\..\Binaries\DotNET\UnrealBuildTool</OutputPath>
+    <DocumentationFile>..\..\..\Binaries\DotNET\UnrealBuildTool\UnrealBuildTool.XML</DocumentationFile>
+  </PropertyGroup>
+  <ItemGroup>
+    <Reference Include="Ionic.Zip.Reduced">
+      <HintPath>..\..\..\Binaries\DotNET\Ionic.Zip.Reduced.dll</HintPath>
+      <Private>False</Private>
+    </Reference>
+    <Reference Include="Microsoft.VisualStudio.Setup.Configuration.Interop">
+      <HintPath>..\..\..\Binaries\ThirdParty\VisualStudio\Microsoft.VisualStudio.Setup.Configuration.Interop.dll</HintPath>
+      <EmbedInteropTypes>True</EmbedInteropTypes>
+    </Reference>
+  </ItemGroup>
+  <Import Project="Sdk.targets" Sdk="Microsoft.NET.Sdk" />
+
+  <ItemGroup>
+    <Compile Remove="obj\**" />
+    <EmbeddedResource Remove="obj\**" />
+    <None Remove="obj\**" />
+  </ItemGroup>
+
+  <Choose>
+    <When Condition="Exists('..\..\..\Restricted\NotForLicensees\Source\Programs\UnrealBuildTool\Executors\FASTBuild.cs')">
+      <PropertyGroup>
+        <DefineConstants>$(DefineConstants);__FASTBUILD_AVAILABLE__</DefineConstants>
+      </PropertyGroup>
+    </When>
+  </Choose>
+  <ItemGroup>
+    <Compile Include="..\DotNETCommon\MetaData.cs">
+      <Link>Properties\MetaData.cs</Link>
+    </Compile>
+  </ItemGroup>
+
+  <!-- The following will glob all platform extensions' UnrealBuildTool source
+  files and include them in the build. They will also appear in VisualStudio.
+  RecursiveDir does not work in VS Mac (and maybe Linux?), so only use it in Windows. 
+  See: https://developercommunity.visualstudio.com/content/problem/527478/recursivedir-causes-infinitely-nested-subdirectori.html -->
+  <ItemGroup>
+    <Compile Include="../../../Platforms/*/Source/Programs/UnrealBuildTool/**/*.cs">
+      <!-- Since RecursiveDir is bugged on VS Mac, we use the FullPath, and use Regex to pull out portions of the path to get Platform name -->
+      <Link>Platform/$([System.Text.RegularExpressions.Regex]::Replace(%(FullPath), '^(.*[\\/]Platforms[\\/])(.+?)[\\/](.*)$', '$2'))/%(Filename)%(Extension)</Link>
+    </Compile>
+    <Compile Include="../../../Restricted/*/Platforms/*/Source/Programs/UnrealBuildTool/**/*.cs">
+      <!-- Restricted may or may not be in a Platforms dir, so just use the subdir directly, even tho it makes for a
+       messier directory structure for these few files -->
+      <Link>Platform/$([System.Text.RegularExpressions.Regex]::Replace(%(FullPath), '^(.*[\\/]Platforms[\\/])(.+?)[\\/](.*)$', '$2'))/%(Filename)%(Extension)</Link>
+    </Compile>
+    <Compile Include="../../../Restricted/*/Source/Programs/UnrealBuildTool/**/*.cs">
+      <!-- Restricted may or may not be in a Platforms dir, so just use the subdir directly, even tho it makes for a
+       messier directory structure for these few files -->
+      <Link>$([System.Text.RegularExpressions.Regex]::Replace(%(FullPath), '^.+?[\\/]UnrealBuildTool[\\/]', ''))</Link>
+    </Compile>
+  </ItemGroup>
+  <!-- To modify your build process, add your task inside one of the targets below and uncomment it. 
+       Other similar extension points exist, see Microsoft.Common.targets.
+  <Target Name="BeforeBuild">
+  </Target>
+  <Target Name="AfterBuild">
+  </Target>
+  -->
+
+  <ItemGroup>
+    <ProjectReference Include="..\DotNETCommon\BuildUtilities\BuildUtilitiesCore.csproj" />
+    <ProjectReference Include="..\DotNETCommon\DotNETUtilities\DotNETUtilitiesCore.csproj" />
+  </ItemGroup>
+  <ItemGroup>
+    <PackageReference Include="Microsoft.CodeAnalysis.CSharp" Version="3.7.0" />
+    <PackageReference Include="Microsoft.Win32.Registry" Version="4.7.0" />
+    <PackageReference Include="System.Management" Version="4.7.0" />
+    <PackageReference Include="System.Security.Cryptography.Csp" Version="4.3.0" />
+    <PackageReference Include="System.ServiceProcess.ServiceController" Version="4.7.0" />
+  </ItemGroup>
+</Project>
\ No newline at end of file
