Add an option to force warnings as errors, even for deprecations. Need to enable in CIS to force determinism.

diff --git a/Engine/Source/Programs/UnrealBuildTool/Configuration/TargetRules.cs b/Engine/Source/Programs/UnrealBuildTool/Configuration/TargetRules.cs
--- a/Engine/Source/Programs/UnrealBuildTool/Configuration/TargetRules.cs
+++ b/Engine/Source/Programs/UnrealBuildTool/Configuration/TargetRules.cs
@@ -941,6 +941,13 @@ namespace UnrealBuildTool
 			set { ShadowVariableWarningLevel = (value? WarningLevel.Error : WarningLevel.Warning); }
 		}
 
+		/// <summary>
+		/// Whether to enable all warnings as errors. UE enables most warnings as errors already, but disables a few (such as deprecation warnings).
+		/// </summary>
+		[XmlConfigFile(Category = "BuildConfiguration")]
+		[CommandLine("-WarningsAsErrors")]
+		public bool bWarningsAsErrors = false;
+
 		/// <summary>
 		/// Indicates what warning/error level to treat unsafe type casts as on platforms that support it (e.g., double->float or int64->int32)
 		/// </summary>
@@ -2374,6 +2381,11 @@ namespace UnrealBuildTool
 			get { return Inner.bUndefinedIdentifierErrors; }
 		}
 
+		public bool bWarningsAsErrors
+		{
+			get { return Inner.bWarningsAsErrors; }
+		}
+
 		public bool bUseFastMonoCalls
 		{
 			get { return Inner.bUseFastMonoCalls; }
diff --git a/Engine/Source/Programs/UnrealBuildTool/Configuration/UEBuildTarget.cs b/Engine/Source/Programs/UnrealBuildTool/Configuration/UEBuildTarget.cs
--- a/Engine/Source/Programs/UnrealBuildTool/Configuration/UEBuildTarget.cs
+++ b/Engine/Source/Programs/UnrealBuildTool/Configuration/UEBuildTarget.cs
@@ -3409,6 +3409,7 @@ namespace UnrealBuildTool
 			GlobalCompileEnvironment.ShadowVariableWarningLevel = Rules.ShadowVariableWarningLevel;
 			GlobalCompileEnvironment.UnsafeTypeCastWarningLevel = Rules.UnsafeTypeCastWarningLevel;
 			GlobalCompileEnvironment.bUndefinedIdentifierWarningsAsErrors = Rules.bUndefinedIdentifierErrors;
+			GlobalCompileEnvironment.bWarningsAsErrors = Rules.bWarningsAsErrors;
 			GlobalCompileEnvironment.bOptimizeForSize = Rules.bCompileForSize;
 			GlobalCompileEnvironment.bUseStaticCRT = Rules.bUseStaticCRT;
 			GlobalCompileEnvironment.bOmitFramePointers = Rules.bOmitFramePointers;
diff --git a/Engine/Source/Programs/UnrealBuildTool/Platform/Windows/VCToolChain.cs b/Engine/Source/Programs/UnrealBuildTool/Platform/Windows/VCToolChain.cs
--- a/Engine/Source/Programs/UnrealBuildTool/Platform/Windows/VCToolChain.cs
+++ b/Engine/Source/Programs/UnrealBuildTool/Platform/Windows/VCToolChain.cs
@@ -596,6 +596,12 @@ namespace UnrealBuildTool
 				Arguments.Add("/W0");
 			}
 
+			// Treat warnings as errors
+			if (CompileEnvironment.bWarningsAsErrors)
+			{
+				Arguments.Add("/WX");
+			}
+
 			// Intel compiler options.
 			if (Target.WindowsPlatform.Compiler == WindowsCompiler.Intel)
 			{
diff --git a/Engine/Source/Programs/UnrealBuildTool/System/CppCompileEnvironment.cs b/Engine/Source/Programs/UnrealBuildTool/System/CppCompileEnvironment.cs
--- a/Engine/Source/Programs/UnrealBuildTool/System/CppCompileEnvironment.cs
+++ b/Engine/Source/Programs/UnrealBuildTool/System/CppCompileEnvironment.cs
@@ -207,6 +207,11 @@ namespace UnrealBuildTool
 		/// </summary>
 		public bool bUndefinedIdentifierWarningsAsErrors = false;
 
+		/// <summary>
+		/// Whether to treat all warnings as errors
+		/// </summary>
+		public bool bWarningsAsErrors = false;
+
 		/// <summary>
 		/// True if compiler optimizations should be enabled. This setting is distinct from the configuration (see CPPTargetConfiguration).
 		/// </summary>
@@ -421,6 +426,7 @@ namespace UnrealBuildTool
 			UnsafeTypeCastWarningLevel = Other.UnsafeTypeCastWarningLevel;
 			bUndefinedIdentifierWarningsAsErrors = Other.bUndefinedIdentifierWarningsAsErrors;
 			bEnableUndefinedIdentifierWarnings = Other.bEnableUndefinedIdentifierWarnings;
+			bWarningsAsErrors = Other.bWarningsAsErrors;
 			bOptimizeCode = Other.bOptimizeCode;
 			bOptimizeForSize = Other.bOptimizeForSize;
 			bCreateDebugInfo = Other.bCreateDebugInfo;
