Do not fall back to choosing VS2017 toolchain if a specific 2019 version is not found.

diff --git a/Engine/Source/Programs/UnrealBuildTool/Platform/Windows/UEBuildWindows.cs b/Engine/Source/Programs/UnrealBuildTool/Platform/Windows/UEBuildWindows.cs
--- a/Engine/Source/Programs/UnrealBuildTool/Platform/Windows/UEBuildWindows.cs
+++ b/Engine/Source/Programs/UnrealBuildTool/Platform/Windows/UEBuildWindows.cs
@@ -700,8 +700,8 @@ namespace UnrealBuildTool
 		{
 			VersionNumber.Parse("14.24.28315"), // VS2019 v16.4.3 (installed to 14.24.28314 folder)
 			VersionNumber.Parse("14.22.27905"), // VS2019 v16.2.3
-			VersionNumber.Parse("14.16.27023.2"), // VS2017 v15.9.15
-			VersionNumber.Parse("14.16.27023"), // fallback to VS2017 15.9 toolchain, microsoft updates these in places so for local installs only this version number is present
+//			VersionNumber.Parse("14.16.27023.2"), // VS2017 v15.9.15
+//			VersionNumber.Parse("14.16.27023"), // fallback to VS2017 15.9 toolchain, microsoft updates these in places so for local installs only this version number is present
 		};
 
 		/// <summary>
