- Fixed UBT to not throw errors if a project has a Platforms/XYZ directory, which will try to build with platform XYZ, if the engine doesn't have Engine/Platforms/XYZ/DataDrivenPlatformInfo.ini file. This will make it so the Project/Platforms/XYZ directory would be completely skipped over (and Project/Restricted/Blah/Platforms/XYZ)
- Changed plugin module descriptor internals to allow for an empty platform whitelist (which means no platforms are enabled), in the case of all platforms in the whitelist are not around
- Added some dummy DDPI.ini files for "fake platforms" to the Engine so that the above directory checks will not strip out the shared platform directories from projects

diff --git a/Engine/Source/Programs/UnrealBuildTool/System/DataDrivenPlatformInfo.cs b/Engine/Source/Programs/UnrealBuildTool/System/DataDrivenPlatformInfo.cs
--- a/Engine/Source/Programs/UnrealBuildTool/System/DataDrivenPlatformInfo.cs
+++ b/Engine/Source/Programs/UnrealBuildTool/System/DataDrivenPlatformInfo.cs
@@ -160,7 +160,7 @@ namespace UnrealBuildTool
 		{
 
 			// lookup the platform name (which is not guaranteed to be there)
-			ConfigDataDrivenPlatformInfo Info = null;
+			ConfigDataDrivenPlatformInfo Info;
 			GetAllPlatformInfos().TryGetValue(PlatformName, out Info);
 
 			// return what we found of null if nothing
diff --git a/Engine/Source/Programs/UnrealBuildTool/System/ModuleDescriptor.cs b/Engine/Source/Programs/UnrealBuildTool/System/ModuleDescriptor.cs
--- a/Engine/Source/Programs/UnrealBuildTool/System/ModuleDescriptor.cs
+++ b/Engine/Source/Programs/UnrealBuildTool/System/ModuleDescriptor.cs
@@ -241,15 +241,42 @@ namespace UnrealBuildTool
 			try
 			{
 				string[] WhitelistPlatforms;
+				// it's important we default to null, and don't have an empty whitelist by default, because that will indicate that no
+				// platforms should be compiled (see IsCompiledInConfiguration(), it only checks for null, not length)
+				Module.WhitelistPlatforms = null;
 				if (InObject.TryGetStringArrayField("WhitelistPlatforms", out WhitelistPlatforms))
 				{
-					Module.WhitelistPlatforms = Array.ConvertAll(WhitelistPlatforms, x => UnrealTargetPlatform.Parse(x)).ToList();
+					Module.WhitelistPlatforms = new List<UnrealTargetPlatform>();
+					foreach (string TargetPlatformName in WhitelistPlatforms)
+					{
+						UnrealTargetPlatform Platform;
+						if (UnrealTargetPlatform.TryParse(TargetPlatformName, out Platform))
+						{
+							Module.WhitelistPlatforms.Add(Platform);
+						}
+						else
+						{
+							Log.TraceWarning("Unknown platform {0} while parsing whitelist for module descriptor {1}", TargetPlatformName, Module.Name);
+						}
+					}
 				}
 
 				string[] BlacklistPlatforms;
 				if (InObject.TryGetStringArrayField("BlacklistPlatforms", out BlacklistPlatforms))
 				{
-					Module.BlacklistPlatforms = Array.ConvertAll(BlacklistPlatforms, x => UnrealTargetPlatform.Parse(x)).ToList();
+					Module.BlacklistPlatforms = new List<UnrealTargetPlatform>();
+					foreach (string TargetPlatformName in BlacklistPlatforms)
+					{
+						UnrealTargetPlatform Platform;
+						if (UnrealTargetPlatform.TryParse(TargetPlatformName, out Platform))
+						{
+							Module.BlacklistPlatforms.Add(Platform);
+						}
+						else
+						{
+							Log.TraceWarning("Unknown platform {0} while parsing blacklist for module descriptor {1}", TargetPlatformName, Module.Name);
+						}
+					}
 				}
 			}
 			catch (BuildException Ex)
@@ -313,7 +340,10 @@ namespace UnrealBuildTool
 			Writer.WriteValue("Name", Name);
 			Writer.WriteValue("Type", Type.ToString());
 			Writer.WriteValue("LoadingPhase", LoadingPhase.ToString());
-			if (WhitelistPlatforms != null && WhitelistPlatforms.Count > 0)
+			// important note: we don't check the length of the whitelist platforms, because if an unknown platform was read in, but was not valid, the 
+			// list will exist but be empty. We don't want to remove the whitelist completely, because that would allow this module on all platforms,
+			// which will not be the desired effect
+			if (WhitelistPlatforms != null)
 			{
 				Writer.WriteArrayStart("WhitelistPlatforms");
 				foreach (UnrealTargetPlatform WhitelistPlatform in WhitelistPlatforms)
@@ -431,7 +461,10 @@ namespace UnrealBuildTool
 		public bool IsCompiledInConfiguration(UnrealTargetPlatform Platform, UnrealTargetConfiguration Configuration, string TargetName, TargetType TargetType, bool bBuildDeveloperTools, bool bBuildRequiresCookedData)
 		{
 			// Check the platform is whitelisted
-			if (WhitelistPlatforms != null && WhitelistPlatforms.Count > 0 && !WhitelistPlatforms.Contains(Platform))
+			// important note: we don't check the length of the whitelist platforms, because if an unknown platform was read in, but was not valid, the 
+			// list will exist but be empty. In this case, we need to disallow all platforms from building, otherwise, build errors will occur when
+			// it starts compiling for _all_ platforms
+			if (WhitelistPlatforms != null && !WhitelistPlatforms.Contains(Platform))
 			{
 				return false;
 			}
diff --git a/Engine/Source/Programs/UnrealBuildTool/System/PluginDescriptor.cs b/Engine/Source/Programs/UnrealBuildTool/System/PluginDescriptor.cs
--- a/Engine/Source/Programs/UnrealBuildTool/System/PluginDescriptor.cs
+++ b/Engine/Source/Programs/UnrealBuildTool/System/PluginDescriptor.cs
@@ -247,20 +247,23 @@ namespace UnrealBuildTool
 			RawObject.TryGetStringArrayField("SupportedPrograms", out SupportedPrograms);
 			RawObject.TryGetBoolField("bIsPluginExtension", out bIsPluginExtension);
 
-			try
+			string[] SupportedTargetPlatformNames;
+			if (RawObject.TryGetStringArrayField("SupportedTargetPlatforms", out SupportedTargetPlatformNames))
 			{
-				string[] SupportedTargetPlatformNames;
-				if (RawObject.TryGetStringArrayField("SupportedTargetPlatforms", out SupportedTargetPlatformNames))
+				SupportedTargetPlatforms = new List<UnrealTargetPlatform>();
+				foreach (string TargetPlatformName in SupportedTargetPlatformNames)
 				{
-					SupportedTargetPlatforms = Array.ConvertAll(SupportedTargetPlatformNames, x => UnrealTargetPlatform.Parse(x)).ToList();
+					UnrealTargetPlatform Platform;
+					if (UnrealTargetPlatform.TryParse(TargetPlatformName, out Platform))
+					{
+						SupportedTargetPlatforms.Add(Platform);
+					}
+					else
+					{
+						Log.TraceWarning("Unknown platform {0} listed in plugin with FriendlyName {1}", TargetPlatformName, FriendlyName);
+					}
 				}
 			}
-			catch (BuildException Ex)
-			{
-				ExceptionUtils.AddContext(Ex, "while parsing SupportedTargetPlatforms in plugin with FriendlyName '{0}'", FriendlyName);
-				throw;
-			}
-
 
 			JsonObject[] ModulesArray;
 			if (RawObject.TryGetObjectArrayField("Modules", out ModulesArray))
diff --git a/Engine/Source/Programs/UnrealBuildTool/UnrealBuildTool.cs b/Engine/Source/Programs/UnrealBuildTool/UnrealBuildTool.cs
--- a/Engine/Source/Programs/UnrealBuildTool/UnrealBuildTool.cs
+++ b/Engine/Source/Programs/UnrealBuildTool/UnrealBuildTool.cs
@@ -166,6 +166,13 @@ namespace UnrealBuildTool
 						}
 					}
 				}
+
+				// remove any platform directories in non-engine locations if the engine doesn't have the platform 
+				if (BaseDir != UnrealBuildTool.EngineDirectory && CachedDirs.Item1.Count > 0)
+				{
+					// if the DDPI.ini file doesn't exist, we haven't synced the platform, so just skip this directory
+					CachedDirs.Item1.RemoveAll(x => DataDrivenPlatformInfo.GetDataDrivenInfoForPlatform(x.GetDirectoryName()) == null);
+				}
 			}
 
 			// now return what the caller wanted (always include BaseDir)
