UBT: Support for module partitions. Also rework dependency mapping to handle modules in files in files with different names.

diff --git a/Engine/Source/Programs/UnrealBuildTool/Modes/BuildMode.cs b/Engine/Source/Programs/UnrealBuildTool/Modes/BuildMode.cs
--- a/Engine/Source/Programs/UnrealBuildTool/Modes/BuildMode.cs
+++ b/Engine/Source/Programs/UnrealBuildTool/Modes/BuildMode.cs
@@ -402,6 +402,8 @@ namespace UnrealBuildTool
 								CompiledModuleInterfaces.Add(PrerequisiteAction.CompiledModuleInterfaceFile);
 							}
 						}
+
+						Dictionary<FileItem, string> ModuleInputs = ModuleOutputs.ToDictionary(x => x.Value, x => x.Key);
 						foreach (Action PrerequisiteAction in PrerequisiteActions)
 						{
 							if (PrerequisiteAction.ActionType == ActionType.CompileModuleInterface)
@@ -415,6 +417,7 @@ namespace UnrealBuildTool
 										if (ModuleOutputs.TryGetValue(ImportedModule, out ModuleOutput))
 										{
 											PrerequisiteAction.PrerequisiteItems.Add(ModuleOutput);
+											PrerequisiteAction.CommandArguments += String.Format(" /reference \"{0}={1}\"", ImportedModule, ModuleOutput.FullName);
 										}
 										else
 										{
@@ -423,6 +426,17 @@ namespace UnrealBuildTool
 									}
 								}
 							}
+							else if(PrerequisiteAction.ActionType == ActionType.Compile)
+							{
+								foreach(FileItem PrerequisiteItem in PrerequisiteAction.PrerequisiteItems)
+								{
+									string ModuleName;
+									if(ModuleInputs.TryGetValue(PrerequisiteItem, out ModuleName))
+									{
+										PrerequisiteAction.CommandArguments += String.Format(" /reference \"{0}={1}\"", ModuleName, PrerequisiteItem.AbsolutePath);
+									}
+								}
+							}
 							else
 							{
 								if (PrerequisiteAction.ActionType != ActionType.GatherModuleDependencies)
diff --git a/Engine/Source/Programs/UnrealBuildTool/Platform/Windows/VCToolChain.cs b/Engine/Source/Programs/UnrealBuildTool/Platform/Windows/VCToolChain.cs
--- a/Engine/Source/Programs/UnrealBuildTool/Platform/Windows/VCToolChain.cs
+++ b/Engine/Source/Programs/UnrealBuildTool/Platform/Windows/VCToolChain.cs
@@ -997,15 +997,6 @@ namespace UnrealBuildTool
 				AddDefinition(SharedArguments, DefinitionArgument);
 			}
 
-			// Module interface paths.
-			if (Target.bEnableCppModules)
-			{
-				foreach (DirectoryReference ModuleInterfacePath in CompileEnvironment.ModuleInterfacePaths)
-				{
-					SharedArguments.Add(String.Format("/ifcSearchDir \"{0}\"", ModuleInterfacePath));
-				}
-			}
-
 			// Create a compile action for each source file.
 			CPPOutput Result = new CPPOutput();
 			foreach (FileItem SourceFile in InputFiles)
