Fixing the path to the dummy workspace in the IOS tool chain.

#rnx

diff --git a/Engine/Source/Programs/UnrealBuildTool/Platform/IOS/IOSToolChain.cs b/Engine/Source/Programs/UnrealBuildTool/Platform/IOS/IOSToolChain.cs
--- a/Engine/Source/Programs/UnrealBuildTool/Platform/IOS/IOSToolChain.cs
+++ b/Engine/Source/Programs/UnrealBuildTool/Platform/IOS/IOSToolChain.cs
@@ -1790,7 +1790,7 @@ namespace UnrealBuildTool
                 {
 					if (AppName == "UnrealGame" || AppName == "UnrealClient" || Target.ProjectFile == null || Target.ProjectFile.IsUnderDirectory(UnrealBuildTool.EngineDirectory))
 					{
-						XcodeWorkspaceDir = DirectoryReference.Combine(UnrealBuildTool.RootDirectory, String.Format("UE4_{0}.xcworkspace", (Target.Platform == UnrealTargetPlatform.IOS ? "IOS" : "TVOS")));
+						XcodeWorkspaceDir = DirectoryReference.Combine(UnrealBuildTool.RootDirectory, String.Format("UE5_{0}.xcworkspace", (Target.Platform == UnrealTargetPlatform.IOS ? "IOS" : "TVOS")));
 					}
 					else
 					{
@@ -1909,7 +1909,7 @@ namespace UnrealBuildTool
 						}
 						else
 						{
-							SchemeName = "UE4";
+							SchemeName = "UE5";
 						}
 					}
 					else
