UBT changes required for netcore projects.

* Added support for more complex Msbuild conditions (using static property methods) as we use this to do per platform checks in the csprojs.
* Tweaks to the parsing of csprojs as expecations are different for netcore (mostly for how configurations are defined)
* Lastly if -dotnetcore flag is present when generating projects, use the netcore project files instead.

diff --git a/Engine/Source/Programs/DotNETCommon/DotNETUtilities/CsProjectInfo.cs b/Engine/Source/Programs/DotNETCommon/DotNETUtilities/CsProjectInfo.cs
--- a/Engine/Source/Programs/DotNETCommon/DotNETUtilities/CsProjectInfo.cs
+++ b/Engine/Source/Programs/DotNETCommon/DotNETUtilities/CsProjectInfo.cs
@@ -756,7 +756,7 @@ namespace Tools.DotNETCommon
 				string[] Tokens = Tokenize(NewText.Substring(Idx + 2, (EndIdx - 1) - (Idx + 2)));
 
 				// Make sure the first token is a valid property name
-				if(Tokens.Length == 0 || !(Char.IsLetter(Tokens[0][0]) || Tokens[0][0] == '_'))
+				if(Tokens.Length == 0 || !(Char.IsLetter(Tokens[0][0]) || Tokens[0][0] == '_' || Tokens[0][0] == '[' ))
 				{
 					throw new Exception(String.Format("Invalid property name '{0}' in .csproj file", Tokens[0]));
 				}
@@ -900,6 +900,76 @@ namespace Tools.DotNETCommon
 					Idx++;
 					Tokens.Add(Condition.Substring(StartIdx, Idx - StartIdx));
 				}
+				else if (Condition[Idx] == '[')
+				{
+					// static property function invoke
+					// format: [Class]::Property
+					// alternatively: [Class]::Method()
+					// we consider the entire invocation to be a single token
+					int StartIdx = Idx++;
+					int ClassEndIdx = 0;
+					int MethodEndIdx = 0;
+					int MethodArgsEndIdx = 0;
+					for (; ; Idx++)
+					{
+						while (Idx < Condition.Length && (Char.IsLetterOrDigit(Condition[Idx]) || Condition[Idx] == '_'))
+						{
+							Idx++;
+						}
+
+						if (Idx == Condition.Length)
+						{
+							throw new Exception(String.Format("Found end of condition when searching for end of static property function for condition string ('{0}')", Condition));
+						}
+						if (Condition[Idx] == ']')
+						{
+							ClassEndIdx = Idx;
+							Idx++;
+							break;
+						}
+					}
+
+					// skip ::
+					if (Condition[Idx] != ':')
+					{
+						throw new Exception(String.Format("Unexpected format of static property function, expected :: after class declaration in condition string ('{0}')", Condition));
+					}
+					Idx += 2;
+
+					while (Idx < Condition.Length && (Char.IsLetterOrDigit(Condition[Idx]) || Condition[Idx] == '_'))
+					{
+						Idx++;
+					}
+
+					MethodEndIdx = Idx;
+
+					if (Condition[Idx] == '(')
+					{
+						// a method invoke
+						for (; ; Idx++)
+						{
+							while (Idx < Condition.Length && (Char.IsLetterOrDigit(Condition[Idx]) || Condition[Idx] == '_'))
+							{
+								Idx++;
+							}
+
+							if (Idx == Condition.Length)
+							{
+								throw new Exception(String.Format("Found end of condition when searching for ) to indicate end of arguments to static property function for condition string ('{0}')", Condition));
+							}
+							if (Condition[Idx] == ')')
+							{
+								MethodArgsEndIdx = Idx;
+								Idx++;
+								break;
+							}
+						}
+
+					}
+
+					Idx++;
+					Tokens.Add(Condition.Substring(StartIdx, Idx - StartIdx));
+				}
 				else if(Char.IsLetterOrDigit(Condition[Idx]) || Condition[Idx] == '_')
 				{
 					// Identifier or number
diff --git a/Engine/Source/Programs/UnrealBuildTool/ProjectFiles/ProjectFileGenerator.cs b/Engine/Source/Programs/UnrealBuildTool/ProjectFiles/ProjectFileGenerator.cs
--- a/Engine/Source/Programs/UnrealBuildTool/ProjectFiles/ProjectFileGenerator.cs
+++ b/Engine/Source/Programs/UnrealBuildTool/ProjectFiles/ProjectFileGenerator.cs
@@ -1624,11 +1624,8 @@ namespace UnrealBuildTool
 
 						if (bIncludeDotNETCoreProjects || !UnrealBuildToolProject.IsDotNETCoreProject())
 						{
-							if (!UnrealBuildToolProject.IsDotNETCoreProject())
-							{
-								// Store it off as we need it when generating target projects.
-								UBTProject = UnrealBuildToolProject;
-							}
+							// Store it off as we need it when generating target projects.
+							UBTProject = UnrealBuildToolProject;
 
 							// Add the project
 							AddExistingProjectFile(UnrealBuildToolProject, bNeedsAllPlatformAndConfigurations: true, bForceDevelopmentConfiguration: true);
@@ -1654,6 +1651,12 @@ namespace UnrealBuildTool
 			VCSharpProjectFile Project = null;
 
 			FileReference ProjectFileName = FileReference.Combine( UnrealBuildTool.EngineSourceDirectory, "Programs", ProjectName, Path.GetFileName( ProjectName ) + ".csproj" );
+			FileReference ProjectFileNameDotnet = FileReference.Combine( UnrealBuildTool.EngineSourceDirectory, "Programs", ProjectName, Path.GetFileName( ProjectName ) + "Core.csproj" );
+	
+			// if we are using the dotnet flag and there is a dotnet version of the project available use that instead
+			if (File.Exists(ProjectFileNameDotnet.FullName) && AllowDotNetCoreProjects && bIncludeDotNETCoreProjects)
+				ProjectFileName = ProjectFileNameDotnet;
+
 			FileInfo Info = new FileInfo( ProjectFileName.FullName );
 			if( Info.Exists )
 			{
@@ -2634,6 +2637,26 @@ namespace UnrealBuildTool
 				WriteMasterProjectFile( UBTProject, PlatformProjectGenerators );
 				Progress.Write(TotalProjectFileCount, TotalProjectFileCount);
 			}
+
+            // Update the dotnet reference files so that the tooling knows which projects to build
+            // TODO: Should UAT build like this for net core or should it use custom logic to determine which projects to build?
+            if (AutomationProjectFiles.Any())
+            {
+	            XNamespace NS = XNamespace.Get("http://schemas.microsoft.com/developer/msbuild/2003");
+
+	            DirectoryReference AutomationToolDir = DirectoryReference.Combine(UnrealBuildTool.EngineSourceDirectory, "Programs", "AutomationTool");
+	            new XDocument(
+		            new XElement(NS + "Project",
+			            new XAttribute("DefaultTargets", "Build"),
+			            new XElement(NS + "ItemGroup",
+				            from AutomationProject in AutomationProjectFiles
+				            select new XElement(NS + "ProjectReference",
+					            new XAttribute("Include", AutomationProject.ProjectFilePath.MakeRelativeTo(AutomationToolDir))
+				            )
+			            )
+		            )
+	            ).Save(FileReference.Combine(AutomationToolDir, "AutomationToolCore.csproj.References").FullName);
+            }
 			return true;
 		}
 
diff --git a/Engine/Source/Programs/UnrealBuildTool/ProjectFiles/VisualStudio/VCProject.cs b/Engine/Source/Programs/UnrealBuildTool/ProjectFiles/VisualStudio/VCProject.cs
--- a/Engine/Source/Programs/UnrealBuildTool/ProjectFiles/VisualStudio/VCProject.cs
+++ b/Engine/Source/Programs/UnrealBuildTool/ProjectFiles/VisualStudio/VCProject.cs
@@ -1773,26 +1773,68 @@ namespace UnrealBuildTool
 
 				// Parse all the configurations and platforms
 				// Parse the basic structure of the document, updating properties and recursing into other referenced projects as we go
-				foreach (XmlElement Element in Document.DocumentElement.ChildNodes.OfType<XmlElement>())
+				if (!IsDotNETCoreProject())
 				{
-					if(Element.Name == "PropertyGroup")
+					foreach (XmlElement Element in Document.DocumentElement.ChildNodes.OfType<XmlElement>())
 					{
-						string Condition = Element.GetAttribute("Condition");
-						if(!String.IsNullOrEmpty(Condition))
+						if(Element.Name == "PropertyGroup")
 						{
-							Match Match = Regex.Match(Condition, "^\\s*'\\$\\(Configuration\\)\\|\\$\\(Platform\\)'\\s*==\\s*'(.+)\\|(.+)'\\s*$");
-							if(Match.Success && Match.Groups.Count == 3)
+							string Condition = Element.GetAttribute("Condition");
+							if(!String.IsNullOrEmpty(Condition))
 							{
-								Configurations.Add(Match.Groups[1].Value);
-								Platforms.Add(Match.Groups[2].Value);
-							}
-							else
-							{
-								Log.TraceWarning("Unable to parse configuration/platform from condition '{0}': {1}", InitFilePath, Condition);
+								Match Match = Regex.Match(Condition, "^\\s*'\\$\\(Configuration\\)\\|\\$\\(Platform\\)'\\s*==\\s*'(.+)\\|(.+)'\\s*$");
+								if(Match.Success && Match.Groups.Count == 3)
+								{
+									Configurations.Add(Match.Groups[1].Value);
+									Platforms.Add(Match.Groups[2].Value);
+								}
+								else
+								{
+									Log.TraceWarning("Unable to parse configuration/platform from condition '{0}': {1}", InitFilePath, Condition);
+								}
 							}
 						}
 					}
 				}
+				else
+				{
+					bool ConfigurationsFound = false;
+					foreach (XmlElement PropertyGroup in Document.DocumentElement.ChildNodes.OfType<XmlElement>()
+						.Where(element => element.Name == "PropertyGroup"))
+					{
+						XmlNodeList ConfigNodeList = PropertyGroup.GetElementsByTagName("Configurations");
+						// if this property group does not set configurations we do not care about it
+						if (ConfigNodeList.Count == 0)
+							continue;
+						
+						if (PropertyGroup.HasAttribute("Condition"))
+						{
+							string Condition = PropertyGroup.GetAttribute("Condition");
+							Log.TraceWarning("Unable to parse configuration from property group with condition '{0}': {1}. UBT Requires you to set the configuration without conditionals.", InitFilePath, Condition);
+							continue;
+						}
+						string[] ParsedConfigurations = ConfigNodeList[0].FirstChild.Value.Split(';');
+						foreach (string c in ParsedConfigurations)
+						{
+							Configurations.Add(c);
+						}
+
+						// platforms change meaning quite a bit in .net core but typically you do not specify this and its derived from the build instead
+						// for most intents it is just Any CPU from .net framework
+						Platforms.Add("AnyCPU");
+
+						ConfigurationsFound = true;
+						break;
+					}
+
+					// dotnet does not require you to specify configurations or platforms, if you do not debug and release are the defaults
+					if (!ConfigurationsFound)
+					{
+						Configurations.Add("Debug");
+						Configurations.Add("Release");
+						Platforms.Add("AnyCPU");
+					}
+				}
 			}
 			catch(Exception Ex)
 			{
diff --git a/Engine/Source/Programs/UnrealBuildTool/ProjectFiles/VisualStudio/VCProjectFileGenerator.cs b/Engine/Source/Programs/UnrealBuildTool/ProjectFiles/VisualStudio/VCProjectFileGenerator.cs
--- a/Engine/Source/Programs/UnrealBuildTool/ProjectFiles/VisualStudio/VCProjectFileGenerator.cs
+++ b/Engine/Source/Programs/UnrealBuildTool/ProjectFiles/VisualStudio/VCProjectFileGenerator.cs
@@ -455,6 +455,27 @@ namespace UnrealBuildTool
 			}
 
 			// Write AutomationReferences file
+			// Write in in net core expected format
+			if (AutomationProjectFiles.Any())
+			{
+				XNamespace NS = XNamespace.Get("http://schemas.microsoft.com/developer/msbuild/2003");
+
+				DirectoryReference AutomationToolDir = DirectoryReference.Combine(UnrealBuildTool.EngineSourceDirectory, "Programs", "AutomationTool");
+				new XDocument(
+					new XElement(NS + "Project",
+						new XAttribute("ToolsVersion", VCProjectFileGenerator.GetProjectFileToolVersionString(Settings.ProjectFileFormat)),
+						new XAttribute("DefaultTargets", "Build"),
+						new XElement(NS + "ItemGroup",
+							from AutomationProject in AutomationProjectFiles
+							select new XElement(NS + "ProjectReference",
+								new XAttribute("Include", AutomationProject.ProjectFilePath.MakeRelativeTo(AutomationToolDir))
+							)
+						)
+					)
+				).Save(FileReference.Combine(AutomationToolDir, "AutomationToolCore.csproj.References").FullName);
+			}
+	
+			// and a copy in the old traditional msbuild format
 			if (AutomationProjectFiles.Any())
 			{
 				XNamespace NS = XNamespace.Get("http://schemas.microsoft.com/developer/msbuild/2003");
diff --git a/Engine/Source/Programs/UnrealBuildTool/System/RulesCompiler.cs b/Engine/Source/Programs/UnrealBuildTool/System/RulesCompiler.cs
--- a/Engine/Source/Programs/UnrealBuildTool/System/RulesCompiler.cs
+++ b/Engine/Source/Programs/UnrealBuildTool/System/RulesCompiler.cs
@@ -278,6 +278,16 @@ namespace UnrealBuildTool
 						Cache.TargetRules.Add(File.Location);
 					}
 				}
+#if NET_CORE
+				else if (File.HasExtension(".automation.core.csproj"))
+				{
+					lock(Cache.AutomationModules)
+					{
+						Cache.AutomationModules.Add(File.Location);
+					}
+					bSearchSubFolders = false;
+				}
+#else
 				else if (File.HasExtension(".automation.csproj"))
 				{
 					lock(Cache.AutomationModules)
@@ -286,6 +296,7 @@ namespace UnrealBuildTool
 					}
 					bSearchSubFolders = false;
 				}
+#endif
 			}
 
 			// If we didn't find anything to stop the search, search all the subdirectories too
