Fixed UBT DataDrivenPlatformInfo to use case insensitive comparisons

#rnx

diff --git a/Engine/Source/Programs/UnrealBuildTool/System/DataDrivenPlatformInfo.cs b/Engine/Source/Programs/UnrealBuildTool/System/DataDrivenPlatformInfo.cs
--- a/Engine/Source/Programs/UnrealBuildTool/System/DataDrivenPlatformInfo.cs
+++ b/Engine/Source/Programs/UnrealBuildTool/System/DataDrivenPlatformInfo.cs
@@ -51,8 +51,8 @@ namespace UnrealBuildTool
 			// need to init?
 			if (PlatformInfos == null)
 			{
-				PlatformInfos = new Dictionary<string, ConfigDataDrivenPlatformInfo>();
-				Dictionary<string, string> IniParents = new Dictionary<string, string>();
+				PlatformInfos = new Dictionary<string, ConfigDataDrivenPlatformInfo>(StringComparer.OrdinalIgnoreCase);
+				Dictionary<string, string> IniParents = new Dictionary<string, string>(StringComparer.OrdinalIgnoreCase);
 
 				// find all platform directories (skipping NFL/NoRedist)
 				foreach (DirectoryReference EngineConfigDir in UnrealBuildTool.GetExtensionDirs(UnrealBuildTool.EngineDirectory, "Config", bIncludeRestrictedDirectories:false))
