Simplified conditions for per platform checks in net core csprojs to make sure our condition parsing in UBT can handle it.

Should fix precompiled binary builds

diff --git a/Engine/Source/Programs/AutomationTool/AutomationToolCore.csproj b/Engine/Source/Programs/AutomationTool/AutomationToolCore.csproj
--- a/Engine/Source/Programs/AutomationTool/AutomationToolCore.csproj
+++ b/Engine/Source/Programs/AutomationTool/AutomationToolCore.csproj
@@ -21,6 +21,9 @@
     <EnableDefaultItems>false</EnableDefaultItems>
     <DefineConstants>NET_CORE</DefineConstants>
     <AppendRuntimeIdentifierToOutputPath>false</AppendRuntimeIdentifierToOutputPath>
+    <IsWindows Condition="'$([System.Runtime.InteropServices.RuntimeInformation]::IsOSPlatform($([System.Runtime.InteropServices.OSPlatform]::Windows)))' == 'true'">true</IsWindows> 
+    <IsOSX Condition="'$([System.Runtime.InteropServices.RuntimeInformation]::IsOSPlatform($([System.Runtime.InteropServices.OSPlatform]::OSX)))' == 'true'">true</IsOSX> 
+    <IsLinux Condition="'$([System.Runtime.InteropServices.RuntimeInformation]::IsOSPlatform($([System.Runtime.InteropServices.OSPlatform]::Linux)))' == 'true'">true</IsLinux> 
   </PropertyGroup>
   <PropertyGroup Condition=" '$(Configuration)|$(Platform)' == 'Development|AnyCPU' ">
     <DefineConstants>$(DefineConstants);TRACE</DefineConstants>
@@ -34,7 +37,7 @@
     <DefineConstants>$(DefineConstants);DEBUG;TRACE</DefineConstants>
   </PropertyGroup>
   <Choose>
-    <When Condition="$([MSBuild]::IsOSPlatform('OSX'))">
+    <When Condition="'$(IsOSX)'=='true'">
       <PropertyGroup>
         <!-- Set runtime identifier for osx so that it produces a app host executable (consistency with the other platforms) -->
         <RuntimeIdentifier>osx-x64</RuntimeIdentifier>
diff --git a/Engine/Source/Programs/UnrealBuildTool/UnrealBuildToolCore.csproj b/Engine/Source/Programs/UnrealBuildTool/UnrealBuildToolCore.csproj
--- a/Engine/Source/Programs/UnrealBuildTool/UnrealBuildToolCore.csproj
+++ b/Engine/Source/Programs/UnrealBuildTool/UnrealBuildToolCore.csproj
@@ -18,25 +18,17 @@
     <AssemblyName>UnrealBuildTool</AssemblyName>
     <RootNamespace>UnrealBuildTool</RootNamespace>
     <AppendRuntimeIdentifierToOutputPath>false</AppendRuntimeIdentifierToOutputPath>
+    <IsWindows Condition="'$([System.Runtime.InteropServices.RuntimeInformation]::IsOSPlatform($([System.Runtime.InteropServices.OSPlatform]::Windows)))' == 'true'">true</IsWindows> 
+    <IsOSX Condition="'$([System.Runtime.InteropServices.RuntimeInformation]::IsOSPlatform($([System.Runtime.InteropServices.OSPlatform]::OSX)))' == 'true'">true</IsOSX> 
+    <IsLinux Condition="'$([System.Runtime.InteropServices.RuntimeInformation]::IsOSPlatform($([System.Runtime.InteropServices.OSPlatform]::Linux)))' == 'true'">true</IsLinux> 
   </PropertyGroup>
   <Choose>
-    <When Condition="$([MSBuild]::IsOSPlatform('OSX'))">
+    <When Condition="'$(IsOSX)'=='true'">
       <PropertyGroup>
         <!-- Set runtie identifier for osx so that it produces a app host executable (consistency with the other platforms) -->
         <RuntimeIdentifier>osx-x64</RuntimeIdentifier>
       </PropertyGroup>
     </When>
-    <!-- The default runtime identifiers are sufficent for Linux and Windows
-    <When Condition="$([MSBuild]::IsOSPlatform('Linux'))">
-      <PropertyGroup>
-        <RuntimeIdentifier>linux-x64</RuntimeIdentifier>
-      </PropertyGroup>
-    </When>
-    <When Condition="$([MSBuild]::IsOSPlatform('Windows'))">
-      <PropertyGroup>
-        <RuntimeIdentifier>win-x64</RuntimeIdentifier>
-      </PropertyGroup>
-    </When>-->
   </Choose>
   <PropertyGroup Condition=" '$(Configuration)|$(Platform)' == 'Development|AnyCPU' ">
     <Optimize>true</Optimize>
