- Added support for finding Windows SDK on builders that don't have the SDK installed - but might have the SDK bits coming from AutoSDK location

diff --git a/Engine/Source/Programs/UnrealBuildTool/Platform/Windows/MicrosoftPlatformSDK.cs b/Engine/Source/Programs/UnrealBuildTool/Platform/Windows/MicrosoftPlatformSDK.cs
--- a/Engine/Source/Programs/UnrealBuildTool/Platform/Windows/MicrosoftPlatformSDK.cs
+++ b/Engine/Source/Programs/UnrealBuildTool/Platform/Windows/MicrosoftPlatformSDK.cs
@@ -3,6 +3,7 @@
 using System;
 using System.Collections.Generic;
 using Tools.DotNETCommon;
+using System.Linq;
 using System.Text.RegularExpressions;
 using Microsoft.Win32;
 
@@ -12,7 +13,7 @@ namespace UnrealBuildTool
 	{
 		public override string GetMainVersion()
 		{
-			// the current and previous versions of the SDK (technically, NX Add-On), with appended NEX version, we support
+			// the current Windows SDK version we expect
 			return "10.0.18632";
 		}
 
@@ -36,15 +37,29 @@ namespace UnrealBuildTool
 			}
 
 			// @todo turnkey: do we support pre-10?
-			// @todo turnkey: MicrosoftPlatformSDK maybe?
+
+			// look for an installed SDK
 			string Version = "v10.0";
 			object Result = Registry.GetValue(@"HKEY_LOCAL_MACHINE\SOFTWARE\Wow6432Node\Microsoft\Microsoft SDKs\Windows\" + Version, "InstallationFolder", null);
-
 			if (Result != null)
 			{
 				return Registry.GetValue(@"HKEY_LOCAL_MACHINE\SOFTWARE\Wow6432Node\Microsoft\Microsoft SDKs\Windows\" + Version, "ProductVersion", null) as string;
 			}
 
+			// look in AutoSDK location (note that it doesn't have a setup.bat, at least right now, so the AutoSDK system isn't used
+			// @todo turnkey: make use of AutoSDKs, and maybe get rid of non-Latest style SDK selection in UEBUildWindows? (and move all that stuff here)
+			DirectoryReference HostAutoSdkDir;
+			if (TryGetHostPlatformAutoSDKDir(out HostAutoSdkDir))
+			{
+				DirectoryReference RootDirAutoSdk = DirectoryReference.Combine(HostAutoSdkDir, "Win64", "Windows Kits", "10", "include");
+				if (DirectoryReference.Exists(RootDirAutoSdk))
+				{
+					// sort the directories under 10, and use the highest one (this logic mirrors UEBuildWindows logic)
+					return DirectoryReference.EnumerateDirectories(RootDirAutoSdk).OrderBy(x => x.GetDirectoryName()).Last().GetDirectoryName();
+				}
+			}
+
+
 			return null;
 		}
 
