Made Ionic.Zip copylocal in netcore, resolves issues loading that assembly when packaging some platforms.

#UE-102997

diff --git a/Engine/Source/Programs/UnrealBuildTool/UnrealBuildToolCore.csproj b/Engine/Source/Programs/UnrealBuildTool/UnrealBuildToolCore.csproj
--- a/Engine/Source/Programs/UnrealBuildTool/UnrealBuildToolCore.csproj
+++ b/Engine/Source/Programs/UnrealBuildTool/UnrealBuildToolCore.csproj
@@ -47,7 +47,6 @@
   <ItemGroup>
     <Reference Include="Ionic.Zip.Reduced">
       <HintPath>..\..\..\Binaries\DotNET\Ionic.Zip.Reduced.dll</HintPath>
-      <Private>False</Private>
     </Reference>
     <Reference Include="Microsoft.VisualStudio.Setup.Configuration.Interop">
       <HintPath>..\..\..\Binaries\ThirdParty\VisualStudio\Microsoft.VisualStudio.Setup.Configuration.Interop.dll</HintPath>
