LiveCoding fails to compile code change.
LiveCoding console reports one of the following messages

Unable to read from .../LiveCoding.json
Build failed with ERROR: Unable to start regular build while Live Coding is active. Press Ctrl+Alt+F11 to trigger a Live Coding compile.

diff --git a/Engine/Source/Programs/UnrealBuildTool/System/HotReload.cs b/Engine/Source/Programs/UnrealBuildTool/System/HotReload.cs
--- a/Engine/Source/Programs/UnrealBuildTool/System/HotReload.cs
+++ b/Engine/Source/Programs/UnrealBuildTool/System/HotReload.cs
@@ -162,7 +162,12 @@ namespace UnrealBuildTool
 		public static void Setup(TargetDescriptor TargetDescriptor, TargetMakefile Makefile, BuildConfiguration BuildConfiguration)
 		{
 			// Get the hot-reload mode
-			if (Makefile.HotReloadModuleNames.Count == 0)
+			if (TargetDescriptor.HotReloadMode == HotReloadMode.LiveCoding)
+			{
+				// In some instances such as packaged builds, we might not have hot reload modules names.
+				// We don't want to lose the live coding setting in that case.
+			}
+			else if (Makefile.HotReloadModuleNames.Count == 0)
 			{
 				TargetDescriptor.HotReloadMode = HotReloadMode.Disabled;
 			}
