Hololens fixes in UBT and UAT under netcore

This was using net framework specific apis, found corresponding apis under net core and switched to them.

diff --git a/Engine/Source/Programs/AutomationTool/HoloLens/HoloLens.Automation.Core.csproj b/Engine/Source/Programs/AutomationTool/HoloLens/HoloLens.Automation.Core.csproj
--- a/Engine/Source/Programs/AutomationTool/HoloLens/HoloLens.Automation.Core.csproj
+++ b/Engine/Source/Programs/AutomationTool/HoloLens/HoloLens.Automation.Core.csproj
@@ -51,26 +51,7 @@
     <When Condition="'$(IsWindows)'=='true'">
       <PropertyGroup>
         <UseWindowsForms>true</UseWindowsForms>
-        <Windows10SdkPath>$(Registry:HKEY_LOCAL_MACHINE\SOFTWARE\Wow6432Node\Microsoft\Microsoft SDKs\Windows\v10.0@InstallationFolder)</Windows10SdkPath>
-        <Windows10SdkLatestVersion>$(Registry:HKEY_LOCAL_MACHINE\SOFTWARE\Wow6432Node\Microsoft\Microsoft SDKs\Windows\v10.0@ProductVersion)</Windows10SdkLatestVersion>
       </PropertyGroup>
-      <ItemGroup>
-        <Reference Include="System.Runtime.WindowsRuntime">
-          <HintPath>$(SystemDrive)\Program Files (x86)\Reference Assemblies\Microsoft\Framework\.NETCore\$(TargetFrameworkVersion)\System.Runtime.WindowsRuntime.dll</HintPath>
-          <HintPath Condition="'$([System.IO.File]::Exists(%(HintPath)))'=='false'">C:\Program Files (x86)\Reference Assemblies\Microsoft\Framework\.NETCore\v4.5\System.Runtime.WindowsRuntime.dll</HintPath>
-          <Private>True</Private>
-        </Reference>
-        <!-- Disable windows reference in net core as this causes infinite rebuilds
-        <Reference Include="Windows">
-          <HintPath>$(Windows10SdkPath)\UnionMetadata\$(Windows10SdkLatestVersion)\Windows.winmd</HintPath>
-          <SDKsWinmdPath Condition="'$(UE_SDKS_ROOT)' != ''">$(UE_SDKS_ROOT)\HostWin64\Win64\Windows Kits\10\UnionMetadata\10.0.18362.0\Windows.winmd</SDKsWinmdPath>
-          <HintPath Condition="'$(UE_SDKS_ROOT)' != '' And '$([System.IO.File]::Exists(%(SDKsWinmdPath)))'!='false'">$(UE_SDKS_ROOT)\HostWin64\Win64\Windows Kits\10\UnionMetadata\10.0.18362.0\Windows.winmd</HintPath>
-          <HintPath Condition="'$([System.IO.File]::Exists(%(HintPath)))'=='false'">$(Windows10SdkPath)\UnionMetadata\$(Windows10SdkLatestVersion).0\Windows.winmd</HintPath>
-          <HintPath Condition="'$([System.IO.File]::Exists(%(HintPath)))'=='false'">$(Windows10SdkPath)\UnionMetadata\Windows.winmd</HintPath>
-          <Private>false</Private>
-        </Reference>
-        -->
-      </ItemGroup>
     </When>
   </Choose>
   <Import Project="Sdk.props" Sdk="Microsoft.NET.Sdk.WindowsDesktop" Condition="'$(IsWindows)'=='true'"/>
@@ -98,5 +79,7 @@
   <ItemGroup>
     <PackageReference Include="Microsoft.CSharp" Version="4.7.0" />
     <PackageReference Include="System.Data.DataSetExtensions" Version="4.5.0" />
+    <PackageReference Include="System.Runtime.WindowsRuntime" Version="4.7.0" />
+    <PackageReference Include="Microsoft.Windows.SDK.Contracts" Version="10.0.19041.1" />
   </ItemGroup>
 </Project>
\ No newline at end of file
diff --git a/Engine/Source/Programs/AutomationTool/HoloLens/HoloLensPlatform.Automation.cs b/Engine/Source/Programs/AutomationTool/HoloLens/HoloLensPlatform.Automation.cs
--- a/Engine/Source/Programs/AutomationTool/HoloLens/HoloLensPlatform.Automation.cs
+++ b/Engine/Source/Programs/AutomationTool/HoloLens/HoloLensPlatform.Automation.cs
@@ -13,7 +13,7 @@ using Tools.DotNETCommon;
 using System.Xml;
 using System.Text.RegularExpressions;
 
-#if !__MonoCS__ && !NET_CORE
+#if !__MonoCS__
 using Windows.Management.Deployment;
 using Windows.Foundation;
 #if false
@@ -196,7 +196,7 @@ namespace HoloLens.Automation
 		}
 	}
 
-#if !__MonoCS__ && !NET_CORE
+#if !__MonoCS__
 	class HoloLensDevicePortalCreatedProcess : IProcessResult
 	{
 		object StateLock;
@@ -585,7 +585,7 @@ namespace HoloLens.Automation
 				{
 					string ArchString = Target.Receipt.Architecture;
 					SC.StageBuildProductsFromReceipt(Target.Receipt, Target.RequireFilesExist, Params.bTreatNonShippingBinariesAsDebugFiles);
-					DeployExports.AddWinMDReferencesFromReceipt(Target.Receipt, Params.RawProjectPath.Directory, SC.LocalRoot.FullName);
+					DeployExports.AddWinMDReferencesFromReceipt(Target.Receipt, Params.RawProjectPath.Directory, SC.LocalRoot.FullName, Windows10SDKVersion);
 
 					// Stage HoloLens-specific assets (tile, splash, etc.)
 					DirectoryReference assetsPath = new DirectoryReference(Path.Combine(ProjectBinariesFolder.FullName, ArchString, "Resources"));
@@ -1224,7 +1224,7 @@ namespace HoloLens.Automation
 			return !result;
 		}
 
-#if !__MonoCS__ && !NET_CORE
+#if !__MonoCS__
 		private void WaitFor(IAsyncOperationWithProgress<DeploymentResult, DeploymentProgress> deploymentOperation)
 		{
 			// This event is signaled when the operation completes
@@ -1262,7 +1262,7 @@ namespace HoloLens.Automation
 
 		private void DeployToLocalDevice(ProjectParams Params, DeploymentContext SC)
 		{
-#if !__MonoCS__ && !NET_CORE
+#if !__MonoCS__
             if (Utils.IsRunningOnMono)
             {
                 return;
@@ -1564,7 +1564,7 @@ namespace HoloLens.Automation
 
 		private IProcessResult RunUsingDevicePortal(string DeviceAddress, ERunOptions ClientRunFlags, string ClientApp, string ClientCmdLine, ProjectParams Params)
 		{
-#if !__MonoCS__ && !NET_CORE
+#if !__MonoCS__
             if (Utils.IsRunningOnMono)
             {
                 return null;
diff --git a/Engine/Source/Programs/UnrealBuildTool/Platform/HoloLens/HoloLensDeploy.cs b/Engine/Source/Programs/UnrealBuildTool/Platform/HoloLens/HoloLensDeploy.cs
--- a/Engine/Source/Programs/UnrealBuildTool/Platform/HoloLens/HoloLensDeploy.cs
+++ b/Engine/Source/Programs/UnrealBuildTool/Platform/HoloLens/HoloLensDeploy.cs
@@ -474,7 +474,15 @@ namespace UnrealBuildTool
 				NewReceipt = new TargetReceipt(Receipt.ProjectFile, Receipt.TargetName, Receipt.TargetType, Receipt.Platform, Receipt.Configuration, Receipt.Version, "Multi");
 			}
 
-			AddWinMDReferencesFromReceipt(Receipt, Receipt.ProjectDir != null ? Receipt.ProjectDir : UnrealBuildTool.EngineDirectory, UnrealBuildTool.EngineDirectory.ParentDirectory.FullName);
+			string SDK = "";
+			var Results = Receipt.AdditionalProperties.Where(x => x.Name == "SDK");
+			if (Results.Any())
+			{
+				SDK = Results.First().Value;
+			}
+			HoloLensExports.InitWindowsSdkToolPath(SDK);
+
+			AddWinMDReferencesFromReceipt(Receipt, Receipt.ProjectDir != null ? Receipt.ProjectDir : UnrealBuildTool.EngineDirectory, UnrealBuildTool.EngineDirectory.ParentDirectory.FullName, SDK);
 
 			//PrepForUATPackageOrDeploy(InTarget.ProjectFile, InAppName, InTarget.ProjectDirectory.FullName, InTarget.OutputPath.FullName, TargetBuildEnvironment.RelativeEnginePath, false, "", false);
 			List<UnrealTargetConfiguration> TargetConfigs = new List<UnrealTargetConfiguration> { Receipt.Configuration };
@@ -487,13 +495,6 @@ namespace UnrealBuildTool
 				Arch = WindowsArchitecture.x64;
 			}
 
-			string SDK = "";
-			var Results = Receipt.AdditionalProperties.Where(x => x.Name == "SDK");
-			if (Results.Any())
-			{
-				SDK = Results.First().Value;
-			}
-			HoloLensExports.InitWindowsSdkToolPath(SDK);
 
 			string AbsoluteExeDirectory = Path.GetDirectoryName(ExePaths[0]);
 			UnrealTargetPlatform Platform = UnrealTargetPlatform.HoloLens;
@@ -513,7 +514,7 @@ namespace UnrealBuildTool
 			return true;
 		}
 
-		public void AddWinMDReferencesFromReceipt(TargetReceipt Receipt, DirectoryReference SourceProjectDir, string DestRelativeTo)
+		public void AddWinMDReferencesFromReceipt(TargetReceipt Receipt, DirectoryReference SourceProjectDir, string DestRelativeTo, string SDKVersion)
 		{
 			// Dependency paths in receipt are already expanded at this point
 			foreach (var Dep in Receipt.RuntimeDependencies)
@@ -528,7 +529,7 @@ namespace UnrealBuildTool
 						DestPath = Dep.Path.FullName.Replace(UnrealBuildTool.EngineDirectory.FullName, Path.Combine(DestRelativeTo, "Engine"));
 						DestPath = DestPath.Replace(SourceProjectDir.FullName, Path.Combine(DestRelativeTo, SourceProjectDir.GetDirectoryName()));
 						DestPath = Utils.MakePathRelativeTo(DestPath, DestRelativeTo);
-						WinMDReferences.Add(new WinMDRegistrationInfo(new FileReference(WinMDFile), DestPath));
+						WinMDReferences.Add(new WinMDRegistrationInfo(new FileReference(WinMDFile), DestPath, SDKVersion));
 					}
 				}
 			}
diff --git a/Engine/Source/Programs/UnrealBuildTool/Platform/HoloLens/HoloLensExports.cs b/Engine/Source/Programs/UnrealBuildTool/Platform/HoloLens/HoloLensExports.cs
--- a/Engine/Source/Programs/UnrealBuildTool/Platform/HoloLens/HoloLensExports.cs
+++ b/Engine/Source/Programs/UnrealBuildTool/Platform/HoloLens/HoloLensExports.cs
@@ -49,9 +49,10 @@ namespace UnrealBuildTool
 		/// <param name="Receipt"></param>
 		/// <param name="SourceProjectDir"></param>
 		/// <param name="DestPackageRoot"></param>
-		public void AddWinMDReferencesFromReceipt(TargetReceipt Receipt, DirectoryReference SourceProjectDir, string DestPackageRoot)
+		/// <param name="SdkVersion"></param>
+		public void AddWinMDReferencesFromReceipt(TargetReceipt Receipt, DirectoryReference SourceProjectDir, string DestPackageRoot, string SdkVersion)
 		{
-			InnerDeploy.AddWinMDReferencesFromReceipt(Receipt, SourceProjectDir, DestPackageRoot);
+			InnerDeploy.AddWinMDReferencesFromReceipt(Receipt, SourceProjectDir, DestPackageRoot, SdkVersion);
 		}
 
 		/// <summary>
diff --git a/Engine/Source/Programs/UnrealBuildTool/Platform/HoloLens/WinMDRegistrationInfo.cs b/Engine/Source/Programs/UnrealBuildTool/Platform/HoloLens/WinMDRegistrationInfo.cs
--- a/Engine/Source/Programs/UnrealBuildTool/Platform/HoloLens/WinMDRegistrationInfo.cs
+++ b/Engine/Source/Programs/UnrealBuildTool/Platform/HoloLens/WinMDRegistrationInfo.cs
@@ -2,9 +2,11 @@
 
 using System;
 using System.Collections.Generic;
+using System.IO;
 using System.Linq;
 using System.Runtime.InteropServices.WindowsRuntime;
 using System.Reflection;
+using System.Runtime.InteropServices;
 using Tools.DotNETCommon;
 
 namespace UnrealBuildTool
@@ -41,28 +43,122 @@ namespace UnrealBuildTool
 			public string ThreadingModelName { get; private set; }
 		}
 
+		private static Version FindLatestVersionDirectory(string InDirectory, Version NoLaterThan)
+		{
+			Version LatestVersion = new Version(0, 0, 0, 0);
+			if (Directory.Exists(InDirectory))
+			{
+				string[] VersionDirectories = Directory.GetDirectories(InDirectory);
+				foreach (string Dir in VersionDirectories)
+				{
+					string VersionString = Path.GetFileName(Dir);
+					Version FoundVersion;
+					if (Version.TryParse(VersionString, out FoundVersion) && FoundVersion > LatestVersion)
+					{
+						if (NoLaterThan == null || FoundVersion <= NoLaterThan)
+						{
+							LatestVersion = FoundVersion;
+						}
+					}
+				}
+			}
+			return LatestVersion;
+		}
+
+		private List<string> ExpandWinMDReferences(string SdkVersion, string[] WinMDReferences)
+		{
+			
+			// Allow bringing in Windows SDK contracts just by naming the contract
+			// These are files that look like References/10.0.98765.0/AMadeUpWindowsApiContract/5.0.0.0/AMadeUpWindowsApiContract.winmd
+			List<string> ExpandedWinMDReferences = new List<string>();
+
+			VersionNumber OutSdkVersion;
+			DirectoryReference OutSdkDir;
+			if (!WindowsPlatform.TryGetWindowsSdkDir(SdkVersion, out OutSdkVersion, out OutSdkDir))
+			{
+				Log.TraceError("Failed to find WinSDK " + SdkVersion);
+				throw new Exception("Failed to find WinSDK");
+			}
+
+			// The first few releases of the Windows 10 SDK didn't put the SDK version in the reference path
+			string ReferenceRoot = Path.Combine(OutSdkDir.FullName, "References");
+			string VersionedReferenceRoot = Path.Combine(ReferenceRoot, OutSdkVersion.ToString());
+			if (Directory.Exists(VersionedReferenceRoot))
+			{
+				ReferenceRoot = VersionedReferenceRoot;
+			}
+
+			foreach (string WinMDRef in WinMDReferences)
+			{
+				if (File.Exists(WinMDRef))
+				{
+					// Already a valid path
+					ExpandedWinMDReferences.Add(WinMDRef);
+				}
+				else
+				{
+					string ContractFolder = Path.Combine(ReferenceRoot, WinMDRef);
+
+					Version ContractVersion = FindLatestVersionDirectory(ContractFolder, null);
+					string ExpandedWinMDRef = Path.Combine(ContractFolder, ContractVersion.ToString(), WinMDRef + ".winmd");
+					if (File.Exists(ExpandedWinMDRef))
+					{
+						ExpandedWinMDReferences.Add(ExpandedWinMDRef);
+					}
+					else
+					{
+						Log.TraceWarning("Unable to resolve location for HoloLens WinMD api contract {0}, file {1}", WinMDRef, ExpandedWinMDRef);
+					}
+				}
+			}
+
+			return ExpandedWinMDReferences;
+		}
+
 		/// <summary>
 		/// WinMD reference info
 		/// </summary>
 		/// <param name="InWindMDSourcePath"></param>
 		/// <param name="InPackageRelativeDllPath"></param>
-		public WinMDRegistrationInfo(FileReference InWindMDSourcePath, string InPackageRelativeDllPath)
+		/// <param name="SdkVersion"></param>
+		public WinMDRegistrationInfo(FileReference InWindMDSourcePath, string InPackageRelativeDllPath, string SdkVersion)
 		{
 			PackageRelativeDllPath = InPackageRelativeDllPath;
 			ResolveSearchPaths.Add(InWindMDSourcePath.Directory.FullName);
+#if NET_CORE			
+			List<string> WinMDAssemblies = ExpandWinMDReferences(SdkVersion, new string[] {
+				"Windows.Foundation.FoundationContract",
+				"Windows.Foundation.UniversalApiContract"
+			});
+			string[] runtimeAssemblies = Directory.GetFiles(RuntimeEnvironment.GetRuntimeDirectory(), "*.dll");
+			// Create the list of assembly paths consisting of runtime assemblies and the input file.
+			List<string> AssemblyPaths = new List<string>(runtimeAssemblies);
+			AssemblyPaths.AddRange(WinMDAssemblies);
+			AssemblyPaths.Add(InWindMDSourcePath.FullName);
 
+			// Create MetadataLoadContext that can resolve assemblies using the created list.
+			PathAssemblyResolver PathAssemblyResolver = new PathAssemblyResolver(AssemblyPaths);
+			MetadataLoadContext Mlc = new MetadataLoadContext(PathAssemblyResolver);
+
+			using (Mlc)
+#else
 			if (BuildHostPlatform.Current.Platform == UnrealTargetPlatform.Win64)
+#endif
 			{
 				ActivatableTypesList = new List<ActivatableType>();
-
+#if NET_CORE
+				Assembly DependsOn = Mlc.LoadFromAssemblyPath(InWindMDSourcePath.FullName);
+#else
 				Assembly DependsOn = Assembly.ReflectionOnlyLoadFrom(InWindMDSourcePath.FullName);
+#endif
 				foreach (Type WinMDType in DependsOn.GetExportedTypes())
 				{
 					bool IsActivatable = false;
 					string ThreadingModelName = "both";
 					foreach (CustomAttributeData Attr in WinMDType.CustomAttributes)
 					{
-						if (Attr.AttributeType.FullName == "Windows.Foundation.Metadata.ActivatableAttribute" || Attr.AttributeType.FullName == "Windows.Foundation.Metadata.StaticAttribute")
+						if (Attr.AttributeType.FullName == "Windows.Foundation.Metadata.ActivatableAttribute" ||
+						    Attr.AttributeType.FullName == "Windows.Foundation.Metadata.StaticAttribute")
 						{
 							IsActivatable = true;
 						}
@@ -96,12 +192,13 @@ namespace UnrealBuildTool
 			}
 		}
 
+#if !NET_CORE
 		static WinMDRegistrationInfo()
 		{
 			if (BuildHostPlatform.Current.Platform == UnrealTargetPlatform.Win64)
 			{
 				AppDomain.CurrentDomain.ReflectionOnlyAssemblyResolve += (Sender, EventArgs) => Assembly.ReflectionOnlyLoad(EventArgs.Name);
-#if !NET_CORE
+
 				WindowsRuntimeMetadata.ReflectionOnlyNamespaceResolve += (Sender, EventArgs) =>
 				{
 					string Path = WindowsRuntimeMetadata.ResolveNamespace(EventArgs.NamespaceName, ResolveSearchPaths).FirstOrDefault();
@@ -111,10 +208,9 @@ namespace UnrealBuildTool
 					}
 					EventArgs.ResolvedAssemblies.Add(Assembly.ReflectionOnlyLoadFrom(Path));
 				};
-#endif	
 			}
 		}
-
+#endif
 		private static List<string> ResolveSearchPaths = new List<string>();
 		private List<ActivatableType> ActivatableTypesList;
 	}
diff --git a/Engine/Source/Programs/UnrealBuildTool/UnrealBuildToolCore.csproj b/Engine/Source/Programs/UnrealBuildTool/UnrealBuildToolCore.csproj
--- a/Engine/Source/Programs/UnrealBuildTool/UnrealBuildToolCore.csproj
+++ b/Engine/Source/Programs/UnrealBuildTool/UnrealBuildToolCore.csproj
@@ -112,5 +112,6 @@
     <PackageReference Include="System.Management" Version="4.7.0" />
     <PackageReference Include="System.Security.Cryptography.Csp" Version="4.3.0" />
     <PackageReference Include="System.ServiceProcess.ServiceController" Version="4.7.0" />
+    <PackageReference Include="System.Reflection.MetadataLoadContext" Version="4.7.2" /> <!-- Required for Hololens -->
   </ItemGroup>
 </Project>
\ No newline at end of file
