Add support for enabling the clang static analyzer on Windows when using the clang compiler.

The clang static analyser is arguably a lot more powerful tool than either PVS or the MSVC, and gives detailed information how how it arrived at its conclusion, unlike the other two (example listed below).

Changed the WindowsStaticAnalyzer.VisualCpp enum to alias to a new entry, WindowsStaticAnalyzer.Default, which picks the default static analysis tool for the given working compiler.

Presently, only a subset of checkers is enabled. As of this CL, the analyzer reports a huge amount of issues, some of which are benign (such as knowingly de-referencing shared pointers when we're sure they're valid, although those could benefit from being changed to shared refs) and some real bugs, for example calls to virtual functions from constructors/destructors, memory leaks in Slate and what looks like a really tricky, potential free-after-use issue with TSharedRef.  Until some of those issues are resolved, I wouldn't recommend enabling this as a analysis pass in lieu of MSVC or PVS. It can be run selectively on individual files for now, however.

Example output:
  E:\dev\ue5_main1\Engine\Source\Developer\StandaloneRenderer\Private\Windows\D3D\SlateD3DIndexBuffer.cpp(85,4): warning: Potential leak of memory pointed to by 'SavedIndices' [cplusplus.NewDeleteLeaks]
                          LogSlateD3DRendererFailure(TEXT("FSlateD3DIndexBuffer::ResizeBuffer() - ID3D11Device::CreateBuffer"), Hr);
                          ^
  E:\dev\ue5_main1\Engine\Source\Developer\StandaloneRenderer\Private\Windows\D3D\SlateD3DIndexBuffer.cpp(45,6): note: Assuming 'NumIndices' is > field 'MaxNumIndices'
          if( NumIndices > MaxNumIndices )
              ^~~~~~~~~~~~~~~~~~~~~~~~~~
  E:\dev\ue5_main1\Engine\Source\Developer\StandaloneRenderer\Private\Windows\D3D\SlateD3DIndexBuffer.cpp(45,2): note: Taking true branch
          if( NumIndices > MaxNumIndices )
          ^
  E:\dev\ue5_main1\Engine\Source\Developer\StandaloneRenderer\Private\Windows\D3D\SlateD3DIndexBuffer.cpp(53,19): note: Memory is allocated
                          SavedIndices = new uint8[MaxNumIndices*sizeof(SlateIndex)];
                                         ^~~~~~~~~~~~~~~~~~~~~~~~~~~~~~~~~~~~~~~~~~~
  E:\dev\ue5_main1\Engine\Source\Developer\StandaloneRenderer\Private\Windows\D3D\SlateD3DIndexBuffer.cpp(68,7): note: Assuming 'Hr' is < 0
                  if (SUCCEEDED(Hr))
                      ^~~~~~~~~~~~~
  C:\Program Files (x86)\Windows Kits\10\include\10.0.18362.0\shared\winerror.h(28788,24): note: expanded from macro 'SUCCEEDED'
  #define SUCCEEDED(hr) (((HRESULT)(hr)) >= 0)
                         ^~~~~~~~~~~~~~~~~~~~
  E:\dev\ue5_main1\Engine\Source\Developer\StandaloneRenderer\Private\Windows\D3D\SlateD3DIndexBuffer.cpp(68,3): note: Taking false branch
                  if (SUCCEEDED(Hr))
                  ^
  E:\dev\ue5_main1\Engine\Source\Developer\StandaloneRenderer\Private\Windows\D3D\SlateD3DIndexBuffer.cpp(85,4): note: Potential leak of memory pointed to by 'SavedIndices'
                          LogSlateD3DRendererFailure(TEXT("FSlateD3DIndexBuffer::ResizeBuffer() - ID3D11Device::CreateBuffer"), Hr);
                          ^

#rnx

diff --git a/Engine/Source/Programs/UnrealBuildTool/Platform/Windows/UEBuildWindows.cs b/Engine/Source/Programs/UnrealBuildTool/Platform/Windows/UEBuildWindows.cs
--- a/Engine/Source/Programs/UnrealBuildTool/Platform/Windows/UEBuildWindows.cs
+++ b/Engine/Source/Programs/UnrealBuildTool/Platform/Windows/UEBuildWindows.cs
@@ -65,10 +65,17 @@ namespace UnrealBuildTool
 		/// </summary>
 		None,
 
+		/// <summary>
+		/// Use the default static analyzer for the selected compiler, if it has one. For
+		/// Visual Studio and Clang, this means using their built-in static analysis tools.
+		/// Any compiler that doesn't support static analysis will ignore this option.
+		/// </summary>
+		Default, 
+
 		/// <summary>
 		/// Use the built-in Visual C++ static analyzer
 		/// </summary>
-		VisualCpp,
+		VisualCpp = Default,
 
 		/// <summary>
 		/// Use PVS-Studio for static analysis
@@ -810,6 +817,15 @@ namespace UnrealBuildTool
 				Target.bCompileCEF3 = false;
 			}
 
+			// Disable static analysis if not appropriate for the compiler.
+			if (Target.WindowsPlatform.StaticAnalyzer == WindowsStaticAnalyzer.Default)
+			{
+				if (Target.WindowsPlatform.Compiler == WindowsCompiler.Intel)
+				{
+					Target.WindowsPlatform.StaticAnalyzer = WindowsStaticAnalyzer.None;
+				}
+			}
+
 			// Set the compiler version if necessary
 			if (Target.WindowsPlatform.Compiler == WindowsCompiler.Default)
 			{
@@ -829,11 +845,15 @@ namespace UnrealBuildTool
 				Target.bDisableLinking = true;
 			}
 
-			// Disable PCHs for PVS studio
+			 // Disable PCHs for PVS studio and clang static analyzer.
 			if(Target.WindowsPlatform.StaticAnalyzer == WindowsStaticAnalyzer.PVSStudio)
 			{
 				Target.bUsePCHFiles = false;
 			}
+			else if (Target.WindowsPlatform.Compiler == WindowsCompiler.Clang && Target.WindowsPlatform.StaticAnalyzer == WindowsStaticAnalyzer.Default)
+			{
+				Target.bUsePCHFiles = false;
+			}
 
 			// Override PCH settings
 			if (Target.WindowsPlatform.Compiler == WindowsCompiler.Intel)
diff --git a/Engine/Source/Programs/UnrealBuildTool/Platform/Windows/VCToolChain.cs b/Engine/Source/Programs/UnrealBuildTool/Platform/Windows/VCToolChain.cs
--- a/Engine/Source/Programs/UnrealBuildTool/Platform/Windows/VCToolChain.cs
+++ b/Engine/Source/Programs/UnrealBuildTool/Platform/Windows/VCToolChain.cs
@@ -157,18 +157,65 @@ namespace UnrealBuildTool
 			// in object file sizes (and link times) due to the amount of stuff we inline.
 			Arguments.Add("/Zc:inline");
 
-			if (Target.WindowsPlatform.Compiler == WindowsCompiler.Clang)
-			{
-				Arguments.Add("-fms-compatibility-version=19.1");
-			}
-
 			// @todo clang: Clang on Windows doesn't respect "#pragma warning (error: ####)", and we're not passing "/WX", so warnings are not
 			// treated as errors when compiling on Windows using Clang right now.
 
 			// NOTE re: clang: the arguments for clang-cl can be found at http://llvm.org/viewvc/llvm-project/cfe/trunk/include/clang/Driver/CLCompatOptions.td?view=markup
 			// This will show the cl.exe options that map to clang.exe ones, which ones are ignored and which ones are unsupported.
+			if (Target.WindowsPlatform.Compiler == WindowsCompiler.Clang)
+			{
+				// FIXME: We should be compatible with the latest MSVC compiler toolkit installed instead of hard-coding it here.
+				Arguments.Add("-fms-compatibility-version=19.1");
+
+				if (Target.WindowsPlatform.StaticAnalyzer == WindowsStaticAnalyzer.Default)
+				{
+					// Enable the static analyzer but only via the backend. Using the frontend
+					// flag ('--analyze') will enable a suite of default checkers, some of which
+					// we don't want (like deadcode.DeadStore)
+					Arguments.Add("-Xclang -analyze");
 
-			if (Target.WindowsPlatform.StaticAnalyzer == WindowsStaticAnalyzer.VisualCpp)
+					// Make sure we get textual output and not XML.
+					Arguments.Add("-Xclang -analyzer-output=text");
+
+					// Make sure we check inside nested blocks (e.g. 'if ((foo = getchar()) == 0) {}')
+					Arguments.Add("-Xclang -analyzer-opt-analyze-nested-blocks");
+
+					// Needed for some of the C++ checkers.
+					Arguments.Add("-Xclang -analyzer-config -Xclang aggressive-binary-operation-simplification=true");
+
+					// Enable only specific checkers of families of checkers
+					// See https://clang.llvm.org/docs/analyzer/checkers.html for a full list. Or run:
+					//    'clang -Xclang -analyzer-checker-help' 
+					// or: 
+					//    'clang -Xclang -analyzer-checker-help-alpha' 
+					// for the list of experimental checkers.
+					String[] EnabledCheckers = 
+					{
+						"core",
+						"unix.Malloc",
+						"unix.MallocSizeof",
+						"cplusplus",
+						"optin.cplusplus.UninitializedObject",
+						"optin.cplusplus.VirtualCall",
+						// "deadstore",					// Check for dead stores (noisy in UE).
+						// "security.FloatLoopCounter",	// Check if using floats for loop counters
+						// Experimental checkers (as of clang 11.0)
+						// "alpha.core.Conversion",		// Sign conversion (noisy)
+						// "alpha.core.PointerArithm",	// Sketchy pointer arithmetic
+						// "alpha.cplusplus.DeleteWithNonVirtualDtor",
+						// "alpha.cplusplus.InvalidatedIterator",
+						// "alpha.cplusplus.IteratorRange",
+						// "alpha.cplusplus.MismatchedIterator",
+						// "alpha.cplusplus.InvalidatedIterator",
+					};
+
+					foreach (String Checker in EnabledCheckers)
+					{
+						Arguments.Add(String.Format("-Xclang -analyzer-checker={0}", Checker));
+					}
+				}
+			}
+			else if (Target.WindowsPlatform.StaticAnalyzer == WindowsStaticAnalyzer.Default)
 			{
 				Arguments.Add("/analyze");
 
@@ -221,7 +268,7 @@ namespace UnrealBuildTool
 			// Allow 750% of the default memory allocation limit when using the static analyzer, and 1000% at other times.
 			if(Target.WindowsPlatform.PCHMemoryAllocationFactor == 0)
 			{
-				if (Target.WindowsPlatform.StaticAnalyzer == WindowsStaticAnalyzer.VisualCpp)
+				if (Target.WindowsPlatform.StaticAnalyzer == WindowsStaticAnalyzer.Default)
 				{
 					Arguments.Add("/Zm750");
 				}
@@ -245,7 +292,7 @@ namespace UnrealBuildTool
 			if (Target.WindowsPlatform.Compiler >= WindowsCompiler.VisualStudio2017 && Target.WindowsPlatform.bStrictConformanceMode)
 			{
 				// Do not enable if we are running the MSVC code analyzer as it will fail to compile.
-				if (Target.WindowsPlatform.StaticAnalyzer != WindowsStaticAnalyzer.VisualCpp)
+				if (Target.WindowsPlatform.StaticAnalyzer != WindowsStaticAnalyzer.Default)
 				{
 					Arguments.Add("/permissive-");
 					Arguments.Add("/Zc:strictStrings-"); // Have to disable strict const char* semantics due to Windows headers not being compliant.
