UBT: Add preprocessor work from Private-ArtifactCache stream.

#rnx

diff --git a/Engine/Source/Programs/UnrealBuildTool/Preprocessor/Identifier.cs b/Engine/Source/Programs/UnrealBuildTool/Preprocessor/Identifier.cs
new file mode 100644
--- /dev/null
+++ b/Engine/Source/Programs/UnrealBuildTool/Preprocessor/Identifier.cs
@@ -0,0 +1,135 @@
+// Copyright Epic Games, Inc. All Rights Reserved.
+
+using System;
+using System.Collections.Concurrent;
+using System.Collections.Generic;
+using System.Linq;
+using System.Text;
+using System.Threading.Tasks;
+using Tools.DotNETCommon;
+
+namespace UnrealBuildTool
+{
+	/// <summary>
+	/// Contains the name of an identifier. Individual objects are unique.
+	/// </summary>
+	class Identifier : IComparable<Identifier>
+	{
+		/// <summary>
+		/// The name of this identifier
+		/// </summary>
+		string Name;
+
+		/// <summary>
+		/// Global map of name to identifier instance
+		/// </summary>
+		static ConcurrentDictionary<string, Identifier> NameToIdentifier = new ConcurrentDictionary<string, Identifier>(StringComparer.Ordinal);
+
+		/// <summary>
+		/// Constructor
+		/// </summary>
+		/// <param name="Name">Name of this identifier</param>
+		private Identifier(string Name)
+		{
+			this.Name = Name;
+		}
+
+		/// <summary>
+		/// Finds or adds an identifer with the given name
+		/// </summary>
+		/// <param name="Name">Name of the identifier</param>
+		/// <returns>New Identifier instance</returns>
+		public static Identifier FindOrAdd(string Name)
+		{
+			Identifier Result;
+			if(!NameToIdentifier.TryGetValue(Name, out Result))
+			{
+				Identifier NewIdentifier = new Identifier(Name);
+				if(NameToIdentifier.TryAdd(Name, NewIdentifier))
+				{
+					Result = NewIdentifier;
+				}
+				else
+				{
+					Result = NameToIdentifier[Name];
+				}
+			}
+			return Result;
+		}
+
+		/// <summary>
+		/// Compares this identifier to another identifier
+		/// </summary>
+		/// <param name="Other">Identifier to compare to</param>
+		/// <returns>Value indicating which identifier should sort first</returns>
+		public int CompareTo(Identifier Other)
+		{
+			return Name.CompareTo(Other.Name);
+		}
+
+		/// <summary>
+		/// Formats this identifer as a string for debugging
+		/// </summary>
+		/// <returns>Name of this identifier</returns>
+		public override string ToString()
+		{
+			return Name;
+		}
+	}
+
+	/// <summary>
+	/// Well known predefined identifiers
+	/// </summary>
+	class Identifiers
+	{
+		public static readonly Identifier Include = Identifier.FindOrAdd("include");
+		public static readonly Identifier Define = Identifier.FindOrAdd("define");
+		public static readonly Identifier Undef = Identifier.FindOrAdd("undef");
+		public static readonly Identifier If = Identifier.FindOrAdd("if");
+		public static readonly Identifier Ifdef = Identifier.FindOrAdd("ifdef");
+		public static readonly Identifier Ifndef = Identifier.FindOrAdd("ifndef");
+		public static readonly Identifier Elif = Identifier.FindOrAdd("elif");
+		public static readonly Identifier Else = Identifier.FindOrAdd("else");
+		public static readonly Identifier Endif = Identifier.FindOrAdd("endif");
+		public static readonly Identifier Defined = Identifier.FindOrAdd("defined");
+		public static readonly Identifier Pragma = Identifier.FindOrAdd("pragma");
+		public static readonly Identifier Once = Identifier.FindOrAdd("once");
+		public static readonly Identifier Error = Identifier.FindOrAdd("error");
+		public static readonly Identifier Warning = Identifier.FindOrAdd("warning");
+		public static readonly Identifier __VA_ARGS__ = Identifier.FindOrAdd("__VA_ARGS__");
+		public static readonly Identifier __FILE__ = Identifier.FindOrAdd("__FILE__");
+		public static readonly Identifier __LINE__ = Identifier.FindOrAdd("__LINE__");
+		public static readonly Identifier __COUNTER__ = Identifier.FindOrAdd("__COUNTER__");
+		public static readonly Identifier Sizeof = Identifier.FindOrAdd("sizeof");
+		public static readonly Identifier Alignof = Identifier.FindOrAdd("alignof");
+		public static readonly Identifier __has_feature = Identifier.FindOrAdd("__has_feature");
+		public static readonly Identifier __building_module = Identifier.FindOrAdd("__building_module");
+		public static readonly Identifier __pragma = Identifier.FindOrAdd("__pragma");
+	}
+
+	/// <summary>
+	/// Helper functions for serialization
+	/// </summary>
+	static class IdentifierExtensionMethods
+	{
+		/// <summary>
+		/// Read an identifier from a binary archive
+		/// </summary>
+		/// <param name="Reader">Reader to serialize data from</param>
+		/// <returns>Instance of the serialized identifier</returns>
+		public static Identifier ReadIdentifier(this BinaryArchiveReader Reader)
+		{
+			return Reader.ReadObjectReference<Identifier>(() => Identifier.FindOrAdd(Reader.ReadString()));
+		}
+
+		/// <summary>
+		/// Write an identifier to a binary archive
+		/// </summary>
+		/// <param name="Writer">Writer to serialize data to</param>
+		/// <param name="Identifier">Identifier to write</param>
+		public static void WriteIdentifier(this BinaryArchiveWriter Writer, Identifier Identifier)
+		{
+			Writer.WriteObjectReference<Identifier>(Identifier, () => Writer.WriteString(Identifier.ToString()));
+		}
+	}
+}
diff --git a/Engine/Source/Programs/UnrealBuildTool/Preprocessor/Preprocessor.cs b/Engine/Source/Programs/UnrealBuildTool/Preprocessor/Preprocessor.cs
new file mode 100644
--- /dev/null
+++ b/Engine/Source/Programs/UnrealBuildTool/Preprocessor/Preprocessor.cs
@@ -0,0 +1,1397 @@
+// Copyright Epic Games, Inc. All Rights Reserved.
+
+using System;
+using System.Collections.Generic;
+using System.Linq;
+using System.Text;
+using System.Threading.Tasks;
+using System.Diagnostics;
+using System.IO;
+using System.Collections;
+using System.Collections.Concurrent;
+using Tools.DotNETCommon;
+
+namespace UnrealBuildTool
+{
+	/// <summary>
+	/// Exception class for the preprocessor, which contains the file and position of the code causing an error
+	/// </summary>
+	class PreprocessorException : Exception
+	{
+		/// <summary>
+		/// The context when the error was encountered
+		/// </summary>
+		PreprocessorContext Context;
+
+		/// <summary>
+		/// Constructor
+		/// </summary>
+		/// <param name="Context">The current preprocesor context</param>
+		/// <param name="Format">Format string, to be passed to String.Format</param>
+		/// <param name="Args">Optional argument list for the format string</param>
+		public PreprocessorException(PreprocessorContext Context, string Format, params object[] Args)
+			: base(String.Format(Format, Args))
+		{
+			this.Context = Context;
+		}
+	}
+
+	/// <summary>
+	/// Implementation of a C++ preprocessor.
+	/// </summary>
+	class Preprocessor
+	{
+		/// <summary>
+		/// Type of an include path
+		/// </summary>
+		public enum IncludePathType
+		{
+			/// <summary>
+			/// Regular include path, enclosed by quotes
+			/// </summary>
+			Normal,
+
+			/// <summary>
+			/// System include path, enclosed by angle brackets
+			/// </summary>
+			System,
+		}
+
+		/// <summary>
+		/// Include paths to look in
+		/// </summary>
+		List<DirectoryItem> IncludeDirectories = new List<DirectoryItem>();
+
+		/// <summary>
+		/// Set of all included files with the #pragma once directive
+		/// </summary>
+		HashSet<FileItem> PragmaOnceFiles = new HashSet<FileItem>();
+
+		/// <summary>
+		/// The current state of the preprocessor
+		/// </summary>
+		PreprocessorState State = new PreprocessorState();
+
+		/// <summary>
+		/// Predefined token containing the constant "0"
+		/// </summary>
+		static readonly byte[] ZeroLiteral = Encoding.UTF8.GetBytes("0");
+
+		/// <summary>
+		/// Predefined token containing the constant "1"
+		/// </summary>
+		static readonly byte[] OneLiteral = Encoding.UTF8.GetBytes("1");
+
+		/// <summary>
+		/// Value of the __COUNTER__ variable
+		/// </summary>
+		int Counter;
+
+		/// <summary>
+		/// Default constructor
+		/// </summary>
+		public Preprocessor()
+		{
+			DateTime Now = DateTime.Now;
+			AddLiteralMacro("__DATE__", TokenType.String, String.Format("\"{0} {1,2} {2}\"", Now.ToString("MMM"), Now.Day, Now.Year));
+			AddLiteralMacro("__TIME__", TokenType.String, "\"" + Now.ToString("HH:mm:ss") + "\"");
+			AddLiteralMacro("__FILE__", TokenType.String, "\"<unknown>\"");
+			AddLiteralMacro("__LINE__", TokenType.Number, "-1");
+			AddLiteralMacro("__COUNTER__", TokenType.Number, "-1");
+			AddLiteralMacro("CHAR_BIT", TokenType.Number, "8"); // Workaround for #include_next not being supported on Linux for limit.h
+		}
+
+		/// <summary>
+		/// Determines whether the current preprocessor branch is active
+		/// </summary>
+		/// <returns>True if the current branch is active</returns>
+		public bool IsCurrentBranchActive()
+		{
+			return State.IsCurrentBranchActive();
+		}
+
+		/// <summary>
+		/// Defines a macro. May have an optional '=Value' suffix.
+		/// </summary>
+		/// <param name="Definition">Macro to define</param>
+		public void AddDefinition(string Definition)
+		{
+			List<Token> Tokens = TokenReader.Tokenize(Definition);
+			if(Tokens.Count == 0)
+			{
+				throw new PreprocessorException(null, "Missing macro name");
+			}
+			if(Tokens[0].Type != TokenType.Identifier)
+			{
+				throw new PreprocessorException(null, "'{0}' is not a valid macro name", Tokens[0].ToString());
+			}
+
+			List<Token> ValueTokens = new List<Token>();
+			if(Tokens.Count == 1)
+			{
+				ValueTokens.Add(new Token(TokenType.Number, TokenFlags.None, OneLiteral));
+			}
+			else if(Tokens[1].Type != TokenType.Equals)
+			{
+				throw new PreprocessorException(null, "Unable to parse macro definition '{0}'", Definition);
+			}
+			else
+			{
+				ValueTokens.AddRange(Tokens.Skip(2));
+			}
+
+			PreprocessorMacro Macro = new PreprocessorMacro(Tokens[0].Identifier, null, Tokens);
+			State.DefineMacro(Macro);
+		}
+
+		/// <summary>
+		/// Defines a macro
+		/// </summary>
+		/// <param name="Name">Name of the macro</param>
+		/// <param name="Value">String to be parsed for the macro's value</param>
+		public void AddDefinition(string Name, string Value)
+		{
+			List<Token> Tokens = new List<Token>();
+
+			TokenReader Reader = new TokenReader(Value);
+			while(Reader.MoveNext())
+			{
+				Tokens.Add(Reader.Current);
+			}
+
+			PreprocessorMacro Macro = new PreprocessorMacro(Identifier.FindOrAdd(Name), null, Tokens);
+			State.DefineMacro(Macro);
+		}
+
+		/// <summary>
+		/// Defines a macro
+		/// </summary>
+		/// <param name="Macro">The macro definition</param>
+		public void AddDefinition(PreprocessorMacro Macro)
+		{
+			State.DefineMacro(Macro);
+		}
+
+		/// <summary>
+		/// Adds an include path to the preprocessor
+		/// </summary>
+		/// <param name="Directory">The include path</param>
+		public void AddIncludePath(DirectoryItem Directory)
+		{
+			if(!IncludeDirectories.Contains(Directory))
+			{
+				IncludeDirectories.Add(Directory);
+			}
+		}
+
+		/// <summary>
+		/// Adds an include path to the preprocessor
+		/// </summary>
+		/// <param name="Location">The include path</param>
+		public void AddIncludePath(DirectoryReference Location)
+		{
+			DirectoryItem Directory = DirectoryItem.GetItemByDirectoryReference(Location);
+			if(!Directory.Exists)
+			{
+				throw new FileNotFoundException("Unable to find " + Location.FullName);
+			}
+			AddIncludePath(Directory);
+		}
+
+		/// <summary>
+		/// Adds an include path to the preprocessor
+		/// </summary>
+		/// <param name="DirectoryName">The include path</param>
+		public void AddIncludePath(string DirectoryName)
+		{
+			AddIncludePath(new DirectoryReference(DirectoryName));
+		}
+
+		/// <summary>
+		/// Try to resolve an quoted include against the list of include directories. Uses search order described by https://msdn.microsoft.com/en-us/library/36k2cdd4.aspx.
+		/// </summary>
+		/// <param name="Context">The current preprocessor context</param>
+		/// <param name="IncludePath">The path appearing in an #include directive</param>
+		/// <param name="Type">Specifies rules for how to resolve the include path (normal/system)</param>
+		/// <param name="File">If found, receives the resolved file</param>
+		/// <returns>True if the The resolved file</returns>
+		public bool TryResolveIncludePath(PreprocessorContext Context, string IncludePath, IncludePathType Type, out FileItem File)
+		{
+			// From MSDN (https://msdn.microsoft.com/en-us/library/36k2cdd4.aspx?f=255&MSPPError=-2147217396)
+			//
+			// The preprocessor searches for include files in this order:
+			//
+			// Quoted form:
+			//   1) In the same directory as the file that contains the #include statement.
+			//   2) In the directories of the currently opened include files, in the reverse order in which they were opened. 
+			//      The search begins in the directory of the parent include file and continues upward through the directories of any grandparent include files.
+			//   3) Along the path that's specified by each /I compiler option.
+			//   4) Along the paths that are specified by the INCLUDE environment variable.
+			//
+			// Angle-bracket form:
+			//   1) Along the path that's specified by each /I compiler option.
+			//   2) Along the paths that are specified by the INCLUDE environment variable.
+
+			// If it's an absolute path, return it immediately
+			if(Path.IsPathRooted(IncludePath))
+			{
+				FileItem FileItem = FileItem.GetItemByPath(IncludePath);
+				if(FileItem.Exists)
+				{
+					File = FileItem;
+					return true;
+				}
+				else
+				{
+					File = null;
+					return false;
+				}
+			}
+
+			// Split the path into fragments
+			string[] Fragments = IncludePath.Split(Path.DirectorySeparatorChar, Path.AltDirectorySeparatorChar);
+
+			// Try to match the include path against any of the included directories
+			if(Type == IncludePathType.Normal)
+			{
+				for(PreprocessorContext OuterContext = Context; OuterContext != null; OuterContext = OuterContext.Outer)
+				{
+					PreprocessorFileContext OuterFileContext = OuterContext as PreprocessorFileContext;
+					if(OuterFileContext != null)
+					{
+						FileItem ResolvedFile;
+						if(TryResolveRelativeIncludePath(OuterFileContext.Directory, Fragments, out ResolvedFile))
+						{
+							File = ResolvedFile;
+							return true;
+						}
+					}
+				}
+			}
+
+			// Try to match the include path against any of the system directories
+			foreach(DirectoryItem BaseDirectory in IncludeDirectories)
+			{
+				FileItem ResolvedFile;
+				if(TryResolveRelativeIncludePath(BaseDirectory, Fragments, out ResolvedFile))
+				{
+					File = ResolvedFile;
+					return true;
+				}
+			}
+
+			// Failed to find the file
+			File = null;
+			return false;
+		}
+
+		/// <summary>
+		/// Try to resolve an quoted include against the list of include directories. Uses search order described by https://msdn.microsoft.com/en-us/library/36k2cdd4.aspx.
+		/// </summary>
+		/// <param name="BaseDirectory">The base directory to search from</param>
+		/// <param name="Fragments">Fragments of the relative path to follow</param>
+		/// <param name="File">The file that was found, if successful</param>
+		/// <returns>True if the The resolved file</returns>
+		public bool TryResolveRelativeIncludePath(DirectoryItem BaseDirectory, string[] Fragments, out FileItem File)
+		{
+			DirectoryItem Directory = BaseDirectory;
+			for(int Idx = 0; Idx < Fragments.Length - 1; Idx++)
+			{
+				if(!Directory.TryGetDirectory(Fragments[Idx], out Directory))
+				{
+					File = null;
+					return false;
+				}
+			}
+			return Directory.TryGetFile(Fragments[Fragments.Length - 1], out File);
+		}
+
+		/// <summary>
+		/// Parses a file recursively 
+		/// </summary>
+		/// <param name="File">File to parse</param>
+		/// <param name="Fragments">Lists of fragments that are parsed</param>
+		/// <param name="OuterContext">Outer context information, for error messages</param>
+		/// <param name="SourceFileCache">Cache for source files</param>
+		/// <param name="bShowIncludes">Show all the included files, in order</param>
+		public void ParseFile(FileItem File, List<SourceFileFragment> Fragments, PreprocessorContext OuterContext, SourceFileMetadataCache SourceFileCache, bool bShowIncludes)
+		{
+			// If the file has already been included and had a #pragma once directive, don't include it again
+			if(PragmaOnceFiles.Contains(File))
+			{
+				return;
+			}
+
+			// Output a trace of the included files
+			if(bShowIncludes)
+			{
+				Log.TraceInformation("Note: including file: {0}", File.Location);
+			}
+
+			// If the file had a header guard, and the macro is still defined, don't include it again
+			SourceFile SourceFile = SourceFileCache.GetSourceFile(File);
+			if(SourceFile.HeaderGuardMacro != null && State.IsMacroDefined(SourceFile.HeaderGuardMacro))
+			{
+				return;
+			}
+
+			// Create a context for this file
+			PreprocessorFileContext Context = new PreprocessorFileContext(SourceFile, OuterContext);
+
+			// Parse the markup for this file
+			while(Context.MarkupIdx < SourceFile.Markup.Length)
+			{
+				SourceFileMarkup Markup = SourceFile.Markup[Context.MarkupIdx];
+				if(Markup.Type == SourceFileMarkupType.Include)
+				{
+					if(State.IsCurrentBranchActive())
+					{
+						// Parse the directive
+						FileItem IncludedFile = ParseIncludeDirective(Markup, Context);
+
+						// Parse the included file
+						ParseFile(IncludedFile, Fragments, Context, SourceFileCache, bShowIncludes);
+					}
+					Context.MarkupIdx++;
+				}
+				else
+				{
+					// Get the next fragment
+					SourceFileFragment Fragment = SourceFile.Fragments[Context.FragmentIdx];
+					Debug.Assert(Fragment.MarkupMin == Context.MarkupIdx);
+
+					// Parse this fragment
+					ParseFragment(SourceFile, Fragment, Context);
+
+					// Add this fragment to the list
+					Fragments.Add(Fragment);
+					Context.FragmentIdx++;
+				}
+			}
+		}
+
+		/// <summary>
+		/// Parse an include directive and resolve the file it references
+		/// </summary>
+		/// <param name="Markup">Markup for the include directive</param>
+		/// <param name="Context">Current preprocessor context</param>
+		/// <returns>Included file</returns>
+		FileItem ParseIncludeDirective(SourceFileMarkup Markup, PreprocessorContext Context)
+		{
+			// Expand macros in the given tokens
+			List<Token> ExpandedTokens = new List<Token>();
+			ExpandMacros(Markup.Tokens, ExpandedTokens, false, Context);
+						
+			// Convert the string to a single token
+			string IncludeToken = Token.Format(ExpandedTokens);
+
+			// Expand any macros in them and resolve it
+			IncludePathType Type;
+			if(IncludeToken.Length >= 2 && IncludeToken[0] == '\"' && IncludeToken[IncludeToken.Length - 1] == '\"')
+			{
+				Type = IncludePathType.Normal;
+			}
+			else if(IncludeToken.Length >= 2 && IncludeToken[0] == '<' && IncludeToken[IncludeToken.Length - 1] == '>')
+			{
+				Type = IncludePathType.System;
+			}
+			else
+			{
+				throw new PreprocessorException(Context, "Couldn't resolve include '{0}'", IncludeToken);
+			}
+
+			// Get the include path
+			string IncludePath = IncludeToken.Substring(1, IncludeToken.Length - 2);
+
+			// Resolve the included file
+			FileItem IncludedFile;
+			if(!TryResolveIncludePath(Context, IncludePath, Type, out IncludedFile))
+			{
+				throw new PreprocessorException(Context, "Couldn't resolve include '{0}'", IncludePath);
+			}
+			return IncludedFile;
+		}
+
+		/// <summary>
+		/// Parse a source file fragment, using cached transforms if possible
+		/// </summary>
+		/// <param name="SourceFile">The source file being parsed</param>
+		/// <param name="Fragment">Fragment to parse</param>
+		/// <param name="Context">Current preprocessor context</param>
+		void ParseFragment(SourceFile SourceFile, SourceFileFragment Fragment, PreprocessorFileContext Context)
+		{
+			// Check if there's a valid transform that matches the current state
+			int TransformIdx = 0;
+			for(;;)
+			{
+				PreprocessorTransform[] Transforms;
+				lock(Fragment)
+				{
+					Transforms = Fragment.Transforms;
+					if(TransformIdx == Transforms.Length)
+					{
+						// Attach a new transform to the current state
+						PreprocessorTransform Transform = State.BeginCapture();
+						for(; Context.MarkupIdx < Fragment.MarkupMax; Context.MarkupIdx++)
+						{
+							SourceFileMarkup Markup = SourceFile.Markup[Context.MarkupIdx];
+							ParseMarkup(Markup.Type, Markup.Tokens, Context);
+						}
+						Transform = State.EndCapture();
+
+						// Add it to the fragment for future fragments
+						PreprocessorTransform[] NewTransforms = new PreprocessorTransform[Fragment.Transforms.Length + 1];
+						for(int Idx = 0; Idx < Transforms.Length; Idx++)
+						{
+							NewTransforms[Idx] = Transforms[Idx];
+						}
+						NewTransforms[Transforms.Length] = Transform;
+
+						// Assign it to the fragment
+						Fragment.Transforms = NewTransforms;
+						return;
+					}
+				}
+
+				for(; TransformIdx < Transforms.Length; TransformIdx++)
+				{
+					PreprocessorTransform Transform = Transforms[TransformIdx];
+					if(State.TryToApply(Transform))
+					{
+						// Update the pragma once state
+						if(Transform.bHasPragmaOnce)
+						{
+							PragmaOnceFiles.Add(SourceFile.File);
+						}
+
+						// Move to the end of the fragment
+						Context.MarkupIdx = Fragment.MarkupMax;
+						return;
+					}
+				}
+			}
+		}
+
+		/// <summary>
+		/// Validate and add a macro using the given parameter and token list
+		/// </summary>
+		/// <param name="Context">The current preprocessor context</param>
+		/// <param name="Name">Name of the macro</param>
+		/// <param name="Parameters">Parameter list for the macro</param>
+		/// <param name="Tokens">List of tokens</param>
+		void AddMacro(PreprocessorContext Context, Identifier Name, List<Identifier> Parameters, List<Token> Tokens)
+		{
+			if(Tokens.Count == 0)
+			{
+				Tokens.Add(new Token(TokenType.Placemarker, TokenFlags.None));
+			}
+			else
+			{
+				if(Tokens[0].HasLeadingSpace)
+				{
+					Tokens[0] = Tokens[0].RemoveFlags(TokenFlags.HasLeadingSpace);
+				}
+				if (Tokens[0].Type == TokenType.HashHash || Tokens[Tokens.Count - 1].Type == TokenType.HashHash)
+				{
+					throw new PreprocessorException(Context, "Invalid use of concatenation at start or end of token sequence");
+				}
+				if (Parameters == null || Parameters.Count == 0 || Parameters[Parameters.Count - 1] != Identifiers.__VA_ARGS__)
+				{
+					if(Tokens.Any(x => x.Identifier == Identifiers.__VA_ARGS__))
+					{
+						throw new PreprocessorException(Context, "Invalid reference to {0}", Identifiers.__VA_ARGS__);
+					}
+				}
+			}
+			State.DefineMacro(new PreprocessorMacro(Name, Parameters, Tokens));
+		}
+		
+		/// <summary>
+		/// Set a predefined macro to a given value
+		/// </summary>
+		/// <param name="Name">Name of the macro</param>
+		/// <param name="Type">Type of the macro token</param>
+		/// <param name="Value">Value of the macro</param>
+		/// <returns>The created macro</returns>
+		void AddLiteralMacro(string Name, TokenType Type, string Value)
+		{
+			Token Token = new Token(Type, TokenFlags.None, Value);
+			PreprocessorMacro Macro = new PreprocessorMacro(Identifier.FindOrAdd(Name), null, new List<Token> { Token });
+			State.DefineMacro(Macro);
+		}
+
+		/// <summary>
+		/// Parse a marked up directive from a file
+		/// </summary>
+		/// <param name="Type">The markup type</param>
+		/// <param name="Tokens">Tokens for the directive</param>
+		/// <param name="Context">The context that this markup is being parsed in</param>
+		public void ParseMarkup(SourceFileMarkupType Type, List<Token> Tokens, PreprocessorContext Context)
+		{
+			switch(Type)
+			{
+				case SourceFileMarkupType.Include:
+					throw new PreprocessorException(Context, "Include directives should be handled by the caller.");
+				case SourceFileMarkupType.Define:
+					ParseDefineDirective(Tokens, Context);
+					break;
+				case SourceFileMarkupType.Undef:
+					ParseUndefDirective(Tokens, Context);
+					break;
+				case SourceFileMarkupType.If:
+					ParseIfDirective(Tokens, Context);
+					break;
+				case SourceFileMarkupType.Ifdef:
+					ParseIfdefDirective(Tokens, Context);
+					break;
+				case SourceFileMarkupType.Ifndef:
+					ParseIfndefDirective(Tokens, Context);
+					break;
+				case SourceFileMarkupType.Elif:
+					ParseElifDirective(Tokens, Context);
+					break;
+				case SourceFileMarkupType.Else:
+					ParseElseDirective(Tokens, Context);
+					break;
+				case SourceFileMarkupType.Endif:
+					ParseEndifDirective(Tokens, Context);
+					break;
+				case SourceFileMarkupType.Pragma:
+					ParsePragmaDirective(Tokens, Context);
+					break;
+			}
+		}
+
+		/// <summary>
+		/// Read a macro definition
+		/// </summary>
+		/// <param name="Tokens">List of tokens in the directive</param>
+		/// <param name="Context">The context that this directive is being parsed in</param>
+		public void ParseDefineDirective(List<Token> Tokens, PreprocessorContext Context)
+		{
+			if(State.IsCurrentBranchActive())
+			{
+				// Check there's a name token
+				if(Tokens.Count < 1 || Tokens[0].Type != TokenType.Identifier || Tokens[0].Identifier == Identifiers.Defined)
+				{
+					throw new PreprocessorException(Context, "Invalid macro name");
+				}
+
+				// Read the macro name
+				Identifier Name = Tokens[0].Identifier;
+				int TokenIdx = 1;
+
+				// Read the macro parameter list, if there is one
+				List<Identifier> Parameters = null;
+				if (TokenIdx < Tokens.Count && !Tokens[TokenIdx].HasLeadingSpace && Tokens[TokenIdx].Type == TokenType.LeftParen)
+				{
+					Parameters = new List<Identifier>();
+					if(++TokenIdx == Tokens.Count)
+					{
+						throw new PreprocessorException(Context, "Unexpected end of macro parameter list");
+					}
+					if(Tokens[TokenIdx].Type != TokenType.RightParen)
+					{
+						for(;;TokenIdx++)
+						{
+							// Check there's enough tokens left for a parameter name, plus ',' or ')'
+							if(TokenIdx + 2 > Tokens.Count)
+							{
+								throw new PreprocessorException(Context, "Unexpected end of macro parameter list");
+							}
+
+							// Check it's a valid name, and add it to the list
+							Token NameToken = Tokens[TokenIdx++];
+							if(NameToken.Type == TokenType.Ellipsis)
+							{
+								if(Tokens[TokenIdx].Type != TokenType.RightParen)
+								{
+									throw new PreprocessorException(Context, "Variadic macro arguments must be last in list");
+								}
+								else
+								{
+									NameToken = new Token(Identifiers.__VA_ARGS__, NameToken.Flags & TokenFlags.HasLeadingSpace);
+								}
+							}
+							else
+							{
+								if (NameToken.Type != TokenType.Identifier || NameToken.Identifier == Identifiers.__VA_ARGS__)
+								{
+									throw new PreprocessorException(Context, "Invalid preprocessor token: {0}", NameToken);
+								}
+								if (Parameters.Contains(NameToken.Identifier))
+								{
+									throw new PreprocessorException(Context, "'{0}' has already been used as an argument name", NameToken);
+								}
+							}
+							Parameters.Add(NameToken.Identifier);
+
+							// Read the separator
+							Token SeparatorToken = Tokens[TokenIdx];
+							if (SeparatorToken.Type == TokenType.RightParen)
+							{
+								break;
+							}
+							if (SeparatorToken.Type != TokenType.Comma)
+							{
+								throw new PreprocessorException(Context, "Expected ',' or ')'");
+							}
+						}
+					}
+					TokenIdx++;
+				}
+
+				// Read the macro tokens
+				AddMacro(Context, Name, Parameters, Tokens.Skip(TokenIdx).ToList());
+			}
+		}
+
+		/// <summary>
+		/// Parse an #undef directive
+		/// </summary>
+		/// <param name="Tokens">List of tokens in the directive</param>
+		/// <param name="Context">The context that this directive is being parsed in</param>
+		public void ParseUndefDirective(List<Token> Tokens, PreprocessorContext Context)
+		{
+			if(State.IsCurrentBranchActive())
+			{
+				// Check there's a name token
+				if (Tokens.Count != 1)
+				{
+					throw new PreprocessorException(Context, "Expected a single token after #undef");
+				}
+				if (Tokens[0].Type != TokenType.Identifier)
+				{
+					throw new PreprocessorException(Context, "Invalid macro name '{0}'", Tokens[0]);
+				}
+
+				// Remove the macro from the list of definitions
+				State.UndefMacro(Tokens[0].Identifier);
+			}
+		}
+
+		/// <summary>
+		/// Parse an #if directive
+		/// </summary>
+		/// <param name="Tokens">List of tokens in the directive</param>
+		/// <param name="Context">The context that this directive is being parsed in</param>
+		public void ParseIfDirective(List<Token> Tokens, PreprocessorContext Context)
+		{
+			PreprocessorBranch Branch = PreprocessorBranch.HasIfDirective;
+			if (State.IsCurrentBranchActive())
+			{
+				// Read a line into the buffer and expand the macros in it
+				List<Token> ExpandedTokens = new List<Token>();
+				ExpandMacros(Tokens, ExpandedTokens, true, Context);
+
+				// Evaluate the condition
+				long Result = PreprocessorExpression.Evaluate(Context, ExpandedTokens);
+				if (Result != 0)
+				{
+					Branch |= PreprocessorBranch.Active | PreprocessorBranch.Taken;
+				}
+			}
+			State.PushBranch(Branch);
+		}
+
+		/// <summary>
+		/// Parse an #ifdef directive
+		/// </summary>
+		/// <param name="Tokens">List of tokens in the directive</param>
+		/// <param name="Context">The context that this directive is being parsed in</param>
+		public void ParseIfdefDirective(List<Token> Tokens, PreprocessorContext Context)
+		{
+			PreprocessorBranch Branch = PreprocessorBranch.HasIfdefDirective;
+			if (State.IsCurrentBranchActive())
+			{
+				// Make sure there's only one token
+				if(Tokens.Count != 1 || Tokens[0].Type != TokenType.Identifier)
+				{
+					throw new PreprocessorException(Context, "Missing or invalid macro name for #ifdef directive");
+				}
+
+				// Check if the macro is defined
+				if(State.IsMacroDefined(Tokens[0].Identifier))
+				{
+					Branch |= PreprocessorBranch.Active | PreprocessorBranch.Taken;
+				}
+			}
+			State.PushBranch(Branch);
+		}
+
+		/// <summary>
+		/// Parse an #ifndef directive
+		/// </summary>
+		/// <param name="Tokens">List of tokens for this directive</param>
+		/// <param name="Context">The context that this directive is being parsed in</param>
+		public void ParseIfndefDirective(List<Token> Tokens, PreprocessorContext Context)
+		{
+			PreprocessorBranch Branch = PreprocessorBranch.HasIfndefDirective;
+			if (State.IsCurrentBranchActive())
+			{
+				// Make sure there's only one token
+				if (Tokens.Count != 1 || Tokens[0].Type != TokenType.Identifier)
+				{
+					throw new PreprocessorException(Context, "Missing or invalid macro name for #ifndef directive");
+				}
+
+				// Check if the macro is defined
+				if(!State.IsMacroDefined(Tokens[0].Identifier))
+				{
+					Branch |= PreprocessorBranch.Active | PreprocessorBranch.Taken;
+				}
+			}
+			State.PushBranch(Branch);
+		}
+
+		/// <summary>
+		/// Parse an #elif directive
+		/// </summary>
+		/// <param name="Tokens">List of tokens for this directive</param>
+		/// <param name="Context">The context that this directive is being parsed in</param>
+		public void ParseElifDirective(List<Token> Tokens, PreprocessorContext Context)
+		{
+			// Check we're in a branch, and haven't already read an #else directive
+			PreprocessorBranch Branch;
+			if (!State.TryPopBranch(out Branch))
+			{
+				throw new PreprocessorException(Context, "#elif directive outside conditional block");
+			}
+			if (Branch.HasFlag(PreprocessorBranch.Complete))
+			{
+				throw new PreprocessorException(Context, "#elif directive cannot appear after #else");
+			}
+
+			// Pop the current branch state at this depth, so we can test against whether the parent state is enabled
+			Branch = (Branch | PreprocessorBranch.HasElifDirective) & ~PreprocessorBranch.Active;
+			if (State.IsCurrentBranchActive())
+			{
+				// Read a line into the buffer and expand the macros in it
+				List<Token> ExpandedTokens = new List<Token>();
+				ExpandMacros(Tokens, ExpandedTokens, true, Context);
+
+				// Check we're at the end of a conditional block
+				if (!Branch.HasFlag(PreprocessorBranch.Taken))
+				{
+					long Result = PreprocessorExpression.Evaluate(Context, ExpandedTokens);
+					if (Result != 0)
+					{
+						Branch |= PreprocessorBranch.Active | PreprocessorBranch.Taken;
+					}
+				}
+			}
+			State.PushBranch(Branch);
+		}
+
+		/// <summary>
+		/// Parse an #else directive
+		/// </summary>
+		/// <param name="Tokens">List of tokens in the directive</param>
+		/// <param name="Context">The context that this directive is being parsed in</param>
+		public void ParseElseDirective(List<Token> Tokens, PreprocessorContext Context)
+		{
+			// Make sure there's nothing else on the line
+			if (Tokens.Count > 0)
+			{
+				throw new PreprocessorException(Context, "Garbage after #else directive");
+			}
+
+			// Check we're in a branch, and haven't already read an #else directive
+			PreprocessorBranch Branch;
+			if (!State.TryPopBranch(out Branch))
+			{
+				throw new PreprocessorException(Context, "#else directive without matching #if directive");
+			}
+			if ((Branch & PreprocessorBranch.Complete) != 0)
+			{
+				throw new PreprocessorException(Context, "Only one #else directive can appear in a conditional block");
+			}
+
+			// Check whether to take this branch, but only allow activating if the parent state is active.
+			Branch &= ~PreprocessorBranch.Active;
+			if(State.IsCurrentBranchActive() && !Branch.HasFlag(PreprocessorBranch.Taken))
+			{
+				Branch |= PreprocessorBranch.Active | PreprocessorBranch.Taken;
+			}
+			State.PushBranch(Branch | PreprocessorBranch.Complete);
+		}
+
+		/// <summary>
+		/// Parse an #endif directive
+		/// </summary>
+		/// <param name="Tokens">List of tokens in the directive</param>
+		/// <param name="Context">The context that this directive is being parsed in</param>
+		public void ParseEndifDirective(List<Token> Tokens, PreprocessorContext Context)
+		{
+			// Pop the branch off the stack
+			PreprocessorBranch Branch;
+			if(!State.TryPopBranch(out Branch))
+			{
+				throw new PreprocessorException(Context, "#endif directive without matching #if/#ifdef/#ifndef directive");
+			}
+		}
+
+		/// <summary>
+		/// Parse a #pragma directive
+		/// </summary>
+		/// <param name="Tokens">List of tokens in the directive</param>
+		/// <param name="Context">The context that this directive is being parsed in</param>
+		public void ParsePragmaDirective(List<Token> Tokens, PreprocessorContext Context)
+		{
+			if(State.IsCurrentBranchActive())
+			{
+				if(Tokens.Count == 1 && Tokens[0].Identifier == Identifiers.Once)
+				{
+					SourceFile SourceFile = GetCurrentSourceFile(Context);
+					PragmaOnceFiles.Add(SourceFile.File);
+					State.MarkPragmaOnce();
+				}
+			}
+		}
+
+		/// <summary>
+		/// Expand macros in the given sequence.
+		/// </summary>
+		/// <param name="InputTokens">Sequence of input tokens</param>
+		/// <param name="OutputTokens">List to receive the expanded tokens</param>
+		/// <param name="bIsConditional">Whether a conditional expression is being evaluated (and 'defined' expressions are valid)</param>
+		/// <param name="Context">The context that this directive is being parsed in</param>
+		public void ExpandMacros(IEnumerable<Token> InputTokens, List<Token> OutputTokens, bool bIsConditional, PreprocessorContext Context)
+		{
+			List<PreprocessorMacro> IgnoreMacros = new List<PreprocessorMacro>();
+			ExpandMacrosRecursively(InputTokens, OutputTokens, bIsConditional, IgnoreMacros, Context);
+		}
+
+		/// <summary>
+		/// Expand macros in the given sequence, ignoring previously expanded macro names from a list.
+		/// </summary>
+		/// <param name="InputTokens">Sequence of input tokens</param>
+		/// <param name="OutputTokens">List to receive the expanded tokens</param>
+		/// <param name="bIsConditional">Whether a conditional expression is being evaluated (and 'defined' expressions are valid)</param>
+		/// <param name="IgnoreMacros">List of macros to ignore</param>
+		/// <param name="Context">The context that this directive is being parsed in</param>
+		void ExpandMacrosRecursively(IEnumerable<Token> InputTokens, List<Token> OutputTokens, bool bIsConditional, List<PreprocessorMacro> IgnoreMacros, PreprocessorContext Context)
+		{
+			IEnumerator<Token> InputEnumerator = InputTokens.GetEnumerator();
+			if(InputEnumerator.MoveNext())
+			{
+				for(;;)
+				{
+					if(!ReadExpandedToken(InputEnumerator, OutputTokens, bIsConditional, IgnoreMacros, Context))
+					{
+						break;
+					}
+				}
+			}
+		}
+
+		/// <summary>
+		/// Merges an optional leading space flag into the given token (recycling the original token if possible).
+		/// </summary>
+		/// <param name="Token">The token to merge a leading space into</param>
+		/// <param name="bHasLeadingSpace">The leading space flag</param>
+		/// <returns>New token with the leading space flag set, or the existing token</returns>
+		Token MergeLeadingSpace(Token Token, bool bHasLeadingSpace)
+		{
+			Token Result = Token;
+			if(bHasLeadingSpace && !Result.HasLeadingSpace)
+			{
+				Result = Result.AddFlags(TokenFlags.HasLeadingSpace);
+			}
+			return Result;
+		}
+
+		/// <summary>
+		/// Read a token from an enumerator and substitute it if it's a macro or 'defined' expression (reading more tokens as necessary to complete the expression).
+		/// </summary>
+		/// <param name="InputEnumerator">The enumerator of input tokens</param>
+		/// <param name="OutputTokens">List to receive the expanded tokens</param>
+		/// <param name="bIsConditional">Whether a conditional expression is being evaluated (and 'defined' expressions are valid)</param>
+		/// <param name="IgnoreMacros">List of macros to ignore</param>
+		/// <param name="Context">The context that this directive is being parsed in</param>
+		/// <returns>Result from calling the enumerator's MoveNext() method</returns>
+		bool ReadExpandedToken(IEnumerator<Token> InputEnumerator, List<Token> OutputTokens, bool bIsConditional, List<PreprocessorMacro> IgnoreMacros, PreprocessorContext Context)
+		{
+			// Capture the first token, then move to the next
+			OutputTokens.Add(InputEnumerator.Current);
+			bool bMoveNext = InputEnumerator.MoveNext();
+
+			// If it's an identifier, try to resolve it as a macro
+			if (OutputTokens[OutputTokens.Count - 1].Identifier == Identifiers.Defined && bIsConditional)
+			{
+				// Remove the 'defined' keyword
+				OutputTokens.RemoveAt(OutputTokens.Count - 1);
+
+				// Make sure there's another token
+				if (!bMoveNext)
+				{
+					throw new PreprocessorException(Context, "Invalid syntax for 'defined' expression");
+				}
+
+				// Check for the form 'defined identifier'
+				Token NameToken;
+				if(InputEnumerator.Current.Type == TokenType.Identifier)
+				{
+					NameToken = InputEnumerator.Current;
+				}
+				else
+				{
+					// Otherwise assume the form 'defined ( identifier )'
+					if(InputEnumerator.Current.Type != TokenType.LeftParen || !InputEnumerator.MoveNext() || InputEnumerator.Current.Type != TokenType.Identifier)
+					{
+						throw new PreprocessorException(Context, "Invalid syntax for 'defined' expression");
+					}
+					NameToken = InputEnumerator.Current;
+					if(!InputEnumerator.MoveNext() || InputEnumerator.Current.Type != TokenType.RightParen)
+					{
+						throw new PreprocessorException(Context, "Invalid syntax for 'defined' expression");
+					}
+				}
+
+				// Insert a token for whether it's defined or not
+				OutputTokens.Add(new Token(TokenType.Number, TokenFlags.None, State.IsMacroDefined(NameToken.Identifier)? OneLiteral : ZeroLiteral));
+				bMoveNext = InputEnumerator.MoveNext();
+			}
+			else
+			{
+				// Repeatedly try to expand the last token into the list
+				while(OutputTokens[OutputTokens.Count - 1].Type == TokenType.Identifier && !OutputTokens[OutputTokens.Count - 1].Flags.HasFlag(TokenFlags.DisableExpansion))
+				{
+					// Try to get a macro for the current token
+					PreprocessorMacro Macro;
+					if (!State.TryGetMacro(OutputTokens[OutputTokens.Count - 1].Identifier, out Macro) || IgnoreMacros.Contains(Macro))
+					{
+						break;
+					}
+					if(Macro.IsFunctionMacro && (!bMoveNext || InputEnumerator.Current.Type != TokenType.LeftParen))
+					{
+						break;
+					}
+
+					// Remove the macro name from the output list
+					bool bHasLeadingSpace = OutputTokens[OutputTokens.Count - 1].HasLeadingSpace;
+					OutputTokens.RemoveAt(OutputTokens.Count - 1);
+
+					// Save the initial number of tokens in the list, so we can tell if it expanded
+					int NumTokens = OutputTokens.Count;
+
+					// If it's an object macro, expand it immediately into the output buffer
+					if(Macro.IsObjectMacro)
+					{
+						// Expand the macro tokens into the output buffer
+						ExpandObjectMacro(Macro, OutputTokens, bIsConditional, IgnoreMacros, Context);
+					}
+					else
+					{
+						// Read balanced token for argument list
+						List<Token> ArgumentTokens = new List<Token>();
+						bMoveNext = ReadBalancedToken(InputEnumerator, ArgumentTokens, Context);
+
+						// Expand the macro tokens into the output buffer
+						ExpandFunctionMacro(Macro, ArgumentTokens, OutputTokens, bIsConditional, IgnoreMacros, Context);
+					}
+
+					// If the macro expanded to nothing, quit
+					if(OutputTokens.Count <= NumTokens)
+					{
+						break;
+					}
+
+					// Make sure the space is propagated to the expanded macro
+					OutputTokens[NumTokens] = MergeLeadingSpace(OutputTokens[NumTokens], bHasLeadingSpace);
+
+					// Mark any tokens matching the macro name as not to be expanded again. This can happen with recursive object macros, eg. #define DWORD ::DWORD
+					for(int Idx = NumTokens; Idx < OutputTokens.Count; Idx++)
+					{
+						if(OutputTokens[Idx].Type == TokenType.Identifier && OutputTokens[Idx].Identifier == Macro.Name)
+						{
+							OutputTokens[Idx] = OutputTokens[Idx].AddFlags(TokenFlags.DisableExpansion);
+						}
+					}
+				}
+			}
+            return bMoveNext;
+		}
+
+		/// <summary>
+		/// Gets a string for the __FILE__ macro
+		/// </summary>
+		/// <param name="Context">Context to scan to find the current file</param>
+		/// <returns>String representing the current context</returns>
+		string GetCurrentFileMacroValue(PreprocessorContext Context)
+		{
+			SourceFile SourceFile = GetCurrentSourceFile(Context);
+			if(SourceFile == null)
+			{
+				return "<unknown>";
+			}
+			else
+			{
+				return SourceFile.Location.FullName;
+			}
+		}
+
+		/// <summary>
+		/// Gets a string for the current file
+		/// </summary>
+		/// <param name="Context">Context to scan to find the current file</param>
+		/// <returns>Current source file being parsed</returns>
+		SourceFile GetCurrentSourceFile(PreprocessorContext Context)
+		{
+			SourceFile SourceFile = null;
+			for(PreprocessorContext OuterContext = Context; OuterContext != null; OuterContext = OuterContext.Outer)
+			{
+				PreprocessorFileContext OuterFileContext = OuterContext as PreprocessorFileContext;
+				if(OuterFileContext != null)
+				{
+					SourceFile = OuterFileContext.SourceFile;
+					break;
+				}
+			}
+			return SourceFile;
+		}
+
+		/// <summary>
+		/// Gets the current line number
+		/// </summary>
+		/// <param name="Context">Context to scan to find the current file</param>
+		/// <returns>Line number in the first file encountered</returns>
+		int GetCurrentLine(PreprocessorContext Context)
+		{
+			for(PreprocessorContext OuterContext = Context; OuterContext != null; OuterContext = OuterContext.Outer)
+			{
+				PreprocessorFileContext OuterFileContext = OuterContext as PreprocessorFileContext;
+				if(OuterFileContext != null)
+				{
+					return OuterFileContext.SourceFile.Markup[OuterFileContext.MarkupIdx].LineNumber;
+				}
+			}
+			return 0;
+		}
+
+		/// <summary>
+		/// Expand an object macro
+		/// </summary>
+		/// <param name="Macro">The functional macro</param>
+		/// <param name="OutputTokens">The list to receive the output tokens</param>
+		/// <param name="bIsConditional">Whether the macro is being expanded in a conditional context, allowing use of the 'defined' keyword</param>
+		/// <param name="IgnoreMacros">List of macros currently being expanded, which should be ignored for recursion</param>
+		/// <param name="Context">The context that this directive is being parsed in</param>
+		void ExpandObjectMacro(PreprocessorMacro Macro, List<Token> OutputTokens, bool bIsConditional, List<PreprocessorMacro> IgnoreMacros, PreprocessorContext Context)
+		{
+			// Special handling for the __LINE__ directive, since we need an updated line number for the current token
+			if(Macro.Name == Identifiers.__FILE__)
+			{
+				Token Token = new Token(TokenType.String, TokenFlags.None, String.Format("\"{0}\"", GetCurrentFileMacroValue(Context).Replace("\\", "\\\\")));
+				OutputTokens.Add(Token);
+			}
+			else if(Macro.Name == Identifiers.__LINE__)
+			{
+				Token Token = new Token(TokenType.Number, TokenFlags.None, GetCurrentLine(Context).ToString());
+				OutputTokens.Add(Token);
+			}
+			else if(Macro.Name == Identifiers.__COUNTER__)
+			{
+				Token Token = new Token(TokenType.Number, TokenFlags.None, (Counter++).ToString());
+				OutputTokens.Add(Token);
+			}
+			else
+			{
+				int OutputTokenCount = OutputTokens.Count;
+
+				// Expand all the macros
+				IgnoreMacros.Add(Macro);
+				ExpandMacrosRecursively(Macro.Tokens, OutputTokens, bIsConditional, IgnoreMacros, Context);
+				IgnoreMacros.RemoveAt(IgnoreMacros.Count - 1);
+
+				// Concatenate any adjacent tokens
+				for(int Idx = OutputTokenCount + 1; Idx < OutputTokens.Count - 1; Idx++)
+				{
+					if(OutputTokens[Idx].Type == TokenType.HashHash)
+					{
+						OutputTokens[Idx - 1] = Token.Concatenate(OutputTokens[Idx - 1], OutputTokens[Idx + 1], Context);
+						OutputTokens.RemoveRange(Idx, 2);
+						Idx--;
+					}
+				}
+			}
+		}
+
+		/// <summary>
+		/// Expand a function macro
+		/// </summary>
+		/// <param name="Macro">The functional macro</param>
+		/// <param name="ArgumentListTokens">Identifiers for each argument token</param>
+		/// <param name="OutputTokens">The list to receive the output tokens</param>
+		/// <param name="bIsConditional">Whether the macro is being expanded in a conditional context, allowing use of the 'defined' keyword</param>
+		/// <param name="IgnoreMacros">List of macros currently being expanded, which should be ignored for recursion</param>
+		/// <param name="Context">The context that this macro is being expanded</param>
+		void ExpandFunctionMacro(PreprocessorMacro Macro, List<Token> ArgumentListTokens, List<Token> OutputTokens, bool bIsConditional, List<PreprocessorMacro> IgnoreMacros, PreprocessorContext Context)
+		{
+			// Replace any newlines with spaces, and merge them with the following token
+			for(int Idx = 0; Idx < ArgumentListTokens.Count; Idx++)
+			{
+				if(ArgumentListTokens[Idx].Type == TokenType.Newline)
+				{
+					if(Idx + 1 < ArgumentListTokens.Count)
+					{
+						ArgumentListTokens[Idx + 1] = MergeLeadingSpace(ArgumentListTokens[Idx + 1], true);
+					}
+					ArgumentListTokens.RemoveAt(Idx--);
+				}
+			}
+
+			// Split the arguments out into separate lists
+			List<List<Token>> Arguments = new List<List<Token>>();
+			if(ArgumentListTokens.Count > 2)
+			{
+				for(int Idx = 1;;Idx++)
+				{
+					if (!Macro.HasVariableArgumentList || Arguments.Count < Macro.Parameters.Count)
+					{
+						Arguments.Add(new List<Token>());
+					}
+
+					List<Token> Argument = Arguments[Arguments.Count - 1];
+
+					int InitialIdx = Idx;
+					while (Idx < ArgumentListTokens.Count - 1 && ArgumentListTokens[Idx].Type != TokenType.Comma)
+					{
+						if (!ReadBalancedToken(ArgumentListTokens, ref Idx, Argument))
+						{
+							throw new PreprocessorException(Context, "Invalid argument");
+						}
+					}
+
+					if(Argument.Count > 0 && Arguments[Arguments.Count - 1][0].HasLeadingSpace)
+					{
+						Argument[0] = Argument[0].RemoveFlags(TokenFlags.HasLeadingSpace);
+					}
+
+					bool bHasLeadingSpace = false;
+					for(int TokenIdx = 0; TokenIdx < Argument.Count; TokenIdx++)
+					{
+						if(Argument[TokenIdx].Text.Length == 0)
+						{
+							bHasLeadingSpace |= Argument[TokenIdx].HasLeadingSpace;
+							Argument.RemoveAt(TokenIdx--);
+						}
+						else
+						{
+							Argument[TokenIdx] = MergeLeadingSpace(Argument[TokenIdx], bHasLeadingSpace);
+							bHasLeadingSpace = false;
+						}
+					}
+
+					if (Argument.Count == 0)
+					{
+						Argument.Add(new Token(TokenType.Placemarker, TokenFlags.None));
+						Argument.Add(new Token(TokenType.Placemarker, bHasLeadingSpace? TokenFlags.HasLeadingSpace : TokenFlags.None));
+					}
+
+					if(Idx == ArgumentListTokens.Count - 1)
+					{
+						break;
+					}
+
+					if (ArgumentListTokens[Idx].Type != TokenType.Comma)
+					{
+						throw new PreprocessorException(Context, "Expected ',' between arguments");
+					}
+
+					if (Macro.HasVariableArgumentList && Arguments.Count == Macro.Parameters.Count && Idx < ArgumentListTokens.Count - 1)
+					{
+						Arguments[Arguments.Count - 1].Add(ArgumentListTokens[Idx]);
+					}
+				}
+			}
+
+			// Add an empty variable argument if one was not specified
+			if (Macro.HasVariableArgumentList && Arguments.Count == Macro.Parameters.Count - 1)
+			{
+				Arguments.Add(new List<Token> { new Token(TokenType.Placemarker, TokenFlags.None) });
+			}
+
+			// Validate the argument list
+			if (Arguments.Count != Macro.Parameters.Count)
+			{
+				throw new PreprocessorException(Context, "Incorrect number of arguments to macro");
+			}
+
+			// Expand each one of the arguments
+			List<List<Token>> ExpandedArguments = new List<List<Token>>();
+			for (int Idx = 0; Idx < Arguments.Count; Idx++)
+			{
+				List<Token> NewArguments = new List<Token>();
+				ExpandMacrosRecursively(Arguments[Idx], NewArguments, bIsConditional, IgnoreMacros, Context);
+				ExpandedArguments.Add(NewArguments);
+			}
+
+			// Substitute all the argument tokens
+			List<Token> ExpandedTokens = new List<Token>();
+			for(int Idx = 0; Idx < Macro.Tokens.Count; Idx++)
+			{
+				Token Token = Macro.Tokens[Idx];
+				if(Token.Type == TokenType.Hash && Idx + 1 < Macro.Tokens.Count)
+				{
+					// Stringizing operator
+					int ParamIdx = Macro.FindParameterIndex(Macro.Tokens[++Idx].Identifier);
+					if (ParamIdx == -1)
+					{
+						throw new PreprocessorException(Context, "{0} is not an argument name", Macro.Tokens[Idx].Text);
+					}
+					ExpandedTokens.Add(new Token(TokenType.String, Token.Flags & TokenFlags.HasLeadingSpace, String.Format("\"{0}\"", Token.Format(Arguments[ParamIdx]).Replace("\\", "\\\\").Replace("\"", "\\\""))));
+				}
+				else if(Macro.HasVariableArgumentList && Idx + 2 < Macro.Tokens.Count && Token.Type == TokenType.Comma && Macro.Tokens[Idx + 1].Type == TokenType.HashHash && Macro.Tokens[Idx + 2].Identifier == Identifiers.__VA_ARGS__)
+				{
+					// Special MSVC/GCC extension: ', ## __VA_ARGS__' removes the comma if __VA_ARGS__ is empty. MSVC seems to format the result with a forced space.
+					List<Token> ExpandedArgument = ExpandedArguments[ExpandedArguments.Count - 1];
+					if(ExpandedArgument.Any(x => x.Text.Length > 0))
+					{
+						ExpandedTokens.Add(Token);
+						AppendTokensWithWhitespace(ExpandedTokens, ExpandedArgument, false);
+						Idx += 2;
+					}
+					else
+					{
+						ExpandedTokens.Add(new Token(TokenType.Placemarker, Token.Flags & TokenFlags.HasLeadingSpace));
+						ExpandedTokens.Add(new Token(TokenType.Placemarker, TokenFlags.HasLeadingSpace));
+						Idx += 2;
+					}
+				}
+				else if (Token.Type == TokenType.Identifier)
+				{
+					// Expand a parameter
+					int ParamIdx = Macro.FindParameterIndex(Token.Identifier);
+					if(ParamIdx == -1)
+					{
+						ExpandedTokens.Add(Token);
+					}
+					else if(Idx > 0 && Macro.Tokens[Idx - 1].Type == TokenType.HashHash)
+					{
+						AppendTokensWithWhitespace(ExpandedTokens, Arguments[ParamIdx], Token.HasLeadingSpace);
+					}
+					else if(Idx + 1 < Macro.Tokens.Count && Macro.Tokens[Idx + 1].Type == TokenType.HashHash)
+					{
+						AppendTokensWithWhitespace(ExpandedTokens, Arguments[ParamIdx], Token.HasLeadingSpace);
+					}
+					else
+					{
+						AppendTokensWithWhitespace(ExpandedTokens, ExpandedArguments[ParamIdx], Token.HasLeadingSpace);
+					}
+				}
+				else
+				{
+					ExpandedTokens.Add(Token);
+				}
+			}
+
+			// Concatenate adjacent tokens
+			for (int Idx = 1; Idx < ExpandedTokens.Count - 1; Idx++)
+			{
+				if(ExpandedTokens[Idx].Type == TokenType.HashHash)
+				{
+					Token ConcatenatedToken = Token.Concatenate(ExpandedTokens[Idx - 1], ExpandedTokens[Idx + 1], Context);
+					ExpandedTokens.RemoveRange(Idx, 2);
+					ExpandedTokens[--Idx] = ConcatenatedToken;
+				}
+			}
+
+			// Finally, return the expansion of this
+			IgnoreMacros.Add(Macro);
+			ExpandMacrosRecursively(ExpandedTokens, OutputTokens, bIsConditional, IgnoreMacros, Context);
+			IgnoreMacros.RemoveAt(IgnoreMacros.Count - 1);
+		}
+
+		/// <summary>
+		/// Appends a list of tokens to another list, setting the leading whitespace flag to the given value
+		/// </summary>
+		/// <param name="OutputTokens">List to receive the appended tokens</param>
+		/// <param name="InputTokens">List of tokens to append</param>
+		/// <param name="bHasLeadingSpace">Whether there is space before the first token</param>
+		void AppendTokensWithWhitespace(List<Token> OutputTokens, List<Token> InputTokens, bool bHasLeadingSpace)
+		{
+			if(InputTokens.Count > 0)
+			{
+				OutputTokens.Add(MergeLeadingSpace(InputTokens[0], bHasLeadingSpace));
+				OutputTokens.AddRange(InputTokens.Skip(1));
+			}
+		}
+
+		/// <summary>
+		/// Copies a single token from one list of tokens to another, or if it's an opening parenthesis, the entire subexpression until the closing parenthesis.
+		/// </summary>
+		/// <param name="InputTokens">The input token list</param>
+		/// <param name="InputIdx">First token index in the input token list. Set to the last uncopied token index on return.</param>
+		/// <param name="OutputTokens">List to recieve the output tokens</param>
+		/// <returns>True if a balanced expression was read, or false if the end of the list was encountered before finding a matching token</returns>
+		bool ReadBalancedToken(List<Token> InputTokens, ref int InputIdx, List<Token> OutputTokens)
+		{
+			// Copy a single token to the output list
+			Token Token = InputTokens[InputIdx++];
+			OutputTokens.Add(Token);
+
+			// If it was the start of a subexpression, copy until the closing parenthesis
+			if (Token.Type == TokenType.LeftParen)
+			{
+				// Copy the contents of the subexpression
+				for (;;)
+				{
+					if (InputIdx == InputTokens.Count)
+					{
+						return false;
+					}
+					if (InputTokens[InputIdx].Type == TokenType.RightParen)
+					{
+						break;
+					}
+					if (!ReadBalancedToken(InputTokens, ref InputIdx, OutputTokens))
+					{
+						return false;
+					}
+				}
+
+				// Copy the closing parenthesis
+				Token = InputTokens[InputIdx++];
+				OutputTokens.Add(Token);
+			}
+			return true;
+		}
+
+		/// <summary>
+		/// Copies a single token from one list of tokens to another, or if it's an opening parenthesis, the entire subexpression until the closing parenthesis.
+		/// </summary>
+		/// <param name="InputEnumerator">The input token list</param>
+		/// <param name="OutputTokens">List to recieve the output tokens</param>
+		/// <param name="Context">The context that the parser is in</param>
+		/// <returns>True if a balanced expression was read, or false if the end of the list was encountered before finding a matching token</returns>
+		bool ReadBalancedToken(IEnumerator<Token> InputEnumerator, List<Token> OutputTokens, PreprocessorContext Context)
+		{
+			// Copy a single token to the output list
+			Token Token = InputEnumerator.Current;
+			bool bMoveNext = InputEnumerator.MoveNext();
+			OutputTokens.Add(Token);
+
+			// If it was the start of a subexpression, copy until the closing parenthesis
+			if (Token.Type == TokenType.LeftParen)
+			{
+				// Copy the contents of the subexpression
+				for (;;)
+				{
+					if (!bMoveNext)
+					{
+						throw new PreprocessorException(Context, "Unbalanced token sequence");
+					}
+					if (InputEnumerator.Current.Type == TokenType.RightParen)
+					{
+						OutputTokens.Add(InputEnumerator.Current);
+						bMoveNext = InputEnumerator.MoveNext();
+						break;
+					}
+					bMoveNext = ReadBalancedToken(InputEnumerator, OutputTokens, Context);
+				}
+			}
+			return bMoveNext;
+		}
+	}
+}
diff --git a/Engine/Source/Programs/UnrealBuildTool/Preprocessor/PreprocessorContext.cs b/Engine/Source/Programs/UnrealBuildTool/Preprocessor/PreprocessorContext.cs
new file mode 100644
--- /dev/null
+++ b/Engine/Source/Programs/UnrealBuildTool/Preprocessor/PreprocessorContext.cs
@@ -0,0 +1,102 @@
+// Copyright Epic Games, Inc. All Rights Reserved.
+
+using System;
+using System.Collections.Generic;
+using System.Linq;
+using System.Text;
+using System.Threading.Tasks;
+
+namespace UnrealBuildTool
+{
+	/// <summary>
+	/// Represents a context that the preprocessor is working in. Used to form error messages, and determine things like the __FILE__ and __LINE__ directives.
+	/// </summary>
+	abstract class PreprocessorContext
+	{
+		/// <summary>
+		/// The outer context
+		/// </summary>
+		public readonly PreprocessorContext Outer;
+
+		/// <summary>
+		/// Constructor
+		/// </summary>
+		/// <param name="Outer">The outer context</param>
+		public PreprocessorContext(PreprocessorContext Outer)
+		{
+			this.Outer = Outer;
+		}
+	}
+
+	/// <summary>
+	/// Context for a command line argument
+	/// </summary>
+	class PreprocessorCommandLineContext : PreprocessorContext
+	{
+		/// <summary>
+		/// Constructor
+		/// </summary>
+		public PreprocessorCommandLineContext()
+			: base(null)
+		{
+		}
+
+		/// <summary>
+		/// Formats this context for error messages
+		/// </summary>
+		/// <returns>String describing this context</returns>
+		public override string ToString()
+		{
+			return "From command line";
+		}
+	}
+
+	/// <summary>
+	/// Represents a context that the preprocessor is reading from
+	/// </summary>
+	class PreprocessorFileContext : PreprocessorContext
+	{
+		/// <summary>
+		/// The source file being read
+		/// </summary>
+		public SourceFile SourceFile;
+
+		/// <summary>
+		/// The directory containing this file. When searching for included files, MSVC will check this directory.
+		/// </summary>
+		public DirectoryItem Directory;
+
+		/// <summary>
+		/// Index of the current markup object being processed
+		/// </summary>
+		public int MarkupIdx;
+
+		/// <summary>
+		/// Index of the next fragment to be read
+		/// </summary>
+		public int FragmentIdx;
+
+		/// <summary>
+		/// Constructor
+		/// </summary>
+		/// <param name="SourceFile">The source file being parsed</param>
+		/// <param name="Outer">The outer context</param>
+		public PreprocessorFileContext(SourceFile SourceFile, PreprocessorContext Outer)
+			: base(Outer)
+		{
+			this.SourceFile = SourceFile;
+			this.Directory = DirectoryItem.GetItemByDirectoryReference(SourceFile.Location.Directory);
+			this.MarkupIdx = 0;
+			this.FragmentIdx = 0;
+		}
+
+		/// <summary>
+		/// Format this file for the debugger, and error messages
+		/// </summary>
+		/// <returns></returns>
+		public override string ToString()
+		{
+			return String.Format("{0}({1})", SourceFile.Location, SourceFile.Markup[MarkupIdx].LineNumber);
+		}
+	}
+}
diff --git a/Engine/Source/Programs/UnrealBuildTool/Preprocessor/PreprocessorExpression.cs b/Engine/Source/Programs/UnrealBuildTool/Preprocessor/PreprocessorExpression.cs
new file mode 100644
--- /dev/null
+++ b/Engine/Source/Programs/UnrealBuildTool/Preprocessor/PreprocessorExpression.cs
@@ -0,0 +1,428 @@
+// Copyright Epic Games, Inc. All Rights Reserved.
+
+using System;
+using System.Collections.Generic;
+using System.Linq;
+using System.Text;
+using System.Threading.Tasks;
+
+namespace UnrealBuildTool
+{
+	/// <summary>
+	/// Class to allow evaluating a preprocessor expression
+	/// </summary>
+	static class PreprocessorExpression
+	{
+		/// <summary>
+		/// Precedence table for binary operators (indexed by TokenType). Lower values indicate higher precedence.
+		/// </summary>
+		static byte[] Precedence;
+
+		/// <summary>
+		/// Static constructor. Initializes the Precedence table.
+		/// </summary>
+		static PreprocessorExpression()
+		{
+			Precedence = new byte[(int)TokenType.Max];
+			Precedence[(int)TokenType.Multiply] = 10;
+			Precedence[(int)TokenType.Divide] = 10;
+			Precedence[(int)TokenType.Modulo] = 10;
+			Precedence[(int)TokenType.Plus] = 9;
+			Precedence[(int)TokenType.Minus] = 9;
+			Precedence[(int)TokenType.LeftShift] = 8;
+			Precedence[(int)TokenType.RightShift] = 8;
+			Precedence[(int)TokenType.CompareLess] = 7;
+			Precedence[(int)TokenType.CompareLessOrEqual] = 7;
+			Precedence[(int)TokenType.CompareGreater] = 7;
+			Precedence[(int)TokenType.CompareGreaterOrEqual] = 7;
+			Precedence[(int)TokenType.CompareEqual] = 6;
+			Precedence[(int)TokenType.CompareNotEqual] = 6;
+			Precedence[(int)TokenType.BitwiseAnd] = 5;
+			Precedence[(int)TokenType.BitwiseXor] = 4;
+			Precedence[(int)TokenType.BitwiseOr] = 3;
+			Precedence[(int)TokenType.LogicalAnd] = 2;
+			Precedence[(int)TokenType.LogicalOr] = 1;
+		}
+
+		/// <summary>
+		/// Evaluate a preprocessor expression
+		/// </summary>
+		/// <param name="Context">The current preprocessor context</param>
+		/// <param name="Tokens">List of tokens in the expression</param>
+		/// <returns>Value of the expression</returns>
+		public static long Evaluate(PreprocessorContext Context, List<Token> Tokens)
+		{
+			int Idx = 0;
+			long Result = EvaluateTernary(Context, Tokens, ref Idx);
+			if (Idx != Tokens.Count)
+			{
+				throw new PreprocessorException(Context, "Garbage after end of expression");
+			}
+			return Result;
+		}
+
+		/// <summary>
+		/// Evaluates a ternary expression (a? b :c)
+		/// </summary>
+		/// <param name="Context">The current preprocessor context</param>
+		/// <param name="Tokens">List of tokens in the expression</param>
+		/// <param name="Idx">Index into the token list of the next token to read</param>
+		/// <returns>Value of the expression</returns>
+		static long EvaluateTernary(PreprocessorContext Context, List<Token> Tokens, ref int Idx)
+		{
+			long Result = EvaluateUnary(Context, Tokens, ref Idx);
+			if(Idx < Tokens.Count && Precedence[(int)Tokens[Idx].Type] != 0)
+			{
+				Result = EvaluateBinary(Context, Tokens, ref Idx, Result, 1);
+			}
+			if(Idx < Tokens.Count && Tokens[Idx].Type == TokenType.QuestionMark)
+			{
+				// Read the left expression
+				Idx++;
+				long Lhs = EvaluateTernary(Context, Tokens, ref Idx);
+
+				// Check for the colon in the middle
+				if(Tokens[Idx].Type != TokenType.Colon)
+				{
+					throw new PreprocessorException(Context, "Expected colon for conditional operator");
+				}
+
+				// Read the right expression
+				Idx++;
+				long Rhs = EvaluateTernary(Context, Tokens, ref Idx);
+
+				// Evaluate it
+				Result = (Result != 0) ? Lhs : Rhs;
+			}
+			return Result;
+		}
+
+		/// <summary>
+		/// Recursive path for evaluating a sequence of binary operators, given a known LHS. Implementation of the shunting yard algorithm.
+		/// </summary>
+		/// <param name="Context">The current preprocessor context</param>
+		/// <param name="Tokens">List of tokens in the expression</param>
+		/// <param name="Idx">Index into the token list of the next token to read</param>
+		/// <param name="Lhs">The LHS value for the binary expression</param>
+		/// <param name="MinPrecedence">Minimum precedence value to consume</param>
+		/// <returns>Value of the expression</returns>
+		static long EvaluateBinary(PreprocessorContext Context, List<Token> Tokens, ref int Idx, long Lhs, byte MinPrecedence)
+		{
+			while(Idx < Tokens.Count)
+			{
+				// Read the next operator token
+				Token Operator = Tokens[Idx];
+
+				// Check if this operator binds more tightly than the outer expression, and exit if it does.
+				int OperatorPrecedence = Precedence[(int)Operator.Type];
+				if(OperatorPrecedence < MinPrecedence)
+				{
+					break;
+				}
+
+				// Move to the next token
+				Idx++;
+
+				// Evaluate the immediate RHS value, then recursively evaluate any subsequent binary operators that bind more tightly to it than this.
+				long Rhs = EvaluateUnary(Context, Tokens, ref Idx);
+				while(Idx < Tokens.Count)
+				{
+					byte NextOperatorPrecedence = Precedence[(int)Tokens[Idx].Type];
+					if(NextOperatorPrecedence <= OperatorPrecedence)
+					{
+						break;
+					}
+					Rhs = EvaluateBinary(Context, Tokens, ref Idx, Rhs, NextOperatorPrecedence);
+				}
+
+				// Evalute the result of applying the operator to the LHS and RHS values
+				switch(Operator.Type)
+				{
+					case TokenType.LogicalOr:
+						Lhs = ((Lhs != 0) || (Rhs != 0))? 1 : 0;
+						break;
+					case TokenType.LogicalAnd:
+						Lhs = ((Lhs != 0) && (Rhs != 0))? 1 : 0;
+						break;
+					case TokenType.BitwiseOr:
+						Lhs |= Rhs;
+						break;
+					case TokenType.BitwiseXor:
+						Lhs ^= Rhs;
+						break;
+					case TokenType.BitwiseAnd:
+						Lhs &= Rhs;
+						break;
+					case TokenType.CompareEqual:
+						Lhs = (Lhs == Rhs)? 1 : 0;
+						break;
+					case TokenType.CompareNotEqual:
+						Lhs = (Lhs != Rhs)? 1 : 0;
+						break;
+					case TokenType.CompareLess:
+						Lhs = (Lhs < Rhs)? 1 : 0;
+						break;
+					case TokenType.CompareGreater:
+						Lhs = (Lhs > Rhs)? 1 : 0;
+						break;
+					case TokenType.CompareLessOrEqual:
+						Lhs = (Lhs <= Rhs)? 1 : 0;
+						break;
+					case TokenType.CompareGreaterOrEqual:
+						Lhs = (Lhs >= Rhs)? 1 : 0;
+						break;
+					case TokenType.LeftShift:
+						Lhs <<= (int)Rhs;
+						break;
+					case TokenType.RightShift:
+						Lhs >>= (int)Rhs;
+						break;
+					case TokenType.Plus:
+						Lhs += Rhs;
+						break;
+					case TokenType.Minus:
+						Lhs -= Rhs;
+						break;
+					case TokenType.Multiply:
+						Lhs *= Rhs;
+						break;
+					case TokenType.Divide:
+						Lhs /= Rhs;
+						break;
+					case TokenType.Modulo:
+						Lhs %= Rhs;
+						break;
+					default:
+						throw new NotImplementedException(String.Format("Binary operator '{0}' has not been implemented", Operator.Type));
+				}
+			}
+
+			return Lhs;
+		}
+
+		/// <summary>
+		/// Evaluates a unary expression (+, -, !, ~)
+		/// </summary>
+		/// <param name="Context">The current preprocessor context</param>
+		/// <param name="Tokens">List of tokens in the expression</param>
+		/// <param name="Idx">Index into the token list of the next token to read</param>
+		/// <returns>Value of the expression</returns>
+		static long EvaluateUnary(PreprocessorContext Context, List<Token> Tokens, ref int Idx)
+		{
+			if(Idx == Tokens.Count)
+			{
+				throw new PreprocessorException(Context, "Early end of expression");
+			}
+
+			switch(Tokens[Idx].Type)
+			{
+				case TokenType.Plus:
+					Idx++;
+					return +EvaluateUnary(Context, Tokens, ref Idx);
+				case TokenType.Minus:
+					Idx++;
+					return -EvaluateUnary(Context, Tokens, ref Idx);
+				case TokenType.LogicalNot:
+					Idx++;
+					return (EvaluateUnary(Context, Tokens, ref Idx) == 0) ? 1 : 0;
+				case TokenType.BitwiseNot:
+					Idx++;
+					return ~EvaluateUnary(Context, Tokens, ref Idx);
+				case TokenType.Identifier:
+					Identifier Text = Tokens[Idx].Identifier;
+					if(Text == Identifiers.Sizeof)
+					{
+						throw new NotImplementedException();
+					}
+					else if(Text == Identifiers.Alignof)
+					{
+						throw new NotImplementedException();
+					}
+					else if(Text == Identifiers.__has_feature)
+					{
+						if(Tokens[Idx + 1].Type != TokenType.LeftParen || Tokens[Idx + 3].Type != TokenType.RightParen)
+						{
+							throw new NotImplementedException();
+						}
+						Idx += 4;
+						return 0;
+					}
+					else if(Text == Identifiers.__building_module)
+					{
+						if(Tokens[Idx + 1].Type != TokenType.LeftParen || Tokens[Idx + 3].Type != TokenType.RightParen)
+						{
+							throw new NotImplementedException();
+						}
+						Idx += 4;
+						return 0;
+					}
+					else
+					{
+						return EvaluatePrimary(Context, Tokens, ref Idx);
+					}
+				default:
+					return EvaluatePrimary(Context, Tokens, ref Idx);
+			}
+		}
+
+		/// <summary>
+		/// Evaluates a primary expression
+		/// </summary>
+		/// <param name="Context">The current preprocessor context</param>
+		/// <param name="Tokens">List of tokens in the expression</param>
+		/// <param name="Idx">Index into the token list of the next token to read</param>
+		/// <returns>Value of the expression</returns>
+		static long EvaluatePrimary(PreprocessorContext Context, List<Token> Tokens, ref int Idx)
+		{
+			if(Tokens[Idx].Type == TokenType.Identifier)
+			{
+				Idx++;
+				return 0;
+			}
+			else if(Tokens[Idx].Type == TokenType.LeftParen)
+			{
+				// Read the expression
+				Idx++;
+				long Result = EvaluateTernary(Context, Tokens, ref Idx);
+
+				// Check for a closing parenthesis
+				if (Tokens[Idx].Type != TokenType.RightParen)
+				{
+					throw new PreprocessorException(Context, "Missing closing parenthesis");
+				}
+
+				// Return the value
+				Idx++;
+				return Result;
+			}
+			else if(Tokens[Idx].Type == TokenType.Number)
+			{
+				return ParseNumericLiteral(Context, Tokens[Idx++].Literal);
+			}
+			else
+			{
+				throw new PreprocessorException(Context, "Unexpected token in expression: {0}", Tokens[Idx]);
+			}
+		}
+
+		/// <summary>
+		/// Parse a numeric literal from the given token
+		/// </summary>
+		/// <param name="Context">The current preprocessor context</param>
+		/// <param name="Literal">The utf-8 literal characters</param>
+		/// <returns>The literal value of the token</returns>
+		static long ParseNumericLiteral(PreprocessorContext Context, byte[] Literal)
+		{
+			// Parse the token
+			long Value = 0;
+			if(Literal.Length >= 2 && Literal[0] == '0' && (Literal[1] == 'x' || Literal[1] == 'X'))
+			{
+				for (int Idx = 2; Idx < Literal.Length; Idx++)
+				{
+					if (Literal[Idx] >= '0' && Literal[Idx] <= '9')
+					{
+						Value = (Value << 4) + (Literal[Idx] - '0');
+					}
+					else if (Literal[Idx] >= 'a' && Literal[Idx] <= 'f')
+					{
+						Value = (Value << 4) + (Literal[Idx] - 'a') + 10;
+					}
+					else if (Literal[Idx] >= 'A' && Literal[Idx] <= 'F')
+					{
+						Value = (Value << 4) + (Literal[Idx] - 'A') + 10;
+					}
+					else if(IsValidIntegerSuffix(Literal, Idx))
+					{
+						break;
+					}
+					else
+					{
+						throw new PreprocessorException(Context, "Invalid hexadecimal literal: '{0}'", Encoding.UTF8.GetString(Literal));
+					}
+				}
+			}
+			else if(Literal[0] == '0')
+			{
+				for (int Idx = 1; Idx < Literal.Length; Idx++)
+				{
+					if (Literal[Idx] >= '0' && Literal[Idx] <= '7')
+					{
+						Value = (Value << 3) + (Literal[Idx] - '0');
+					}
+					else if(IsValidIntegerSuffix(Literal, Idx))
+					{
+						break;
+					}
+					else
+					{
+						throw new PreprocessorException(Context, "Invalid octal literal: '{0}'", Encoding.UTF8.GetString(Literal));
+					}
+				}
+			}
+			else
+			{
+				for (int Idx = 0; Idx < Literal.Length; Idx++)
+				{
+					if (Literal[Idx] >= '0' && Literal[Idx] <= '9')
+					{
+						Value = (Value * 10) + (Literal[Idx] - '0');
+					}
+					else if(IsValidIntegerSuffix(Literal, Idx))
+					{
+						break;
+					}
+					else
+					{
+						throw new PreprocessorException(Context, "Invalid decimal literal: '{0}'", Encoding.UTF8.GetString(Literal));
+					}
+				}
+			}
+			return Value;
+		}
+
+		/// <summary>
+		/// Checks whether the given literal has a valid integer suffix, starting at the given position
+		/// </summary>
+		/// <param name="Literal">The literal token</param>
+		/// <param name="Index">Index within the literal to start reading</param>
+		/// <returns>The literal value of the token</returns>
+		static bool IsValidIntegerSuffix(byte[] Literal, int Index)
+		{
+			bool bAllowTrailingUnsigned = true;
+
+			if(Literal[Index] == 'u' || Literal[Index] == 'U')
+			{
+				// 'U'
+				Index++;
+				bAllowTrailingUnsigned = false;
+			}
+
+			if(Index < Literal.Length && (Literal[Index] == 'l' || Literal[Index] == 'L'))
+			{
+				// 'LL' or 'L'
+				if(Index + 1 < Literal.Length && Literal[Index + 1] == Literal[Index])
+				{
+					Index += 2;
+				}
+				else
+				{
+					Index++;
+				}
+			}
+			else if(Index + 2 < Literal.Length && (Literal[Index] == 'i' || Literal[Index] == 'I') && Literal[Index + 1] == '6' && Literal[Index + 2] == '4')
+			{
+				// 'I64' (Microsoft extension)
+				Index += 3;
+				bAllowTrailingUnsigned = false;
+			}
+
+			if(bAllowTrailingUnsigned && Index < Literal.Length && (Literal[Index] == 'u' || Literal[Index] == 'U'))
+			{
+				// 'U'
+				Index++;
+			}
+
+			return Index == Literal.Length;
+		}
+	}
+}
diff --git a/Engine/Source/Programs/UnrealBuildTool/Preprocessor/PreprocessorMacro.cs b/Engine/Source/Programs/UnrealBuildTool/Preprocessor/PreprocessorMacro.cs
new file mode 100644
--- /dev/null
+++ b/Engine/Source/Programs/UnrealBuildTool/Preprocessor/PreprocessorMacro.cs
@@ -0,0 +1,178 @@
+// Copyright Epic Games, Inc. All Rights Reserved.
+
+using System;
+using System.Collections.Generic;
+using System.Linq;
+using System.Text;
+using System.Threading.Tasks;
+using Tools.DotNETCommon;
+
+namespace UnrealBuildTool
+{
+	/// <summary>
+	/// Stores information about a defined preprocessor macro
+	/// </summary>
+	class PreprocessorMacro
+	{
+		/// <summary>
+		/// Name of the macro
+		/// </summary>
+		public Identifier Name;
+
+		/// <summary>
+		/// Parameter names for the macro. The '...' placeholder is represented by the __VA_ARGS__ string.
+		/// </summary>
+		public List<Identifier> Parameters;
+
+		/// <summary>
+		/// Raw list of tokens for this macro
+		/// </summary>
+		public List<Token> Tokens;
+
+		/// <summary>
+		/// Construct a preprocessor macro
+		/// </summary>
+		/// <param name="Name">Name of the macro</param>
+		/// <param name="Parameters">Parameter list for the macro. Should be null for object macros. Ownership of this list is transferred.</param>
+		/// <param name="Tokens">Tokens for the macro. Ownership of this list is transferred.</param>
+		public PreprocessorMacro(Identifier Name, List<Identifier> Parameters, List<Token> Tokens)
+		{
+			this.Name = Name;
+			this.Parameters = Parameters;
+			this.Tokens = Tokens;
+		}
+
+		/// <summary>
+		/// Read a macro from a binary archive
+		/// </summary>
+		/// <param name="Reader">Reader to serialize from</param>
+		public PreprocessorMacro(BinaryArchiveReader Reader)
+		{
+			Name = Reader.ReadIdentifier();
+			Parameters = Reader.ReadList(() => Reader.ReadIdentifier());
+			Tokens = Reader.ReadList(() => Reader.ReadToken());
+		}
+
+		/// <summary>
+		/// Write a macro to a binary archive
+		/// </summary>
+		/// <param name="Writer">Writer to serialize to</param>
+		public void Write(BinaryArchiveWriter Writer)
+		{
+			Writer.WriteIdentifier(Name);
+			Writer.WriteList(Parameters, x => Writer.WriteIdentifier(x));
+			Writer.WriteList(Tokens, x => Writer.WriteToken(x));
+		}
+
+		/// <summary>
+		/// Finds the index of a parameter in the parameter list
+		/// </summary>
+		/// <param name="Parameter">Parameter name to look for</param>
+		/// <returns>Index of the parameter, or -1 if it's not found.</returns>
+		public int FindParameterIndex(Identifier Parameter)
+		{
+			for (int Idx = 0; Idx < Parameters.Count; Idx++)
+			{
+				if (Parameters[Idx] == Parameter)
+				{
+					return Idx;
+				}
+			}
+			return -1;
+		}
+
+		/// <summary>
+		/// Checks whether this macro definition is equivalent to another macro definition
+		/// </summary>
+		/// <param name="Other">The macro definition to compare against</param>
+		/// <returns>True if the macro definitions are equivalent</returns>
+		public bool IsEquivalentTo(PreprocessorMacro Other)
+		{
+			if(this != Other)
+			{
+				if(Name != Other.Name || Tokens.Count != Other.Tokens.Count)
+				{
+					return false;
+				}
+				if(Parameters != null)
+				{
+					if(Other.Parameters == null || Other.Parameters.Count != Parameters.Count || !Enumerable.SequenceEqual(Parameters, Other.Parameters))
+					{
+						return false;
+					}
+				}
+				else
+				{
+					if(Other.Parameters != null)
+					{
+						return false;
+					}
+				}
+				if(!Enumerable.SequenceEqual(Tokens, Other.Tokens))
+				{
+					return false;
+				}
+			}
+			return true;
+		}
+
+		/// <summary>
+		/// True if the macro is an object macro
+		/// </summary>
+		public bool IsObjectMacro
+		{
+			get { return Parameters == null; }
+		}
+
+		/// <summary>
+		/// True if the macro is a function macro
+		/// </summary>
+		public bool IsFunctionMacro
+		{
+			get { return Parameters != null; }
+		}
+		
+		/// <summary>
+		/// The number of required parameters. For variadic macros, the last parameter is optional.
+		/// </summary>
+		public int MinRequiredParameters
+		{
+			get { return HasVariableArgumentList? (Parameters.Count - 1) : Parameters.Count; }
+		}
+
+		/// <summary>
+		/// True if the macro has a variable argument list
+		/// </summary>
+		public bool HasVariableArgumentList
+		{
+			get { return Parameters.Count > 0 && Parameters[Parameters.Count - 1] == Identifiers.__VA_ARGS__; }
+		}
+
+		/// <summary>
+		/// Converts this macro to a string for debugging
+		/// </summary>
+		/// <returns>The tokens in this macro</returns>
+		public override string ToString()
+		{
+			StringBuilder Result = new StringBuilder(Name.ToString());
+			if (Parameters != null)
+			{
+				Result.AppendFormat("({0})", String.Join(", ", Parameters));
+			}
+			Result.Append("=");
+			if (Tokens.Count > 0)
+			{
+				Result.Append(Tokens[0].Text);
+				for (int Idx = 1; Idx < Tokens.Count; Idx++)
+				{
+					if(Tokens[Idx].HasLeadingSpace)
+					{
+						Result.Append(" ");
+					}
+					Result.Append(Tokens[Idx].Text);
+				}
+			}
+			return Result.ToString();
+		}
+	}
+}
diff --git a/Engine/Source/Programs/UnrealBuildTool/Preprocessor/PreprocessorState.cs b/Engine/Source/Programs/UnrealBuildTool/Preprocessor/PreprocessorState.cs
new file mode 100644
--- /dev/null
+++ b/Engine/Source/Programs/UnrealBuildTool/Preprocessor/PreprocessorState.cs
@@ -0,0 +1,362 @@
+// Copyright Epic Games, Inc. All Rights Reserved.
+
+using System;
+using System.Collections.Generic;
+using System.Linq;
+using System.Text;
+using System.Threading.Tasks;
+
+namespace UnrealBuildTool
+{
+	/// <summary>
+	/// Flags indicating the state of each nested conditional block in a stack.
+	/// </summary>
+	[Flags]
+	enum PreprocessorBranch
+	{
+		/// <summary>
+		/// The current conditional branch is active.
+		/// </summary>
+		Active = 0x01,
+
+		/// <summary>
+		/// A branch within the current conditional block has been taken. Any #else/#elif blocks will not be taken.
+		/// </summary>
+		Taken = 0x02,
+
+		/// <summary>
+		/// The #else directive for this conditional block has been encountered. An #endif directive is expected next.
+		/// </summary>
+		Complete = 0x04,
+
+		/// <summary>
+		/// The first condition in this branch was an #if directive (as opposed to an #ifdef or #ifndef directive)
+		/// </summary>
+		HasIfDirective = 0x10,
+
+		/// <summary>
+		/// The first condition in this branch was an #if directive (as opposed to an #ifdef or #ifndef directive)
+		/// </summary>
+		HasIfdefDirective = 0x20,
+
+		/// <summary>
+		/// The first condition in this branch was an #ifndef directive (as opposed to an #ifdef or #if directive)
+		/// </summary>
+		HasIfndefDirective = 0x40,
+
+		/// <summary>
+		/// The branch has an #elif directive
+		/// </summary>
+		HasElifDirective = 0x80,
+	}
+
+	/// <summary>
+	/// Marshals access to the preprocessor state, ensuring that any changes are tracked within the active transform
+	/// </summary>
+	class PreprocessorState
+	{
+		/// <summary>
+		/// Stack of conditional branch states
+		/// </summary>
+		List<PreprocessorBranch> Branches = new List<PreprocessorBranch>();
+
+		/// <summary>
+		/// Mapping of name to macro definition
+		/// </summary>
+		Dictionary<Identifier, PreprocessorMacro> NameToMacro = new Dictionary<Identifier, PreprocessorMacro>();
+
+		/// <summary>
+		/// The current transform. Any queries or state modifications will be recorded in this.
+		/// </summary>
+		PreprocessorTransform Transform;
+
+		/// <summary>
+		/// Enumerates the current branches. Do not use this for decision logic; any dependencies will not be tracked.
+		/// </summary>
+		public IEnumerable<PreprocessorBranch> CurrentBranches
+		{
+			get { return Branches; }
+		}
+
+		/// <summary>
+		/// Enumerates the current macros. Do not use this for decision logic; any dependencies will not be tracked.
+		/// </summary>
+		public IEnumerable<PreprocessorMacro> CurrentMacros
+		{
+			get { return NameToMacro.Values; }
+		}
+
+		/// <summary>
+		/// Initialize an empty preprocessor state
+		/// </summary>
+		public PreprocessorState()
+		{
+		}
+
+		/// <summary>
+		/// Duplicates another preprocessor state. Throws an exception if a transform is currently being built in the other state.
+		/// </summary>
+		/// <param name="Other">The preprocessor state to copy</param>
+		public PreprocessorState(PreprocessorState Other)
+		{
+			if(Other.Transform != null)
+			{
+				throw new Exception("Unable to copy another preprocessor state while a transform is being built.");
+			}
+
+			Branches.AddRange(Other.Branches);
+
+			foreach(KeyValuePair<Identifier, PreprocessorMacro> Pair in Other.NameToMacro)
+			{
+				NameToMacro[Pair.Key] = Pair.Value;
+			}
+		}
+
+		/// <summary>
+		/// Create a new transform object, and assign it to be current
+		/// </summary>
+		/// <returns>The new transform object</returns>
+		public PreprocessorTransform BeginCapture()
+		{
+			Transform = new PreprocessorTransform();
+			return Transform;
+		}
+
+		/// <summary>
+		/// Detach the current transform
+		/// </summary>
+		/// <returns>The transform object</returns>
+		public PreprocessorTransform EndCapture()
+		{
+			PreprocessorTransform PrevTransform = Transform;
+			Transform = null;
+			return PrevTransform;
+		}
+
+		/// <summary>
+		/// Sets a flag indicating that a #pragma once directive was encountered
+		/// </summary>
+		public void MarkPragmaOnce()
+		{
+			if(Transform != null)
+			{
+				Transform.bHasPragmaOnce = true;
+			}
+		}
+
+		/// <summary>
+		/// Set a macro to a specific definition
+		/// </summary>
+		/// <param name="Macro">Macro definition</param>
+		public void DefineMacro(PreprocessorMacro Macro)
+		{
+			NameToMacro[Macro.Name] = Macro;
+
+			if(Transform != null)
+			{
+				Transform.NewMacros[Macro.Name] = Macro;
+			}
+		}
+
+		/// <summary>
+		/// Checks if a macro with the given name is defined
+		/// </summary>
+		/// <param name="Name">The macro name</param>
+		/// <returns>True if the macro is defined</returns>
+		public bool IsMacroDefined(Identifier Name)
+		{
+			// Could account for the fact that we don't need the full definition later...
+			PreprocessorMacro Macro;
+			return TryGetMacro(Name, out Macro);
+		}
+
+		/// <summary>
+		/// Removes a macro definition
+		/// </summary>
+		/// <param name="Name">Name of the macro</param>
+		public void UndefMacro(Identifier Name)
+		{
+			NameToMacro.Remove(Name);
+
+			if(Transform != null)
+			{
+				Transform.NewMacros[Name] = null;
+			}
+		}
+
+		/// <summary>
+		/// Attemps to get the definition for a macro
+		/// </summary>
+		/// <param name="Name">Name of the macro</param>
+		/// <param name="Macro">Receives the macro definition, or null if it's not defined</param>
+		/// <returns>True if the macro is defined, otherwise false</returns>
+		public bool TryGetMacro(Identifier Name, out PreprocessorMacro Macro)
+		{
+			bool bResult = NameToMacro.TryGetValue(Name, out Macro);
+
+			if(Transform != null && !Transform.NewMacros.ContainsKey(Name))
+			{
+				Transform.RequiredMacros[Name] = Macro;
+			}
+
+			return bResult;
+		}
+
+		/// <summary>
+		/// Pushes the preprocessor branch onto the stack
+		/// </summary>
+		/// <param name="Branch">The branch state</param>
+		public void PushBranch(PreprocessorBranch Branch)
+		{
+			Branches.Add(Branch);
+
+			if(Transform != null)
+			{
+				Transform.NewBranches.Add(Branch);
+			}
+		}
+
+		/// <summary>
+		/// Pops a preprocessor branch from the stack
+		/// </summary>
+		/// <returns>The popped branch state</returns>
+		public PreprocessorBranch PopBranch()
+		{
+			PreprocessorBranch Branch;
+			if(!TryPopBranch(out Branch))
+			{
+				throw new InvalidOperationException("Branch stack is empty");
+			}
+			return Branch;
+		}
+
+		/// <summary>
+		/// Attempts to pops a preprocessor branch from the stack
+		/// </summary>
+		/// <param name="Branch">On success, receives the preprocessor branch state</param>
+		/// <returns>True if a branch was active, else false</returns>
+		public bool TryPopBranch(out PreprocessorBranch Branch)
+		{
+			if(Branches.Count == 0)
+			{
+				Branch = 0;
+				return false;
+			}
+			else
+			{
+				Branch = Branches[Branches.Count - 1];
+				Branches.RemoveAt(Branches.Count - 1);
+
+				if(Transform != null)
+				{
+					if(Transform.NewBranches.Count > 0)
+					{
+						Transform.NewBranches.RemoveAt(Transform.NewBranches.Count - 1);
+					}
+					else
+					{
+						Transform.RequiredBranches.Add(Branch);
+						Transform.bRequireTopmostActive = null;
+					}
+				}
+
+				return true;
+			}
+		}
+
+		/// <summary>
+		/// Determines if the branch that the preprocessor is in is active
+		/// </summary>
+		/// <returns>True if the branch is active, false otherwise</returns>
+		public bool IsCurrentBranchActive()
+		{
+			bool bActive = (Branches.Count == 0 || Branches[Branches.Count - 1].HasFlag(PreprocessorBranch.Active));
+			if(Transform != null && Transform.NewBranches.Count == 0)
+			{
+				Transform.bRequireTopmostActive = bActive;
+			}
+			return bActive;
+		}
+
+		/// <summary>
+		/// Determines if the given transform can apply to the current preprocessor state
+		/// </summary>
+		/// <param name="Transform">The transform to test</param>
+		/// <returns>True if the transform can be applied to the current state</returns>
+		public bool CanApply(PreprocessorTransform Transform)
+		{
+			// Check all the required branches match
+			for(int Idx = 0; Idx < Transform.RequiredBranches.Count; Idx++)
+			{
+				if(Branches[Branches.Count - Idx - 1] != Transform.RequiredBranches[Idx])
+				{
+					return false;
+				}
+			}
+
+			// Check the topmost branch is active
+			if(Transform.bRequireTopmostActive.HasValue)
+			{
+				bool bTopmostActive = (Branches.Count == Transform.RequiredBranches.Count || Branches[Branches.Count - Transform.RequiredBranches.Count - 1].HasFlag(PreprocessorBranch.Active));
+				if(Transform.bRequireTopmostActive.Value != bTopmostActive)
+				{
+					return false;
+				}
+			}
+
+			// Check all the required macros match
+			foreach(KeyValuePair<Identifier, PreprocessorMacro> RequiredPair in Transform.RequiredMacros)
+			{
+				PreprocessorMacro Macro;
+				if(NameToMacro.TryGetValue(RequiredPair.Key, out Macro))
+				{
+					if(RequiredPair.Value == null || !Macro.IsEquivalentTo(RequiredPair.Value))
+					{
+						return false;
+					}
+				} 
+				else
+				{
+					if(RequiredPair.Value != null)
+					{
+						return false;
+					}
+				}
+			}
+
+			return true;
+		}
+
+		/// <summary>
+		/// Apply the given transform to the current preprocessor state
+		/// </summary>
+		/// <param name="Transform">The transform to apply</param>
+		/// <returns>True if the transform was applied to the current state</returns>
+		public bool TryToApply(PreprocessorTransform Transform)
+		{
+			if(!CanApply(Transform))
+			{
+				return false;
+			}
+
+			// Update the branch state
+			Branches.RemoveRange(Branches.Count - Transform.RequiredBranches.Count, Transform.RequiredBranches.Count);
+			Branches.AddRange(Transform.NewBranches);
+
+			// Update the macro definitions
+			foreach(KeyValuePair<Identifier, PreprocessorMacro> NewMacro in Transform.NewMacros)
+			{
+				if(NewMacro.Value == null)
+				{
+					NameToMacro.Remove(NewMacro.Key);
+				}
+				else
+				{
+					NameToMacro[NewMacro.Key] = NewMacro.Value;
+				}
+			}
+
+			return true;
+		}
+	}
+}
diff --git a/Engine/Source/Programs/UnrealBuildTool/Preprocessor/PreprocessorTransform.cs b/Engine/Source/Programs/UnrealBuildTool/Preprocessor/PreprocessorTransform.cs
new file mode 100644
--- /dev/null
+++ b/Engine/Source/Programs/UnrealBuildTool/Preprocessor/PreprocessorTransform.cs
@@ -0,0 +1,95 @@
+// Copyright Epic Games, Inc. All Rights Reserved.
+
+using System;
+using System.Collections.Generic;
+using System.Linq;
+using System.Text;
+using System.Threading.Tasks;
+using Tools.DotNETCommon;
+
+namespace UnrealBuildTool
+{
+	/// <summary>
+	/// Tracks dependencies for transforming the current preprocessor state
+	/// </summary>
+	class PreprocessorTransform
+	{
+		/// <summary>
+		/// Whether the preprocessor must be in an active state after the branches below are popped. This cannot be treated as a 'pop', because it might be the global scope.
+		/// </summary>
+		public bool? bRequireTopmostActive;
+
+		/// <summary>
+		/// List of branch states that will be popped from the stack
+		/// </summary>
+		public List<PreprocessorBranch> RequiredBranches;
+
+		/// <summary>
+		/// Map of macro name to their required value
+		/// </summary>
+		public Dictionary<Identifier, PreprocessorMacro> RequiredMacros;
+
+		/// <summary>
+		/// List of new branches that will be pushed onto the stack
+		/// </summary>
+		public List<PreprocessorBranch> NewBranches;
+
+		/// <summary>
+		/// List of new macros that will be defined
+		/// </summary>
+		public Dictionary<Identifier, PreprocessorMacro> NewMacros;
+
+		/// <summary>
+		/// This fragment contains a pragma once directive
+		/// </summary>
+		public bool bHasPragmaOnce;
+
+		/// <summary>
+		/// Constructor
+		/// </summary>
+		public PreprocessorTransform()
+		{
+			RequiredBranches = new List<PreprocessorBranch>();
+			RequiredMacros = new Dictionary<Identifier, PreprocessorMacro>();
+			NewBranches = new List<PreprocessorBranch>();
+			NewMacros = new Dictionary<Identifier, PreprocessorMacro>();
+		}
+
+		/// <summary>
+		/// Reads a transform from disk
+		/// </summary>
+		/// <param name="Reader">Archive to serialize from</param>
+		public PreprocessorTransform(BinaryArchiveReader Reader)
+		{
+			if(Reader.ReadBool())
+			{
+				bRequireTopmostActive = Reader.ReadBool();
+			}
+
+			RequiredBranches = Reader.ReadList(() => (PreprocessorBranch)Reader.ReadByte());
+			RequiredMacros = Reader.ReadDictionary(() => Reader.ReadIdentifier(), () => Reader.ReadObjectReference(() => new PreprocessorMacro(Reader)));
+			NewBranches = Reader.ReadList(() => (PreprocessorBranch)Reader.ReadByte());
+			NewMacros = Reader.ReadDictionary(() => Reader.ReadIdentifier(), () => Reader.ReadObjectReference(() => new PreprocessorMacro(Reader)));
+			bHasPragmaOnce = Reader.ReadBool();
+		}
+
+		/// <summary>
+		/// Writes a transform to disk
+		/// </summary>
+		/// <param name="Writer">Archive to serialize to</param>
+		public void Write(BinaryArchiveWriter Writer)
+		{
+			Writer.WriteBool(bRequireTopmostActive.HasValue);
+			if(bRequireTopmostActive.HasValue)
+			{
+				Writer.WriteBool(bRequireTopmostActive.Value);
+			}
+
+			Writer.WriteList(RequiredBranches, x => Writer.WriteByte((byte)x));
+			Writer.WriteDictionary(RequiredMacros, k => Writer.WriteIdentifier(k), v => Writer.WriteObjectReference(v, () => v.Write(Writer)));
+			Writer.WriteList(NewBranches, x => Writer.WriteByte((byte)x));
+			Writer.WriteDictionary(NewMacros, k => Writer.WriteIdentifier(k), v => Writer.WriteObjectReference(v, () => v.Write(Writer)));
+			Writer.WriteBool(bHasPragmaOnce);
+		}
+	}
+}
diff --git a/Engine/Source/Programs/UnrealBuildTool/Preprocessor/SourceFile.cs b/Engine/Source/Programs/UnrealBuildTool/Preprocessor/SourceFile.cs
new file mode 100644
--- /dev/null
+++ b/Engine/Source/Programs/UnrealBuildTool/Preprocessor/SourceFile.cs
@@ -0,0 +1,250 @@
+// Copyright Epic Games, Inc. All Rights Reserved.
+
+using System;
+using System.Collections.Concurrent;
+using System.Collections.Generic;
+using System.Diagnostics;
+using System.IO;
+using System.Linq;
+using System.Text;
+using System.Threading.Tasks;
+using Tools.DotNETCommon;
+
+namespace UnrealBuildTool
+{
+	/// <summary>
+	/// Encapsulates metadata for a C++ source file
+	/// </summary>
+	[Serializable]
+	class SourceFile
+	{
+		/// <summary>
+		/// The corresponding file item
+		/// </summary>
+		public readonly FileItem File;
+
+		/// <summary>
+		/// Hash of the source file data
+		/// </summary>
+		public readonly ContentHash Hash;
+
+		/// <summary>
+		/// Last write time of this file when constructed
+		/// </summary>
+		public readonly long LastWriteTimeUtc;
+
+		/// <summary>
+		/// Markup dividing up this source file
+		/// </summary>
+		public readonly SourceFileMarkup[] Markup;
+
+		/// <summary>
+		/// Macro which is defined as part of the header guard. If this macro is defined, we don't need to include the header more than once.
+		/// </summary>
+		public Identifier HeaderGuardMacro;
+
+		/// <summary>
+		/// Parsed fragments in this file. Only set for files which are being optimized.
+		/// </summary>
+		public SourceFileFragment[] Fragments;
+
+		/// <summary>
+		/// Gets the location of this file
+		/// </summary>
+		public FileReference Location
+		{
+			get { return File.Location; }
+		}
+
+		/// <summary>
+		/// Construct a SourceFile from a file on disk
+		/// </summary>
+		/// <param name="File">Location of the file</param>
+		public SourceFile(FileItem File)
+			: this(ReadFileWithNullTerminator(File.Location))
+		{
+			this.File = File;
+			this.LastWriteTimeUtc = File.LastWriteTimeUtc.Ticks;
+		}
+
+		/// <summary>
+		/// Construct a SourceFile from a raw byte array
+		/// </summary>
+		/// <param name="Data">Contents of the file</param>
+		public SourceFile(byte[] Data)
+		{
+			// Compute the hash
+			Hash = ContentHash.MD5(Data);
+
+			// Read the preprocessor markup
+			TokenReader Reader = new TokenReader(Data);
+			Markup = SourceFileMarkup.Parse(Reader);
+
+			// Try to parse a header guard from this file
+			ParseHeaderGuard();
+
+			// Create the fragments
+			List<SourceFileFragment> FragmentsList = new List<SourceFileFragment>();
+			for(int MaxIdx = 0; MaxIdx < Markup.Length; MaxIdx++)
+			{
+				if(Markup[MaxIdx].Type != SourceFileMarkupType.Include)
+				{
+					int MinIdx = MaxIdx++;
+					while(MaxIdx < Markup.Length && Markup[MaxIdx].Type != SourceFileMarkupType.Include)
+					{
+						MaxIdx++;
+					}
+					FragmentsList.Add(new SourceFileFragment(this, MinIdx, MaxIdx));
+				}
+			}
+			Fragments = FragmentsList.ToArray();
+		}
+
+		/// <summary>
+		/// Read a source file from a binary reader
+		/// </summary>
+		/// <param name="Reader">The reader to serialize from</param>
+		public SourceFile(BinaryArchiveReader Reader)
+		{
+			File = Reader.ReadFileItem();
+			Hash = Reader.ReadContentHash();
+			LastWriteTimeUtc = Reader.ReadLong();
+			Markup = Reader.ReadArray(() => new SourceFileMarkup(Reader));
+			HeaderGuardMacro = Reader.ReadIdentifier();
+			Fragments = Reader.ReadArray(() => new SourceFileFragment(this, Reader));
+		}
+
+		/// <summary>
+		/// Read a source file from a binary reader
+		/// </summary>
+		/// <param name="Writer">The reader to serialize to</param>
+		public void Write(BinaryArchiveWriter Writer)
+		{
+			Writer.WriteFileItem(File);
+			Writer.WriteContentHash(Hash);
+			Writer.WriteLong(LastWriteTimeUtc);
+			Writer.WriteArray(Markup, x => x.Write(Writer));
+			Writer.WriteIdentifier(HeaderGuardMacro);
+			Writer.WriteArray(Fragments, x => x.Write(Writer));
+		}
+
+		/// <summary>
+		/// Reads the contents of a file into a buffer, with a null terminator
+		/// </summary>
+		/// <param name="Location">Location of the file</param>
+		/// <returns>Array of bytes, with a null terminator</returns>
+		public static byte[] ReadFileWithNullTerminator(FileReference Location)
+		{
+			using(FileStream Stream = FileReference.Open(Location, FileMode.Open, FileAccess.Read, FileShare.Read))
+			{
+				int Length = (int)Stream.Length;
+
+				byte[] Data = new byte[Length + 1];
+				Stream.Read(Data, 0, Length);
+
+				return ConvertToNullTerminatedUtf8(Data, Length);
+			}
+		}
+
+		/// <summary>
+		/// Converts a buffer from a given encoding to UTF8, appending a null terminator
+		/// </summary>
+		/// <param name="Data">The source data</param>
+		/// <param name="Length">Length of the data in the source buffer</param>
+		/// <returns>Array of utf-8 encoded characters, with a null terminator</returns>
+		public static byte[] ConvertToNullTerminatedUtf8(byte[] Data, int Length)
+		{
+			if(Length >= 2)
+			{
+				if(Data[0] == 0xfe && Data[1] == 0xff)
+				{
+					return ConvertToNullTerminatedUtf8(Encoding.BigEndianUnicode, Data, Length);
+				}
+				else if(Data[0] == 0xff && Data[1] == 0xfe)
+				{
+					return ConvertToNullTerminatedUtf8(Encoding.Unicode, Data, Length);
+				}
+			}
+			return Data;
+		}
+
+		/// <summary>
+		/// Converts a buffer from a given encoding to UTF8, appending a null terminator
+		/// </summary>
+		/// <param name="SourceEncoding">The encoding to convert from</param>
+		/// <param name="Data">The source data</param>
+		/// <param name="Length">Length of the data in the source buffer</param>
+		/// <returns>Array of utf-8 encoded characters, with a null terminator</returns>
+		public static byte[] ConvertToNullTerminatedUtf8(Encoding SourceEncoding, byte[] Data, int Length)
+		{
+			byte[] NewBytes = Data;
+			if(!(SourceEncoding is UTF8Encoding))
+			{
+				char[] Chars = SourceEncoding.GetChars(Data, 0, Length);
+
+				int NewLength = Encoding.UTF8.GetByteCount(Chars);
+				NewBytes = new byte[NewLength + 1];
+
+				Encoding.UTF8.GetBytes(Chars, 0, Chars.Length, NewBytes, 0);
+			}
+			return NewBytes;
+		}
+
+		/// <summary>
+		/// Checks whether the markup matches a C-style header guard
+		/// </summary>
+		void ParseHeaderGuard()
+		{
+			// Make sure there are enough markup entries in this list
+			if(Markup.Length >= 3)
+			{
+				int MinIdx = 0;
+
+				// Get the define used for the header guard
+				Identifier PossibleHeaderGuardMacro;
+				if(Markup[MinIdx].Type == SourceFileMarkupType.Ifndef && Markup[MinIdx].Tokens.Count == 1 && Markup[MinIdx].Tokens[0].Type == TokenType.Identifier)
+				{
+					PossibleHeaderGuardMacro = Markup[MinIdx].Tokens[0].Identifier;
+				}
+				else if(Markup[MinIdx].Type == SourceFileMarkupType.If && Markup[MinIdx].Tokens.Count == 3 && Markup[MinIdx].Tokens[0].Type == TokenType.LogicalNot && Markup[MinIdx].Tokens[1].Identifier == Identifiers.Defined && Markup[MinIdx].Tokens[2].Type == TokenType.Identifier)
+				{
+					PossibleHeaderGuardMacro = Markup[MinIdx].Tokens[2].Identifier;
+				}
+				else
+				{
+					PossibleHeaderGuardMacro = null;
+				}
+
+				// Check it starts with an #if or #ifdef
+				if (Markup[MinIdx + 1].Type == SourceFileMarkupType.Define && Markup[MinIdx + 1].Tokens.Count == 1 && Markup[MinIdx + 1].Tokens[0].Identifier == PossibleHeaderGuardMacro)
+				{
+					// Find the point at which the include depth goes back to zero
+					int MaxIdx = MinIdx + 1;
+					for (int Depth = 1; MaxIdx < Markup.Length; MaxIdx++)
+					{
+						Depth += Markup[MaxIdx].GetConditionDepthDelta();
+						if (Depth == 0)
+						{
+							break;
+						}
+					}
+
+					// Check it matched the end of the file
+					if (MaxIdx == Markup.Length - 1)
+					{
+						HeaderGuardMacro = PossibleHeaderGuardMacro;
+					}
+				}
+			}
+		}
+
+		/// <summary>
+		/// Converts a source file to a string for debugging
+		/// </summary>
+		/// <returns>The full path to the file</returns>
+		public override string ToString()
+		{
+			return File.ToString();
+		}
+	}
+}
diff --git a/Engine/Source/Programs/UnrealBuildTool/Preprocessor/SourceFileFragment.cs b/Engine/Source/Programs/UnrealBuildTool/Preprocessor/SourceFileFragment.cs
new file mode 100644
--- /dev/null
+++ b/Engine/Source/Programs/UnrealBuildTool/Preprocessor/SourceFileFragment.cs
@@ -0,0 +1,92 @@
+// Copyright Epic Games, Inc. All Rights Reserved.
+
+using System;
+using System.Collections.Generic;
+using System.IO;
+using System.Linq;
+using System.Text;
+using System.Threading.Tasks;
+using Tools.DotNETCommon;
+
+namespace UnrealBuildTool
+{
+	/// <summary>
+	/// Represents an fragment of a source file consisting of a sequence of includes followed by some arbitrary directives or source code.
+	/// </summary>
+	[Serializable]
+	class SourceFileFragment
+	{
+		/// <summary>
+		/// The file that this fragment is part of
+		/// </summary>
+		public readonly SourceFile File;
+
+		/// <summary>
+		/// Index into the file's markup array of the start of this fragment (inclusive)
+		/// </summary>
+		public readonly int MarkupMin;
+
+		/// <summary>
+		/// Index into the file's markup array of the end of this fragment (exclusive). Set to zero for external files.
+		/// </summary>
+		public readonly int MarkupMax;
+
+		/// <summary>
+		/// List of known transforms for this fragment
+		/// </summary>
+		public PreprocessorTransform[] Transforms = new PreprocessorTransform[0];
+
+		/// <summary>
+		/// Constructor
+		/// </summary>
+		/// <param name="File">The source file that this fragment is part of</param>
+		/// <param name="MarkupMin">Index into the file's markup array of the start of this fragment (inclusive)</param>
+		/// <param name="MarkupMax">Index into the file's markup array of the end of this fragment (exclusive)</param>
+		public SourceFileFragment(SourceFile File, int MarkupMin, int MarkupMax)
+		{
+			this.File = File;
+			this.MarkupMin = MarkupMin;
+			this.MarkupMax = MarkupMax;
+		}
+
+		/// <summary>
+		/// Constructor
+		/// </summary>
+		/// <param name="File">The source file that this fragment is part of</param>
+		/// <param name="Reader">Reader to deserialize this object from</param>
+		public SourceFileFragment(SourceFile File, BinaryArchiveReader Reader)
+		{
+			this.File = File;
+			this.MarkupMin = Reader.ReadInt();
+			this.MarkupMax = Reader.ReadInt();
+			this.Transforms = Reader.ReadArray(() => new PreprocessorTransform(Reader));
+		}
+
+		/// <summary>
+		/// Write this fragment to an archive
+		/// </summary>
+		/// <param name="Writer">The writer to serialize to</param>
+		public void Write(BinaryArchiveWriter Writer)
+		{
+			Writer.WriteInt(MarkupMin);
+			Writer.WriteInt(MarkupMax);
+			Writer.WriteArray(Transforms, x => x.Write(Writer));
+		}
+
+		/// <summary>
+		/// Summarize this fragment for the debugger
+		/// </summary>
+		/// <returns>String representation of this fragment for the debugger</returns>
+		public override string ToString()
+		{
+			if(MarkupMax == 0)
+			{
+				return File.ToString();
+			}
+			else
+			{
+				return String.Format("{0}: {1}-{2}", File.ToString(), File.Markup[MarkupMin].LineNumber, File.Markup[MarkupMax - 1].LineNumber + 1);
+			}
+		}
+	}
+}
diff --git a/Engine/Source/Programs/UnrealBuildTool/Preprocessor/SourceFileMarkup.cs b/Engine/Source/Programs/UnrealBuildTool/Preprocessor/SourceFileMarkup.cs
new file mode 100644
--- /dev/null
+++ b/Engine/Source/Programs/UnrealBuildTool/Preprocessor/SourceFileMarkup.cs
@@ -0,0 +1,351 @@
+// Copyright Epic Games, Inc. All Rights Reserved.
+
+using System;
+using System.Collections.Generic;
+using System.Linq;
+using System.Text;
+using System.Threading;
+using System.Threading.Tasks;
+using Tools.DotNETCommon;
+
+namespace UnrealBuildTool
+{
+	/// <summary>
+	/// Possible types of source file markup
+	/// </summary>
+	enum SourceFileMarkupType
+	{
+		/// <summary>
+		/// A span of tokens which are not preprocessor directives
+		/// </summary>
+		Text,
+
+		/// <summary>
+		/// An #include directive
+		/// </summary>
+		Include,
+
+		/// <summary>
+		/// A #define directive
+		/// </summary>
+		Define,
+
+		/// <summary>
+		/// An #undef directive
+		/// </summary>
+		Undef,
+
+		/// <summary>
+		/// An #if directive
+		/// </summary>
+		If,
+
+		/// <summary>
+		/// An #ifdef directive
+		/// </summary>
+		Ifdef,
+
+		/// <summary>
+		/// An #ifndef directive
+		/// </summary>
+		Ifndef,
+
+		/// <summary>
+		/// An #elif directive
+		/// </summary>
+		Elif,
+
+		/// <summary>
+		/// An #else directive
+		/// </summary>
+		Else,
+
+		/// <summary>
+		/// An #endif directive
+		/// </summary>
+		Endif,
+
+		/// <summary>
+		/// A #pragma directive
+		/// </summary>
+		Pragma,
+
+		/// <summary>
+		/// An #error directive
+		/// </summary>
+		Error,
+
+		/// <summary>
+		/// A #warning directive
+		/// </summary>
+		Warning,
+
+		/// <summary>
+		/// An empty '#' on a line of its own
+		/// </summary>
+		Empty,
+
+		/// <summary>
+		/// Some other directive
+		/// </summary>
+		OtherDirective
+	}
+
+	/// <summary>
+	/// Base class for an annotated section of a source file
+	/// </summary>
+	[Serializable]
+	class SourceFileMarkup
+	{
+		/// <summary>
+		/// The directive corresponding to this markup
+		/// </summary>
+		public SourceFileMarkupType Type;
+
+		/// <summary>
+		/// The one-based line number of this markup
+		/// </summary>
+		public int LineNumber;
+
+		/// <summary>
+		/// The tokens parsed for this markup. Set for directives.
+		/// </summary>
+		public List<Token> Tokens;
+
+		/// <summary>
+		/// Construct the annotation with the given range
+		/// </summary>
+		/// <param name="Type">The type of this directive</param>
+		/// <param name="LineNumber">The line number of this markup</param>
+		/// <param name="Tokens">List of tokens</param>
+		public SourceFileMarkup(SourceFileMarkupType Type, int LineNumber, List<Token> Tokens)
+		{
+			this.Type = Type;
+			this.LineNumber = LineNumber;
+			this.Tokens = Tokens;
+		}
+
+		/// <summary>
+		/// Constructs a markup object using data read from an archive
+		/// </summary>
+		/// <param name="Reader">The reader to deserialize from</param>
+		public SourceFileMarkup(BinaryArchiveReader Reader)
+		{
+			Type = (SourceFileMarkupType)Reader.ReadByte();
+			LineNumber = Reader.ReadInt();
+			Tokens = Reader.ReadList(() => Reader.ReadToken());
+		}
+
+		/// <summary>
+		/// Serializes this object to a binary archive
+		/// </summary>
+		/// <param name="Writer">Writer to serialize to</param>
+		public void Write(BinaryArchiveWriter Writer)
+		{
+			Writer.WriteByte((byte)Type);
+			Writer.WriteInt(LineNumber);
+			Writer.WriteList(Tokens, x => Writer.WriteToken(x));
+		}
+
+		/// <summary>
+		/// Determines if this markup indicates a conditional preprocessor directive
+		/// </summary>
+		/// <returns>True if this object is a conditional preprocessor directive</returns>
+		public bool IsConditionalPreprocessorDirective()
+		{
+			switch(Type)
+			{
+				case SourceFileMarkupType.If:
+				case SourceFileMarkupType.Ifdef:
+				case SourceFileMarkupType.Ifndef:
+				case SourceFileMarkupType.Elif:
+				case SourceFileMarkupType.Else:
+				case SourceFileMarkupType.Endif:
+					return true;
+			}
+			return false;
+		}
+
+		/// <summary>
+		/// How this condition modifies the condition depth. Opening "if" statements have a value of +1, "endif" statements have a value of -1, and "else" statements have a value of 0.
+		/// </summary>
+		public int GetConditionDepthDelta()
+		{
+			if(Type == SourceFileMarkupType.If || Type == SourceFileMarkupType.Ifdef || Type == SourceFileMarkupType.Ifndef)
+			{
+				return +1;
+			}
+			else if(Type == SourceFileMarkupType.Endif)
+			{
+				return -1;
+			}
+			else
+			{
+				return 0;
+			}
+		}
+
+		/// <summary>
+		/// Generate a string describing this annotation
+		/// </summary>
+		/// <returns>String representation for debugging</returns>
+		public override string ToString()
+		{
+			StringBuilder Result = new StringBuilder();
+			Result.AppendFormat("[{0}] ", LineNumber);
+
+			if(Type == SourceFileMarkupType.Text)
+			{
+				Result.Append("...");
+			}
+			else
+			{
+				Result.Append("#");
+				if(Type != SourceFileMarkupType.OtherDirective)
+				{
+					Result.Append(Type.ToString().ToLowerInvariant());
+				}
+				if(Tokens != null && Tokens.Count > 0)
+				{
+					Result.Append(' ');
+					Token.Format(Tokens, Result);
+				}
+			}
+			return Result.ToString();
+		}
+
+		/// <summary>
+		/// Create markup for the given file
+		/// </summary>
+		/// <param name="Reader">Reader for tokens in the file</param>
+		/// <returns>Array of markup objects which split up the given text buffer</returns>
+		public static SourceFileMarkup[] Parse(TokenReader Reader)
+		{
+			List<SourceFileMarkup> Markup = new List<SourceFileMarkup>();
+			if(Reader.MoveNext())
+			{
+				bool bMoveNext = true;
+				while(bMoveNext)
+				{
+					int StartLineNumber = Reader.LineNumber;
+					if (Reader.Current.Type == TokenType.Hash)
+					{
+						// Create the appropriate markup object for the directive
+						SourceFileMarkupType Type = SourceFileMarkupType.OtherDirective;
+						if(Reader.MoveNext())
+						{
+							if(Reader.Current.Type == TokenType.Identifier)
+							{
+								Identifier Directive = Reader.Current.Identifier;
+								if(Directive == Identifiers.Include)
+								{
+									Type = SourceFileMarkupType.Include;
+								}
+								else if(Directive == Identifiers.Define)
+								{
+									Type = SourceFileMarkupType.Define;
+								}
+								else if(Directive == Identifiers.Undef)
+								{
+									Type = SourceFileMarkupType.Undef;
+								}
+								else if(Directive == Identifiers.If)
+								{
+									Type = SourceFileMarkupType.If;
+								}
+								else if(Directive == Identifiers.Ifdef)
+								{
+									Type = SourceFileMarkupType.Ifdef;
+								}
+								else if(Directive == Identifiers.Ifndef)
+								{
+									Type = SourceFileMarkupType.Ifndef;
+								}
+								else if(Directive == Identifiers.Elif)
+								{
+									Type = SourceFileMarkupType.Elif;
+								}
+								else if(Directive == Identifiers.Else)
+								{
+									Type = SourceFileMarkupType.Else;
+								}
+								else if(Directive == Identifiers.Endif)
+								{
+									Type = SourceFileMarkupType.Endif;
+								}
+								else if(Directive == Identifiers.Pragma)
+								{
+									Type = SourceFileMarkupType.Pragma;
+								}
+								else if(Directive == Identifiers.Error)
+								{
+									Type = SourceFileMarkupType.Error;
+								}
+								else if(Directive == Identifiers.Warning)
+								{
+									Type = SourceFileMarkupType.Warning;
+								}
+							}
+							else if(Reader.Current.Type == TokenType.Newline)
+							{
+								Type = SourceFileMarkupType.Empty;
+							}
+						}
+
+						// Create the token list
+						List<Token> Tokens = new List<Token>();
+						if(Type == SourceFileMarkupType.OtherDirective)
+						{
+							Tokens.Add(Reader.Current);
+						}
+
+						// Read the first token
+						if(Type == SourceFileMarkupType.Empty)
+						{
+							bMoveNext = true;
+						}
+						else if(Type == SourceFileMarkupType.Include)
+						{
+							bMoveNext = Reader.MoveNextIncludePath();
+						}
+						else if(Type == SourceFileMarkupType.Error || Type == SourceFileMarkupType.Warning)
+						{
+							bMoveNext = Reader.MoveNextTokenString();
+						}
+						else
+						{
+							bMoveNext = Reader.MoveNext();
+						}
+
+						// Read the rest of the tokens
+						while(bMoveNext && Reader.Current.Type != TokenType.Newline)
+						{
+							Tokens.Add(Reader.Current);
+							bMoveNext = Reader.MoveNext();
+						}
+
+						// Create the markup
+						Markup.Add(new SourceFileMarkup(Type, StartLineNumber, Tokens));
+
+						// Move to the next token
+						bMoveNext = Reader.MoveNext();
+					}
+					else if(Reader.Current.Type != TokenType.Newline)
+					{
+						// Create the new fragment
+						Markup.Add(new SourceFileMarkup(SourceFileMarkupType.Text, StartLineNumber, null));
+
+						// Move to the next directive
+						bMoveNext = Reader.MoveToNextDirective();
+					}
+					else
+					{
+						// Skip the empty line
+						bMoveNext = Reader.MoveNext();
+					}
+				}
+			}
+			return Markup.ToArray();
+		}
+	}
+}
diff --git a/Engine/Source/Programs/UnrealBuildTool/Preprocessor/Token.cs b/Engine/Source/Programs/UnrealBuildTool/Preprocessor/Token.cs
new file mode 100644
--- /dev/null
+++ b/Engine/Source/Programs/UnrealBuildTool/Preprocessor/Token.cs
@@ -0,0 +1,538 @@
+// Copyright Epic Games, Inc. All Rights Reserved.
+
+using System;
+using System.Collections.Generic;
+using System.Diagnostics;
+using System.Linq;
+using System.Text;
+using System.Threading.Tasks;
+using Tools.DotNETCommon;
+
+namespace UnrealBuildTool
+{
+	/// <summary>
+	/// Flags for a single token
+	/// </summary>
+	[Flags]
+	enum TokenFlags : byte
+	{
+		/// <summary>
+		/// No flags
+		/// </summary>
+		None = 0x00,
+
+		/// <summary>
+		/// The token has space before it
+		/// </summary>
+		HasLeadingSpace = 0x01,
+
+		/// <summary>
+		/// Do not replace any instances of this token with the corresponding macro.
+		/// </summary>
+		DisableExpansion = 0x02,
+	}
+
+	/// <summary>
+	/// Enumeration for token types
+	/// </summary>
+	enum TokenType : byte
+	{
+		End,
+		LeftParen,
+		RightParen,
+		Comma,
+		Identifier,
+		Number,
+		String,
+		Character,
+		Newline,
+		Ellipsis,
+		StringOfTokens,
+		SystemInclude,
+		Placemarker,
+		Dot,
+		QuestionMark,
+		Colon,
+		LogicalNot,
+		LogicalAnd,
+		LogicalOr,
+		BitwiseXor,
+		BitwiseAnd,
+		BitwiseNot,
+		BitwiseOr,
+		Equals,
+		LeftShift,
+		RightShift,
+		CompareEqual,
+		CompareNotEqual,
+		CompareLessOrEqual,
+		CompareLess,
+		CompareGreaterOrEqual,
+		CompareGreater,
+		Plus,
+		Minus,
+		Multiply,
+		Divide,
+		Modulo,
+		Hash,
+		HashHash,
+		Unknown,
+		Max
+	}
+
+	/// <summary>
+	/// Single lexical token
+	/// </summary>
+	[DebuggerDisplay("{Text}")]
+	sealed class Token : IEquatable<Token>
+	{
+		/// <summary>
+		/// Static token names
+		/// </summary>
+		static string[] StaticTokens;
+
+		/// <summary>
+		/// Contains the type of the token
+		/// </summary>
+		public TokenType Type
+		{
+			get;
+		}
+
+		/// <summary>
+		/// Properties for this token
+		/// </summary>
+		public TokenFlags Flags
+		{
+			get;
+		}
+
+		/// <summary>
+		/// If this token is identifier, contains the identifier name
+		/// </summary>
+		public Identifier Identifier
+		{
+			get;
+		}
+
+		/// <summary>
+		/// If this token is a literal, contains the raw utf-8 representation of it.
+		/// </summary>
+		public byte[] Literal
+		{
+			get;
+		}
+
+		/// <summary>
+		/// Constructor
+		/// </summary>
+		/// <param name="Type">The token type</param>
+		/// <param name="Flags">Flags for this token</param>
+		public Token(TokenType Type, TokenFlags Flags)
+		{
+			this.Type = Type;
+			this.Flags = Flags;
+		}
+
+		/// <summary>
+		/// Constructor
+		/// </summary>
+		/// <param name="Identifier">The token identifier</param>
+		/// <param name="Flags">Flags for this token</param>
+		public Token(Identifier Identifier, TokenFlags Flags)
+		{
+			this.Type = TokenType.Identifier;
+			this.Flags = Flags;
+			this.Identifier = Identifier;
+		}
+
+		/// <summary>
+		/// Constructs a literal token
+		/// </summary>
+		public Token(TokenType Type, TokenFlags Flags, byte[] Literal) 
+			: this(Type, Flags)
+		{
+			Debug.Assert(IsLiteral(Type));
+			this.Literal = Literal;
+		}
+
+		/// <summary>
+		/// Constructs a literal token
+		/// </summary>
+		public Token(TokenType Type, TokenFlags Flags, string Literal) 
+			: this(Type, Flags)
+		{
+			Debug.Assert(IsLiteral(Type));
+			this.Literal = Encoding.UTF8.GetBytes(Literal);
+		}
+
+		/// <summary>
+		/// Reads a token from a binary archive
+		/// </summary>
+		/// <param name="Reader">Archive to read from</param>
+		public Token(BinaryArchiveReader Reader)
+		{
+			Type = (TokenType)Reader.ReadByte();
+			Flags = (TokenFlags)Reader.ReadByte();
+			if(Type == TokenType.Identifier)
+			{
+				Identifier = Reader.ReadIdentifier();
+			}
+			else if(IsLiteral(Type))
+			{
+				Literal = Reader.ReadByteArray();
+			}
+		}
+
+		/// <summary>
+		/// Writes a token to a binary archive
+		/// </summary>
+		/// <param name="Writer">The writer to serialize to</param>
+		public void Write(BinaryArchiveWriter Writer)
+		{
+			Writer.WriteByte((byte)Type);
+			Writer.WriteByte((byte)Flags);
+			if(Type == TokenType.Identifier)
+			{
+				Writer.WriteIdentifier(Identifier);
+			}
+			else if(IsLiteral(Type))
+			{
+				Writer.WriteByteArray(Literal);
+			}
+		}
+
+		/// <summary>
+		/// Checks if a token is equal to another object
+		/// </summary>
+		/// <param name="Other">The object to compare against</param>
+		/// <returns>True if the objects are equal, false otherwise</returns>
+		public override bool Equals(object Other)
+		{
+			return Equals(Other as Token);
+		}
+
+		/// <summary>
+		/// Checks if two tokens are equivalent
+		/// </summary>
+		/// <param name="Other">The object to compare against</param>
+		/// <returns>True if the tokens are equal, false otherwise</returns>
+		public bool Equals(Token Other)
+		{
+			if((object)Other == null)
+			{
+				return false;
+			}
+			if(Type != Other.Type || Flags != Other.Flags || Identifier != Other.Identifier)
+			{
+				return false;
+			}
+			if(Literal != null)
+			{
+				if(Other.Literal == null || Literal.Length != Other.Literal.Length || !Enumerable.SequenceEqual(Literal, Other.Literal))
+				{
+					return false;
+				}
+			}
+			else
+			{
+				if(Other.Literal != null)
+				{
+					return false;
+				}
+			}
+			return true;
+		}
+
+		/// <summary>
+		/// Compares two tokens for equality
+		/// </summary>
+		/// <param name="A">The first token to compare</param>
+		/// <param name="B">The second token to compare</param>
+		/// <returns>True if the objects are equal, false otherwise</returns>
+		public static bool operator==(Token A, Token B)
+		{
+			if((object)A == null)
+			{
+				return ((object)B == null);
+			}
+			else
+			{
+				return A.Equals(B);
+			}
+		}
+
+		/// <summary>
+		/// Compares two tokens for inequality
+		/// </summary>
+		/// <param name="A">The first token to compare</param>
+		/// <param name="B">The second token to compare</param>
+		/// <returns>True if the objects are not equal, false otherwise</returns>
+		public static bool operator!=(Token A, Token B)
+		{
+			return !(A == B);
+		}
+
+		/// <summary>
+		/// Gets a hash code for this object
+		/// </summary>
+		/// <returns>Hash code for the object</returns>
+		public override int GetHashCode()
+		{
+			int Result = (int)Type + (int)Flags * 7;
+			if(Identifier != null)
+			{
+				Result = (Result * 11) + Identifier.GetHashCode();
+			}
+			if(Literal != null)
+			{
+				for(int Idx = 0; Idx < Literal.Length; Idx++)
+				{
+					Result = (Result * 13) + Literal[Idx];
+				}
+			}
+			return base.GetHashCode();
+		}
+
+		/// <summary>
+		/// Text corresponding to this token
+		/// </summary>
+		public string Text
+		{
+			get
+			{
+				if(Identifier != null)
+				{
+					return Identifier.ToString();
+				}
+				else if(Literal != null)
+				{
+					return Encoding.UTF8.GetString(Literal);
+				}
+				else
+				{
+					return StaticTokens[(int)Type];
+				}
+			}
+		}
+
+		/// <summary>
+		/// Returns a new token with different flags
+		/// </summary>
+		/// <param name="FlagsToAdd">Flags to add from the token</param>
+		/// <returns>New token with updated flags</returns>
+		public Token AddFlags(TokenFlags FlagsToAdd)
+		{
+			if(Identifier != null)
+			{
+				return new Token(Identifier, Flags | FlagsToAdd);
+			}
+			else if(Literal != null)
+			{
+				return new Token(Type, Flags | FlagsToAdd, Literal);
+			}
+			else
+			{
+				return new Token(Type, Flags | FlagsToAdd);
+			}
+		}
+
+		/// <summary>
+		/// Returns a new token with different flags
+		/// </summary>
+		/// <param name="FlagsToRemove">Flags to remove from the token</param>
+		/// <returns>New token with updated flags</returns>
+		public Token RemoveFlags(TokenFlags FlagsToRemove)
+		{
+			if(Identifier != null)
+			{
+				return new Token(Identifier, Flags & ~FlagsToRemove);
+			}
+			else if(Literal != null)
+			{
+				return new Token(Type, Flags & ~FlagsToRemove, Literal);
+			}
+			else
+			{
+				return new Token(Type, Flags & ~FlagsToRemove);
+			}
+		}
+
+		/// <summary>
+		/// Accessor for whether this token has leading whitespace
+		/// </summary>
+		public bool HasLeadingSpace
+		{
+			get { return (Flags & TokenFlags.HasLeadingSpace) != 0; }
+		}
+
+		/// <summary>
+		/// Checks whether two tokens match
+		/// </summary>
+		/// <param name="Other">The token to compare against</param>
+		/// <returns>True if the tokens match, false otherwise</returns>
+		public bool Matches(Token Other)
+		{
+			if(Type != Other.Type)
+			{
+				return false;
+			}
+
+			if(Literal == null)
+			{
+				return Identifier == Other.Identifier;
+			}
+			else 
+			{
+				return Enumerable.SequenceEqual(Literal, Other.Literal);
+			}
+		}
+
+		/// <summary>
+		/// Determines whether the given token type is a literal
+		/// </summary>
+		/// <param name="Type">The type to test</param>
+		/// <returns>True if the given type is a literal</returns>
+		public static bool IsLiteral(TokenType Type)
+		{
+			return Type == TokenType.Unknown || Type == TokenType.Character || Type == TokenType.String || Type == TokenType.Number || Type == TokenType.StringOfTokens || Type == TokenType.SystemInclude;
+		}
+
+		/// <summary>
+		/// Concatenate a sequence of tokens into a string
+		/// </summary>
+		/// <param name="Tokens">The sequence of tokens to concatenate</param>
+		/// <returns>String containing the concatenated tokens</returns>
+		public static string Format(IEnumerable<Token> Tokens)
+		{
+			StringBuilder Result = new StringBuilder();
+			Format(Tokens, Result);
+			return Result.ToString();
+		}
+
+		/// <summary>
+		/// Concatenate a sequence of tokens into a string
+		/// </summary>
+		/// <param name="Tokens">The sequence of tokens to concatenate</param>
+		/// <param name="Result">Receives the formatted string</param>
+		public static void Format(IEnumerable<Token> Tokens, StringBuilder Result)
+		{
+			IEnumerator<Token> Enumerator = Tokens.GetEnumerator();
+			if(Enumerator.MoveNext())
+			{
+				Result.Append(Enumerator.Current.Text);
+
+				Token LastToken = Enumerator.Current;
+				while(Enumerator.MoveNext())
+				{
+					Token Token = Enumerator.Current;
+					if(Token.HasLeadingSpace && (Token.Type != TokenType.LeftParen || LastToken.Type != TokenType.Identifier || LastToken.Identifier != Identifiers.__pragma))
+					{
+						Result.Append(" ");
+					}
+					Result.Append(Token.Text);
+					LastToken = Token;
+				}
+			}
+		}
+
+		/// <summary>
+		/// Concatenate two tokens together
+		/// </summary>
+		/// <param name="FirstToken">The first token</param>
+		/// <param name="SecondToken">The second token</param>
+		/// <param name="Context">Current preprocessor context</param>
+		/// <returns>The combined token</returns>
+		public static Token Concatenate(Token FirstToken, Token SecondToken, PreprocessorContext Context)
+		{
+			string Text = FirstToken.Text.ToString() + SecondToken.Text.ToString();
+
+			List<Token> Tokens = new List<Token>();
+
+			TokenReader Reader = new TokenReader(Text);
+			while(Reader.MoveNext())
+			{
+				Tokens.Add(Reader.Current);
+			}
+
+			if(Tokens.Count == 0)
+			{
+				return new Token(TokenType.Placemarker, TokenFlags.None);
+			}
+			else if(Tokens.Count == 1)
+			{
+				return Tokens[0];
+			}
+			else
+			{
+				return new Token(TokenType.Unknown, TokenFlags.None, Text);
+			}
+		}
+
+		/// <summary>
+		/// Static constructor. Initializes the StaticTokens array.
+		/// </summary>
+		static Token()
+		{
+			StaticTokens = new string[(int)TokenType.Max];
+			StaticTokens[(int)TokenType.LeftParen] = "(";
+			StaticTokens[(int)TokenType.RightParen] = ")";
+			StaticTokens[(int)TokenType.Comma] = ",";
+			StaticTokens[(int)TokenType.Newline] = "\n";
+			StaticTokens[(int)TokenType.Ellipsis] = "...";
+			StaticTokens[(int)TokenType.Placemarker] = "";
+			StaticTokens[(int)TokenType.Dot] = ".";
+			StaticTokens[(int)TokenType.QuestionMark] = "?";
+			StaticTokens[(int)TokenType.Colon] = ":";
+			StaticTokens[(int)TokenType.LogicalNot] = "!";
+			StaticTokens[(int)TokenType.LogicalAnd] = "&&";
+			StaticTokens[(int)TokenType.LogicalOr] = "||";
+			StaticTokens[(int)TokenType.BitwiseXor] = "^";
+			StaticTokens[(int)TokenType.BitwiseAnd] = "&";
+			StaticTokens[(int)TokenType.BitwiseNot] = "~";
+			StaticTokens[(int)TokenType.BitwiseOr] = "|";
+			StaticTokens[(int)TokenType.Equals] = "=";
+			StaticTokens[(int)TokenType.LeftShift] = "<<";
+			StaticTokens[(int)TokenType.RightShift] = ">>";
+			StaticTokens[(int)TokenType.CompareEqual] = "==";
+			StaticTokens[(int)TokenType.CompareNotEqual] = "!=";
+			StaticTokens[(int)TokenType.CompareLessOrEqual] = "<=";
+			StaticTokens[(int)TokenType.CompareLess] = "<";
+			StaticTokens[(int)TokenType.CompareGreaterOrEqual] = ">=";
+			StaticTokens[(int)TokenType.CompareGreater] = ">";
+			StaticTokens[(int)TokenType.Plus] = "+";
+			StaticTokens[(int)TokenType.Minus] = "-";
+			StaticTokens[(int)TokenType.Multiply] = "*";
+			StaticTokens[(int)TokenType.Divide] = "/";
+			StaticTokens[(int)TokenType.Modulo] = "%";
+			StaticTokens[(int)TokenType.Hash] = "#";
+			StaticTokens[(int)TokenType.HashHash] = "##";
+		}
+	}
+
+	/// <summary>
+	/// Helper functions for serialization
+	/// </summary>
+	static class TokenExtensionMethods
+	{
+		/// <summary>
+		/// Read a token from a binary archive
+		/// </summary>
+		/// <param name="Reader">Reader to serialize data from</param>
+		/// <returns>Token read from the archive</returns>
+		public static Token ReadToken(this BinaryArchiveReader Reader)
+		{
+			return new Token(Reader);
+		}
+
+		/// <summary>
+		/// Write a token to a binary archive
+		/// </summary>
+		/// <param name="Writer">Writer to serialize data to</param>
+		/// <param name="Token">Token to write</param>
+		public static void WriteToken(this BinaryArchiveWriter Writer, Token Token)
+		{
+			Token.Write(Writer);
+		}
+	}
+}
diff --git a/Engine/Source/Programs/UnrealBuildTool/Preprocessor/TokenReader.cs b/Engine/Source/Programs/UnrealBuildTool/Preprocessor/TokenReader.cs
new file mode 100644
--- /dev/null
+++ b/Engine/Source/Programs/UnrealBuildTool/Preprocessor/TokenReader.cs
@@ -0,0 +1,821 @@
+// Copyright Epic Games, Inc. All Rights Reserved.
+
+using System;
+using System.Collections;
+using System.Collections.Generic;
+using System.Linq;
+using System.Text;
+using System.Threading.Tasks;
+
+namespace UnrealBuildTool
+{
+	/// <summary>
+	/// Tokenizer for C++ source files. Provides functions for navigating a source file skipping whitespace, comments, and so on when required.
+	/// </summary>
+	class TokenReader : IEnumerator<Token>
+	{
+		/// <summary>
+		/// Flags used to classify the first character in a token
+		/// </summary>
+		enum FirstCharacterClass : byte
+		{
+			Unknown,
+			Terminator,
+			Whitespace,
+			Newline,
+			Identifier,
+			Number,
+			Character,
+			String,
+			Dot,
+			QuestionMark,
+			Colon,
+			ExclamationMark,
+			Pipe,
+			Ampersand,
+			Caret,
+			Equals,
+			LeftTriangleBracket,
+			RightTriangleBracket,
+			Plus,
+			Minus,
+			Star,
+			Slash,
+			PercentSign,
+			Tilde,
+			LeftParen,
+			RightParen,
+			Hash,
+			Comma,
+		}
+
+		/// <summary>
+		/// Flags used to classify subsequent characters in a token
+		/// </summary>
+		[Flags]
+		enum CharacterFlags : byte
+		{
+			Identifier = 1,
+			Digit = 2,
+			NumberTail = 4,
+			Whitespace = 8,
+		}
+
+		/// <summary>
+		/// Map of utf-8 leading bytes to their class
+		/// </summary>
+		static FirstCharacterClass[] FirstCharacters;
+
+		/// <summary>
+		/// Array of flags for different leading utf-8 sequences
+		/// </summary>
+		static CharacterFlags[] Characters;
+
+		/// <summary>
+		/// The current buffer being read from. Encoded as UTF-8 with a null terminator.
+		/// </summary>
+		byte[] Data;
+
+		/// <summary>
+		/// Current offset within the buffer
+		/// </summary>
+		int Offset;
+
+		/// <summary>
+		/// The current token
+		/// </summary>
+		Token CurrentToken;
+
+		/// <summary>
+		/// Constructor
+		/// </summary>
+		/// <param name="Data">The raw byte stream to read from, encoded as UTF-8</param>
+		public TokenReader(byte[] Data)
+		{
+			this.Data = Data;
+			this.LineNumber = 0;
+			this.LineNumberAfterToken = 1;
+
+			// Make sure the input data has a null terminator
+			if(Data.Length == 0 || Data[Data.Length - 1] != 0)
+			{
+				throw new ArgumentException("Data parameter must be null terminated.");
+			}
+
+			// If the data contains a UTF-8 BOM, skip over it
+			if(Data[0] == 0xef && Data[1] == 0xbb && Data[2] == 0xbf)
+			{
+				Offset = 3;
+			}
+		}
+
+		/// <summary>
+		/// Constructor
+		/// </summary>
+		/// <param name="Text">The text to read from</param>
+		public TokenReader(string Text)
+			: this(GetNullTerminatedByteArray(Text))
+		{
+		}
+
+		/// <summary>
+		/// Copy constructor
+		/// </summary>
+		/// <param name="Other">Token reader to copy from</param>
+		public TokenReader(TokenReader Other)
+		{
+			this.Data = Other.Data;
+			this.Offset = Other.Offset;
+			this.LineNumber = Other.LineNumber;
+			this.LineNumberAfterToken = Other.LineNumberAfterToken;
+		}
+
+		/// <summary>
+		/// Gets a null terminated byte array from the given string
+		/// </summary>
+		/// <param name="Text">String to convert into bytes</param>
+		/// <returns>Array of bytes</returns>
+		public static byte[] GetNullTerminatedByteArray(string Text)
+		{
+			byte[] Bytes = new byte[Encoding.UTF8.GetByteCount(Text) + 1];
+			Encoding.UTF8.GetBytes(Text, 0, Text.Length, Bytes, 0);
+			return Bytes;
+		}
+
+		/// <summary>
+		/// Tokenize an input string
+		/// </summary>
+		/// <param name="Text">Text to tokenize</param>
+		/// <returns>List of tokens parsed from the text</returns>
+		public static List<Token> Tokenize(string Text)
+		{
+			List<Token> Tokens = new List<Token>();
+
+			TokenReader Reader = new TokenReader(Text);
+			while(Reader.MoveNext())
+			{
+				Tokens.Add(Reader.Current);
+			}
+
+			return Tokens;
+		}
+
+		/// <summary>
+		/// Current line number (one-based)
+		/// </summary>
+		public int LineNumber
+		{
+			get;
+			private set;
+		}
+
+		/// <summary>
+		/// Line number at the end of this token (one-based)
+		/// </summary>
+		public int LineNumberAfterToken
+		{
+			get;
+			private set;
+		}
+
+		/// <summary>
+		/// Returns the current token
+		/// </summary>
+		public Token Current
+		{
+			get { return CurrentToken; }
+		}
+
+		/// <summary>
+		/// Untyped implementation of Current for IEnumerator.
+		/// </summary>
+		object IEnumerator.Current
+		{
+			get { return Current; }
+		}
+
+		/// <summary>
+		/// Override of IEnumerator.Dispose. Not required.
+		/// </summary>
+		void IDisposable.Dispose()
+		{
+		}
+
+		/// <summary>
+		/// Move to the next token
+		/// </summary>
+		/// <returns>True if the reader could move to the next token, false otherwise</returns>
+		public bool MoveNext()
+		{
+			// Skip past the leading whitespace
+			TokenFlags Flags = SkipWhitespace();
+
+			// Update the line number to the one after the current token
+			LineNumber = LineNumberAfterToken;
+
+			// Capture the start of the token
+			int StartOffset = Offset;
+
+			// Initialize the flags for the next token
+			byte FirstByte = Data[Offset++];
+			switch(FirstCharacters[FirstByte])
+			{
+				case FirstCharacterClass.Unknown:
+					if((FirstByte & 0x80) != 0)
+					{
+						while((Data[Offset] & 0x80) != 0)
+						{
+							Offset++;
+						}
+					}
+					CurrentToken = CreateLiteral(TokenType.Unknown, Flags, StartOffset, Offset - StartOffset);
+					return true;
+				case FirstCharacterClass.Terminator:
+					Offset = StartOffset;
+					CurrentToken = new Token(TokenType.End, Flags);
+					return false;
+				case FirstCharacterClass.Newline:
+					LineNumberAfterToken++;
+					CurrentToken = new Token(TokenType.Newline, Flags);
+					return true;
+				case FirstCharacterClass.Identifier:
+					// Identifier (or text literal with prefix)
+					while((Characters[Data[Offset]] & CharacterFlags.Identifier) != 0)
+					{
+						Offset++;
+					}
+
+					// Check if it's a prefixed text literal
+					if(Data[Offset] == '\'')
+					{
+						Offset++;
+						SkipTextLiteral('\'');
+						CurrentToken = CreateLiteral(TokenType.Character, Flags, StartOffset, Offset - StartOffset);
+						return true;
+					}
+					else if(Data[Offset] == '\"')
+					{
+						Offset++;
+						SkipTextLiteral('\"');
+						CurrentToken = CreateLiteral(TokenType.String, Flags, StartOffset, Offset - StartOffset);
+						return true;
+					}
+					else
+					{
+						string Name = Encoding.UTF8.GetString(Data, StartOffset, Offset - StartOffset);
+						Identifier Identifier = Identifier.FindOrAdd(Name);
+						CurrentToken = new Token(Identifier, Flags);
+						return true;
+					}
+				case FirstCharacterClass.Number:
+					// Numeric literal
+					SkipNumericLiteral();
+					CurrentToken = CreateLiteral(TokenType.Number, Flags, StartOffset, Offset - StartOffset);
+					return true;
+				case FirstCharacterClass.Character:
+					// Character literal
+					SkipTextLiteral('\'');
+					CurrentToken = CreateLiteral(TokenType.Character, Flags, StartOffset, Offset - StartOffset);
+					return true;
+				case FirstCharacterClass.String:
+					// String literal
+					SkipTextLiteral('\"');
+					CurrentToken = CreateLiteral(TokenType.String, Flags, StartOffset, Offset - StartOffset);
+					return true;
+				case FirstCharacterClass.Dot:
+					// Numeric literal, ellipsis, or dot
+					if((Characters[Data[Offset]] & CharacterFlags.Digit) != 0)
+					{
+						Offset++;
+						SkipNumericLiteral();
+						CurrentToken = CreateLiteral(TokenType.Number, Flags, StartOffset, Offset - StartOffset);
+						return true;
+					}
+					else if(Data[Offset] == '.' && Data[Offset + 1] == '.')
+					{
+						Offset += 2;
+						CurrentToken = new Token(TokenType.Ellipsis, Flags);
+						return true;
+					}
+					else
+					{
+						CurrentToken = new Token(TokenType.Dot, Flags);
+						return true;
+					}
+				case FirstCharacterClass.QuestionMark:
+					CurrentToken = new Token(TokenType.QuestionMark, Flags);
+					return true;
+				case FirstCharacterClass.Colon:
+					CurrentToken = new Token(TokenType.Colon, Flags);
+					return true;
+				case FirstCharacterClass.ExclamationMark:
+					if(ReadCharacter('='))
+					{
+						CurrentToken = new Token(TokenType.CompareNotEqual, Flags);
+					}
+					else
+					{
+						CurrentToken = new Token(TokenType.LogicalNot, Flags);
+					}
+					return true;
+				case FirstCharacterClass.Pipe:
+					if(ReadCharacter('|'))
+					{
+						CurrentToken = new Token(TokenType.LogicalOr, Flags);
+					}
+					else
+					{
+						CurrentToken = new Token(TokenType.BitwiseOr, Flags);
+					}
+					return true;
+				case FirstCharacterClass.Ampersand:
+					if(ReadCharacter('&'))
+					{
+						CurrentToken = new Token(TokenType.LogicalAnd, Flags);
+					}
+					else
+					{
+						CurrentToken = new Token(TokenType.BitwiseAnd, Flags);
+					}
+					return true;
+				case FirstCharacterClass.Caret:
+					CurrentToken = new Token(TokenType.BitwiseXor, Flags);
+					return true;
+				case FirstCharacterClass.Equals:
+					if(ReadCharacter('='))
+					{
+						CurrentToken = new Token(TokenType.CompareEqual, Flags);
+					}
+					else
+					{
+						CurrentToken = new Token(TokenType.Equals, Flags);
+					}
+					return true;
+				case FirstCharacterClass.LeftTriangleBracket:
+					if(ReadCharacter('<'))
+					{
+						CurrentToken = new Token(TokenType.LeftShift, Flags);
+					}
+					else if(ReadCharacter('='))
+					{
+						CurrentToken = new Token(TokenType.CompareLessOrEqual, Flags);
+					}
+					else
+					{
+						CurrentToken = new Token(TokenType.CompareLess, Flags);
+					}
+					return true;
+				case FirstCharacterClass.RightTriangleBracket:
+					if(ReadCharacter('>'))
+					{
+						CurrentToken = new Token(TokenType.RightShift, Flags);
+					}
+					else if(ReadCharacter('='))
+					{
+						CurrentToken = new Token(TokenType.CompareGreaterOrEqual, Flags);
+					}
+					else
+					{
+						CurrentToken = new Token(TokenType.CompareGreater, Flags);
+					}
+					return true;
+				case FirstCharacterClass.Plus:
+					CurrentToken = new Token(TokenType.Plus, Flags);
+					return true;
+				case FirstCharacterClass.Minus:
+					CurrentToken = new Token(TokenType.Minus, Flags);
+					return true;
+				case FirstCharacterClass.Star:
+					CurrentToken = new Token(TokenType.Multiply, Flags);
+					return true;
+				case FirstCharacterClass.Slash:
+					CurrentToken = new Token(TokenType.Divide, Flags);
+					return true;
+				case FirstCharacterClass.PercentSign:
+					CurrentToken = new Token(TokenType.Modulo, Flags);
+					return true;
+				case FirstCharacterClass.Tilde:
+					CurrentToken = new Token(TokenType.BitwiseNot, Flags);
+					return true;
+				case FirstCharacterClass.LeftParen:
+					CurrentToken = new Token(TokenType.LeftParen, Flags);
+					return true;
+				case FirstCharacterClass.RightParen:
+					CurrentToken = new Token(TokenType.RightParen, Flags);
+					return true;
+				case FirstCharacterClass.Hash:
+					if(ReadCharacter('#'))
+					{
+						CurrentToken = new Token(TokenType.HashHash, Flags);
+					}
+					else
+					{
+						CurrentToken = new Token(TokenType.Hash, Flags);
+					}
+					return true;
+				case FirstCharacterClass.Comma:
+					CurrentToken = new Token(TokenType.Comma, Flags);
+					return true;
+				default:
+					throw new NotImplementedException();
+			}
+		}
+
+		/// <summary>
+		/// Move to the next token, parsing a single literal token string until the end of the line (in the form used by error directives)
+		/// </summary>
+		/// <returns>True if a token was parsed</returns>
+		public bool MoveNextTokenString()
+		{
+			TokenFlags Flags = SkipWhitespace();
+
+			int StartOffset = Offset;
+			while(Data[Offset] != 0 && Data[Offset] != '\n')
+			{
+				Offset++;
+			}
+
+			CurrentToken = CreateLiteral(TokenType.StringOfTokens, TokenFlags.None, StartOffset, Offset - StartOffset);
+			return true;
+		}
+
+		/// <summary>
+		/// Move to the next token, allowing include directives in triangle brackets in the style permitted by system include directives
+		/// </summary>
+		/// <returns>True if a token was read</returns>
+		public bool MoveNextIncludePath()
+		{
+			bool bResult = MoveNext();
+			if(bResult && CurrentToken.Type == TokenType.CompareLess)
+			{
+				int StartOffset = Offset - 1;
+				for(int EndOffset = Offset; Data[EndOffset] != 0 && Data[EndOffset] != '\n'; EndOffset++)
+				{
+					if(Data[EndOffset] == '>')
+					{
+						CurrentToken = CreateLiteral(TokenType.SystemInclude, TokenFlags.None, StartOffset, EndOffset + 1 - StartOffset);
+						Offset = EndOffset + 1;
+						break;
+					}
+				}
+			}
+			return bResult;
+		}
+
+		/// <summary>
+		/// Creates a literal token
+		/// </summary>
+		/// <param name="Type">Type of token to create</param>
+		/// <param name="Flags">Flags for the token</param>
+		/// <param name="Offset">Offset of the literal within the source data stream</param>
+		/// <param name="Length">Length of the literal</param>
+		/// <returns>New token for the literal</returns>
+		private Token CreateLiteral(TokenType Type, TokenFlags Flags, int Offset, int Length)
+		{
+			byte[] Literal = new byte[Length];
+			for(int Idx = 0; Idx < Length; Idx++)
+			{
+				Literal[Idx] = Data[Offset + Idx];
+			}
+			return new Token(Type, Flags, Literal);
+		}
+
+		/// <summary>
+		/// Scan ahead until we reach the next directive (a hash token after a newline). Assumes that the current line does not contain a directive.
+		/// </summary>
+		/// <returns>True if we were able to find another token</returns>
+		public bool MoveToNextDirective()
+		{
+			// Scan lines until we reach a directive
+			for(;;)
+			{
+				// Move to the next newline
+				while(Data[Offset] != '\n')
+				{
+					if(Data[Offset] == 0)
+					{
+						return false;
+					}
+					else if(Data[Offset] == '\\' && Data[Offset + 1] == '\r' && Data[Offset + 2] == '\n')
+					{
+						LineNumberAfterToken++;
+						Offset += 3;
+					}
+					else if(Data[Offset] == '\\' && Data[Offset + 1] == '\n')
+					{
+						LineNumberAfterToken++;
+						Offset += 2;
+					}
+					else if (Data[Offset] == '/' && Data[Offset + 1] == '*')
+					{
+						Offset += 2;
+						for(; Data[Offset] != 0; Offset++)
+						{
+							if(Data[Offset] == '\n')
+							{
+								LineNumberAfterToken++;
+								continue;
+							}
+							if(Data[Offset] == '*' && Data[Offset + 1] == '/')
+							{
+								Offset += 2;
+								break;
+							}
+						}
+					}
+					else
+					{
+						Offset++;
+					}
+				}
+
+				// Move past the newline
+				Offset++;
+				LineNumberAfterToken++;
+
+				// Skip any horizontal whitespace
+				TokenFlags Flags = SkipWhitespace();
+
+				// Check if this is a line marker
+				if(Data[Offset] == '#' && Data[Offset + 1] != '#')
+				{
+					Offset++;
+					LineNumber = LineNumberAfterToken;
+					CurrentToken = new Token(TokenType.Hash, Flags);
+					return true;
+				}
+			}
+		}
+
+		/// <summary>
+		/// Definition of IEnumerator.Reset(). Not supported.
+		/// </summary>
+		void IEnumerator.Reset()
+		{
+			throw new NotSupportedException();
+		}
+
+		/// <summary>
+		/// Attempts to read a given character from the stream
+		/// </summary>
+		/// <param name="Character">Character to read</param>
+		/// <returns>True if the character was read (and the current position was updated)</returns>
+		bool ReadCharacter(char Character)
+		{
+			if(Data[Offset] == Character)
+			{
+				Offset++;
+				return true;
+			}
+			return false;
+		}
+
+		/// <summary>
+		/// Advances the given position past any horizontal whitespace or comments
+		/// </summary>
+		/// <returns>Flags for the following token</returns>
+		TokenFlags SkipWhitespace()
+		{
+			TokenFlags Flags = TokenFlags.None;
+			for(;;)
+			{
+				// Quickly skip over trivial whitespace
+				while((Characters[Data[Offset]] & CharacterFlags.Whitespace) != 0)
+				{
+					Offset++;
+					Flags |= TokenFlags.HasLeadingSpace;
+				}
+
+				// Look at what's next
+				char Character = (char)Data[Offset];
+				if (Character == '\\')
+				{
+					byte NextCharacter = Data[Offset + 1];
+					if(NextCharacter == '\r' && Data[Offset + 2] == '\n')
+					{
+						LineNumberAfterToken++;
+						Offset += 3;
+					}
+					else if (NextCharacter == '\n')
+					{
+						LineNumberAfterToken++;
+						Offset += 2;
+					}
+					else
+					{
+						break;
+					}
+				}
+				else if (Character == '/')
+				{
+					byte NextCharacter = Data[Offset + 1];
+					if (NextCharacter == '/')
+					{
+						Offset += 2;
+						for(; Data[Offset] != 0; Offset++)
+						{
+							if(Data[Offset] == '\n')
+							{
+								break;
+							}
+						}
+						Flags |= TokenFlags.HasLeadingSpace;
+					}
+					else if (NextCharacter == '*')
+					{
+						Offset += 2;
+						for(; Data[Offset] != 0; Offset++)
+						{
+							if(Data[Offset] == '\n')
+							{
+								LineNumberAfterToken++;
+								continue;
+							}
+							if(Data[Offset] == '*' && Data[Offset + 1] == '/')
+							{
+								Offset += 2;
+								break;
+							}
+						}
+						Flags |= TokenFlags.HasLeadingSpace;
+					}
+					else
+					{
+						break;
+					}
+				}
+				else
+				{
+					if(Character == '\r')
+					{
+						Offset++;
+					}
+					else
+					{
+						break;
+					}
+				}
+			}
+			return Flags;
+		}
+
+		/// <summary>
+		/// Skip past a text literal (a quoted character literal or string literal)
+		/// </summary>
+		/// <param name="LastCharacter">The terminating character to look for, ignoring escape sequences</param>
+		void SkipTextLiteral(char LastCharacter)
+		{
+			for(;;)
+			{
+				char Character = (char)Data[Offset];
+				if(Character == '\0')
+				{
+					throw new Exception("Unexpected end of file in text literal");
+				}
+
+				Offset++;
+
+				if(Character == '\\' && Data[Offset] != 0)
+				{
+					Offset++;
+				}
+				else if(Character == LastCharacter)
+				{
+					break;
+				}
+			}
+		}
+
+		/// <summary>
+		/// Skips over a numeric literal after the initial digit or dot/digit pair.
+		/// </summary>
+		void SkipNumericLiteral()
+		{
+			for (;;)
+			{
+				while((Characters[Data[Offset]] & CharacterFlags.NumberTail) != 0)
+				{
+					Offset++;
+				}
+				if((Data[Offset] == '+' || Data[Offset] == '-') && (Data[Offset - 1] == 'e' || Data[Offset - 1] == 'E'))
+				{
+					Offset++;
+				}
+				else
+				{
+					break;
+				}
+			}
+		}
+
+		/// <summary>
+		/// Static constructor. Initializes the lookup tables used by the lexer.
+		/// </summary>
+		static TokenReader()
+		{
+			FirstCharacters = new FirstCharacterClass[256];
+
+			// End of file
+			FirstCharacters[0] = FirstCharacterClass.Terminator;
+
+			// Horizontal whitespace
+			FirstCharacters[' '] = FirstCharacterClass.Whitespace;
+			FirstCharacters['\t'] = FirstCharacterClass.Whitespace;
+			FirstCharacters['\v'] = FirstCharacterClass.Whitespace;
+			FirstCharacters['\r'] = FirstCharacterClass.Whitespace;
+
+			// Newline
+			FirstCharacters['\n'] = FirstCharacterClass.Newline;
+
+			// Identifiers
+			FirstCharacters['_'] = FirstCharacterClass.Identifier;
+			for(int Idx = 'a'; Idx <= 'z'; Idx++)
+			{
+				FirstCharacters[Idx] = FirstCharacterClass.Identifier;
+			}
+			for(int Idx = 'A'; Idx <= 'Z'; Idx++)
+			{
+				FirstCharacters[Idx] = FirstCharacterClass.Identifier;
+			}
+
+			// Numeric literals
+			for(int Idx = '0'; Idx <= '9'; Idx++)
+			{
+				FirstCharacters[Idx] = FirstCharacterClass.Number;
+			}
+
+			// Character literals
+			FirstCharacters['\''] = FirstCharacterClass.Character;
+
+			// String literals
+			FirstCharacters['\"'] = FirstCharacterClass.String;
+
+			// Other symbols
+			FirstCharacters['.'] = FirstCharacterClass.Dot;
+			FirstCharacters['?'] = FirstCharacterClass.QuestionMark;
+			FirstCharacters[':'] = FirstCharacterClass.Colon;
+			FirstCharacters['!'] = FirstCharacterClass.ExclamationMark;
+			FirstCharacters['|'] = FirstCharacterClass.Pipe;
+			FirstCharacters['&'] = FirstCharacterClass.Ampersand;
+			FirstCharacters['^'] = FirstCharacterClass.Caret;
+			FirstCharacters['='] = FirstCharacterClass.Equals;
+			FirstCharacters['<'] = FirstCharacterClass.LeftTriangleBracket;
+			FirstCharacters['>'] = FirstCharacterClass.RightTriangleBracket;
+			FirstCharacters['+'] = FirstCharacterClass.Plus;
+			FirstCharacters['-'] = FirstCharacterClass.Minus;
+			FirstCharacters['*'] = FirstCharacterClass.Star;
+			FirstCharacters['/'] = FirstCharacterClass.Slash;
+			FirstCharacters['%'] = FirstCharacterClass.PercentSign;
+			FirstCharacters['~'] = FirstCharacterClass.Tilde;
+			FirstCharacters['('] = FirstCharacterClass.LeftParen;
+			FirstCharacters[')'] = FirstCharacterClass.RightParen;
+			FirstCharacters['#'] = FirstCharacterClass.Hash;
+			FirstCharacters[','] = FirstCharacterClass.Comma;
+
+			// Flags for secondary characters
+			Characters = new CharacterFlags[256];
+
+			// Identifiers
+			Characters['_'] |= CharacterFlags.Identifier;
+			for(int Idx = 'a'; Idx <= 'z'; Idx++)
+			{
+				Characters[Idx] |= CharacterFlags.Identifier;
+			}
+			for(int Idx = 'A'; Idx <= 'Z'; Idx++)
+			{
+				Characters[Idx] |= CharacterFlags.Identifier;
+			}
+			for(int Idx = '0'; Idx <= '9'; Idx++)
+			{
+				Characters[Idx] |= CharacterFlags.Identifier;
+			}
+
+			// Numbers
+			for(int Idx = '0'; Idx <= '9'; Idx++)
+			{
+				Characters[Idx] |= CharacterFlags.Digit;
+			}
+
+			// Preprocessing number tail
+			Characters['.'] |= CharacterFlags.NumberTail;
+			Characters['_'] |= CharacterFlags.NumberTail;
+			for(int Idx = '0'; Idx <= '9'; Idx++)
+			{
+				Characters[Idx] |= CharacterFlags.NumberTail;
+			}
+			for(int Idx = 'a'; Idx <= 'z'; Idx++)
+			{
+				Characters[Idx] |= CharacterFlags.NumberTail;
+			}
+			for(int Idx = 'A'; Idx <= 'Z'; Idx++)
+			{
+				Characters[Idx] |= CharacterFlags.NumberTail;
+			}
+
+			// Whitespace
+			Characters[' '] |= CharacterFlags.Whitespace;
+			Characters['\t'] |= CharacterFlags.Whitespace;
+			Characters['\v'] |= CharacterFlags.Whitespace;
+		}
+	}
+}
diff --git a/Engine/Source/Programs/UnrealBuildTool/Properties/AssemblyInfo.cs b/Engine/Source/Programs/UnrealBuildTool/Properties/AssemblyInfo.cs
--- a/Engine/Source/Programs/UnrealBuildTool/Properties/AssemblyInfo.cs
+++ b/Engine/Source/Programs/UnrealBuildTool/Properties/AssemblyInfo.cs
@@ -1,6 +1,7 @@
 // Copyright Epic Games, Inc. All Rights Reserved.
 
 using System.Reflection;
+using System.Runtime.CompilerServices;
 
 // General Information about an assembly is controlled through the following 
 // set of attributes. Change these attribute values to modify the information
@@ -8,3 +9,5 @@ using System.Reflection;
 [assembly: AssemblyTitle("UnrealBuildTool")]
 [assembly: AssemblyDescription("A tool to build UE4")]
 [assembly: AssemblyConfiguration("")]
+
+[assembly: InternalsVisibleTo("UnrealBuildToolTests")]
diff --git a/Engine/Source/Programs/UnrealBuildTool/System/SourceFileMetadataCache.cs b/Engine/Source/Programs/UnrealBuildTool/System/SourceFileMetadataCache.cs
--- a/Engine/Source/Programs/UnrealBuildTool/System/SourceFileMetadataCache.cs
+++ b/Engine/Source/Programs/UnrealBuildTool/System/SourceFileMetadataCache.cs
@@ -80,6 +80,11 @@ namespace UnrealBuildTool
 		/// </summary>
 		ConcurrentDictionary<FileItem, ReflectionInfo> FileToReflectionInfo = new ConcurrentDictionary<FileItem, ReflectionInfo>();
 
+		/// <summary>
+		/// Map from file item to source file info
+		/// </summary>
+		ConcurrentDictionary<FileItem, SourceFile> FileToSourceFile = new ConcurrentDictionary<FileItem, SourceFile>();
+
 		/// <summary>
 		/// Whether the cache has been modified and needs to be saved
 		/// </summary>
@@ -152,6 +157,50 @@ namespace UnrealBuildTool
 			}
 		}
 
+		/// <summary>
+		/// Finds or adds a SourceFile class for the given file
+		/// </summary>
+		/// <param name="File">File to fetch the source file data for</param>
+		/// <returns>SourceFile instance corresponding to the given source file</returns>
+		public SourceFile GetSourceFile(FileItem File)
+		{
+			if (Parent != null && !File.Location.IsUnderDirectory(BaseDirectory))
+			{
+				return Parent.GetSourceFile(File);
+			}
+			else
+			{
+				SourceFile Result;
+				if (!FileToSourceFile.TryGetValue(File, out Result) || File.LastWriteTimeUtc.Ticks > Result.LastWriteTimeUtc)
+				{
+					SourceFile NewSourceFile = new SourceFile(File);
+					if (Result == null)
+					{
+						if (FileToSourceFile.TryAdd(File, NewSourceFile))
+						{
+							Result = NewSourceFile;
+						}
+						else
+						{
+							Result = FileToSourceFile[File];
+						}
+					}
+					else
+					{
+						if (FileToSourceFile.TryUpdate(File, NewSourceFile, Result))
+						{
+							Result = NewSourceFile;
+						}
+						else
+						{
+							Result = FileToSourceFile[File];
+						}
+					}
+				}
+				return Result;
+			}
+		}
+
 		/// <summary>
 		/// Determines whether the given file contains reflection markup
 		/// </summary>
diff --git a/Engine/Source/Programs/UnrealBuildTool/UnrealBuildTool.csproj b/Engine/Source/Programs/UnrealBuildTool/UnrealBuildTool.csproj
--- a/Engine/Source/Programs/UnrealBuildTool/UnrealBuildTool.csproj
+++ b/Engine/Source/Programs/UnrealBuildTool/UnrealBuildTool.csproj
@@ -134,6 +134,18 @@
     <Compile Include="Platform\Mac\ApplePlatformSDK.cs" />
     <Compile Include="Platform\Windows\MicrosoftPlatformSDK.cs" />
     <Compile Include="Platform\Windows\VCManifestGenerator.cs" />
+    <Compile Include="Preprocessor\Identifier.cs" />
+    <Compile Include="Preprocessor\Preprocessor.cs" />
+    <Compile Include="Preprocessor\PreprocessorContext.cs" />
+    <Compile Include="Preprocessor\PreprocessorExpression.cs" />
+    <Compile Include="Preprocessor\PreprocessorMacro.cs" />
+    <Compile Include="Preprocessor\PreprocessorState.cs" />
+    <Compile Include="Preprocessor\PreprocessorTransform.cs" />
+    <Compile Include="Preprocessor\SourceFile.cs" />
+    <Compile Include="Preprocessor\SourceFileFragment.cs" />
+    <Compile Include="Preprocessor\SourceFileMarkup.cs" />
+    <Compile Include="Preprocessor\Token.cs" />
+    <Compile Include="Preprocessor\TokenReader.cs" />
     <Compile Include="ProjectFiles\CLion\CLionGenerator.cs" />
     <Compile Include="ProjectFiles\CMake\CMakefileGenerator.cs" />
     <Compile Include="ProjectFiles\CodeLite\CodeLiteGenerator.cs" />
diff --git a/Engine/Source/Programs/UnrealBuildToolTests/PreprocessorExpressionTests.cs b/Engine/Source/Programs/UnrealBuildToolTests/PreprocessorExpressionTests.cs
new file mode 100644
--- /dev/null
+++ b/Engine/Source/Programs/UnrealBuildToolTests/PreprocessorExpressionTests.cs
@@ -0,0 +1,152 @@
+// Copyright 1998-2019 Epic Games, Inc. All Rights Reserved.
+
+using Microsoft.VisualStudio.TestTools.UnitTesting;
+using System;
+using System.Collections.Generic;
+using System.Linq;
+using System.Text;
+using System.Threading.Tasks;
+using UnrealBuildTool;
+
+namespace UnrealBuildToolTests
+{
+	[TestClass]
+	public class PreprocessorExpressionTests
+	{
+		[TestMethod]
+		public void Addition()
+		{
+			RunTest("1 + 2 > 3", 0);
+			RunTest("1 + 2 >= 3", 1);
+		}
+
+		[TestMethod]
+		public void Equality()
+		{
+			RunTest("1 == 1", 1);
+			RunTest("1 == 2", 0);
+			RunTest("1 != 1", 0);
+			RunTest("1 != 2", 1);
+			RunTest("1 < 0", 0);
+			RunTest("1 < 1", 0);
+			RunTest("1 < 2", 1);
+			RunTest("1 <= 0", 0);
+			RunTest("1 <= 1", 1);
+			RunTest("1 <= 2", 1);
+			RunTest("1 > 0", 1);
+			RunTest("1 > 1", 0);
+			RunTest("1 > 2", 0);
+			RunTest("1 >= 0", 1);
+			RunTest("1 >= 1", 1);
+			RunTest("1 >= 2", 0);
+		}
+
+		[TestMethod]
+		public void UnaryOperators()
+		{
+			RunTest("0 + 1 == +1", 1);
+			RunTest("0 - 1 == -1", 1);
+			RunTest("!0", 1);
+			RunTest("!1", 0);
+			RunTest("~0", -1L);
+		}
+
+		[TestMethod]
+		public void ArithmeticOperators()
+		{
+			RunTest("3 + 7", 10);
+			RunTest("3 - 7", -4);
+			RunTest("3 * 7", 21);
+			RunTest("21 / 3", 7);
+			RunTest("24 % 7", 3);
+			RunTest("2 << 4", 32);
+			RunTest("32 >> 4", 2);
+			RunTest("(0xab & 0xf)", 0xb);
+			RunTest("(0xab | 0xf)", 0xaf);
+			RunTest("(0xab ^ 0xf)", 0xa4);
+		}
+
+		[TestMethod]
+		public void LogicalOperators()
+		{
+			RunTest("0 || 0", 0);
+			RunTest("0 || 1", 1);
+			RunTest("1 || 0", 1);
+			RunTest("1 || 1", 1);
+			RunTest("0 && 0", 0);
+			RunTest("0 && 1", 0);
+			RunTest("1 && 0", 0);
+			RunTest("1 && 1", 1);
+		}
+
+		[TestMethod]
+		public void TernaryOperator()
+		{
+			RunTest("((0)? 123 : 456) == 456", 1);
+			RunTest("((1)? 123 : 456) == 456", 0);
+		}
+
+		[TestMethod]
+		public void Precedence()
+		{
+			RunTest("3 + 27 / 3", 12);
+			RunTest("(3 + 27) / 3", 10);
+		}
+
+		[TestMethod]
+		public void Literals()
+		{
+			RunTest("123", 123);
+			RunTest("123l", 123);
+			RunTest("123L", 123);
+			RunTest("123ll", 123);
+			RunTest("123LL", 123);
+			RunTest("123u", 123);
+			RunTest("123U", 123);
+			RunTest("123ul", 123);
+			RunTest("123Ul", 123);
+			RunTest("123UL", 123);
+			RunTest("123Ull", 123);
+			RunTest("123ULL", 123);
+			RunTest("123ui64", 123);
+			RunTest("0x123i64", 0x123);
+			RunTest("0123i64", 83);
+        }
+
+		class PreprocessorTestContext : PreprocessorContext
+		{
+			public PreprocessorTestContext()
+				: base(null)
+			{
+			}
+		}
+
+		/// <summary>
+		/// Tokenize the given expression and evaluate it, and check it matches the expected result
+		/// </summary>
+		/// <param name="Expression">The expression to evaluate, as a string</param>
+		/// <param name="ExpectedResult">The expected value of the expression</param>
+        static void RunTest(string Expression, long? ExpectedResult)
+        {
+			TokenReader Reader = new TokenReader(Expression);
+
+			List<Token> Tokens = new List<Token>();
+			while (Reader.MoveNext())
+			{
+				Tokens.Add(Reader.Current);
+			}
+
+			long? Result;
+			try
+			{
+				Result = PreprocessorExpression.Evaluate(new PreprocessorTestContext(), Tokens);
+			}
+			catch(PreprocessorException)
+			{
+				Result = null;
+			}
+
+			Assert.AreEqual(ExpectedResult, Result);
+		}
+	}
+}
diff --git a/Engine/Source/Programs/UnrealBuildToolTests/PreprocessorStateTests.cs b/Engine/Source/Programs/UnrealBuildToolTests/PreprocessorStateTests.cs
new file mode 100644
--- /dev/null
+++ b/Engine/Source/Programs/UnrealBuildToolTests/PreprocessorStateTests.cs
@@ -0,0 +1,150 @@
+// Copyright 1998-2019 Epic Games, Inc. All Rights Reserved.
+
+using Microsoft.VisualStudio.TestTools.UnitTesting;
+using System;
+using System.Collections.Generic;
+using System.Linq;
+using System.Text;
+using System.Threading.Tasks;
+using Tools.DotNETCommon;
+using UnrealBuildTool;
+
+namespace UnrealBuildToolTests
+{
+	[TestClass]
+	public class PreprocessorStateTests
+	{
+		[TestMethod]
+		public void Run()
+		{
+			PreprocessorState BaseState = new PreprocessorState();
+			BaseState.DefineMacro(ParseMacro("FOO", null, "123"));
+			BaseState.DefineMacro(ParseMacro("BAR", new List<string>{ "Arg1", "Arg2" }, "Arg1 + Arg2"));
+			Assert.AreEqual("m BAR(Arg1, Arg2)=Arg1 + Arg2\nm FOO=123\n", FormatState(BaseState));
+
+			// Modify the state and revert it
+			{
+				PreprocessorState State = new PreprocessorState(BaseState);
+				Assert.AreEqual(FormatState(BaseState), FormatState(State));
+
+				PreprocessorTransform Transform1 = State.BeginCapture();
+				State.UndefMacro(Identifier.FindOrAdd("FOO"));
+				State.DefineMacro(ParseMacro("FOO2", null, "FOO()"));
+				State.PushBranch(PreprocessorBranch.Active);
+				Assert.AreEqual("b Active\nm BAR(Arg1, Arg2)=Arg1 + Arg2\nm FOO2=FOO()\n", FormatState(State));
+				Assert.AreEqual("+b Active\n+m FOO undef\n+m FOO2=FOO()\n", FormatTransform(Transform1));
+
+				PreprocessorState State2 = new PreprocessorState(BaseState);
+				State2.TryToApply(Transform1);
+				Assert.AreEqual(FormatState(State), FormatState(State2));
+
+				PreprocessorTransform Transform2 = State.BeginCapture();
+				State.PopBranch();
+				State.UndefMacro(Identifier.FindOrAdd("FOO"));
+				State.DefineMacro(ParseMacro("FOO", null, "123"));
+				State.UndefMacro(Identifier.FindOrAdd("FOO2"));
+				Assert.AreEqual(FormatState(BaseState), FormatState(State));
+				Assert.AreEqual("-b Active\n+m FOO=123\n+m FOO2 undef\n", FormatTransform(Transform2));
+
+				State2.TryToApply(Transform2);
+				Assert.AreEqual(FormatState(BaseState), FormatState(State2));
+			}
+
+			// Check the tracking of branches
+			{
+				PreprocessorState State = new PreprocessorState(BaseState);
+				PreprocessorTransform Transform1 = State.BeginCapture();
+				State.IsCurrentBranchActive();
+				Assert.AreEqual("=b Active\n", FormatTransform(Transform1));
+
+				State.PushBranch(PreprocessorBranch.Active);
+				PreprocessorTransform Transform2 = State.BeginCapture();
+				State.IsCurrentBranchActive();
+				State.PopBranch();
+				State.EndCapture();
+				Assert.AreEqual("-b Active\n", FormatTransform(Transform2));
+
+				State.PushBranch(0);
+				Assert.AreEqual("False", State.CanApply(Transform2).ToString());
+
+				State.PopBranch();
+				State.PushBranch(PreprocessorBranch.Active);
+				Assert.AreEqual("True", State.CanApply(Transform2).ToString());
+			}
+		}
+	
+		static string FormatState(PreprocessorState State)
+		{
+			StringBuilder Result = new StringBuilder();
+			foreach(PreprocessorBranch Branch in State.CurrentBranches)
+			{
+				Result.AppendFormat("b {0}\n", Branch);
+			}
+			foreach(PreprocessorMacro Macro in State.CurrentMacros.OrderBy(x => x.Name))
+			{
+				Result.AppendFormat("m {0}\n", Macro.ToString().TrimEnd());
+			}
+			return Result.ToString();
+		}
+
+		static string FormatTransform(PreprocessorTransform Transform)
+		{
+			StringBuilder Result = new StringBuilder();
+			if(Transform.bRequireTopmostActive.HasValue)
+			{
+				Result.AppendFormat("=b {0}\n", Transform.bRequireTopmostActive.Value? "Active" : "0");
+			}
+			foreach(PreprocessorBranch Branch in Transform.RequiredBranches)
+			{
+				Result.AppendFormat("-b {0}\n", Branch);
+			}
+			foreach(PreprocessorBranch Branch in Transform.NewBranches)
+			{
+				Result.AppendFormat("+b {0}\n", Branch);
+			}
+			foreach(KeyValuePair<Identifier, PreprocessorMacro> Macro in Transform.RequiredMacros)
+			{
+				if(Macro.Value == null)
+				{
+					Result.AppendFormat("=m {0} undef\n", Macro.Key);
+				}
+				else
+				{
+					Result.AppendFormat("=m {0}\n", Macro.Value.ToString().TrimEnd());
+				}
+			}
+			foreach(KeyValuePair<Identifier, PreprocessorMacro> Macro in Transform.NewMacros)
+			{
+				if(Macro.Value == null)
+				{
+					Result.AppendFormat("+m {0} undef\n", Macro.Key);
+				}
+				else
+				{
+					Result.AppendFormat("+m {0}\n", Macro.Value.ToString().TrimEnd());
+				}
+			}
+			return Result.ToString();
+		}
+
+		static PreprocessorMacro ParseMacro(string Name, List<string> Parameters, string Value)
+		{
+			List<Token> Tokens = new List<Token>();
+
+			TokenReader Reader = new TokenReader(Value);
+			while(Reader.MoveNext())
+			{
+				Tokens.Add(Reader.Current);
+			}
+
+			if(Parameters == null)
+			{
+				return new PreprocessorMacro(Identifier.FindOrAdd(Name), null, Tokens);
+			}
+			else
+			{
+				return new PreprocessorMacro(Identifier.FindOrAdd(Name), Parameters.Select(x => Identifier.FindOrAdd(x)).ToList(), Tokens);
+			}
+		}
+	}
+}
diff --git a/Engine/Source/Programs/UnrealBuildToolTests/PreprocessorTests.cs b/Engine/Source/Programs/UnrealBuildToolTests/PreprocessorTests.cs
new file mode 100644
--- /dev/null
+++ b/Engine/Source/Programs/UnrealBuildToolTests/PreprocessorTests.cs
@@ -0,0 +1,274 @@
+using Microsoft.VisualStudio.TestTools.UnitTesting;
+using System.Collections.Generic;
+using System.Text;
+using UnrealBuildTool;
+
+namespace UnrealBuildToolTests
+{
+	[TestClass]
+	public class PreprocessorTests
+	{
+		[TestMethod]
+		public void BasicTests()
+		{
+			RunTest("token", "token\n");
+			RunTest("#", "");
+		}
+
+		[TestMethod]
+		public void ObjectMacroTests()
+		{
+			RunTest("#define X token\nX", "token\n");
+			RunTest("#define A B C D\nA", "B C D\n");
+			RunTest("#define A B ## D\nA", "BD\n");
+
+			RunTest("#define X list of tokens\nX", "list of tokens\n");
+			RunTest("#define LST list\n#define TOKS tokens\n#define LOTS LST of TOKS\nLOTS LOTS", "list of tokens list of tokens\n");
+			RunTest("#define LOTS LST of TOKS\n#define LST list\n#define TOKS tokens\nLOTS LOTS", "list of tokens list of tokens\n");
+
+			RunTest("#define A MACRO\nA", "MACRO\n");
+			RunTest("#define A MACRO\n#undef A\nA", "A\n");
+		}
+
+		[TestMethod]
+		public void FunctionMacroTests()
+		{
+			RunTest("#define X(x) token x other\nX(and one).", "token and one other.\n");
+			RunTest("#define X(x,y) token x and y\nX(1, 2).", "token 1 and 2.\n");
+			RunTest("#define INC(x) (x + 2)\nINC", "INC\n");
+			RunTest("#define TEST(x) x\\\n?\nTEST(A)", "A?\n");
+
+			RunTest("#define F(A) <A>\nF(x)", "<x>\n");
+			RunTest("#define F(A,B) <A,B>\nF(x,y) + 1", "<x,y> + 1\n");
+			RunTest("#define F(A,B,C) <A,B,C>\nF(x,y,z)", "<x,y,z>\n");
+			RunTest("#define F(...) <__VA_ARGS__>\nF(x)", "<x>\n");
+			RunTest("#define F(...) <__VA_ARGS__>\nF(x,y)", "<x,y>\n");
+			RunTest("#define F(A,...) <A,__VA_ARGS__>\nF(x,y)", "<x,y>\n");
+			RunTest("#define F(A,...) <A, __VA_ARGS__>\nF(x,y)", "<x, y>\n");
+			RunTest("#define F(A,...) <A, __VA_ARGS__>\nF(x,y, z)", "<x, y, z>\n");
+
+			RunTest("#define FUNC(x) arg=x.\nFUNC(var) FUNC(2)", "arg=var. arg=2.\n");
+			RunTest("#define FUNC(x,y,z) int n = z+y*x;\nFUNC(1,2,3)", "int n = 3+2*1;\n");
+			RunTest("#define X 20\n#define FUNC(x,y) x+y\nx=FUNC(X,Y);", "x=20+Y;\n");
+			RunTest("#define FA(x,y) FB(x,y)\n#define FB(x,y) x + y\nFB(1,2);", "1 + 2;\n");
+			RunTest("#define PRINTF(...) printf(__VA_ARGS__)\nPRINTF()", "printf()\n");
+			RunTest("#define PRINTF(...) printf(__VA_ARGS__)\nPRINTF(\"hello\")", "printf(\"hello\")\n");
+			RunTest("#define PRINTF(...) printf(__VA_ARGS__)\nPRINTF(\"%d\", 1)", "printf(\"%d\", 1)\n");
+			RunTest("#define PRINTF(FORMAT, ...) printf(FORMAT, __VA_ARGS__)\nPRINTF(\"test\")", "printf(\"test\", )\n");
+			RunTest("#define PRINTF(FORMAT, ...) printf(FORMAT, __VA_ARGS__)\nPRINTF(\"test %s\", \"hello\")", "printf(\"test %s\", \"hello\")\n");
+			RunTest("#define PRINTF(FORMAT, ...) printf(FORMAT, __VA_ARGS__)\nPRINTF(\"test %s %d\", \"hello\", 1)", "printf(\"test %s %d\", \"hello\", 1)\n");
+			RunTest("#define PRINTF(FORMAT, ...) printf(FORMAT, ## __VA_ARGS__)\nPRINTF(\"test\")", "printf(\"test\" )\n");
+
+			RunTest("#define INC(x) (x + 2)\nINC(\n1\n)", "(1 + 2)\n");
+
+			RunTest("#define STRINGIZE(ARG) #ARG\nSTRINGIZE(+=)", "\"+=\"\n");
+			RunTest("#define STRINGIZE(ARG) #ARG\nSTRINGIZE(:>)", "\":>\"\n");
+			RunTest("#define STRINGIZE(ARG) #ARG\nSTRINGIZE(3.1415)", "\"3.1415\"\n");
+
+			RunTest("#define CONCAT(X, Y) X ## Y\nCONCAT(+, =)", "+=\n");
+			RunTest("#define CONCAT(X, Y) X ## Y\nCONCAT(3, .14159)", "3.14159\n");
+			RunTest("#define CONCAT(X, Y) X ## Y\nCONCAT(3, hello)", "3hello\n");
+			RunTest("#define CONCAT(X, Y) X ## #Y\nCONCAT(u, hello)", "u\"hello\"\n");
+			RunTest("#define CONCAT(X, ...) X ## __VA_ARGS__\nCONCAT(hello) there", "hello there\n");
+			RunTest("#define CONCAT(X, ...) X ## __VA_ARGS__\nCONCAT(hello, there)", "hellothere\n");
+
+			RunTest("#define A(X) MACRO X\nA(x)", "MACRO x\n");
+			RunTest("#define A(X) MACRO X\n#undef A\nA(x)", "A(x)\n");
+		}
+
+		[TestMethod]
+		public void ConditionalTests()
+		{
+			RunTest("#if 1\nfirst_branch\n#endif\nend", "first_branch\nend\n");
+			RunTest("#if 0\nfirst_branch\n#endif\nend", "end\n");
+			RunTest("#if 1\nbranch_1\n#else\nbranch_2\n#endif", "branch_1\n");
+			RunTest("#if 0\nbranch_1\n#else\nbranch_2\n#endif", "branch_2\n");
+			RunTest("#define A\n#ifdef A\nYes\n#endif", "Yes\n");
+			RunTest("#define B\n#ifdef A\nYes\n#endif\nNo", "No\n");
+			RunTest("#define A\n#ifndef A\nYes\n#endif\nNo", "No\n");
+			RunTest("#define B\n#ifndef A\nYes\n#endif", "Yes\n");
+			RunTest("#define A\n#undef A\n#ifdef A\nYes\n#endif\nNo", "No\n");
+			RunTest("#define A\n#undef A\n#ifndef A\nYes\n#endif", "Yes\n");
+			RunTest("#define A\n#ifdef A\nYes\n#else\nNo\n#endif", "Yes\n");
+			RunTest("#define B\n#ifdef A\nYes\n#else\nNo\n#endif", "No\n");
+		}
+
+		[TestMethod]
+		public void ExpressionTests()
+		{
+			RunTest("#if 1 + 2 > 3\nYes\n#endif\nNo", "No\n");
+			RunTest("#if 1 + 2 >= 3\nYes\n#endif", "Yes\n");
+			RunTest("#define ONE 1\n#define TWO 2\n#define PLUS(x, y) x + y\n#if PLUS(ONE, TWO) > 3\nYes\n#endif\nNo", "No\n");
+			RunTest("#define ONE 1\n#define TWO 2\n#define PLUS(x, y) x + y\n#if PLUS(ONE, TWO) >= 3\nYes\n#endif", "Yes\n");
+			RunTest("#define ONE 1\n#if defined ONE\nOne\n#elif defined TWO\nTwo\n#else\nThree\n#endif", "One\n");
+			RunTest("#define TWO 1\n#if defined ONE\nOne\n#elif defined TWO\nTwo\n#else\nThree\n#endif", "Two\n");
+			RunTest("#define ONE 0\n#if defined(ONE) + defined(TWO) >= 1\nYes\n#else\nNo\n#endif", "Yes\n");
+			RunTest("#define ONE 0\n#if defined(ONE) + defined(TWO) >= 2\nYes\n#else\nNo\n#endif", "No\n");
+			RunTest("#define ONE 0\n#define TWO\n#if defined(ONE) + defined(TWO) >= 2\nYes\n#else\nNo\n#endif", "Yes\n");
+		}
+
+		[TestMethod]
+		public void ConcatenationTests()
+		{
+			RunTest("#define FUNC(x) Value = x\n#define PP_JOIN(x, y) x ## y\n#define RESULT(x) PP_JOIN(FU, NC)(x)\nRESULT(1234)", "Value = 1234\n");
+			RunTest("#define VALUE 1234\n#define PP_JOIN(x, y) x ## y\n#define RESULT PP_JOIN(V, ALUE)\nRESULT", "1234\n");
+			RunTest("#define V 1\n#define ALUE 2\n#define VALUE 1234\n#define PP_JOIN(x, y) x ## y\n#define RESULT PP_JOIN(V, ALUE)\nRESULT", "1234\n");
+		}
+
+		[TestMethod]
+		public void VarargTests()
+		{
+			RunTest("#define FUNC(fmt,...) (fmt, ## __VA_ARGS__)\nFUNC(a)", "(a )\n");
+			RunTest("#define FUNC(fmt,...) (fmt, ## __VA_ARGS__)\nFUNC(a,b)", "(a,b)\n");
+			RunTest("#define FUNC(fmt,...) (fmt, ## __VA_ARGS__)\nFUNC(a,b )", "(a,b)\n");
+			RunTest("#define FUNC(fmt, ...) (fmt, ## __VA_ARGS__)\nFUNC(a)", "(a )\n");
+			RunTest("#define FUNC(fmt, ...) (fmt, ## __VA_ARGS__)\nFUNC(a,b)", "(a,b)\n");
+			RunTest("#define FUNC(fmt, ...) (fmt, ## __VA_ARGS__)\nFUNC(a,b )", "(a,b)\n");
+		}
+
+		[TestMethod]
+		public void EmptyTokenTests()
+		{
+			RunTest("#define EMPTY_TOKEN\n#define FUNC(_FuncName, _HType1, _HArg1) _FuncName(_HType1 _HArg1);\nFUNC(hello, EMPTY_TOKEN, int)", "hello( int);\n");
+
+			RunTest("#define EMPTY_TOKEN\n#define GCC_EXTENSION(...) 123 , ## __VA_ARGS__\nGCC_EXTENSION(EMPTY_TOKEN)", "123  \n");
+
+			RunTest("#define EMPTY_TOKEN\n#define FUNC(x) (x)\nFUNC(EMPTY_TOKEN A)", "( A)\n");
+			RunTest("#define EMPTY_TOKEN\n#define FUNC(x,y) (x y)\nFUNC(EMPTY_TOKEN,A)", "( A)\n");
+			//			RunTest("#define EMPTY_TOKEN\n#define FUNC(x) (x)\n#define FUNC_2(x,y) FUNC(x y)\nFUNC_2(EMPTY_TOKEN,A)", "( A)\n");
+			//			RunTest("#define EMPTY_TOKEN\n#define FUNC(x) (x EMPTY_TOKEN)\nFUNC(A)", "(A )\n");
+			//			RunTest("#define EMPTY_TOKEN\n#define FUNC(x,y) (x y)\nFUNC(A,)", "(A)\n");
+
+
+			RunTest("#define EMPTY\n#define FUNC(x) (x)\nFUNC(x    y)", "(x y)\n");
+			RunTest("#define EMPTY\n#define FUNC(x) (x)\nFUNC(x    EMPTY    y)", "(x  y)\n");
+			RunTest("#define EMPTY\n#define FUNC(x) (x)\nFUNC(x    EMPTY    y    EMPTY)", "(x  y )\n");
+			RunTest("#define EMPTY\n#define FUNC(x) (x)\nFUNC(EMPTY    x    EMPTY    y)", "( x  y)\n");
+			RunTest("#define EMPTY\n#define FUNC(x) (x)\nFUNC(    EMPTY    x    EMPTY    y)", "( x  y)\n");
+			RunTest("#define EMPTY\n#define FUNC(x) (x)\nFUNC(EMPTY    x    EMPTY    y    )", "( x  y)\n");
+
+			RunTest("#define EMPTY\n#define FUNC(x) (x)\n#define FUNC_2(x) FUNC(x)\nFUNC_2(x    y)", "(x y)\n");
+			RunTest("#define EMPTY\n#define FUNC(x) (x)\n#define FUNC_2(x) FUNC(x)\nFUNC_2(x    EMPTY    y)", "(x y)\n");
+			RunTest("#define EMPTY\n#define FUNC(x) (x)\n#define FUNC_2(x) FUNC(x)\nFUNC_2(x    EMPTY    y    EMPTY)", "(x y)\n");
+			RunTest("#define EMPTY\n#define FUNC(x) (x)\n#define FUNC_2(x) FUNC(x)\nFUNC_2(EMPTY    x    EMPTY    y)", "( x y)\n");
+			RunTest("#define EMPTY\n#define FUNC(x) (x)\n#define FUNC_2(x) FUNC(x)\nFUNC_2(    EMPTY    x    EMPTY    y)", "( x y)\n");
+			RunTest("#define EMPTY\n#define FUNC(x) (x)\n#define FUNC_2(x) FUNC(x)\nFUNC_2(EMPTY    x    EMPTY    y    )", "( x y)\n");
+
+			RunTest("#define EMPTY\n#define FUNC(x) ( x )\nFUNC(EMPTY EMPTY)", "(   )\n");
+			RunTest("#define EMPTY\n#define FUNC(x) ( x)\nFUNC(EMPTY EMPTY x)", "(   x)\n");
+			RunTest("#define EMPTY\n#define FUNC(x,y) ( x y)\n#define FUNC2(x,y) FUNC(x,y)\nFUNC2(EMPTY EMPTY EMPTY, x)", "(   x)\n");
+			RunTest("#define EMPTY\n#define DOUBLE_EMPTY EMPTY EMPTY\n#define FUNC(x) ( x)\nFUNC(DOUBLE_EMPTY x)", "(   x)\n");
+			RunTest("#define EMPTY\n#define DOUBLE_EMPTY EMPTY EMPTY\n#define FUNC(x,y) ( x y )\n#define FUNC2(x,y) FUNC(x,y)\nFUNC2(DOUBLE_EMPTY EMPTY, x)", "(   x )\n");
+			RunTest("#define EMPTY\n#define FUNC(x,y) ( x y )\nFUNC(EMPTY EMPTY EMPTY EMPTY EMPTY EMPTY, x)", "(       x )\n");
+		}
+
+		[TestMethod]
+		public void MiscTests()
+		{
+			RunTest("#define REGISTER_NAME(num,name) NAME_##name = num,\nREGISTER_NAME(201, TRUE)", "NAME_TRUE = 201,\n");
+			RunTest("#define FUNC_2(X) VALUE=X\n#define FUNC_N(X) FUNC_##X\n#define OBJECT FUNC_N(2)\nOBJECT(1234)", "VALUE=1234\n");
+			RunTest("#define GCC_EXTENSION(...) 123, ## __VA_ARGS__\nGCC_EXTENSION(456)", "123,456\n");
+
+			RunTest("#define _NON_MEMBER_CALL(FUNC, CV_REF_OPT) FUNC(X,CV_REF_OPT)\n#define _NON_MEMBER_CALL_CV(FUNC, REF_OPT) _NON_MEMBER_CALL(FUNC, REF_OPT) _NON_MEMBER_CALL(FUNC, const REF_OPT)\n#define _NON_MEMBER_CALL_CV_REF(FUNC) _NON_MEMBER_CALL_CV(FUNC, ) _NON_MEMBER_CALL_CV(FUNC, &)\n#define _IS_FUNCTION(X,CV_REF_OPT) (CV_REF_OPT)\n_NON_MEMBER_CALL_CV_REF(_IS_FUNCTION)", "() (const) (&) (const &)\n");
+
+			RunTest("#define TEXT(x) L ## x\n#define checkf(expr, format, ...) FDebug::AssertFailed(#expr, format, ##__VA_ARGS__)\ncheckf( true, TEXT( \"hello world\" ) );", "FDebug::AssertFailed(\"true\", L\"hello world\" );\n");
+		}
+
+		[TestMethod]
+		public void ErrorTests()
+		{
+			RunTest("#define", null);
+			RunTest("#define INC(x) (x + 2)\nINC(", null);
+			RunTest("#define +", null);
+			RunTest("#define A B __VA_ARGS__ D\nA", null);
+			RunTest("#define A ## B\nA", null);
+			RunTest("#define A B ##\nA", null);
+			RunTest("#define defined not_defined", null);
+
+			RunTest("#define F(A) <A>\nF(x,y) + 1", null);
+			RunTest("#define F(A,) <A>\nF(x)", null);
+			RunTest("#define F(A+) <A>\nF(x)", null);
+			RunTest("#define F(A,A) <A>\nF(x,y)", null);
+			RunTest("#define F(A,B) <A,B>\nF(x) + 1", null);
+			RunTest("#define F(A,B,) <A,B>\nF(x,y) + 1", null);
+//			RunTest("#define F(...) <__VA_ARGS__>\nF(x,)", null);
+			RunTest("#define F(A...) <A, __VA_ARGS__>\nF(x)", null);
+			RunTest("#define F(A...) <A, __VA_ARGS__>\nF(x,y)", null);
+			RunTest("#define F(A,__VA_ARGS__) <A, __VA_ARGS__>\nF(x,y)", null);
+			RunTest("#define F(A) #+\nF(x)", null);
+			RunTest("#define F(A) #B\nF(x)", null);
+			RunTest("#define F(A) ## A\nF(x)", null);
+			RunTest("#define F(A) A ##\nF(x)", null);
+			RunTest("#define F(A) <__VA_ARGS__>\nF(x)", null);
+
+			RunTest("#define INC(x\n)", null);
+			RunTest("#if 1\nbranch_1\n#else garbage\nbranch_2\n#endif\nend", null);
+//			RunTest("#if 1\nbranch_1\n#else\nbranch_2\n#endif garbage\nend", null);
+			RunTest("#if 1\nbranch_1\n#else\nbranch_2\n#elif 1\nbranch_3\n#endif\nend", null);
+			RunTest("#if 1\nbranch_1\n#else\nbranch_2\n#else\nbranch_3\n#endif", null);
+			RunTest("#if 0\nbranch_1\n#else\nbranch_2\n#else\nbranch_3\n#endif", null);
+			RunTest("#elif\nbranch_1\n#else\nbranch_2\n#endif\nend", null);
+			RunTest("#ifdef +\nbranch_1\n#else\nbranch_2\n#endif", null);
+			RunTest("#ifdef A 1\nbranch_1\n#else\nbranch_2\n#endif", null);
+			RunTest("#define A()\n#if A(\n)\nOK\n#endif", null);
+			RunTest("#define A MACRO\n#undef A B\nA", null);
+		}
+
+		/// <summary>
+		/// Preprocess a fragment of code, and check it results in an expected sequence of tokens
+		/// </summary>
+		/// <param name="Fragment">The code fragment to preprocess</param>
+		/// <param name="ExpectedResult">The expected sequence of tokens, as a string. Null to indicate that the input is invalid, and an exception is expected.</param>
+		static void RunTest(string Fragment, string ExpectedResult)
+		{
+			string[] Lines = Fragment.Split('\n');
+
+			SourceFile File = new SourceFile(TokenReader.GetNullTerminatedByteArray(Fragment));
+
+			Preprocessor Instance = new Preprocessor();
+
+			string Result;
+			try
+			{
+				List<Token> OutputTokens = new List<Token>();
+				for (int MarkupIdx = 0; MarkupIdx < File.Markup.Length; MarkupIdx++)
+				{
+					SourceFileMarkup Markup = File.Markup[MarkupIdx];
+					if (Markup.Type == SourceFileMarkupType.Text)
+					{
+						if (Instance.IsCurrentBranchActive())
+						{
+							StringBuilder SourceText = new StringBuilder();
+
+							int LastLineIndex = (MarkupIdx + 1 < File.Markup.Length) ? File.Markup[MarkupIdx + 1].LineNumber - 1 : Lines.Length;
+							for (int LineIndex = Markup.LineNumber - 1; LineIndex < LastLineIndex; LineIndex++)
+							{
+								SourceText.AppendLine(Lines[LineIndex]);
+							}
+
+							List<Token> Tokens = new List<Token>();
+
+							TokenReader Reader = new TokenReader(SourceText.ToString());
+							while (Reader.MoveNext())
+							{
+								Tokens.Add(Reader.Current);
+							}
+
+							Instance.ExpandMacros(Tokens, OutputTokens, false, null);
+						}
+					}
+					else
+					{
+						Instance.ParseMarkup(Markup.Type, Markup.Tokens, null);
+					}
+				}
+				Result = Token.Format(OutputTokens);
+			}
+			catch (PreprocessorException)
+			{
+				Result = null;
+			}
+			Assert.AreEqual(ExpectedResult, Result);
+		}
+	}
+}
diff --git a/Engine/Source/Programs/UnrealBuildToolTests/SourceFileTests.cs b/Engine/Source/Programs/UnrealBuildToolTests/SourceFileTests.cs
new file mode 100644
--- /dev/null
+++ b/Engine/Source/Programs/UnrealBuildToolTests/SourceFileTests.cs
@@ -0,0 +1,89 @@
+// Copyright 1998-2019 Epic Games, Inc. All Rights Reserved.
+
+using Microsoft.VisualStudio.TestTools.UnitTesting;
+using System;
+using System.Collections.Concurrent;
+using System.Collections.Generic;
+using System.Linq;
+using System.Text;
+using System.Threading.Tasks;
+using Tools.DotNETCommon;
+using UnrealBuildTool;
+
+namespace UnrealBuildToolTests
+{
+	/// <summary>
+	/// Tests for reading source file markup
+	/// </summary>
+	[TestClass]
+	public class SourceFileTests
+	{
+		[TestMethod]
+		public void Run()
+		{
+			List<DirectoryReference> BaseDirectories = new List<DirectoryReference>();
+			BaseDirectories.Add(DirectoryReference.Combine(UnrealBuildTool.UnrealBuildTool.EngineSourceDirectory, "Runtime"));
+			BaseDirectories.Add(DirectoryReference.Combine(UnrealBuildTool.UnrealBuildTool.EngineSourceDirectory, "Developer"));
+			BaseDirectories.Add(DirectoryReference.Combine(UnrealBuildTool.UnrealBuildTool.EngineSourceDirectory, "Editor"));
+
+			foreach(FileReference PluginFile in Plugins.EnumeratePlugins((FileReference)null))
+			{
+				DirectoryReference PluginSourceDir = DirectoryReference.Combine(PluginFile.Directory, "Source");
+				if(DirectoryReference.Exists(PluginSourceDir))
+				{
+					BaseDirectories.Add(PluginSourceDir);
+				}
+			}
+
+			ConcurrentBag<SourceFile> SourceFiles = new ConcurrentBag<SourceFile>();
+			using(Timeline.ScopeEvent("Scanning source files"))
+			{
+				using(ThreadPoolWorkQueue Queue = new ThreadPoolWorkQueue())
+				{
+					foreach(DirectoryReference BaseDirectory in BaseDirectories)
+					{
+						Queue.Enqueue(() => ParseSourceFiles(DirectoryItem.GetItemByDirectoryReference(BaseDirectory), SourceFiles, Queue));
+					}
+				}
+			}
+
+			Log.TraceLog("Read {0} source files", SourceFiles.Count);
+
+			FileReference TempDataFile = FileReference.Combine(UnrealBuildTool.UnrealBuildTool.EngineDirectory, "Intermediate", "Temp", "SourceFileTests.bin");
+			DirectoryReference.CreateDirectory(TempDataFile.Directory);
+
+			using(Timeline.ScopeEvent("Writing source file data"))
+			{
+				using(BinaryArchiveWriter Writer = new BinaryArchiveWriter(TempDataFile))
+				{
+					Writer.WriteList(SourceFiles.ToList(), x => x.Write(Writer));
+				}
+			}
+
+			List<SourceFile> ReadSourceFiles = new List<SourceFile>();
+			using(Timeline.ScopeEvent("Reading source file data"))
+			{
+				using(BinaryArchiveReader Reader = new BinaryArchiveReader(TempDataFile))
+				{
+					ReadSourceFiles = Reader.ReadList(() => new SourceFile(Reader));
+				}
+			}
+		}
+
+		static void ParseSourceFiles(DirectoryItem Directory, ConcurrentBag<SourceFile> SourceFiles, ThreadPoolWorkQueue Queue)
+		{
+			foreach(DirectoryItem SubDirectory in Directory.EnumerateDirectories())
+			{
+				Queue.Enqueue(() => ParseSourceFiles(SubDirectory, SourceFiles, Queue));
+			}
+
+			foreach(FileItem File in Directory.EnumerateFiles())
+			{
+				if(File.HasExtension(".h") || File.HasExtension(".cpp"))
+				{
+					Queue.Enqueue(() => SourceFiles.Add(new SourceFile(File)));
+				}
+			}
+		}
+	}
+}
diff --git a/Engine/Source/Programs/UnrealBuildToolTests/TokenReaderTests.cs b/Engine/Source/Programs/UnrealBuildToolTests/TokenReaderTests.cs
new file mode 100644
--- /dev/null
+++ b/Engine/Source/Programs/UnrealBuildToolTests/TokenReaderTests.cs
@@ -0,0 +1,174 @@
+// Copyright 1998-2019 Epic Games, Inc. All Rights Reserved.
+
+using Microsoft.VisualStudio.TestTools.UnitTesting;
+using System;
+using System.Collections.Generic;
+using System.Linq;
+using System.Text;
+using System.Threading.Tasks;
+using UnrealBuildTool;
+
+namespace UnrealBuildToolTests
+{
+	[TestClass]
+	public class TokenReaderTests
+	{
+		[TestMethod]
+		public void Run()
+		{
+			RunLexerTest("ABC//hello\nDEF", "Identifier(ABC), Newline, Identifier(DEF)");
+			RunLexerTest("ABC/*hello\nworld*/DEF", "Identifier(ABC), Identifier(DEF)");
+			RunLexerTest("ABC//hello\r\nDEF", "Identifier(ABC), Newline, Identifier(DEF)");
+			RunLexerTest("ABC123", "Identifier(ABC123)");
+			RunLexerTest("123ABC", "Number(123ABC)");
+			RunLexerTest("FOO(1,2,3,...)", "Identifier(FOO), LeftParen, Number(1), Comma, Number(2), Comma, Number(3), Comma, Ellipsis, RightParen");
+		}
+
+		[TestMethod]
+		public void SimpleTokens()
+		{
+			RunLexerTest("=<><<+", "Equals, CompareLess, CompareGreater, LeftShift, Plus");
+			RunLexerTest("+=+++++", "Plus, Equals, Plus, Plus, Plus, Plus, Plus");
+			RunLexerTest("/+++=", "Divide, Plus, Plus, Plus, Equals");
+			RunLexerTest("###", "HashHash, Hash");
+			RunLexerTest("():?", "LeftParen, RightParen, Colon, QuestionMark");
+			RunLexerTest("..*->->*", "Dot, Dot, Multiply, Minus, CompareGreater, Minus, CompareGreater, Multiply");
+			RunLexerTest("+-*/%", "Plus, Minus, Multiply, Divide, Modulo");
+			RunLexerTest("|^&~", "BitwiseOr, BitwiseXor, BitwiseAnd, BitwiseNot");
+			RunLexerTest("+=-=*=/=", "Plus, Equals, Minus, Equals, Multiply, Equals, Divide, Equals");
+			RunLexerTest("%=|=^=", "Modulo, Equals, BitwiseOr, Equals, BitwiseXor, Equals");
+			RunLexerTest("&&||!", "LogicalAnd, LogicalOr, LogicalNot");
+			RunLexerTest("=<><=>=", "Equals, CompareLess, CompareGreater, CompareLessOrEqual, CompareGreaterOrEqual");
+			RunLexerTest("==!=", "CompareEqual, CompareNotEqual");
+			RunLexerTest(".....", "Ellipsis, Dot, Dot");
+		}
+
+		[TestMethod]
+		public void IntegerLiterals()
+		{
+			RunLexerTest("123+", "Number(123), Plus");
+			RunLexerTest("01234567+", "Number(01234567), Plus");
+			RunLexerTest("0x123456789ABCDEFabcdef+", "Number(0x123456789ABCDEFabcdef), Plus");
+			RunLexerTest("123Ull+", "Number(123Ull), Plus");
+			RunLexerTest("123_hello+", "Number(123_hello), Plus");
+			RunLexerTest("01234567_hello+", "Number(01234567_hello), Plus");
+			RunLexerTest("0x123456789ABCDEFabcdef_hello+", "Number(0x123456789ABCDEFabcdef_hello), Plus");
+		}
+
+		[TestMethod]
+		public void FloatLiterals()
+		{
+			RunLexerTest("123.+", "Number(123.), Plus");
+			RunLexerTest("123e123+", "Number(123e123), Plus");
+			RunLexerTest("123e+123+", "Number(123e+123), Plus");
+			RunLexerTest(".123+", "Number(.123), Plus");
+			RunLexerTest(".123e123+", "Number(.123e123), Plus");
+			RunLexerTest(".123e+123+", "Number(.123e+123), Plus");
+			RunLexerTest("0.123+", "Number(0.123), Plus");
+			RunLexerTest("0.123e123+", "Number(0.123e123), Plus");
+			RunLexerTest("0.123e+123+", "Number(0.123e+123), Plus");
+			RunLexerTest("123.123+", "Number(123.123), Plus");
+			RunLexerTest("123.123e123+", "Number(123.123e123), Plus");
+			RunLexerTest("123.123e+123+", "Number(123.123e+123), Plus");
+			RunLexerTest("1.f+", "Number(1.f), Plus");
+			RunLexerTest("123.123e+123f+", "Number(123.123e+123f), Plus");
+			RunLexerTest("123.123e+123F+", "Number(123.123e+123F), Plus");
+		}
+
+		[TestMethod]
+		public void CharacterLiterals()
+		{
+			RunLexerTest("'hello'+", "Character('hello'), Plus");
+			RunLexerTest("'hello\\'\\''+", "Character('hello\\'\\''), Plus");
+			RunLexerTest("u'hello'+", "Character(u'hello'), Plus");
+			RunLexerTest("U'hello'+", "Character(U'hello'), Plus");
+			RunLexerTest("L'hello'+", "Character(L'hello'), Plus");
+			RunLexerTest("u8'hello'+", "Character(u8'hello'), Plus");
+		}
+
+		[TestMethod]
+		public void StringLiterals()
+		{
+			RunLexerTest("\"hello\"+", "String(\"hello\"), Plus");
+			RunLexerTest("\"hello\\\"\\\"\"+", "String(\"hello\\\"\\\"\"), Plus");
+			RunLexerTest("u\"hello\"+", "String(u\"hello\"), Plus");
+			RunLexerTest("U\"hello\"+", "String(U\"hello\"), Plus");
+			RunLexerTest("L\"hello\"+", "String(L\"hello\"), Plus");
+			RunLexerTest("u8\"hello\"+", "String(u8\"hello\"), Plus");
+		}
+
+		[TestMethod]
+		public void Identifiers()
+		{
+			RunLexerTest("abcdefghijklmnopqrstuvwxyz1234567890ABCDEFGHIJKLMNOPQRSTUVWXYZ_+", "Identifier(abcdefghijklmnopqrstuvwxyz1234567890ABCDEFGHIJKLMNOPQRSTUVWXYZ_), Plus");
+		}
+
+		[TestMethod]
+		public void Comments()
+		{
+			// Single-line comments
+			RunLexerTest("//hello\n+", "Newline, Plus");
+			RunLexerTest("//hello\n//world\n+", "Newline, Newline, Plus");
+
+			// Multi-line comments
+			RunLexerTest("/*hello*/+", "Plus");
+			RunLexerTest("/*/hello*/+", "Plus");
+			RunLexerTest("/*/hello**/+", "Plus");
+		}
+
+		[TestMethod]
+		public void LineNumbers()
+		{
+			// Test reported line numbers are accurate
+			RunLineNumberTest("+", 1);
+			RunLineNumberTest("\n+", 2);
+			RunLineNumberTest("\r\n+", 2);
+			RunLineNumberTest("\r\n\r\n+", 3);
+			RunLineNumberTest("\r\n\n+", 3);
+			RunLineNumberTest("//comment\n+", 2);
+			RunLineNumberTest("//comment\\\n\n+", 3);
+			RunLineNumberTest("/*comment\n*/+", 2);
+			RunLineNumberTest("/*comment\r\n*/+", 2);
+			RunLineNumberTest("/*comment\r\n\r\n*/+", 3);
+			RunLineNumberTest("first line \\\n+", 2);
+		}
+
+		public static void RunLexerTest(string InputText, string ExpectedOutputText)
+		{
+			List<Token> Tokens = new List<Token>();
+
+			TokenReader Reader = new TokenReader(InputText);
+			while(Reader.MoveNext())
+			{
+				Tokens.Add(Reader.Current);
+			}
+
+			string OutputText = String.Join(", ", Tokens.Select(x => FormatToken(x)));
+			Assert.AreEqual(ExpectedOutputText, OutputText);
+		}
+
+		static string FormatToken(Token Token)
+		{
+			StringBuilder Result = new StringBuilder();
+			Result.AppendFormat(Token.Type.ToString());
+			if(Token.Type == TokenType.Identifier || Token.Type == TokenType.Number || Token.Type == TokenType.Character || Token.Type == TokenType.String)
+			{
+				Result.AppendFormat("({0})", Token.Text);
+			}
+			return Result.ToString();
+		}
+
+		public static void RunLineNumberTest(string InputText, int ExpectedLineNumber)
+		{
+			TokenReader Reader = new TokenReader(InputText);
+			while(Reader.MoveNext())
+			{
+				if(Reader.Current.Type == TokenType.Plus)
+				{
+					break;
+				}
+			}
+			Assert.AreEqual(ExpectedLineNumber, Reader.LineNumber);
+		}
+	}
+}
diff --git a/Engine/Source/Programs/UnrealBuildToolTests/UnrealBuildToolTests.csproj b/Engine/Source/Programs/UnrealBuildToolTests/UnrealBuildToolTests.csproj
new file mode 100644
--- /dev/null
+++ b/Engine/Source/Programs/UnrealBuildToolTests/UnrealBuildToolTests.csproj
@@ -0,0 +1,20 @@
+<Project Sdk="Microsoft.NET.Sdk">
+
+  <PropertyGroup>
+    <TargetFramework>netcoreapp3.1</TargetFramework>
+
+    <IsPackable>false</IsPackable>
+  </PropertyGroup>
+
+  <ItemGroup>
+    <PackageReference Include="Microsoft.NET.Test.Sdk" Version="16.5.0" />
+    <PackageReference Include="MSTest.TestAdapter" Version="2.1.0" />
+    <PackageReference Include="MSTest.TestFramework" Version="2.1.0" />
+    <PackageReference Include="coverlet.collector" Version="1.2.0" />
+  </ItemGroup>
+
+  <ItemGroup>
+    <ProjectReference Include="..\UnrealBuildTool\UnrealBuildToolCore.csproj" />
+  </ItemGroup>
+
+</Project>
