UBT: Add utility class for caching MSVC predefined macros.

#rnx

diff --git a/Engine/Build/Windows/PreludeMacros.txt b/Engine/Build/Windows/PreludeMacros.txt
new file mode 100644
--- /dev/null
+++ b/Engine/Build/Windows/PreludeMacros.txt
@@ -0,0 +1,199 @@
+; This file is used to generate Prelude.h, which sets up the initial environment for the preprocessor.
+
+; Standard macros
+__STDC__
+__STDC_VERSION__
+__STDC_HOSTED__
+__cplusplus
+
+; Microsoft defines, taken from https://msdn.microsoft.com/en-us/library/b0084kay.aspx
+_ATL_VER
+__AVX__
+__AVX2__
+_CHAR_UNSIGNED
+__CLR_VER
+__cplusplus_cli
+__cplusplus_winrt
+_CPPRTTI
+_CPPUNWIND
+_DEBUG
+_DLL
+_INTEGRAL_MAX_BITS
+_M_AMD64
+_M_ARM
+_M_CEE
+_M_CEE_PURE
+_M_CEE_SAFE
+_M_IX86
+_M_ARM_FP
+_M_IX86_FP
+_M_X64
+_MANAGED
+_MFC_VER
+_MSC_BUILD
+_MSC_EXTENSIONS
+_MSC_FULL_VER
+_MSC_VER
+__MSVC_RUNTIME_CHECKS
+_MT
+_NATIVE_WCHAR_T_DEFINED
+_OPENMP
+_VC_NODEFAULTLIB
+_WCHAR_T_DEFINED
+_WIN32
+_WIN64
+
+; Microsoft defines, undocumented but referenced by headers
+_M_FP_FAST
+
+; GCC defines, taken from https://gcc.gnu.org/onlinedocs/cpp/Common-Predefined-Macros.html#Common-Predefined-Macros
+__OBJC__
+__ASSEMBLER__
+__GFORTRAN__
+__GNUC__
+__GNUC_MINOR__
+__GNUC_PATCHLEVEL__
+__GNUG__
+__STRICT_ANSI__
+__BASE_FILE__
+__INCLUDE_LEVEL__
+__ELF__
+__VERSION__
+__OPTIMIZE__
+__OPTIMIZE_SIZE__
+__NO_INLINE__
+__GNUC_GNU_INLINE__
+__GNUC_STDC_INLINE__
+__CHAR_UNSIGNED__
+__WCHAR_UNSIGNED__
+__REGISTER_PREFIX__
+__USER_LABEL_PREFIX__
+__SIZE_TYPE__
+__PTRDIFF_TYPE__
+__WCHAR_TYPE__
+__WINT_TYPE__
+__INTMAX_TYPE__
+__UINTMAX_TYPE__
+__SIG_ATOMIC_TYPE__
+__INT8_TYPE__
+__INT16_TYPE__
+__INT32_TYPE__
+__INT64_TYPE__
+__UINT8_TYPE__
+__UINT16_TYPE__
+__UINT32_TYPE__
+__UINT64_TYPE__
+__INT_LEAST8_TYPE__
+__INT_LEAST16_TYPE__
+__INT_LEAST32_TYPE__
+__INT_LEAST64_TYPE__
+__UINT_LEAST8_TYPE__
+__UINT_LEAST16_TYPE__
+__UINT_LEAST32_TYPE__
+__UINT_LEAST64_TYPE__
+__INT_FAST8_TYPE__
+__INT_FAST16_TYPE__
+__INT_FAST32_TYPE__
+__INT_FAST64_TYPE__
+__UINT_FAST8_TYPE__
+__UINT_FAST16_TYPE__
+__UINT_FAST32_TYPE__
+__UINT_FAST64_TYPE__
+__INTPTR_TYPE__
+__UINTPTR_TYPE__
+__CHAR_BIT__
+__SCHAR_MAX__
+__WCHAR_MAX__
+__SHRT_MAX__
+__INT_MAX__
+__LONG_MAX__
+__LONG_LONG_MAX__
+__WINT_MAX__
+__SIZE_MAX__
+__PTRDIFF_MAX__
+__INTMAX_MAX__
+__UINTMAX_MAX__
+__SIG_ATOMIC_MAX__
+__INT8_MAX__
+__INT16_MAX__
+__INT32_MAX__
+__INT64_MAX__
+__UINT8_MAX__
+__UINT16_MAX__
+__UINT32_MAX__
+__UINT64_MAX__
+__INT_LEAST8_MAX__
+__INT_LEAST16_MAX__
+__INT_LEAST32_MAX__
+__INT_LEAST64_MAX__
+__UINT_LEAST8_MAX__
+__UINT_LEAST16_MAX__
+__UINT_LEAST32_MAX__
+__UINT_LEAST64_MAX__
+__INT_FAST8_MAX__
+__INT_FAST16_MAX__
+__INT_FAST32_MAX__
+__INT_FAST64_MAX__
+__UINT_FAST8_MAX__
+__UINT_FAST16_MAX__
+__UINT_FAST32_MAX__
+__UINT_FAST64_MAX__
+__INTPTR_MAX__
+__UINTPTR_MAX__
+__WCHAR_MIN__
+__WINT_MIN__
+__SIG_ATOMIC_MIN__
+__INT8_C
+__INT16_C
+__INT32_C
+__INT64_C
+__UINT8_C
+__UINT16_C
+__UINT32_C
+__UINT64_C
+__INTMAX_C
+__UINTMAX_C
+__SIZEOF_INT__
+__SIZEOF_LONG__
+__SIZEOF_LONG_LONG__
+__SIZEOF_SHORT__
+__SIZEOF_POINTER__
+__SIZEOF_FLOAT__
+__SIZEOF_DOUBLE__
+__SIZEOF_LONG_DOUBLE__
+__SIZEOF_SIZE_T__
+__SIZEOF_WCHAR_T__
+__SIZEOF_WINT_T__
+__SIZEOF_PTRDIFF_T__
+__BYTE_ORDER__
+__ORDER_LITTLE_ENDIAN__
+__ORDER_BIG_ENDIAN__
+__ORDER_PDP_ENDIAN__
+__FLOAT_WORD_ORDER__
+__DEPRECATED
+__EXCEPTIONS
+__GXX_RTTI
+__USING_SJLJ_EXCEPTIONS__
+__GXX_EXPERIMENTAL_CXX0X__
+__GXX_WEAK__
+__NEXT_RUNTIME__
+__LP64__
+_LP64
+__SSP__
+__SSP_ALL__
+__SSP_STRONG__
+__SSP_EXPLICIT__
+__SANITIZE_ADDRESS__
+__SANITIZE_THREAD__
+__GCC_HAVE_SYNC_COMPARE_AND_SWAP_1
+__GCC_HAVE_SYNC_COMPARE_AND_SWAP_2
+__GCC_HAVE_SYNC_COMPARE_AND_SWAP_4
+__GCC_HAVE_SYNC_COMPARE_AND_SWAP_8
+__GCC_HAVE_SYNC_COMPARE_AND_SWAP_16
+__GCC_HAVE_DWARF2_CFI_ASM
+__FP_FAST_FMA
+__FP_FAST_FMAF
+__FP_FAST_FMAL
+__GCC_IEC_559
+__GCC_IEC_559_COMPLEX
+__NO_MATH_ERRNO__
diff --git a/Engine/Source/Programs/UnrealBuildTool/Platform/Windows/VCPreludeCache.cs b/Engine/Source/Programs/UnrealBuildTool/Platform/Windows/VCPreludeCache.cs
new file mode 100644
--- /dev/null
+++ b/Engine/Source/Programs/UnrealBuildTool/Platform/Windows/VCPreludeCache.cs
@@ -0,0 +1,181 @@
+// Copyright Epic Games, Inc. All Rights Reserved.
+
+using System;
+using System.Collections.Generic;
+using System.Diagnostics;
+using System.IO;
+using System.Linq;
+using Tools.DotNETCommon;
+
+namespace UnrealBuildTool
+{
+	/// <summary>
+	/// Caches predefined state ("prelude") for compiler versions and options
+	/// </summary>
+	static class VCPreludeCache
+	{
+		/// <summary>
+		/// Set of options which can influence the prelude state. Any options matching entries in this set will be used to generate a separate prelude file.
+		/// </summary>
+		static HashSet<string> FilterArguments = new HashSet<string>(StringComparer.OrdinalIgnoreCase)
+		{
+			"/EH",  // Exception handling mode
+			"/GR",  // RTTI settings
+			"/MD",  // Multithreaded DLL runtime library
+			"/MDd", // Multithreaded debug DLL runtime library
+			"/MT",  // Multithreaded static runtime library
+			"/MTd", // Multithreaded debug static runtime library
+		};
+
+		/// <summary>
+		/// Object used for ensuring only one thread can create items in the cache at once
+		/// </summary>
+		static object LockObject = new object();
+
+		/// <summary>
+		/// File containing a list of macros to be used to generate a prelude file
+		/// </summary>
+		static readonly FileReference MacroNamesFile = FileReference.Combine(UnrealBuildTool.EngineDirectory, "Build", "Windows", "PreludeMacros.txt");
+
+		/// <summary>
+		/// Base directory for the cache
+		/// </summary>
+		static readonly DirectoryReference PreludeCacheDir = DirectoryReference.Combine(UnrealBuildTool.EngineDirectory, "Saved", "UnrealBuildTool", "Prelude", "MSVC");
+
+		/// <summary>
+		/// The parsed list of macros
+		/// </summary>
+		static List<string> MacroNames = new List<string>();
+
+		/// <summary>
+		/// Path to the file used to generate prelude headers
+		/// </summary>
+		static FileReference PreludeGeneratorFile;
+
+		/// <summary>
+		/// Cached mapping from compiler filename to its version string
+		/// </summary>
+		static Dictionary<FileReference, string> CompilerExeToVersion = new Dictionary<FileReference, string>();
+
+		/// <summary>
+		/// Get the prelude file (or create it) for the given compiler executable and set of options
+		/// </summary>
+		/// <param name="Compiler">Path to the compiler executable</param>
+		/// <param name="Arguments">Command line arguments for generating the prelude file</param>
+		/// <returns>Path to the prelude file</returns>
+		public static FileReference GetPreludeHeader(FileReference Compiler, List<string> Arguments)
+		{
+			// Cache the compiler version
+			string CompilerVersion;
+			lock (CompilerExeToVersion)
+			{
+				if (!CompilerExeToVersion.TryGetValue(Compiler, out CompilerVersion))
+				{
+					FileVersionInfo VersionInfo = FileVersionInfo.GetVersionInfo(Compiler.FullName);
+					CompilerVersion = String.Format("{0} {1}", VersionInfo.ProductName, VersionInfo.ProductVersion);
+				}
+			}
+
+			// Cache the macro list
+			lock (MacroNames)
+			{
+				if (MacroNames.Count == 0)
+				{
+					string[] MacroListLines = FileReference.ReadAllLines(MacroNamesFile);
+					foreach (string MacroListLine in MacroListLines)
+					{
+						string MacroName = MacroListLine.Trim();
+						if (MacroName.Length > 0 && MacroName[0] != ';')
+						{
+							MacroNames.Add(MacroName);
+						}
+					}
+				}
+			}
+
+			// Create a digest from the command line
+			string CommandLine = String.Join(" ", Arguments.Where(x => FilterArguments.Contains(x)).OrderBy(x => x));
+			ContentHash Digest = ContentHash.SHA1(String.Format("{0}\n{1}\n{2}", CompilerVersion, CommandLine, String.Join("\n", MacroNames)));
+
+			// Get the prelude file
+			FileReference PreludeHeader = FileReference.Combine(PreludeCacheDir, String.Format("{0}.h", Digest));
+			if (!FileReference.Exists(PreludeHeader))
+			{
+				lock (LockObject)
+				{
+					if (!FileReference.Exists(PreludeHeader))
+					{
+						CreatePreludeHeader(PreludeHeader, Compiler, CompilerVersion, CommandLine);
+					}
+				}
+			}
+			return PreludeHeader;
+		}
+
+		/// <summary>
+		/// Creates a prelude header
+		/// </summary>
+		/// <param name="PreludeHeader">Path to the file to create</param>
+		/// <param name="Compiler">Path to the compiler</param>
+		/// <param name="CompilerVersion">Version number of the compiler</param>
+		/// <param name="CommandLine">Command line arguments used to generate the header</param>
+		static void CreatePreludeHeader(FileReference PreludeHeader, FileReference Compiler, string CompilerVersion, string CommandLine)
+		{
+			// Make sure the cache directory exists
+			DirectoryReference.CreateDirectory(PreludeCacheDir);
+
+			// Make sure the prelude generator file exists
+			if (PreludeGeneratorFile == null)
+			{
+				PreludeGeneratorFile = FileReference.Combine(PreludeCacheDir, "PreludeGenerator.cpp");
+				using (StreamWriter Writer = new StreamWriter(PreludeGeneratorFile.FullName))
+				{
+					Writer.WriteLine("#define STRINGIZE_2(x) #x");
+					Writer.WriteLine("#define STRINGIZE(x) STRINGIZE_2(x)");
+					foreach (string MacroName in MacroNames)
+					{
+						Writer.WriteLine("");
+						Writer.WriteLine("#ifdef {0}", MacroName);
+						Writer.WriteLine("#pragma message(\"#define {0} \" STRINGIZE({0}))", MacroName);
+						Writer.WriteLine("#endif");
+					}
+				}
+			}
+
+			// Invoke the compiler and capture the output
+			using (ManagedProcessGroup ProcessGroup = new ManagedProcessGroup())
+			{
+				string FullCommandLine = String.Format("{0} /nologo /c {1}", CommandLine, Utils.MakePathSafeToUseWithCommandLine(PreludeGeneratorFile.FullName));
+				using (ManagedProcess Process = new ManagedProcess(ProcessGroup, Compiler.FullName, FullCommandLine, PreludeCacheDir.FullName, null, null, ProcessPriorityClass.Normal))
+				{
+					// Filter the output
+					FileReference TempPreludeFile = new FileReference(PreludeHeader.FullName + ".temp");
+					using (StreamWriter Writer = new StreamWriter(TempPreludeFile.FullName))
+					{
+						Writer.WriteLine("// Compiler: {0}", CompilerVersion);
+						Writer.WriteLine("// Arguments: {0}", CommandLine);
+						Writer.WriteLine();
+
+						string Line;
+						while (Process.TryReadLine(out Line))
+						{
+							Line = Line.Trim();
+							if (Line.Length > 0)
+							{
+								if (Line.StartsWith("#define", StringComparison.Ordinal))
+								{
+									Writer.WriteLine(Line);
+								}
+								else if (Line.Trim() != PreludeGeneratorFile.GetFileName())
+								{
+									throw new BuildException("Unexpected output from compiler: {0}", Line);
+								}
+							}
+						}
+					}
+					FileReference.Move(TempPreludeFile, PreludeHeader);
+				}
+			}
+		}
+	}
+}
