UBT: Add interfaces for CAS providers, and a default implementation for the filesystem.

#rnx

diff --git a/Engine/Source/Programs/UnrealBuildTool/Storage/IStorageProvider.cs b/Engine/Source/Programs/UnrealBuildTool/Storage/IStorageProvider.cs
new file mode 100644
--- /dev/null
+++ b/Engine/Source/Programs/UnrealBuildTool/Storage/IStorageProvider.cs
@@ -0,0 +1,64 @@
+using System;
+using System.Collections.Generic;
+using System.Diagnostics;
+using System.IO;
+using System.Linq;
+using System.Security.Cryptography;
+using System.Text;
+using System.Threading;
+using System.Threading.Tasks;
+using Tools.DotNETCommon;
+
+namespace UnrealBuildTool.Storage
+{
+	/// <summary>
+	/// Interface for reads from content-addressible storage
+	/// </summary>
+	interface IStorageReader : IDisposable
+	{
+		/// <summary>
+		/// Accessor for the read stream
+		/// </summary>
+		Stream Stream { get; }
+
+		/// <summary>
+		/// Whether the reader contains any data
+		/// </summary>
+		bool IsValid { get; }
+	}
+
+	/// <summary>
+	/// Interface for writes to content-addressible storage
+	/// </summary>
+	interface IStorageWriter : IDisposable
+	{
+		/// <summary>
+		/// Accessor for the write stream
+		/// </summary>
+		Stream Stream { get; }
+
+		/// <summary>
+		/// Commits the written data to the cache
+		/// </summary>
+		void Commit();
+	}
+
+	/// <summary>
+	/// Interface for an artifact cache
+	/// </summary>
+	interface IStorageProvider
+	{
+		/// <summary>
+		/// Attempts to open a file from the output cache
+		/// </summary>
+		/// <param name="Digest">Digest of the item to retrieve</param>
+		/// <returns>True if the item exists in the cache, false otherwise</returns>
+		IStorageReader CreateReader(ContentHash Digest);
+
+		/// <summary>
+		/// Opens a stream for writing into the cache. The digest 
+		/// </summary>
+		/// <returns></returns>
+		IStorageWriter CreateWriter(ContentHash Digest);
+	}
+}
diff --git a/Engine/Source/Programs/UnrealBuildTool/Storage/Impl/FileSystemStorageProvider.cs b/Engine/Source/Programs/UnrealBuildTool/Storage/Impl/FileSystemStorageProvider.cs
new file mode 100644
--- /dev/null
+++ b/Engine/Source/Programs/UnrealBuildTool/Storage/Impl/FileSystemStorageProvider.cs
@@ -0,0 +1,141 @@
+using System;
+using System.Collections.Generic;
+using System.Diagnostics;
+using System.IO;
+using System.Text;
+using Tools.DotNETCommon;
+
+namespace UnrealBuildTool.Storage.Impl
+{
+	/// <summary>
+	/// Storage provider which uses the filesystem
+	/// </summary>
+	class FileSystemStorageProvider : IStorageProvider
+	{
+		/// <summary>
+		/// Cached copy of the current process id
+		/// </summary>
+		static int ProcessId = Process.GetCurrentProcess().Id;
+
+		/// <summary>
+		/// Implements <see cref="IStorageReader"/> for writes to the backing storage
+		/// </summary>
+		class StorageReader : IStorageReader
+		{
+			public Stream Stream
+			{
+				get;
+				private set;
+			}
+
+			public bool IsValid => Stream != null;
+
+			public StorageReader(FileReference Location)
+			{
+				if (FileReference.Exists(Location))
+				{
+					Stream = FileReference.Open(Location, FileMode.Open, FileAccess.Read, FileShare.Read | FileShare.Delete);
+				}
+			}
+
+			public void Dispose()
+			{
+				if (Stream != null)
+				{
+					Stream.Dispose();
+					Stream = null;
+				}
+			}
+		}
+
+		/// <summary>
+		/// Represents a write transaction to the cache. Commit() should be called once the transaction is complete.
+		/// </summary>
+		class StorageWriter : IStorageWriter
+		{
+			FileReference TempLocation;
+			FileReference FinalLocation;
+
+			public Stream Stream
+			{
+				get;
+				private set;
+			}
+
+			public StorageWriter(FileReference Location)
+			{
+				this.FinalLocation = Location;
+				DirectoryReference.CreateDirectory(FinalLocation.Directory);
+
+				this.TempLocation = new FileReference(String.Format("{0}.{1}", Location.FullName, ProcessId));
+				Stream = FileReference.Open(TempLocation, FileMode.Create, FileAccess.Write, FileShare.Read | FileShare.Delete);
+			}
+
+			public void Commit()
+			{
+				if (TempLocation == null)
+				{
+					throw new InvalidOperationException("Item has already been committed");
+				}
+
+				Stream.Close();
+
+				try
+				{
+					FileReference.Move(TempLocation, FinalLocation);
+					TempLocation = null;
+				}
+				catch
+				{
+					FileReference.Delete(TempLocation);
+					TempLocation = null;
+					throw;
+				}
+			}
+
+			public void Dispose()
+			{
+				if (Stream != null)
+				{
+					Stream.Dispose();
+					Stream = null;
+				}
+				if (TempLocation != null)
+				{
+					FileReference.Delete(TempLocation);
+					TempLocation = null;
+				}
+			}
+		}
+
+		/// <summary>
+		/// Attempts to open a file from the output cache
+		/// </summary>
+		/// <param name="Digest">Digest of the item to retrieve</param>
+		/// <returns>Reader interface for the file</returns>
+		public IStorageReader CreateReader(ContentHash Digest)
+		{
+			return new StorageReader(GetFileForDigest(Digest));
+		}
+
+		/// <summary>
+		/// Opens a stream for writing into the cache. The digest 
+		/// </summary>
+		/// <returns></returns>
+		public IStorageWriter CreateWriter(ContentHash Digest)
+		{
+			return new StorageWriter(GetFileForDigest(Digest));
+		}
+
+		/// <summary>
+		/// Gets the filename on disk to use for a particular digest
+		/// </summary>
+		/// <param name="Digest">The digest to find a filename for</param>
+		/// <returns>Filename to use for the given digest</returns>
+		static FileReference GetFileForDigest(ContentHash Digest)
+		{
+			string DigestText = Digest.ToString();
+			return FileReference.Combine(UnrealBuildTool.EngineDirectory, "Saved", "UnrealBuildTool", "Cache", String.Format("{0}/{1}/{2}/{3}.bin", DigestText[0], DigestText[1], DigestText[2], DigestText));
+		}
+	}
+}
diff --git a/Engine/Source/Programs/UnrealBuildTool/UnrealBuildTool.csproj b/Engine/Source/Programs/UnrealBuildTool/UnrealBuildTool.csproj
--- a/Engine/Source/Programs/UnrealBuildTool/UnrealBuildTool.csproj
+++ b/Engine/Source/Programs/UnrealBuildTool/UnrealBuildTool.csproj
@@ -169,6 +169,8 @@
     <Compile Include="ProjectFiles\Xcode\XcodeFrameworkWrapperProject.cs" />
     <Compile Include="ProjectFiles\Xcode\XcodeProject.cs" />
     <Compile Include="ProjectFiles\Xcode\XcodeProjectFileGenerator.cs" />
+    <Compile Include="Storage\Impl\FileSystemStorageProvider.cs" />
+    <Compile Include="Storage\IStorageProvider.cs" />
     <Compile Include="System\Action.cs" />
     <Compile Include="System\ActionExecutor.cs" />
     <Compile Include="Modes\AggregateParsedTimingInfo.cs" />
