UBT: Remove the DistCC executor. We don't maintain this any more.

diff --git a/Engine/Source/Programs/UnrealBuildTool/Configuration/BuildConfiguration.cs b/Engine/Source/Programs/UnrealBuildTool/Configuration/BuildConfiguration.cs
--- a/Engine/Source/Programs/UnrealBuildTool/Configuration/BuildConfiguration.cs
+++ b/Engine/Source/Programs/UnrealBuildTool/Configuration/BuildConfiguration.cs
@@ -93,13 +93,6 @@ namespace UnrealBuildTool
 		[CommandLine("-NoUBTMakefiles", Value = "false")]
 		public bool bUseUBTMakefiles = true;
 
-		/// <summary>
-		/// Whether DMUCS/Distcc may be used.
-		/// Distcc requires some setup -- so by default, disable it so that we do not break local or remote building.
-		/// </summary>
-		[XmlConfigFile]
-		public bool bAllowDistcc = false;
-
 		/// <summary>
 		/// Whether to allow using parallel executor on Windows.
 		/// </summary>
diff --git a/Engine/Source/Programs/UnrealBuildTool/Configuration/UEBuildPlatform.cs b/Engine/Source/Programs/UnrealBuildTool/Configuration/UEBuildPlatform.cs
--- a/Engine/Source/Programs/UnrealBuildTool/Configuration/UEBuildPlatform.cs
+++ b/Engine/Source/Programs/UnrealBuildTool/Configuration/UEBuildPlatform.cs
@@ -503,14 +503,6 @@ namespace UnrealBuildTool
 			return CanUseXGE();
 		}
 
-		/// <summary>
-		/// If this platform can be compiled with DMUCS/Distcc
-		/// </summary>
-		public virtual bool CanUseDistcc()
-		{
-			return false;
-		}
-
 		/// <summary>
 		/// If this platform can be compiled with SN-DBS
 		/// </summary>
diff --git a/Engine/Source/Programs/UnrealBuildTool/Executors/Distcc.cs b/Engine/Source/Programs/UnrealBuildTool/Executors/Distcc.cs
deleted file mode 100644
diff --git a/Engine/Source/Programs/UnrealBuildTool/Modes/BuildMode.cs b/Engine/Source/Programs/UnrealBuildTool/Modes/BuildMode.cs
--- a/Engine/Source/Programs/UnrealBuildTool/Modes/BuildMode.cs
+++ b/Engine/Source/Programs/UnrealBuildTool/Modes/BuildMode.cs
@@ -508,7 +508,6 @@ namespace UnrealBuildTool
 				{
 					UEBuildPlatform BuildPlatform = UEBuildPlatform.GetBuildPlatform(TargetDescriptor.Platform);
 					BuildConfiguration.bAllowXGE &= BuildPlatform.CanUseXGE();
-					BuildConfiguration.bAllowDistcc &= BuildPlatform.CanUseDistcc();
 					BuildConfiguration.bAllowSNDBS &= BuildPlatform.CanUseSNDBS();
 				}
 
diff --git a/Engine/Source/Programs/UnrealBuildTool/Platform/IOS/UEBuildIOS.cs b/Engine/Source/Programs/UnrealBuildTool/Platform/IOS/UEBuildIOS.cs
--- a/Engine/Source/Programs/UnrealBuildTool/Platform/IOS/UEBuildIOS.cs
+++ b/Engine/Source/Programs/UnrealBuildTool/Platform/IOS/UEBuildIOS.cs
@@ -874,11 +874,6 @@ namespace UnrealBuildTool
 			return false;
 		}
 
-		public override bool CanUseDistcc()
-		{
-			return true;
-		}
-
 		public bool HasCustomIcons(DirectoryReference ProjectDirectoryName)
 		{
 			string IconDir = Path.Combine(ProjectDirectoryName.FullName, "Build", "IOS", "Resources", "Graphics");
diff --git a/Engine/Source/Programs/UnrealBuildTool/Platform/Mac/UEBuildMac.cs b/Engine/Source/Programs/UnrealBuildTool/Platform/Mac/UEBuildMac.cs
--- a/Engine/Source/Programs/UnrealBuildTool/Platform/Mac/UEBuildMac.cs
+++ b/Engine/Source/Programs/UnrealBuildTool/Platform/Mac/UEBuildMac.cs
@@ -121,11 +121,6 @@ namespace UnrealBuildTool
 			return false;
 		}
 
-		public override bool CanUseDistcc()
-		{
-			return true;
-		}
-
 		public override void ResetTarget(TargetRules Target)
 		{			
 		}
diff --git a/Engine/Source/Programs/UnrealBuildTool/System/ActionGraph.cs b/Engine/Source/Programs/UnrealBuildTool/System/ActionGraph.cs
--- a/Engine/Source/Programs/UnrealBuildTool/System/ActionGraph.cs
+++ b/Engine/Source/Programs/UnrealBuildTool/System/ActionGraph.cs
@@ -184,10 +184,6 @@ namespace UnrealBuildTool
 				{
 					Executor = new FASTBuild();
 				}
-				else if (BuildConfiguration.bAllowDistcc)
-				{
-					Executor = new Distcc();
-				}
 				else if(BuildConfiguration.bAllowSNDBS && SNDBS.IsAvailable())
 				{
 					Executor = new SNDBS();
diff --git a/Engine/Source/Programs/UnrealBuildTool/UnrealBuildTool.csproj b/Engine/Source/Programs/UnrealBuildTool/UnrealBuildTool.csproj
--- a/Engine/Source/Programs/UnrealBuildTool/UnrealBuildTool.csproj
+++ b/Engine/Source/Programs/UnrealBuildTool/UnrealBuildTool.csproj
@@ -117,7 +117,6 @@
     <Compile Include="Configuration\UEBuildModuleExternal.cs" />
     <Compile Include="Configuration\UEBuildPlatformFactory.cs" />
     <Compile Include="Configuration\UEBuildPlugin.cs" />
-    <Compile Include="Executors\Distcc.cs" />
     <Compile Include="Executors\LocalExecutor.cs" />
     <Compile Include="Executors\ParallelExecutor.cs" />
     <Compile Include="Executors\SNDBS.cs" />
@@ -134,6 +133,7 @@
     <Compile Include="Platform\Mac\ApplePlatformSDK.cs" />
     <Compile Include="Platform\Windows\MicrosoftPlatformSDK.cs" />
     <Compile Include="Platform\Windows\VCManifestGenerator.cs" />
+    <Compile Include="Platform\Windows\VCPreludeCache.cs" />
     <Compile Include="Preprocessor\Identifier.cs" />
     <Compile Include="Preprocessor\Preprocessor.cs" />
     <Compile Include="Preprocessor\PreprocessorContext.cs" />
