UBT: Remove timing fields from the action class. This is only used by the local executor, which is no longer the default.

diff --git a/Engine/Source/Programs/UnrealBuildTool/Configuration/BuildConfiguration.cs b/Engine/Source/Programs/UnrealBuildTool/Configuration/BuildConfiguration.cs
--- a/Engine/Source/Programs/UnrealBuildTool/Configuration/BuildConfiguration.cs
+++ b/Engine/Source/Programs/UnrealBuildTool/Configuration/BuildConfiguration.cs
@@ -36,12 +36,6 @@ namespace UnrealBuildTool
 		[CommandLine("-PrintDebugInfo", Value = "true")]
 		public bool bPrintDebugInfo = false;
 
-		/// <summary>
-		/// Whether to log detailed action stats. This forces local execution.
-		/// </summary>
-		[XmlConfigFile]
-		public bool bLogDetailedActionStats = false;
-
 		/// <summary>
 		/// Whether the hybrid executor will be used (a remote executor and local executor).
 		/// </summary>
diff --git a/Engine/Source/Programs/UnrealBuildTool/Executors/Experimental/FASTBuild.cs b/Engine/Source/Programs/UnrealBuildTool/Executors/Experimental/FASTBuild.cs
--- a/Engine/Source/Programs/UnrealBuildTool/Executors/Experimental/FASTBuild.cs
+++ b/Engine/Source/Programs/UnrealBuildTool/Executors/Experimental/FASTBuild.cs
@@ -387,7 +387,7 @@ namespace UnrealBuildTool
 			}
 		}
 
-		public override bool ExecuteActions(List<Action> Actions, bool bLogDetailedActionStats)
+		public override bool ExecuteActions(List<Action> Actions)
 		{
 			if (Actions.Count <= 0)
 				return true;
@@ -415,7 +415,7 @@ namespace UnrealBuildTool
 
 			if (PreCompileActions.Any())
 			{
-				bool bResult = new LocalExecutor().ExecuteActions(PreCompileActions, bLogDetailedActionStats);
+				bool bResult = new LocalExecutor().ExecuteActions(PreCompileActions);
 
 				if (!bResult)
 					return false;
@@ -444,7 +444,7 @@ namespace UnrealBuildTool
 
 			if (PostCompileActions.Any())
 			{
-				bool bResult = new LocalExecutor().ExecuteActions(PostCompileActions, bLogDetailedActionStats);
+				bool bResult = new LocalExecutor().ExecuteActions(PostCompileActions);
 
 				if (!bResult)
 					return false;
diff --git a/Engine/Source/Programs/UnrealBuildTool/Executors/HybridExecutor.cs b/Engine/Source/Programs/UnrealBuildTool/Executors/HybridExecutor.cs
--- a/Engine/Source/Programs/UnrealBuildTool/Executors/HybridExecutor.cs
+++ b/Engine/Source/Programs/UnrealBuildTool/Executors/HybridExecutor.cs
@@ -68,9 +68,8 @@ namespace UnrealBuildTool
 		/// Execute the given actions
 		/// </summary>
 		/// <param name="ActionsToExecute">Actions to be executed</param>
-		/// <param name="bLogDetailedActionStats">Whether to log detailed action stats</param>
 		/// <returns>True if the build succeeded, false otherwise</returns>
-		public override bool ExecuteActions(List<Action> ActionsToExecute, bool bLogDetailedActionStats)
+		public override bool ExecuteActions(List<Action> ActionsToExecute)
 		{
 			// Find the number of dependants for each action
 			Dictionary<Action, int> ActionToNumDependents = ActionsToExecute.ToDictionary(x => x, x => 0);
@@ -133,7 +132,7 @@ namespace UnrealBuildTool
 			// Execute the remote actions
 			if(RemoteActionsToExecute.Count > 0)
 			{
-				if (!RemoteExecutor.ExecuteActions(RemoteActionsToExecute, bLogDetailedActionStats))
+				if (!RemoteExecutor.ExecuteActions(RemoteActionsToExecute))
 				{
 					return false;
 				}
@@ -142,7 +141,7 @@ namespace UnrealBuildTool
 			// Pass all the local actions through to the parallel executor
 			if(LocalActionsToExecute.Count > 0)
 			{
-				if(!LocalExecutor.ExecuteActions(LocalActionsToExecute, bLogDetailedActionStats))
+				if(!LocalExecutor.ExecuteActions(LocalActionsToExecute))
 				{
 					return false;
 				}
diff --git a/Engine/Source/Programs/UnrealBuildTool/Executors/LocalExecutor.cs b/Engine/Source/Programs/UnrealBuildTool/Executors/LocalExecutor.cs
--- a/Engine/Source/Programs/UnrealBuildTool/Executors/LocalExecutor.cs
+++ b/Engine/Source/Programs/UnrealBuildTool/Executors/LocalExecutor.cs
@@ -68,9 +68,6 @@ namespace UnrealBuildTool
 		/// </summary>
 		private void ThreadFunc()
 		{
-			// thread start time
-			Action.StartTime = DateTimeOffset.Now;
-
 			string Args = Action.CommandArguments;
 #if NET_CORE
 			// Process Arguments follow windows conventions in .NET Core
@@ -183,9 +180,6 @@ namespace UnrealBuildTool
 					ActionProcess.Close();
 			}
 
-			// track how long it took
-			Action.EndTime = DateTimeOffset.Now;
-
 			// we are done!!
 			bComplete = true;
 		}
@@ -302,7 +296,7 @@ namespace UnrealBuildTool
 		/// Executes the specified actions locally.
 		/// </summary>
 		/// <returns>True if all the tasks successfully executed, or false if any of them failed.</returns>
-		public override bool ExecuteActions(List<Action> Actions, bool bLogDetailedActionStats)
+		public override bool ExecuteActions(List<Action> Actions)
 		{
 			// Time to sleep after each iteration of the loop in order to not busy wait.
 			const float LoopSleepTime = 0.1f;
@@ -426,19 +420,8 @@ namespace UnrealBuildTool
 				}
 			}
 
-			Log.WriteLineIf(bLogDetailedActionStats, LogEventType.Console, "-------- Begin Detailed Action Stats ----------------------------------------------------------");
-			Log.WriteLineIf(bLogDetailedActionStats, LogEventType.Console, "^Action Type^Duration (seconds)^Tool^Task^Using PCH");
-
-			double TotalThreadSeconds = 0;
-
 			// Check whether any of the tasks failed and log action stats if wanted.
 			bool bSuccess = true;
-			double TotalBuildProjectTime = 0;
-			double TotalCompileTime = 0;
-			double TotalCreateAppBundleTime = 0;
-			double TotalGenerateDebugInfoTime = 0;
-			double TotalLinkTime = 0;
-			double TotalOtherActionsTime = 0;
 			foreach (KeyValuePair<Action, ActionThread> ActionProcess in ActionThreadDictionary)
 			{
 				Action Action = ActionProcess.Key;
@@ -455,68 +438,8 @@ namespace UnrealBuildTool
 				{
 					bSuccess = false;
 				}
-				// Log CPU time, tool and task.
-				double ThreadSeconds = Action.Duration.TotalSeconds;
-
-				Log.WriteLineIf(bLogDetailedActionStats,
-					LogEventType.Console,
-					"^{0}^{1:0.00}^{2}^{3}",
-					Action.ActionType.ToString(),
-					ThreadSeconds,
-					Action.CommandPath.GetFileName(),
-					Action.StatusDescription);
-
-				// Update statistics
-				switch (Action.ActionType)
-				{
-					case ActionType.BuildProject:
-						TotalBuildProjectTime += ThreadSeconds;
-						break;
-
-					case ActionType.Compile:
-						TotalCompileTime += ThreadSeconds;
-						break;
-
-					case ActionType.CreateAppBundle:
-						TotalCreateAppBundleTime += ThreadSeconds;
-						break;
-
-					case ActionType.GenerateDebugInfo:
-						TotalGenerateDebugInfoTime += ThreadSeconds;
-						break;
-
-					case ActionType.Link:
-						TotalLinkTime += ThreadSeconds;
-						break;
-
-					default:
-						TotalOtherActionsTime += ThreadSeconds;
-						break;
-				}
-
-				// Keep track of total thread seconds spent on tasks.
-				TotalThreadSeconds += ThreadSeconds;
 			}
 
-			Log.WriteLineIf(bLogDetailedActionStats, LogEventType.Console, "-------- End Detailed Actions Stats -----------------------------------------------------------");
-
-			// Log total CPU seconds and numbers of processors involved in tasks.
-			Log.WriteLineIf(bLogDetailedActionStats,
-				LogEventType.Console, "Cumulative thread seconds ({0} processors): {1:0.00}", System.Environment.ProcessorCount, TotalThreadSeconds);
-
-			// Log detailed stats
-			Log.WriteLineIf(bLogDetailedActionStats,
-				LogEventType.Console,
-				"Cumulative action seconds ({0} processors): {1:0.00} building projects, {2:0.00} compiling, {3:0.00} creating app bundles, {4:0.00} generating debug info, {5:0.00} linking, {6:0.00} other",
-				System.Environment.ProcessorCount,
-				TotalBuildProjectTime,
-				TotalCompileTime,
-				TotalCreateAppBundleTime,
-				TotalGenerateDebugInfoTime,
-				TotalLinkTime,
-				TotalOtherActionsTime
-			);
-
 			return bSuccess;
 		}
 	};
diff --git a/Engine/Source/Programs/UnrealBuildTool/Executors/ParallelExecutor.cs b/Engine/Source/Programs/UnrealBuildTool/Executors/ParallelExecutor.cs
--- a/Engine/Source/Programs/UnrealBuildTool/Executors/ParallelExecutor.cs
+++ b/Engine/Source/Programs/UnrealBuildTool/Executors/ParallelExecutor.cs
@@ -96,7 +96,7 @@ namespace UnrealBuildTool
 		/// Executes the specified actions locally.
 		/// </summary>
 		/// <returns>True if all the tasks successfully executed, or false if any of them failed.</returns>
-		public override bool ExecuteActions(List<Action> InputActions, bool bLogDetailedActionStats)
+		public override bool ExecuteActions(List<Action> InputActions)
 		{
 			Log.TraceInformation("Building {0} {1} with {2} {3}...", InputActions.Count, (InputActions.Count == 1) ? "action" : "actions", NumParallelProcesses, (NumParallelProcesses == 1)? "process" : "processes");
 
diff --git a/Engine/Source/Programs/UnrealBuildTool/Executors/SNDBS.cs b/Engine/Source/Programs/UnrealBuildTool/Executors/SNDBS.cs
--- a/Engine/Source/Programs/UnrealBuildTool/Executors/SNDBS.cs
+++ b/Engine/Source/Programs/UnrealBuildTool/Executors/SNDBS.cs
@@ -44,7 +44,7 @@ namespace UnrealBuildTool
 			return ServiceController.GetServices().Any(s => s.ServiceName.StartsWith("SNDBS") && s.Status == ServiceControllerStatus.Running);
 		}
 
-		public override bool ExecuteActions(List<Action> Actions, bool bLogDetailedActionStats)
+		public override bool ExecuteActions(List<Action> Actions)
 		{
 			if (Actions.Count == 0)
 				return true;
diff --git a/Engine/Source/Programs/UnrealBuildTool/Executors/XGE.cs b/Engine/Source/Programs/UnrealBuildTool/Executors/XGE.cs
--- a/Engine/Source/Programs/UnrealBuildTool/Executors/XGE.cs
+++ b/Engine/Source/Programs/UnrealBuildTool/Executors/XGE.cs
@@ -320,7 +320,7 @@ namespace UnrealBuildTool
 			Log.TraceInformation("XGEEXPORT: Exported '{0}'", OutFile);
 		}
 
-		public override bool ExecuteActions(List<Action> ActionsToExecute, bool bLogDetailedActionStats)
+		public override bool ExecuteActions(List<Action> ActionsToExecute)
 		{
 			bool XGEResult = true;
 
diff --git a/Engine/Source/Programs/UnrealBuildTool/System/Action.cs b/Engine/Source/Programs/UnrealBuildTool/System/Action.cs
--- a/Engine/Source/Programs/UnrealBuildTool/System/Action.cs
+++ b/Engine/Source/Programs/UnrealBuildTool/System/Action.cs
@@ -159,21 +159,6 @@ namespace UnrealBuildTool
 		public int NumTotalDependentActions = 0;
 
 
-		///
-		/// Assembly only (not serialized)
-		///
-
-		/// <summary>
-		/// Start time of action, optionally set by executor.
-		/// </summary>
-		public DateTimeOffset StartTime = DateTimeOffset.MinValue;
-
-		/// <summary>
-		/// End time of action, optionally set by executor.
-		/// </summary>
-		public DateTimeOffset EndTime = DateTimeOffset.MinValue;
-
-
 		public Action(ActionType InActionType)
 		{
 			ActionType = InActionType;
@@ -484,21 +469,5 @@ namespace UnrealBuildTool
 			}
 			return ReturnString;
 		}
-
-		/// <summary>
-		/// Returns the amount of time that this action is or has been executing in.
-		/// </summary>
-		public TimeSpan Duration
-		{
-			get
-			{
-				if (EndTime == DateTimeOffset.MinValue)
-				{
-					return DateTimeOffset.Now - StartTime;
-				}
-
-				return EndTime - StartTime;
-			}
-		}
 	}
 }
diff --git a/Engine/Source/Programs/UnrealBuildTool/System/ActionExecutor.cs b/Engine/Source/Programs/UnrealBuildTool/System/ActionExecutor.cs
--- a/Engine/Source/Programs/UnrealBuildTool/System/ActionExecutor.cs
+++ b/Engine/Source/Programs/UnrealBuildTool/System/ActionExecutor.cs
@@ -15,7 +15,7 @@ namespace UnrealBuildTool
 			get;
 		}
 
-		public abstract bool ExecuteActions(List<Action> ActionsToExecute, bool bLogDetailedActionStats);
+		public abstract bool ExecuteActions(List<Action> ActionsToExecute);
 	}
 
 }
diff --git a/Engine/Source/Programs/UnrealBuildTool/System/ActionGraph.cs b/Engine/Source/Programs/UnrealBuildTool/System/ActionGraph.cs
--- a/Engine/Source/Programs/UnrealBuildTool/System/ActionGraph.cs
+++ b/Engine/Source/Programs/UnrealBuildTool/System/ActionGraph.cs
@@ -199,7 +199,7 @@ namespace UnrealBuildTool
 
 				// Execute the build
 				Stopwatch Timer = Stopwatch.StartNew();
-				if(!Executor.ExecuteActions(ActionsToExecute, BuildConfiguration.bLogDetailedActionStats))
+				if(!Executor.ExecuteActions(ActionsToExecute))
 				{
 					throw new CompilationResultException(CompilationResult.OtherCompilationError);
 				}
@@ -249,7 +249,26 @@ namespace UnrealBuildTool
 			}
 
 			// Sort actions by number of actions depending on them, descending. Secondary sort criteria is file size.
-			Actions.Sort(Action.Compare);
+			Actions.Sort(CompareActions);
+		}
+
+		/// <summary>
+		/// Compares two actions based on total number of dependent items, descending.
+		/// </summary>
+		/// <param name="A">Action to compare</param>
+		/// <param name="B">Action to compare</param>
+		public static int CompareActions(Action A, Action B)
+		{
+			// Primary sort criteria is total number of dependent files, up to max depth.
+			if (B.NumTotalDependentActions != A.NumTotalDependentActions)
+			{
+				return Math.Sign(B.NumTotalDependentActions - A.NumTotalDependentActions);
+			}
+			// Secondary sort criteria is number of pre-requisites.
+			else
+			{
+				return Math.Sign(B.PrerequisiteItems.Count - A.PrerequisiteItems.Count);
+			}
 		}
 
 		/// <summary>
