UBT: Move methods for checking for conflicts off the Action class.

#rnx

diff --git a/Engine/Source/Programs/UnrealBuildTool/System/Action.cs b/Engine/Source/Programs/UnrealBuildTool/System/Action.cs
--- a/Engine/Source/Programs/UnrealBuildTool/System/Action.cs
+++ b/Engine/Source/Programs/UnrealBuildTool/System/Action.cs
@@ -362,65 +362,6 @@ namespace UnrealBuildTool
 			}
 		}
 
-		/// <summary>
-		/// Finds conflicts betwee two actions, and prints them to the log
-		/// </summary>
-		/// <param name="Other">Other action to compare to.</param>
-		/// <returns>True if any conflicts were found, false otherwise.</returns>
-		public bool CheckForConflicts(Action Other)
-		{
-			bool bResult = true;
-			if(ActionType != Other.ActionType)
-			{
-				LogConflict("action type is different", ActionType.ToString(), Other.ActionType.ToString());
-				bResult = false;
-			}
-			if(!Enumerable.SequenceEqual(PrerequisiteItems, Other.PrerequisiteItems))
-			{
-				LogConflict("prerequisites are different", String.Join(", ", PrerequisiteItems.Select(x => x.Location)), String.Join(", ", Other.PrerequisiteItems.Select(x => x.Location)));
-				bResult = false;
-			}
-			if(!Enumerable.SequenceEqual(DeleteItems, Other.DeleteItems))
-			{
-				LogConflict("deleted items are different", String.Join(", ", DeleteItems.Select(x => x.Location)), String.Join(", ", Other.DeleteItems.Select(x => x.Location)));
-				bResult = false;
-			}
-			if(DependencyListFile != Other.DependencyListFile)
-			{
-				LogConflict("dependency list is different", (DependencyListFile == null)? "(none)" : DependencyListFile.AbsolutePath, (Other.DependencyListFile == null)? "(none)" : Other.DependencyListFile.AbsolutePath);
-				bResult = false;
-			}
-			if(WorkingDirectory != Other.WorkingDirectory)
-			{
-				LogConflict("working directory is different", WorkingDirectory.FullName, Other.WorkingDirectory.FullName);
-				bResult = false;
-			}
-			if(CommandPath != Other.CommandPath)
-			{
-				LogConflict("command path is different", CommandPath.FullName, Other.CommandPath.FullName);
-				bResult = false;
-			}
-			if(CommandArguments != Other.CommandArguments)
-			{
-				LogConflict("command arguments are different", CommandArguments, Other.CommandArguments);
-				bResult = false;
-			}
-			return bResult;
-		}
-
-		/// <summary>
-		/// Adds the description of a merge error to an output message
-		/// </summary>
-		/// <param name="Description">Description of the difference</param>
-		/// <param name="OldValue">Previous value for the field</param>
-		/// <param name="NewValue">Conflicting value for the field</param>
-		void LogConflict(string Description, string OldValue, string NewValue)
-		{
-			Log.TraceError("Unable to merge actions producing {0}: {1}", ProducedItems[0].Location.GetFileName(), Description);
-			Log.TraceLog("  Previous: {0}", OldValue);
-			Log.TraceLog("  Conflict: {0}", NewValue);
-		}
-
 		/// <summary>
 		/// Increment the number of dependents, recursively
 		/// </summary>
diff --git a/Engine/Source/Programs/UnrealBuildTool/System/ActionGraph.cs b/Engine/Source/Programs/UnrealBuildTool/System/ActionGraph.cs
--- a/Engine/Source/Programs/UnrealBuildTool/System/ActionGraph.cs
+++ b/Engine/Source/Programs/UnrealBuildTool/System/ActionGraph.cs
@@ -67,7 +67,7 @@ namespace UnrealBuildTool
 					Action ExistingAction;
 					if(ItemToProducingAction.TryGetValue(ProducedItem, out ExistingAction))
 					{
-						bResult &= ExistingAction.CheckForConflicts(Action);
+						bResult &= CheckForConflicts(ExistingAction, Action);
 					}
 					else
 					{
@@ -82,6 +82,67 @@ namespace UnrealBuildTool
 			}
 		}
 
+		/// <summary>
+		/// Finds conflicts betwee two actions, and prints them to the log
+		/// </summary>
+		/// <param name="A">The first action</param>
+		/// <param name="B">The second action</param>
+		/// <returns>True if any conflicts were found, false otherwise.</returns>
+		public static bool CheckForConflicts(Action A, Action B)
+		{
+			bool bResult = true;
+			if (A.ActionType != B.ActionType)
+			{
+				LogConflict(A, "action type is different", A.ActionType.ToString(), B.ActionType.ToString());
+				bResult = false;
+			}
+			if (!Enumerable.SequenceEqual(A.PrerequisiteItems, B.PrerequisiteItems))
+			{
+				LogConflict(A, "prerequisites are different", String.Join(", ", A.PrerequisiteItems.Select(x => x.Location)), String.Join(", ", B.PrerequisiteItems.Select(x => x.Location)));
+				bResult = false;
+			}
+			if (!Enumerable.SequenceEqual(A.DeleteItems, B.DeleteItems))
+			{
+				LogConflict(A, "deleted items are different", String.Join(", ", A.DeleteItems.Select(x => x.Location)), String.Join(", ", B.DeleteItems.Select(x => x.Location)));
+				bResult = false;
+			}
+			if (A.DependencyListFile != B.DependencyListFile)
+			{
+				LogConflict(A, "dependency list is different", (A.DependencyListFile == null) ? "(none)" : A.DependencyListFile.AbsolutePath, (B.DependencyListFile == null) ? "(none)" : B.DependencyListFile.AbsolutePath);
+				bResult = false;
+			}
+			if (A.WorkingDirectory != B.WorkingDirectory)
+			{
+				LogConflict(A, "working directory is different", A.WorkingDirectory.FullName, B.WorkingDirectory.FullName);
+				bResult = false;
+			}
+			if (A.CommandPath != B.CommandPath)
+			{
+				LogConflict(A, "command path is different", A.CommandPath.FullName, B.CommandPath.FullName);
+				bResult = false;
+			}
+			if (A.CommandArguments != B.CommandArguments)
+			{
+				LogConflict(A, "command arguments are different", A.CommandArguments, B.CommandArguments);
+				bResult = false;
+			}
+			return bResult;
+		}
+
+		/// <summary>
+		/// Adds the description of a merge error to an output message
+		/// </summary>
+		/// <param name="Action">The action with the conflict</param>
+		/// <param name="Description">Description of the difference</param>
+		/// <param name="OldValue">Previous value for the field</param>
+		/// <param name="NewValue">Conflicting value for the field</param>
+		static void LogConflict(Action Action, string Description, string OldValue, string NewValue)
+		{
+			Log.TraceError("Unable to merge actions producing {0}: {1}", Action.ProducedItems[0].Location.GetFileName(), Description);
+			Log.TraceLog("  Previous: {0}", OldValue);
+			Log.TraceLog("  Conflict: {0}", NewValue);
+		}
+
 		/// <summary>
 		/// Builds a list of actions that need to be executed to produce the specified output items.
 		/// </summary>
