UBT: Move all internal derived data for action objects into a separate class. This will allow creating specific action-derived classes for different actions.

diff --git a/Engine/Source/Programs/UnrealBuildTool/Executors/Experimental/FASTBuild.cs b/Engine/Source/Programs/UnrealBuildTool/Executors/Experimental/FASTBuild.cs
--- a/Engine/Source/Programs/UnrealBuildTool/Executors/Experimental/FASTBuild.cs
+++ b/Engine/Source/Programs/UnrealBuildTool/Executors/Experimental/FASTBuild.cs
@@ -269,13 +269,13 @@ namespace UnrealBuildTool
 
 		private ObjectIDGenerator objectIDGenerator = new ObjectIDGenerator();
 
-		private long GetActionID(Action Action)
+		private long GetActionID(QueuedAction Action)
 		{
 			bool bFirstTime = false;
 			return objectIDGenerator.GetId(Action, out bFirstTime);
 		}
 
-		private string ActionToActionString(Action Action)
+		private string ActionToActionString(QueuedAction Action)
 		{
 			return ActionToActionString(GetActionID(Action));
 		}
@@ -299,7 +299,7 @@ namespace UnrealBuildTool
 			return $"\t\t'{ActionToActionString(UniqueId)}', ;{StatusDescription}{ExtraInfoString}";
 		}
 
-		private string ActionToDependencyString(Action Action)
+		private string ActionToDependencyString(QueuedAction Action)
 		{
 			return ActionToDependencyString(GetActionID(Action), Action.StatusDescription, Action.CommandDescription, Action.ActionType);
 		}
@@ -317,33 +317,33 @@ namespace UnrealBuildTool
 
 		private FBBuildType BuildType = FBBuildType.Windows;
 
-		private readonly static Tuple<string, Func<Action, string>, FBBuildType>[] BuildTypeSearchParams = new Tuple<string, Func<Action, string>, FBBuildType>[]
+		private readonly static Tuple<string, Func<QueuedAction, string>, FBBuildType>[] BuildTypeSearchParams = new Tuple<string, Func<QueuedAction, string>, FBBuildType>[]
 		{
-			Tuple.Create<string, Func<Action, string>, FBBuildType>
+			Tuple.Create<string, Func<QueuedAction, string>, FBBuildType>
 			(
 				"Xcode",
 				Action => Action.CommandArguments,
 				FBBuildType.Apple
 			),
-			Tuple.Create<string, Func<Action, string>, FBBuildType>
+			Tuple.Create<string, Func<QueuedAction, string>, FBBuildType>
 			(
 				"apple",
 				Action => Action.CommandArguments.ToLower(),
 				FBBuildType.Apple
 			),
-			Tuple.Create<string, Func<Action, string>, FBBuildType>
+			Tuple.Create<string, Func<QueuedAction, string>, FBBuildType>
 			(
 				"/bin/sh",
 				Action => Action.CommandPath.FullName.ToLower(),
 				FBBuildType.Apple
 			),
-			Tuple.Create<string, Func<Action, string>, FBBuildType>
+			Tuple.Create<string, Func<QueuedAction, string>, FBBuildType>
 			(
 				"Windows",		// Not a great test
 				Action => Action.CommandPath.FullName,
 				FBBuildType.Windows
 			),
-			Tuple.Create<string, Func<Action, string>, FBBuildType>
+			Tuple.Create<string, Func<QueuedAction, string>, FBBuildType>
 			(
 				"Microsoft",	// Not a great test
 				Action => Action.CommandPath.FullName,
@@ -351,11 +351,11 @@ namespace UnrealBuildTool
 			),
 		};
 
-		private bool DetectBuildType(IEnumerable<Action> Actions)
+		private bool DetectBuildType(IEnumerable<QueuedAction> Actions)
 		{
-			foreach (Action Action in Actions)
+			foreach (QueuedAction Action in Actions)
 			{
-				foreach (Tuple<string, Func<Action, string>, FBBuildType> BuildTypeSearchParam in BuildTypeSearchParams)
+				foreach (Tuple<string, Func<QueuedAction, string>, FBBuildType> BuildTypeSearchParam in BuildTypeSearchParams)
 				{
 					if (BuildTypeSearchParam.Item2(Action).Contains(BuildTypeSearchParam.Item1))
 					{
@@ -367,7 +367,7 @@ namespace UnrealBuildTool
 			}
 
 			Log.TraceError("Couldn't detect build type from actions! Unsupported platform?");
-			foreach (Action Action in Actions)
+			foreach (QueuedAction Action in Actions)
 			{
 				PrintActionDetails(Action);
 			}
@@ -387,26 +387,26 @@ namespace UnrealBuildTool
 			}
 		}
 
-		public override bool ExecuteActions(List<Action> Actions)
+		public override bool ExecuteActions(List<QueuedAction> Actions)
 		{
 			if (Actions.Count <= 0)
 				return true;
 
-			IEnumerable<Action> CompileActions		= Actions.Where(Action => (Action.ActionType == ActionType.Compile && Action.bCanExecuteRemotely == true));
-			IEnumerable<Action> NonCompileActions	= Actions.Where(Action => (Action.ActionType != ActionType.Compile || Action.bCanExecuteRemotely == false));
+			IEnumerable<QueuedAction> CompileActions		= Actions.Where(Action => (Action.ActionType == ActionType.Compile && Action.bCanExecuteRemotely == true));
+			IEnumerable<QueuedAction> NonCompileActions	= Actions.Where(Action => (Action.ActionType != ActionType.Compile || Action.bCanExecuteRemotely == false));
 
 			///////////////////////////////////////////////////////////////
 			// Pre Compile Stage
 
 			// We want to complete any non-compile actions locally that are necessary for the distributed compile step
-			List<Action> PreCompileActions =
+			List<QueuedAction> PreCompileActions =
 				NonCompileActions
 				.Where(NonCompileAction => CompileActions.Any(CompileAction => CompileAction.PrerequisiteActions.Contains(NonCompileAction)))
 				.ToList();
 
 			// Precompile actions may have their own prerequisites which need to be executed along with them
-			List<Action> Prerequisites = new List<Action>();
-			foreach (Action PreCompileAction in PreCompileActions)
+			List<QueuedAction> Prerequisites = new List<QueuedAction>();
+			foreach (QueuedAction PreCompileAction in PreCompileActions)
 			{
 				Prerequisites.AddRange(PreCompileAction.PrerequisiteActions);
 			}
@@ -440,7 +440,7 @@ namespace UnrealBuildTool
 			///////////////////////////////////////////////////////////////
 			// Post Compile Stage
 
-			List<Action> PostCompileActions = NonCompileActions.Except(PreCompileActions).ToList();
+			List<QueuedAction> PostCompileActions = NonCompileActions.Except(PreCompileActions).ToList();
 
 			if (PostCompileActions.Any())
 			{
@@ -459,7 +459,7 @@ namespace UnrealBuildTool
 			bffOutputMemoryStream.Write(Info, 0, Info.Length);
 		}
 
-		private void AddPreBuildDependenciesText(IEnumerable<Action> PreBuildDependencies)
+		private void AddPreBuildDependenciesText(IEnumerable<QueuedAction> PreBuildDependencies)
 		{
 			if (!PreBuildDependencies.Any())
 				return;
@@ -717,7 +717,7 @@ namespace UnrealBuildTool
 			return ParsedCompilerOptions;
 		}
 
-		private string GetOptionValue(Dictionary<string, string> OptionsDictionary, string Key, Action Action, bool ProblemIfNotFound = false)
+		private string GetOptionValue(Dictionary<string, string> OptionsDictionary, string Key, QueuedAction Action, bool ProblemIfNotFound = false)
 		{
 			string Value = string.Empty;
 			if (OptionsDictionary.TryGetValue(Key, out Value))
@@ -994,7 +994,7 @@ namespace UnrealBuildTool
 			AddText("}\n\n"); //End Settings
 		}
 
-		private void AddCompileAction(Action Action, IEnumerable<Action> DependencyActions)
+		private void AddCompileAction(QueuedAction Action, IEnumerable<QueuedAction> DependencyActions)
 		{
 			string CompilerName = GetCompilerName();
 			if (Action.CommandPath.FullName.Contains("rc.exe"))
@@ -1109,7 +1109,7 @@ namespace UnrealBuildTool
 			AddText("}\n\n");
 		}
 
-		private void AddLinkAction(Action Action, IEnumerable<Action> DependencyActions)
+		private void AddLinkAction(QueuedAction Action, IEnumerable<QueuedAction> DependencyActions)
 		{
 			string[] SpecialLinkerOptions = IsApple() ? new string[] { "-o" } : new string[] { "/OUT:", "@", "-o" };
 
@@ -1149,13 +1149,13 @@ namespace UnrealBuildTool
 			string ResponseFilePath = GetOptionValue(ParsedLinkerOptions, "@", Action);
 			string OtherCompilerOptions = GetOptionValue(ParsedLinkerOptions, "OtherOptions", Action);
 
-			IEnumerable<Action> PrebuildDependencies = null;
+			IEnumerable<QueuedAction> PrebuildDependencies = null;
 
 			if (Action.CommandPath.FullName.Contains("lib.exe"))
 			{
 				if (DependencyActions.Any())
 				{
-					Func<Action, bool> DoesActionProducePCH = (Action ActionToCheck) =>
+					Func<QueuedAction, bool> DoesActionProducePCH = (QueuedAction ActionToCheck) =>
 					{
 						foreach (FileItem ProducedItem in ActionToCheck.ProducedItems)
 						{
@@ -1256,9 +1256,9 @@ namespace UnrealBuildTool
 				string InputFile = GetOptionValue(ParsedLinkerOptions, "InputFile", Action, ProblemIfNotFound: true);
 				if (!string.IsNullOrEmpty(InputFile))
 				{
-					Action InputFileAction = DependencyActions
+					QueuedAction InputFileAction = DependencyActions
 						.Where(ActionToInspect =>
-							ActionToInspect.ProducedItems.Exists(Item => Item.AbsolutePath == InputFile)
+							ActionToInspect.ProducedItems.Any(Item => Item.AbsolutePath == InputFile)
 						).FirstOrDefault();
 
 					if (InputFileAction != null)
@@ -1304,7 +1304,7 @@ namespace UnrealBuildTool
 			}
 		}
 
-		private void PrintActionDetails(Action ActionToPrint)
+		private void PrintActionDetails(QueuedAction ActionToPrint)
 		{
 			Log.TraceInformation(ActionToActionString(ActionToPrint));
 			Log.TraceInformation($"Action Type: {ActionToPrint.ActionType.ToString()}");
@@ -1314,7 +1314,7 @@ namespace UnrealBuildTool
 
 		private MemoryStream bffOutputMemoryStream = null;
 
-		private bool CreateBffFile(IEnumerable<Action> Actions, string BffFilePath)
+		private bool CreateBffFile(IEnumerable<QueuedAction> Actions, string BffFilePath)
 		{
 			try
 			{
@@ -1326,11 +1326,11 @@ namespace UnrealBuildTool
 
 				WriteEnvironmentSetup(); //Compiler, environment variables and base paths
 
-				foreach (Action Action in Actions)
+				foreach (QueuedAction Action in Actions)
 				{
 					// Resolve the list of prerequisite items for this action to
 					// a list of actions which produce these prerequisites
-					IEnumerable<Action> DependencyActions = Action.PrerequisiteActions.Distinct();
+					IEnumerable<QueuedAction> DependencyActions = Action.PrerequisiteActions.Distinct();
 
 					AddText($";** Function for Action {GetActionID(Action)} **\n");
 					AddText($";** CommandPath: {Action.CommandPath.FullName}\n");
diff --git a/Engine/Source/Programs/UnrealBuildTool/Executors/HybridExecutor.cs b/Engine/Source/Programs/UnrealBuildTool/Executors/HybridExecutor.cs
--- a/Engine/Source/Programs/UnrealBuildTool/Executors/HybridExecutor.cs
+++ b/Engine/Source/Programs/UnrealBuildTool/Executors/HybridExecutor.cs
@@ -69,25 +69,25 @@ namespace UnrealBuildTool
 		/// </summary>
 		/// <param name="ActionsToExecute">Actions to be executed</param>
 		/// <returns>True if the build succeeded, false otherwise</returns>
-		public override bool ExecuteActions(List<Action> ActionsToExecute)
+		public override bool ExecuteActions(List<QueuedAction> ActionsToExecute)
 		{
 			// Find the number of dependants for each action
-			Dictionary<Action, int> ActionToNumDependents = ActionsToExecute.ToDictionary(x => x, x => 0);
-			foreach(Action Action in ActionsToExecute)
+			Dictionary<QueuedAction, int> ActionToNumDependents = ActionsToExecute.ToDictionary(x => x, x => 0);
+			foreach(QueuedAction Action in ActionsToExecute)
 			{
-				foreach(Action PrerequisiteAction in Action.PrerequisiteActions)
+				foreach(QueuedAction PrerequisiteAction in Action.PrerequisiteActions)
 				{
 					ActionToNumDependents[PrerequisiteAction]++;
 				}
 			}
 
 			// Build up a set of leaf actions in several iterations, ensuring that the number of leaf actions in each 
-			HashSet<Action> LeafActions = new HashSet<Action>();
+			HashSet<QueuedAction> LeafActions = new HashSet<QueuedAction>();
 			for(;;)
 			{
 				// Find all the leaf actions in the graph
-				List<Action> NewLeafActions = new List<Action>();
-				foreach(Action Action in ActionsToExecute)
+				List<QueuedAction> NewLeafActions = new List<QueuedAction>();
+				foreach(QueuedAction Action in ActionsToExecute)
 				{
 					if(ActionToNumDependents[Action] == 0 && !LeafActions.Contains(Action))
 					{
@@ -105,9 +105,9 @@ namespace UnrealBuildTool
 				LeafActions.UnionWith(NewLeafActions);
 
 				// Decrement the dependent counts for any of their prerequisites, so we can try and remove those from the tree in another iteration
-				foreach(Action NewLeafAction in NewLeafActions)
+				foreach(QueuedAction NewLeafAction in NewLeafActions)
 				{
-					foreach(Action PrerequisiteAction in NewLeafAction.PrerequisiteActions)
+					foreach(QueuedAction PrerequisiteAction in NewLeafAction.PrerequisiteActions)
 					{
 						ActionToNumDependents[PrerequisiteAction]--;
 					}
@@ -115,9 +115,9 @@ namespace UnrealBuildTool
 			}
 
 			// Split the list of actions into those which should be executed locally and remotely
-			List<Action> LocalActionsToExecute = new List<Action>(LeafActions.Count);
-			List<Action> RemoteActionsToExecute = new List<Action>(ActionsToExecute.Count - LeafActions.Count);
-			foreach(Action ActionToExecute in ActionsToExecute)
+			List<QueuedAction> LocalActionsToExecute = new List<QueuedAction>(LeafActions.Count);
+			List<QueuedAction> RemoteActionsToExecute = new List<QueuedAction>(ActionsToExecute.Count - LeafActions.Count);
+			foreach(QueuedAction ActionToExecute in ActionsToExecute)
 			{
 				if(LeafActions.Contains(ActionToExecute))
 				{
diff --git a/Engine/Source/Programs/UnrealBuildTool/Executors/LocalExecutor.cs b/Engine/Source/Programs/UnrealBuildTool/Executors/LocalExecutor.cs
--- a/Engine/Source/Programs/UnrealBuildTool/Executors/LocalExecutor.cs
+++ b/Engine/Source/Programs/UnrealBuildTool/Executors/LocalExecutor.cs
@@ -28,7 +28,7 @@ namespace UnrealBuildTool
 		/// <summary>
 		/// Cache the action that this thread is managing
 		/// </summary>
-		Action Action;
+		QueuedAction Action;
 
 		/// <summary>
 		/// For reporting status to the user
@@ -39,7 +39,7 @@ namespace UnrealBuildTool
 		/// <summary>
 		/// Constructor, takes the action to process
 		/// </summary>
-		public ActionThread(Action InAction, int InJobNumber, int InTotalJobs)
+		public ActionThread(QueuedAction InAction, int InJobNumber, int InTotalJobs)
 		{
 			Action = InAction;
 			JobNumber = InJobNumber;
@@ -296,7 +296,7 @@ namespace UnrealBuildTool
 		/// Executes the specified actions locally.
 		/// </summary>
 		/// <returns>True if all the tasks successfully executed, or false if any of them failed.</returns>
-		public override bool ExecuteActions(List<Action> Actions)
+		public override bool ExecuteActions(List<QueuedAction> Actions)
 		{
 			// Time to sleep after each iteration of the loop in order to not busy wait.
 			const float LoopSleepTime = 0.1f;
@@ -304,7 +304,7 @@ namespace UnrealBuildTool
 			// The number of actions to execute in parallel is trying to keep the CPU busy enough in presence of I/O stalls.
 			Log.TraceInformation("Performing {0} actions ({1} in parallel)", Actions.Count, NumParallelProcesses);
 
-			Dictionary<Action, ActionThread> ActionThreadDictionary = new Dictionary<Action, ActionThread>();
+			Dictionary<QueuedAction, ActionThread> ActionThreadDictionary = new Dictionary<QueuedAction, ActionThread>();
 			int JobNumber = 1;
 			using (ProgressWriter ProgressWriter = new ProgressWriter("Compiling C++ source code...", false))
 			{
@@ -314,7 +314,7 @@ namespace UnrealBuildTool
 					// Count the number of pending and still executing actions.
 					int NumUnexecutedActions = 0;
 					int NumExecutingActions = 0;
-					foreach (Action Action in Actions)
+					foreach (QueuedAction Action in Actions)
 					{
 						ActionThread ActionThread = null;
 						bool bFoundActionProcess = ActionThreadDictionary.TryGetValue(Action, out ActionThread);
@@ -348,7 +348,7 @@ namespace UnrealBuildTool
 
 					// If there are fewer actions executing than the maximum, look for pending actions that don't have any outdated
 					// prerequisites.
-					foreach (Action Action in Actions)
+					foreach (QueuedAction Action in Actions)
 					{
 						ActionThread ActionProcess = null;
 						bool bFoundActionProcess = ActionThreadDictionary.TryGetValue(Action, out ActionProcess);
@@ -359,7 +359,7 @@ namespace UnrealBuildTool
 								// Determine whether there are any prerequisites of the action that are outdated.
 								bool bHasOutdatedPrerequisites = false;
 								bool bHasFailedPrerequisites = false;
-								foreach (Action PrerequisiteAction in Action.PrerequisiteActions)
+								foreach (QueuedAction PrerequisiteAction in Action.PrerequisiteActions)
 								{
 									if (Actions.Contains(PrerequisiteAction))
 									{
@@ -422,9 +422,9 @@ namespace UnrealBuildTool
 
 			// Check whether any of the tasks failed and log action stats if wanted.
 			bool bSuccess = true;
-			foreach (KeyValuePair<Action, ActionThread> ActionProcess in ActionThreadDictionary)
+			foreach (KeyValuePair<QueuedAction, ActionThread> ActionProcess in ActionThreadDictionary)
 			{
-				Action Action = ActionProcess.Key;
+				QueuedAction Action = ActionProcess.Key;
 				ActionThread ActionThread = ActionProcess.Value;
 
 				// Check for pending actions, preemptive failure
diff --git a/Engine/Source/Programs/UnrealBuildTool/Executors/ParallelExecutor.cs b/Engine/Source/Programs/UnrealBuildTool/Executors/ParallelExecutor.cs
--- a/Engine/Source/Programs/UnrealBuildTool/Executors/ParallelExecutor.cs
+++ b/Engine/Source/Programs/UnrealBuildTool/Executors/ParallelExecutor.cs
@@ -20,7 +20,7 @@ namespace UnrealBuildTool
 		class BuildAction
 		{
 			public int SortIndex;
-			public Action Inner;
+			public QueuedAction Inner;
 
 			public HashSet<BuildAction> Dependencies = new HashSet<BuildAction>();
 			public int MissingDependencyCount;
@@ -96,7 +96,7 @@ namespace UnrealBuildTool
 		/// Executes the specified actions locally.
 		/// </summary>
 		/// <returns>True if all the tasks successfully executed, or false if any of them failed.</returns>
-		public override bool ExecuteActions(List<Action> InputActions)
+		public override bool ExecuteActions(List<QueuedAction> InputActions)
 		{
 			Log.TraceInformation("Building {0} {1} with {2} {3}...", InputActions.Count, (InputActions.Count == 1) ? "action" : "actions", NumParallelProcesses, (NumParallelProcesses == 1)? "process" : "processes");
 
diff --git a/Engine/Source/Programs/UnrealBuildTool/Executors/SNDBS.cs b/Engine/Source/Programs/UnrealBuildTool/Executors/SNDBS.cs
--- a/Engine/Source/Programs/UnrealBuildTool/Executors/SNDBS.cs
+++ b/Engine/Source/Programs/UnrealBuildTool/Executors/SNDBS.cs
@@ -44,7 +44,7 @@ namespace UnrealBuildTool
 			return ServiceController.GetServices().Any(s => s.ServiceName.StartsWith("SNDBS") && s.Status == ServiceControllerStatus.Running);
 		}
 
-		public override bool ExecuteActions(List<Action> Actions)
+		public override bool ExecuteActions(List<QueuedAction> Actions)
 		{
 			if (Actions.Count == 0)
 				return true;
diff --git a/Engine/Source/Programs/UnrealBuildTool/Executors/XGE.cs b/Engine/Source/Programs/UnrealBuildTool/Executors/XGE.cs
--- a/Engine/Source/Programs/UnrealBuildTool/Executors/XGE.cs
+++ b/Engine/Source/Programs/UnrealBuildTool/Executors/XGE.cs
@@ -301,7 +301,7 @@ namespace UnrealBuildTool
 		// precompile the Regex needed to parse the XGE output (the ones we want are of the form "File (Duration at +time)"
 		//private static Regex XGEDurationRegex = new Regex(@"(?<Filename>.*) *\((?<Duration>[0-9:\.]+) at [0-9\+:\.]+\)", RegexOptions.ExplicitCapture);
 
-		public static void ExportActions(List<Action> ActionsToExecute)
+		public static void ExportActions(List<QueuedAction> ActionsToExecute)
 		{
 			for(int FileNum = 0;;FileNum++)
 			{
@@ -314,22 +314,22 @@ namespace UnrealBuildTool
 			}
 		}
 
-		public static void ExportActions(List<Action> ActionsToExecute, string OutFile)
+		public static void ExportActions(List<QueuedAction> ActionsToExecute, string OutFile)
 		{
 			WriteTaskFile(ActionsToExecute, OutFile, ProgressWriter.bWriteMarkup, bXGEExport: true);
 			Log.TraceInformation("XGEEXPORT: Exported '{0}'", OutFile);
 		}
 
-		public override bool ExecuteActions(List<Action> ActionsToExecute)
+		public override bool ExecuteActions(List<QueuedAction> ActionsToExecute)
 		{
 			bool XGEResult = true;
 
 			// Batch up XGE execution by actions with the same output event handler.
-			List<Action> ActionBatch = new List<Action>();
+			List<QueuedAction> ActionBatch = new List<QueuedAction>();
 			ActionBatch.Add(ActionsToExecute[0]);
 			for (int ActionIndex = 1; ActionIndex < ActionsToExecute.Count && XGEResult; ++ActionIndex)
 			{
-				Action CurrentAction = ActionsToExecute[ActionIndex];
+				QueuedAction CurrentAction = ActionsToExecute[ActionIndex];
 				ActionBatch.Add(CurrentAction);
 			}
 			if (ActionBatch.Count > 0 && XGEResult)
@@ -341,7 +341,7 @@ namespace UnrealBuildTool
 			return XGEResult;
 		}
 
-		bool ExecuteActionBatch(List<Action> Actions)
+		bool ExecuteActionBatch(List<QueuedAction> Actions)
 		{
 			bool XGEResult = true;
 			if (Actions.Count > 0)
@@ -358,11 +358,11 @@ namespace UnrealBuildTool
 		/// <summary>
 		/// Writes a XGE task file containing the specified actions to the specified file path.
 		/// </summary>
-		static void WriteTaskFile(List<Action> InActions, string TaskFilePath, bool bProgressMarkup, bool bXGEExport)
+		static void WriteTaskFile(List<QueuedAction> InActions, string TaskFilePath, bool bProgressMarkup, bool bXGEExport)
 		{
 			Dictionary<string, string> ExportEnv = new Dictionary<string, string>();
 
-			List<Action> Actions = InActions;
+			List<QueuedAction> Actions = InActions;
 			if (bXGEExport)
 			{
 				IDictionary CurrentEnvironment = Environment.GetEnvironmentVariables();
@@ -413,7 +413,7 @@ namespace UnrealBuildTool
 
 			for (int ActionIndex = 0; ActionIndex < Actions.Count; ActionIndex++)
 			{
-				Action Action = Actions[ActionIndex];
+				QueuedAction Action = Actions[ActionIndex];
 
 				// <Tool ... />
 				XmlElement ToolElement = XGETaskDocument.CreateElement("Tool");
@@ -457,7 +457,7 @@ namespace UnrealBuildTool
 					"OutputFileMasks",
 					string.Join(
 						",",
-						Action.ProducedItems.ConvertAll<string>(
+						Action.ProducedItems.Select(
 							delegate(FileItem ProducedItem) { return ProducedItem.Location.GetFileName(); }
 							).ToArray()
 						)
@@ -477,7 +477,7 @@ namespace UnrealBuildTool
 
 			for (int ActionIndex = 0; ActionIndex < Actions.Count; ActionIndex++)
 			{
-				Action Action = Actions[ActionIndex];
+				QueuedAction Action = Actions[ActionIndex];
 
 				// <Task ... />
 				XmlElement TaskElement = XGETaskDocument.CreateElement("Task");
@@ -500,7 +500,7 @@ namespace UnrealBuildTool
 
 				// Create a semi-colon separated list of the other tasks this task depends on the results of.
 				List<string> DependencyNames = new List<string>();
-				foreach(Action PrerequisiteAction in Action.PrerequisiteActions)
+				foreach(QueuedAction PrerequisiteAction in Action.PrerequisiteActions)
 				{
 					if (Actions.Contains(PrerequisiteAction))
 					{
diff --git a/Engine/Source/Programs/UnrealBuildTool/Modes/BuildMode.cs b/Engine/Source/Programs/UnrealBuildTool/Modes/BuildMode.cs
--- a/Engine/Source/Programs/UnrealBuildTool/Modes/BuildMode.cs
+++ b/Engine/Source/Programs/UnrealBuildTool/Modes/BuildMode.cs
@@ -315,10 +315,10 @@ namespace UnrealBuildTool
 				}
 
 				// Merge the action graphs together
-				List<Action> MergedActions;
+				List<QueuedAction> MergedActions;
 				if (TargetDescriptors.Count == 1)
 				{
-					MergedActions = new List<Action>(Makefiles[0].Actions);
+					MergedActions = Makefiles[0].Actions.ConvertAll(x => new QueuedAction(x));
 				}
 				else
 				{
@@ -336,7 +336,7 @@ namespace UnrealBuildTool
 				ActionGraph.Link(MergedActions);
 
 				// Get all the actions that are prerequisites for these targets. This forms the list of actions that we want executed.
-				List<Action> PrerequisiteActions = ActionGraph.GatherPrerequisiteActions(MergedActions, MergedOutputItems);
+				List<QueuedAction> PrerequisiteActions = ActionGraph.GatherPrerequisiteActions(MergedActions, MergedOutputItems);
 
 				// Create the action history
 				ActionHistory History = new ActionHistory();
@@ -353,13 +353,13 @@ namespace UnrealBuildTool
 				}
 
 				// Pre-process module interfaces to generate dependency files
-				List<Action> ModuleDependencyActions = PrerequisiteActions.Where(x => x.ActionType == ActionType.GatherModuleDependencies).ToList();
+				List<QueuedAction> ModuleDependencyActions = PrerequisiteActions.Where(x => x.ActionType == ActionType.GatherModuleDependencies).ToList();
 				if (ModuleDependencyActions.Count > 0)
 				{
-					Dictionary<Action, bool> PreprocessActionToOutdatedFlag = new Dictionary<Action, bool>();
+					Dictionary<QueuedAction, bool> PreprocessActionToOutdatedFlag = new Dictionary<QueuedAction, bool>();
 					ActionGraph.GatherAllOutdatedActions(ModuleDependencyActions, History, PreprocessActionToOutdatedFlag, null, BuildConfiguration.bIgnoreOutdatedImportLibraries);
 
-					List<Action> PreprocessActions = PreprocessActionToOutdatedFlag.Where(x => x.Value).Select(x => x.Key).ToList();
+					List<QueuedAction> PreprocessActions = PreprocessActionToOutdatedFlag.Where(x => x.Value).Select(x => x.Key).ToList();
 					if (PreprocessActions.Count > 0)
 					{
 						Log.TraceInformation("Updating module dependencies...");
@@ -373,7 +373,7 @@ namespace UnrealBuildTool
 				}
 
 				// Figure out which actions need to be built
-				Dictionary<Action, bool> ActionToOutdatedFlag = new Dictionary<Action, bool>();
+				Dictionary<QueuedAction, bool> ActionToOutdatedFlag = new Dictionary<QueuedAction, bool>();
 				for (int TargetIdx = 0; TargetIdx < TargetDescriptors.Count; TargetIdx++)
 				{
 					TargetDescriptor TargetDescriptor = TargetDescriptors[TargetIdx];
@@ -390,7 +390,7 @@ namespace UnrealBuildTool
 					if (ModuleDependencyActions.Count > 0)
 					{
 						List<FileItem> CompiledModuleInterfaces = new List<FileItem>();
-						foreach(Action PrerequisiteAction in PrerequisiteActions)
+						foreach(QueuedAction PrerequisiteAction in PrerequisiteActions)
 						{
 							if (PrerequisiteAction.ActionType == ActionType.CompileModuleInterface && PrerequisiteAction.CompiledModuleInterfaceFile != null)
 							{
@@ -404,7 +404,7 @@ namespace UnrealBuildTool
 						}
 
 						Dictionary<FileItem, string> ModuleInputs = ModuleOutputs.ToDictionary(x => x.Value, x => x.Key);
-						foreach (Action PrerequisiteAction in PrerequisiteActions)
+						foreach (QueuedAction PrerequisiteAction in PrerequisiteActions)
 						{
 							if (PrerequisiteAction.ActionType == ActionType.CompileModuleInterface)
 							{
@@ -416,8 +416,8 @@ namespace UnrealBuildTool
 										FileItem ModuleOutput;
 										if (ModuleOutputs.TryGetValue(ImportedModule, out ModuleOutput))
 										{
-											PrerequisiteAction.PrerequisiteItems.Add(ModuleOutput);
-											PrerequisiteAction.CommandArguments += String.Format(" /reference \"{0}={1}\"", ImportedModule, ModuleOutput.FullName);
+											PrerequisiteAction.Inner.PrerequisiteItems.Add(ModuleOutput);
+											PrerequisiteAction.Inner.CommandArguments += String.Format(" /reference \"{0}={1}\"", ImportedModule, ModuleOutput.FullName);
 										}
 										else
 										{
@@ -433,7 +433,7 @@ namespace UnrealBuildTool
 									string ModuleName;
 									if(ModuleInputs.TryGetValue(PrerequisiteItem, out ModuleName))
 									{
-										PrerequisiteAction.CommandArguments += String.Format(" /reference \"{0}={1}\"", ModuleName, PrerequisiteItem.AbsolutePath);
+										PrerequisiteAction.Inner.CommandArguments += String.Format(" /reference \"{0}={1}\"", ModuleName, PrerequisiteItem.AbsolutePath);
 									}
 								}
 							}
@@ -441,7 +441,7 @@ namespace UnrealBuildTool
 							{
 								if (PrerequisiteAction.ActionType != ActionType.GatherModuleDependencies)
 								{
-									PrerequisiteAction.PrerequisiteItems.AddRange(CompiledModuleInterfaces);
+									PrerequisiteAction.Inner.PrerequisiteItems.AddRange(CompiledModuleInterfaces);
 								}
 							}
 						}
@@ -456,7 +456,7 @@ namespace UnrealBuildTool
 					{
 						foreach (FileReference SpecificFile in TargetDescriptor.SpecificFilesToCompile)
 						{
-							foreach (Action PrerequisiteAction in PrerequisiteActions.Where(x => x.PrerequisiteItems.Any(y => y.Location == SpecificFile)))
+							foreach (QueuedAction PrerequisiteAction in PrerequisiteActions.Where(x => x.PrerequisiteItems.Any(y => y.Location == SpecificFile)))
 							{
 								ActionToOutdatedFlag[PrerequisiteAction] = true;
 							}
@@ -465,7 +465,7 @@ namespace UnrealBuildTool
 				}
 
 				// Link the action graph again to sort it
-				List<Action> MergedActionsToExecute = ActionToOutdatedFlag.Where(x => x.Value).Select(x => x.Key).ToList();
+				List<QueuedAction> MergedActionsToExecute = ActionToOutdatedFlag.Where(x => x.Value).Select(x => x.Key).ToList();
 				ActionGraph.Link(MergedActionsToExecute);
 
 				// Allow hot reload to override the actions
@@ -539,7 +539,7 @@ namespace UnrealBuildTool
 					// Write actions to an output file
 					using (Timeline.ScopeEvent("ActionGraph.WriteActions"))
 					{
-						ActionGraph.ExportJson(MergedActionsToExecute, WriteOutdatedActionsFile);
+						ActionGraph.ExportJson(MergedActionsToExecute.ConvertAll(x => x.Inner), WriteOutdatedActionsFile);
 					}
 				}
 				else
@@ -783,25 +783,25 @@ namespace UnrealBuildTool
 		/// <param name="TargetDescriptors">List of target descriptors</param>
 		/// <param name="Makefiles">The makefiles being built</param>
 		/// <returns>List of merged actions</returns>
-		static List<Action> MergeActionGraphs(List<TargetDescriptor> TargetDescriptors, TargetMakefile[] Makefiles)
+		static List<QueuedAction> MergeActionGraphs(List<TargetDescriptor> TargetDescriptors, TargetMakefile[] Makefiles)
 		{
 			// Set of all output items. Knowing that there are no conflicts in produced items, we use this to eliminate duplicate actions.
-			Dictionary<FileItem, Action> OutputItemToProducingAction = new Dictionary<FileItem, Action>();
+			Dictionary<FileItem, QueuedAction> OutputItemToProducingAction = new Dictionary<FileItem, QueuedAction>();
 			for(int TargetIdx = 0; TargetIdx < TargetDescriptors.Count; TargetIdx++)
 			{
 				string GroupPrefix = String.Format("{0}-{1}-{2}", TargetDescriptors[TargetIdx].Name, TargetDescriptors[TargetIdx].Platform, TargetDescriptors[TargetIdx].Configuration);
 				foreach(Action Action in Makefiles[TargetIdx].Actions)
 				{
-					Action ExistingAction;
+					QueuedAction ExistingAction;
 					if(!OutputItemToProducingAction.TryGetValue(Action.ProducedItems[0], out ExistingAction))
 					{
-						OutputItemToProducingAction[Action.ProducedItems[0]] = Action;
-						ExistingAction = Action;
+						ExistingAction = new QueuedAction(Action);
+						OutputItemToProducingAction[Action.ProducedItems[0]] = ExistingAction;
 					}
 					ExistingAction.GroupNames.Add(GroupPrefix);
 				}
 			}
-			return new List<Action>(OutputItemToProducingAction.Values);
+			return new List<QueuedAction>(OutputItemToProducingAction.Values);
 		}
 	}
 }
diff --git a/Engine/Source/Programs/UnrealBuildTool/Modes/ExecuteMode.cs b/Engine/Source/Programs/UnrealBuildTool/Modes/ExecuteMode.cs
--- a/Engine/Source/Programs/UnrealBuildTool/Modes/ExecuteMode.cs
+++ b/Engine/Source/Programs/UnrealBuildTool/Modes/ExecuteMode.cs
@@ -42,10 +42,10 @@ namespace UnrealBuildTool
 			Arguments.ApplyTo(BuildConfiguration);
 
 			// Read the actions file
-			List<Action> Actions;
+			List<QueuedAction> Actions;
 			using(Timeline.ScopeEvent("ActionGraph.ReadActions()"))
 			{
-				Actions = ActionGraph.ImportJson(ActionsFile);
+				Actions = ActionGraph.ImportJson(ActionsFile).ConvertAll(x => new QueuedAction(x));
 			}
 
 			// Link the action graph
diff --git a/Engine/Source/Programs/UnrealBuildTool/Modes/JsonExportMode.cs b/Engine/Source/Programs/UnrealBuildTool/Modes/JsonExportMode.cs
--- a/Engine/Source/Programs/UnrealBuildTool/Modes/JsonExportMode.cs
+++ b/Engine/Source/Programs/UnrealBuildTool/Modes/JsonExportMode.cs
@@ -56,11 +56,12 @@ namespace UnrealBuildTool
 						// Create the makefile
 						const bool bIsAssemblingBuild = true;
 						TargetMakefile Makefile = Target.Build(BuildConfiguration, WorkingSet, bIsAssemblingBuild, TargetDescriptor.SpecificFilesToCompile);
-						ActionGraph.Link(Makefile.Actions);
+						List<QueuedAction> Actions = Makefile.Actions.ConvertAll(x => new QueuedAction(x));
+						ActionGraph.Link(Actions);
 
 						// Filter all the actions to execute
 						HashSet<FileItem> PrerequisiteItems = new HashSet<FileItem>(Makefile.Actions.SelectMany(x => x.ProducedItems).Where(x => x.HasExtension(".h") || x.HasExtension(".cpp")));
-						List<Action> PrerequisiteActions = ActionGraph.GatherPrerequisiteActions(Makefile.Actions, PrerequisiteItems);
+						List<QueuedAction> PrerequisiteActions = ActionGraph.GatherPrerequisiteActions(Actions, PrerequisiteItems);
 
 						// Execute these actions
 						if (PrerequisiteActions.Count > 0)
diff --git a/Engine/Source/Programs/UnrealBuildTool/System/Action.cs b/Engine/Source/Programs/UnrealBuildTool/System/Action.cs
--- a/Engine/Source/Programs/UnrealBuildTool/System/Action.cs
+++ b/Engine/Source/Programs/UnrealBuildTool/System/Action.cs
@@ -37,10 +37,104 @@ namespace UnrealBuildTool
 		ParseTimingInfo,
 	}
 
+	interface IAction
+	{
+		/// <summary>
+		/// The type of this action (for debugging purposes).
+		/// </summary>
+		ActionType ActionType { get; }
+
+		/// <summary>
+		/// Every file this action depends on.  These files need to exist and be up to date in order for this action to even be considered
+		/// </summary>
+		IReadOnlyList<FileItem> PrerequisiteItems { get; }
+
+		/// <summary>
+		/// The files that this action produces after completing
+		/// </summary>
+		IReadOnlyList<FileItem> ProducedItems { get; }
+
+		/// <summary>
+		/// Items that should be deleted before running this action
+		/// </summary>
+		IReadOnlyList<FileItem> DeleteItems { get; }
+
+		/// <summary>
+		/// For C++ source files, specifies a dependency list file used to check changes to header files
+		/// </summary>
+		FileItem DependencyListFile { get; }
+
+		/// <summary>
+		/// The compiled C++ module interface (IFC) file produced by this action
+		/// </summary>
+		FileItem CompiledModuleInterfaceFile { get; }
+
+		/// <summary>
+		/// For C++ source files, specifies a timing file used to track timing information.
+		/// </summary>
+		FileItem TimingFile { get; }
+
+		/// <summary>
+		/// Directory from which to execute the program to create produced items
+		/// </summary>
+		DirectoryReference WorkingDirectory { get; }
+
+		/// <summary>
+		/// True if we should log extra information when we run a program to create produced items
+		/// </summary>
+		bool bPrintDebugInfo { get; }
+
+		/// <summary>
+		/// The command to run to create produced items
+		/// </summary>
+		FileReference CommandPath { get; }
+
+		/// <summary>
+		/// Command-line parameters to pass to the program
+		/// </summary>
+		string CommandArguments { get; }
+
+		/// <summary>
+		/// Optional friendly description of the type of command being performed, for example "Compile" or "Link".  Displayed by some executors.
+		/// </summary>
+		string CommandDescription { get; }
+
+		/// <summary>
+		/// Human-readable description of this action that may be displayed as status while invoking the action.  This is often the name of the file being compiled, or an executable file name being linked.  Displayed by some executors.
+		/// </summary>
+		string StatusDescription { get; }
+
+		/// <summary>
+		/// True if this action is allowed to be run on a remote machine when a distributed build system is being used, such as XGE
+		/// </summary>
+		bool bCanExecuteRemotely { get; }
+
+		/// <summary>
+		/// True if this action is allowed to be run on a remote machine with SNDBS. Files with #import directives must be compiled locally. Also requires bCanExecuteRemotely = true.
+		/// </summary>
+		bool bCanExecuteRemotelyWithSNDBS { get; }
+
+		/// <summary>
+		/// True if this action is using the GCC compiler.  Some build systems may be able to optimize for this case.
+		/// </summary>
+		bool bIsGCCCompiler { get; }
+
+		/// <summary>
+		/// Whether we should log this action, whether executed locally or remotely.  This is useful for actions that take time
+		/// but invoke tools without any console output.
+		/// </summary>
+		bool bShouldOutputStatusDescription { get; }
+
+		/// <summary>
+		/// True if any libraries produced by this action should be considered 'import libraries'
+		/// </summary>
+		bool bProducesImportLibrary { get; }
+	}
+
 	/// <summary>
 	/// A build action.
 	/// </summary>
-	class Action
+	class Action : IAction
 	{
 		///
 		/// Preparation and Assembly (serialized)
@@ -49,115 +143,97 @@ namespace UnrealBuildTool
 		/// <summary>
 		/// The type of this action (for debugging purposes).
 		/// </summary>
-		public ActionType ActionType;
+		public ActionType ActionType { get; set; }
 
 		/// <summary>
 		/// Every file this action depends on.  These files need to exist and be up to date in order for this action to even be considered
 		/// </summary>
-		public List<FileItem> PrerequisiteItems = new List<FileItem>();
+		public List<FileItem> PrerequisiteItems { get; set; } = new List<FileItem>();
 
 		/// <summary>
 		/// The files that this action produces after completing
 		/// </summary>
-		public List<FileItem> ProducedItems = new List<FileItem>();
+		public List<FileItem> ProducedItems { get; set; } = new List<FileItem>();
 
 		/// <summary>
 		/// Items that should be deleted before running this action
 		/// </summary>
-		public List<FileItem> DeleteItems = new List<FileItem>();
+		public List<FileItem> DeleteItems { get; set; } = new List<FileItem>();
 
 		/// <summary>
 		/// For C++ source files, specifies a dependency list file used to check changes to header files
 		/// </summary>
-		public FileItem DependencyListFile;
+		public FileItem DependencyListFile { get; set; }
 
 		/// <summary>
 		/// The compiled C++ module interface (IFC) file produced by this action
 		/// </summary>
-		public FileItem CompiledModuleInterfaceFile;
+		public FileItem CompiledModuleInterfaceFile { get; set; }
 
 		/// <summary>
 		/// For C++ source files, specifies a timing file used to track timing information.
 		/// </summary>
-		public FileItem TimingFile;
-
-		/// <summary>
-		/// Set of other actions that this action depends on. This set is built when the action graph is linked.
-		/// </summary>
-		public HashSet<Action> PrerequisiteActions;
+		public FileItem TimingFile { get; set; }
 
 		/// <summary>
 		/// Directory from which to execute the program to create produced items
 		/// </summary>
-		public DirectoryReference WorkingDirectory = null;
+		public DirectoryReference WorkingDirectory { get; set; } = null;
 
 		/// <summary>
 		/// True if we should log extra information when we run a program to create produced items
 		/// </summary>
-		public bool bPrintDebugInfo = false;
+		public bool bPrintDebugInfo { get; set; } = false;
 
 		/// <summary>
 		/// The command to run to create produced items
 		/// </summary>
-		public FileReference CommandPath = null;
+		public FileReference CommandPath { get; set; } = null;
 
 		/// <summary>
 		/// Command-line parameters to pass to the program
 		/// </summary>
-		public string CommandArguments = null;
+		public string CommandArguments { get; set; } = null;
 
 		/// <summary>
 		/// Optional friendly description of the type of command being performed, for example "Compile" or "Link".  Displayed by some executors.
 		/// </summary>
-		public string CommandDescription = null;
+		public string CommandDescription { get; set; } = null;
 
 		/// <summary>
 		/// Human-readable description of this action that may be displayed as status while invoking the action.  This is often the name of the file being compiled, or an executable file name being linked.  Displayed by some executors.
 		/// </summary>
-		public string StatusDescription = "...";
-
-		/// <summary>
-		/// If set, will be output whenever the group differs to the last executed action. Set when executing multiple targets at once.
-		/// </summary>
-		public List<string> GroupNames = new List<string>();
+		public string StatusDescription { get; set; } = "...";
 
 		/// <summary>
 		/// True if this action is allowed to be run on a remote machine when a distributed build system is being used, such as XGE
 		/// </summary>
-		public bool bCanExecuteRemotely = false;
+		public bool bCanExecuteRemotely { get; set; } = false;
 
 		/// <summary>
 		/// True if this action is allowed to be run on a remote machine with SNDBS. Files with #import directives must be compiled locally. Also requires bCanExecuteRemotely = true.
 		/// </summary>
-		public bool bCanExecuteRemotelyWithSNDBS = true;
+		public bool bCanExecuteRemotelyWithSNDBS { get; set; } = true;
 
 		/// <summary>
 		/// True if this action is using the GCC compiler.  Some build systems may be able to optimize for this case.
 		/// </summary>
-		public bool bIsGCCCompiler = false;
+		public bool bIsGCCCompiler { get; set; } = false;
 
 		/// <summary>
 		/// Whether we should log this action, whether executed locally or remotely.  This is useful for actions that take time
 		/// but invoke tools without any console output.
 		/// </summary>
-		public bool bShouldOutputStatusDescription = true;
+		public bool bShouldOutputStatusDescription { get; set; } = true;
 
 		/// <summary>
 		/// True if any libraries produced by this action should be considered 'import libraries'
 		/// </summary>
-		public bool bProducesImportLibrary = false;
-
-
-
-		///
-		/// Preparation only (not serialized)
-		///
-
-		/// <summary>
-		/// Total number of actions depending on this one.
-		/// </summary>
-		public int NumTotalDependentActions = 0;
+		public bool bProducesImportLibrary { get; set; } = false;
 
+		IReadOnlyList<FileItem> IAction.PrerequisiteItems => PrerequisiteItems;
+		IReadOnlyList<FileItem> IAction.ProducedItems => ProducedItems;
+		IReadOnlyList<FileItem> IAction.DeleteItems => DeleteItems;
 
 		public Action(ActionType InActionType)
 		{
@@ -362,16 +438,88 @@ namespace UnrealBuildTool
 			}
 		}
 
+		public override string ToString()
+		{
+			string ReturnString = "";
+			if (CommandPath != null)
+			{
+				ReturnString += CommandPath + " - ";
+			}
+			if (CommandArguments != null)
+			{
+				ReturnString += CommandArguments;
+			}
+			return ReturnString;
+		}
+	}
+
+	/// <summary>
+	/// Information about an action queued to be executed
+	/// </summary>
+	class QueuedAction : IAction
+	{
+		/// <summary>
+		/// The inner action instance
+		/// </summary>
+		public Action Inner;
+
+		/// <summary>
+		/// Set of other actions that this action depends on. This set is built when the action graph is linked.
+		/// </summary>
+		public HashSet<QueuedAction> PrerequisiteActions;
+
+		/// <summary>
+		/// Total number of actions depending on this one.
+		/// </summary>
+		public int NumTotalDependentActions = 0;
+
+		/// <summary>
+		/// If set, will be output whenever the group differs to the last executed action. Set when executing multiple targets at once.
+		/// </summary>
+		public List<string> GroupNames = new List<string>();
+
+		#region Wrapper implementation of IAction
+
+		public ActionType ActionType => Inner.ActionType;
+		public IReadOnlyList<FileItem> PrerequisiteItems => Inner.PrerequisiteItems;
+		public IReadOnlyList<FileItem> ProducedItems => Inner.ProducedItems;
+		public IReadOnlyList<FileItem> DeleteItems => Inner.DeleteItems;
+		public FileItem DependencyListFile => Inner.DependencyListFile;
+		public FileItem CompiledModuleInterfaceFile => Inner.CompiledModuleInterfaceFile;
+		public FileItem TimingFile => Inner.TimingFile;
+		public DirectoryReference WorkingDirectory => Inner.WorkingDirectory;
+		public bool bPrintDebugInfo => Inner.bPrintDebugInfo;
+		public FileReference CommandPath => Inner.CommandPath;
+		public string CommandArguments => Inner.CommandArguments;
+		public string CommandDescription => Inner.CommandDescription;
+		public string StatusDescription => Inner.StatusDescription;
+		public bool bCanExecuteRemotely => Inner.bCanExecuteRemotely;
+		public bool bCanExecuteRemotelyWithSNDBS => Inner.bCanExecuteRemotelyWithSNDBS;
+		public bool bIsGCCCompiler => Inner.bIsGCCCompiler;
+		public bool bShouldOutputStatusDescription => Inner.bShouldOutputStatusDescription;
+		public bool bProducesImportLibrary => Inner.bProducesImportLibrary;
+
+		#endregion
+
+		/// <summary>
+		/// Constructor
+		/// </summary>
+		/// <param name="Inner">The inner action instance</param>
+		public QueuedAction(Action Inner)
+		{
+			this.Inner = Inner;
+		}
+
 		/// <summary>
 		/// Increment the number of dependents, recursively
 		/// </summary>
 		/// <param name="VisitedActions">Set of visited actions</param>
-		public void IncrementDependentCount(HashSet<Action> VisitedActions)
+		public void IncrementDependentCount(HashSet<QueuedAction> VisitedActions)
 		{
-			if(VisitedActions.Add(this))
+			if (VisitedActions.Add(this))
 			{
 				NumTotalDependentActions++;
-				foreach(Action PrerequisiteAction in PrerequisiteActions)
+				foreach (QueuedAction PrerequisiteAction in PrerequisiteActions)
 				{
 					PrerequisiteAction.IncrementDependentCount(VisitedActions);
 				}
@@ -383,7 +531,7 @@ namespace UnrealBuildTool
 		/// </summary>
 		/// <param name="A">Action to compare</param>
 		/// <param name="B">Action to compare</param>
-		public static int Compare(Action A, Action B)
+		public static int Compare(QueuedAction A, QueuedAction B)
 		{
 			// Primary sort criteria is total number of dependent files, up to max depth.
 			if (B.NumTotalDependentActions != A.NumTotalDependentActions)
@@ -396,19 +544,5 @@ namespace UnrealBuildTool
 				return Math.Sign(B.PrerequisiteItems.Count - A.PrerequisiteItems.Count);
 			}
 		}
-
-		public override string ToString()
-		{
-			string ReturnString = "";
-			if (CommandPath != null)
-			{
-				ReturnString += CommandPath + " - ";
-			}
-			if (CommandArguments != null)
-			{
-				ReturnString += CommandArguments;
-			}
-			return ReturnString;
-		}
 	}
 }
diff --git a/Engine/Source/Programs/UnrealBuildTool/System/ActionExecutor.cs b/Engine/Source/Programs/UnrealBuildTool/System/ActionExecutor.cs
--- a/Engine/Source/Programs/UnrealBuildTool/System/ActionExecutor.cs
+++ b/Engine/Source/Programs/UnrealBuildTool/System/ActionExecutor.cs
@@ -15,7 +15,7 @@ namespace UnrealBuildTool
 			get;
 		}
 
-		public abstract bool ExecuteActions(List<Action> ActionsToExecute);
+		public abstract bool ExecuteActions(List<QueuedAction> ActionsToExecute);
 	}
 
 }
diff --git a/Engine/Source/Programs/UnrealBuildTool/System/ActionGraph.cs b/Engine/Source/Programs/UnrealBuildTool/System/ActionGraph.cs
--- a/Engine/Source/Programs/UnrealBuildTool/System/ActionGraph.cs
+++ b/Engine/Source/Programs/UnrealBuildTool/System/ActionGraph.cs
@@ -18,11 +18,11 @@ namespace UnrealBuildTool
 		/// Links the actions together and sets up their dependencies
 		/// </summary>
 		/// <param name="Actions">List of actions in the graph</param>
-		public static void Link(List<Action> Actions)
+		public static void Link(List<QueuedAction> Actions)
 		{
 			// Build a map from item to its producing action
-			Dictionary<FileItem, Action> ItemToProducingAction = new Dictionary<FileItem, Action>();
-			foreach (Action Action in Actions)
+			Dictionary<FileItem, QueuedAction> ItemToProducingAction = new Dictionary<FileItem, QueuedAction>();
+			foreach (QueuedAction Action in Actions)
 			{
 				foreach (FileItem ProducedItem in Action.ProducedItems)
 				{
@@ -34,12 +34,12 @@ namespace UnrealBuildTool
 			DetectActionGraphCycles(Actions, ItemToProducingAction);
 
 			// Use this map to add all the prerequisite actions
-			foreach (Action Action in Actions)
+			foreach (QueuedAction Action in Actions)
 			{
-				Action.PrerequisiteActions = new HashSet<Action>();
+				Action.PrerequisiteActions = new HashSet<QueuedAction>();
 				foreach(FileItem PrerequisiteItem in Action.PrerequisiteItems)
 				{
-					Action PrerequisiteAction;
+					QueuedAction PrerequisiteAction;
 					if(ItemToProducingAction.TryGetValue(PrerequisiteItem, out PrerequisiteAction))
 					{
 						Action.PrerequisiteActions.Add(PrerequisiteAction);
@@ -146,12 +146,12 @@ namespace UnrealBuildTool
 		/// <summary>
 		/// Builds a list of actions that need to be executed to produce the specified output items.
 		/// </summary>
-		public static List<Action> GetActionsToExecute(List<Action> Actions, CppDependencyCache CppDependencies, ActionHistory History, bool bIgnoreOutdatedImportLibraries)
+		public static List<QueuedAction> GetActionsToExecute(List<QueuedAction> Actions, CppDependencyCache CppDependencies, ActionHistory History, bool bIgnoreOutdatedImportLibraries)
 		{
 			using (Timeline.ScopeEvent("ActionGraph.GetActionsToExecute()"))
 			{
 				// For all targets, build a set of all actions that are outdated.
-				Dictionary<Action, bool> OutdatedActionDictionary = new Dictionary<Action, bool>();
+				Dictionary<QueuedAction, bool> OutdatedActionDictionary = new Dictionary<QueuedAction, bool>();
 				GatherAllOutdatedActions(Actions, History, OutdatedActionDictionary, CppDependencies, bIgnoreOutdatedImportLibraries);
 
 				// Build a list of actions that are both needed for this target and outdated.
@@ -223,7 +223,7 @@ namespace UnrealBuildTool
 		/// <summary>
 		/// Executes a list of actions.
 		/// </summary>
-		public static void ExecuteActions(BuildConfiguration BuildConfiguration, List<Action> ActionsToExecute)
+		public static void ExecuteActions(BuildConfiguration BuildConfiguration, List<QueuedAction> ActionsToExecute)
 		{
 			if(ActionsToExecute.Count == 0)
 			{
@@ -267,7 +267,7 @@ namespace UnrealBuildTool
 				Log.TraceInformation("Total time in {0} executor: {1:0.00} seconds", Executor.Name, Timer.Elapsed.TotalSeconds);
 
 				// Reset the file info for all the produced items
-				foreach (Action BuildAction in ActionsToExecute)
+				foreach (QueuedAction BuildAction in ActionsToExecute)
 				{
 					foreach(FileItem ProducedItem in BuildAction.ProducedItems)
 					{
@@ -276,7 +276,7 @@ namespace UnrealBuildTool
 				}
 
 				// Verify the link outputs were created (seems to happen with Win64 compiles)
-				foreach (Action BuildAction in ActionsToExecute)
+				foreach (QueuedAction BuildAction in ActionsToExecute)
 				{
 					if (BuildAction.ActionType == ActionType.Link)
 					{
@@ -295,57 +295,38 @@ namespace UnrealBuildTool
 		/// <summary>
 		/// Sorts the action list for improved parallelism with local execution.
 		/// </summary>
-		static void SortActionList(List<Action> Actions)
+		static void SortActionList(List<QueuedAction> Actions)
 		{
 			// Clear the current dependent count
-			foreach(Action Action in Actions)
+			foreach(QueuedAction Action in Actions)
 			{
 				Action.NumTotalDependentActions = 0;
 			}
 
 			// Increment all the dependencies
-			foreach(Action Action in Actions)
+			foreach(QueuedAction Action in Actions)
 			{
-				Action.IncrementDependentCount(new HashSet<Action>());
+				Action.IncrementDependentCount(new HashSet<QueuedAction>());
 			}
 
 			// Sort actions by number of actions depending on them, descending. Secondary sort criteria is file size.
-			Actions.Sort(CompareActions);
-		}
-
-		/// <summary>
-		/// Compares two actions based on total number of dependent items, descending.
-		/// </summary>
-		/// <param name="A">Action to compare</param>
-		/// <param name="B">Action to compare</param>
-		public static int CompareActions(Action A, Action B)
-		{
-			// Primary sort criteria is total number of dependent files, up to max depth.
-			if (B.NumTotalDependentActions != A.NumTotalDependentActions)
-			{
-				return Math.Sign(B.NumTotalDependentActions - A.NumTotalDependentActions);
-			}
-			// Secondary sort criteria is number of pre-requisites.
-			else
-			{
-				return Math.Sign(B.PrerequisiteItems.Count - A.PrerequisiteItems.Count);
-			}
+			Actions.Sort(QueuedAction.Compare);
 		}
 
 		/// <summary>
 		/// Checks for cycles in the action graph.
 		/// </summary>
-		static void DetectActionGraphCycles(List<Action> Actions, Dictionary<FileItem, Action> ItemToProducingAction)
+		static void DetectActionGraphCycles(List<QueuedAction> Actions, Dictionary<FileItem, QueuedAction> ItemToProducingAction)
 		{
 			// Starting with actions that only depend on non-produced items, iteratively expand a set of actions that are only dependent on
 			// non-cyclical actions.
-			Dictionary<Action, bool> ActionIsNonCyclical = new Dictionary<Action, bool>();
-			Dictionary<Action, List<Action>> CyclicActions = new Dictionary<Action, List<Action>>();
+			Dictionary<QueuedAction, bool> ActionIsNonCyclical = new Dictionary<QueuedAction, bool>();
+			Dictionary<QueuedAction, List<QueuedAction>> CyclicActions = new Dictionary<QueuedAction, List<QueuedAction>>();
 			while (true)
 			{
 				bool bFoundNewNonCyclicalAction = false;
 
-				foreach (Action Action in Actions)
+				foreach (QueuedAction Action in Actions)
 				{
 					if (!ActionIsNonCyclical.ContainsKey(Action))
 					{
@@ -353,7 +334,7 @@ namespace UnrealBuildTool
 						bool bActionOnlyDependsOnNonCyclicalActions = true;
 						foreach (FileItem PrerequisiteItem in Action.PrerequisiteItems)
 						{
-							Action ProducingAction;
+							QueuedAction ProducingAction;
 							if (ItemToProducingAction.TryGetValue(PrerequisiteItem, out ProducingAction))
 							{
 								if (!ActionIsNonCyclical.ContainsKey(ProducingAction))
@@ -361,10 +342,10 @@ namespace UnrealBuildTool
 									bActionOnlyDependsOnNonCyclicalActions = false;
 									if (!CyclicActions.ContainsKey(Action))
 									{
-										CyclicActions.Add(Action, new List<Action>());
+										CyclicActions.Add(Action, new List<QueuedAction>());
 									}
 
-									List<Action> CyclicPrereq = CyclicActions[Action];
+									List<QueuedAction> CyclicPrereq = CyclicActions[Action];
 									if (!CyclicPrereq.Contains(ProducingAction))
 									{
 										CyclicPrereq.Add(ProducingAction);
@@ -398,7 +379,7 @@ namespace UnrealBuildTool
 			if (ActionIsNonCyclical.Count < Actions.Count)
 			{
 				// Find the index of each action
-				Dictionary<Action, int> ActionToIndex = new Dictionary<Action, int>();
+				Dictionary<QueuedAction, int> ActionToIndex = new Dictionary<QueuedAction, int>();
 				for(int Idx = 0; Idx < Actions.Count; Idx++)
 				{
 					ActionToIndex[Actions[Idx]] = Idx;
@@ -406,7 +387,7 @@ namespace UnrealBuildTool
 
 				// Describe the cyclical actions.
 				string CycleDescription = "";
-				foreach (Action Action in Actions)
+				foreach (QueuedAction Action in Actions)
 				{
 					if (!ActionIsNonCyclical.ContainsKey(Action))
 					{
@@ -423,7 +404,7 @@ namespace UnrealBuildTool
 						CycleDescription += string.Format("\tDepends on cyclic actions:\n");
 						if (CyclicActions.ContainsKey(Action))
 						{
-							foreach (Action CyclicPrerequisiteAction in CyclicActions[Action])
+							foreach (QueuedAction CyclicPrerequisiteAction in CyclicActions[Action])
 							{
 								if (CyclicActions.ContainsKey(CyclicPrerequisiteAction))
 								{
@@ -461,10 +442,10 @@ namespace UnrealBuildTool
 		/// <param name="Actions">All the actions in the graph</param>
 		/// <param name="OutputItems">Set of output items to be built</param>
 		/// <returns>Set of prerequisite actions</returns>
-		public static List<Action> GatherPrerequisiteActions(List<Action> Actions, HashSet<FileItem> OutputItems)
+		public static List<QueuedAction> GatherPrerequisiteActions(List<QueuedAction> Actions, HashSet<FileItem> OutputItems)
 		{
-			HashSet<Action> PrerequisiteActions = new HashSet<Action>();
-			foreach(Action Action in Actions)
+			HashSet<QueuedAction> PrerequisiteActions = new HashSet<QueuedAction>();
+			foreach(QueuedAction Action in Actions)
 			{
 				if(Action.ProducedItems.Any(x => OutputItems.Contains(x)))
 				{
@@ -479,11 +460,11 @@ namespace UnrealBuildTool
 		/// </summary>
 		/// <param name="Action">The root action to scan</param>
 		/// <param name="PrerequisiteActions">Set of prerequisite actions</param>
-		private static void GatherPrerequisiteActions(Action Action, HashSet<Action> PrerequisiteActions)
+		private static void GatherPrerequisiteActions(QueuedAction Action, HashSet<QueuedAction> PrerequisiteActions)
 		{
 			if(PrerequisiteActions.Add(Action))
 			{
-				foreach(Action PrerequisiteAction in Action.PrerequisiteActions)
+				foreach(QueuedAction PrerequisiteAction in Action.PrerequisiteActions)
 				{
 					GatherPrerequisiteActions(PrerequisiteAction, PrerequisiteActions);
 				}
@@ -500,7 +481,7 @@ namespace UnrealBuildTool
 		/// <param name="CppDependencies"></param>
 		/// <param name="bIgnoreOutdatedImportLibraries"></param>
 		/// <returns>true if outdated</returns>
-		public static bool IsActionOutdated(Action RootAction, Dictionary<Action, bool> OutdatedActionDictionary, ActionHistory ActionHistory, CppDependencyCache CppDependencies, bool bIgnoreOutdatedImportLibraries)
+		public static bool IsActionOutdated(QueuedAction RootAction, Dictionary<QueuedAction, bool> OutdatedActionDictionary, ActionHistory ActionHistory, CppDependencyCache CppDependencies, bool bIgnoreOutdatedImportLibraries)
 		{
 			// Only compute the outdated-ness for actions that don't aren't cached in the outdated action dictionary.
 			bool bIsOutdated = false;
@@ -561,7 +542,7 @@ namespace UnrealBuildTool
 			// Check if any of the prerequisite actions are out of date
 			if (!bIsOutdated)
 			{
-				foreach (Action PrerequisiteAction in RootAction.PrerequisiteActions)
+				foreach (QueuedAction PrerequisiteAction in RootAction.PrerequisiteActions)
 				{
 					if (IsActionOutdated(PrerequisiteAction, OutdatedActionDictionary, ActionHistory, CppDependencies, bIgnoreOutdatedImportLibraries))
 					{
@@ -650,7 +631,7 @@ namespace UnrealBuildTool
 		/// <param name="RootAction">The action to check</param>
 		/// <param name="PrerequisiteAction">The action that it depends on</param>
 		/// <returns>True if the only dependency between two actions is for an import library</returns>
-		static bool IsImportLibraryDependency(Action RootAction, Action PrerequisiteAction)
+		static bool IsImportLibraryDependency(QueuedAction RootAction, QueuedAction PrerequisiteAction)
 		{
 			if(PrerequisiteAction.bProducesImportLibrary)
 			{
@@ -668,11 +649,11 @@ namespace UnrealBuildTool
 		/// <param name="RootAction">The action to check</param>
 		/// <param name="PrerequisiteItem">The dependency that is out of date</param>
 		/// <returns>True if the only dependency between two actions is for an import library</returns>
-		static bool IsImportLibraryDependency(Action RootAction, FileItem PrerequisiteItem)
+		static bool IsImportLibraryDependency(QueuedAction RootAction, FileItem PrerequisiteItem)
 		{
 			if(PrerequisiteItem.Location.HasExtension(".lib"))
 			{
-				foreach(Action PrerequisiteAction in RootAction.PrerequisiteActions)
+				foreach(QueuedAction PrerequisiteAction in RootAction.PrerequisiteActions)
 				{
 					if(PrerequisiteAction.bProducesImportLibrary && PrerequisiteAction.ProducedItems.Contains(PrerequisiteItem))
 					{
@@ -687,12 +668,12 @@ namespace UnrealBuildTool
 		/// Builds a dictionary containing the actions from AllActions that are outdated by calling
 		/// IsActionOutdated.
 		/// </summary>
-		public static void GatherAllOutdatedActions(IEnumerable<Action> Actions, ActionHistory ActionHistory, Dictionary<Action, bool> OutdatedActions, CppDependencyCache CppDependencies, bool bIgnoreOutdatedImportLibraries)
+		public static void GatherAllOutdatedActions(IEnumerable<QueuedAction> Actions, ActionHistory ActionHistory, Dictionary<QueuedAction, bool> OutdatedActions, CppDependencyCache CppDependencies, bool bIgnoreOutdatedImportLibraries)
 		{
 			using(Timeline.ScopeEvent("Prefetching include dependencies"))
 			{
 				List<FileItem> Dependencies = new List<FileItem>();
-				foreach(Action Action in Actions)
+				foreach(QueuedAction Action in Actions)
 				{
 					if(Action.DependencyListFile != null)
 					{
@@ -711,9 +692,9 @@ namespace UnrealBuildTool
 		/// Deletes all the items produced by actions in the provided outdated action dictionary.
 		/// </summary>
 		/// <param name="OutdatedActions">List of outdated actions</param>
-		public static void DeleteOutdatedProducedItems(List<Action> OutdatedActions)
+		public static void DeleteOutdatedProducedItems(List<QueuedAction> OutdatedActions)
 		{
-			foreach(Action OutdatedAction in OutdatedActions)
+			foreach(QueuedAction OutdatedAction in OutdatedActions)
 			{
 				foreach (FileItem DeleteItem in OutdatedAction.DeleteItems)
 				{
@@ -730,10 +711,10 @@ namespace UnrealBuildTool
 		/// Creates directories for all the items produced by actions in the provided outdated action
 		/// dictionary.
 		/// </summary>
-		public static void CreateDirectoriesForProducedItems(List<Action> OutdatedActions)
+		public static void CreateDirectoriesForProducedItems(List<QueuedAction> OutdatedActions)
 		{
 			HashSet<DirectoryReference> OutputDirectories = new HashSet<DirectoryReference>();
-			foreach(Action OutdatedAction in OutdatedActions)
+			foreach(QueuedAction OutdatedAction in OutdatedActions)
 			{
 				foreach(FileItem ProducedItem in OutdatedAction.ProducedItems)
 				{
diff --git a/Engine/Source/Programs/UnrealBuildTool/System/HotReload.cs b/Engine/Source/Programs/UnrealBuildTool/System/HotReload.cs
--- a/Engine/Source/Programs/UnrealBuildTool/System/HotReload.cs
+++ b/Engine/Source/Programs/UnrealBuildTool/System/HotReload.cs
@@ -469,7 +469,7 @@ namespace UnrealBuildTool
 		/// <param name="PrerequisiteActions">The actions to execute</param>
 		/// <param name="TargetActionsToExecute">Actions to execute for this target</param>
 		/// <returns>Set of actions to execute</returns>
-		public static List<Action> PatchActionsForTarget(BuildConfiguration BuildConfiguration, TargetDescriptor TargetDescriptor, TargetMakefile Makefile, List<Action> PrerequisiteActions, List<Action> TargetActionsToExecute)
+		public static List<QueuedAction> PatchActionsForTarget(BuildConfiguration BuildConfiguration, TargetDescriptor TargetDescriptor, TargetMakefile Makefile, List<QueuedAction> PrerequisiteActions, List<QueuedAction> TargetActionsToExecute)
 		{
 			// Get the dependency history
 			CppDependencyCache CppDependencies = CppDependencyCache.CreateHierarchy(TargetDescriptor.ProjectFile, TargetDescriptor.Name, TargetDescriptor.Platform, TargetDescriptor.Configuration, Makefile.TargetType, TargetDescriptor.Architecture);
@@ -501,7 +501,7 @@ namespace UnrealBuildTool
 
 					// Find all the binaries that we're actually going to build
 					HashSet<FileReference> OutputFiles = new HashSet<FileReference>();
-					foreach (Action Action in TargetActionsToExecute)
+					foreach (QueuedAction Action in TargetActionsToExecute)
 					{
 						if (Action.ActionType == ActionType.Link)
 						{
@@ -523,12 +523,12 @@ namespace UnrealBuildTool
 				}
 
 				// Filter the prerequisite actions down to just the compile actions, then recompute all the actions to execute
-				PrerequisiteActions = new List<Action>(TargetActionsToExecute.Where(x => x.ActionType == ActionType.Compile));
+				PrerequisiteActions = new List<QueuedAction>(TargetActionsToExecute.Where(x => x.ActionType == ActionType.Compile));
 				TargetActionsToExecute = ActionGraph.GetActionsToExecute(PrerequisiteActions, CppDependencies, History, BuildConfiguration.bIgnoreOutdatedImportLibraries);
 
 				// Update the action graph with these new paths
 				Dictionary<FileReference, FileReference> OriginalFileToPatchedFile = new Dictionary<FileReference, FileReference>();
-				HotReload.PatchActionGraphForLiveCoding(PrerequisiteActions, OriginalFileToPatchedFile);
+				HotReload.PatchActionGraphForLiveCoding(PrerequisiteActions.Select(x => x.Inner), OriginalFileToPatchedFile);
 
 				// Get a new list of actions to execute now that the graph has been modified
 				TargetActionsToExecute = ActionGraph.GetActionsToExecute(PrerequisiteActions, CppDependencies, History, BuildConfiguration.bIgnoreOutdatedImportLibraries);
@@ -583,7 +583,7 @@ namespace UnrealBuildTool
 				for (int LastNumFilesWithNewSuffix = 0; FilesRequiringSuffix.Count > LastNumFilesWithNewSuffix;)
 				{
 					LastNumFilesWithNewSuffix = FilesRequiringSuffix.Count;
-					foreach (Action PrerequisiteAction in PrerequisiteActions)
+					foreach (QueuedAction PrerequisiteAction in PrerequisiteActions)
 					{
 						if (!TargetActionsToExecute.Contains(PrerequisiteAction))
 						{
@@ -610,7 +610,7 @@ namespace UnrealBuildTool
 				}
 
 				// Update the action graph with these new paths
-				HotReload.PatchActionGraph(PrerequisiteActions, OldLocationToNewLocation);
+				HotReload.PatchActionGraph(PrerequisiteActions.Select(x => x.Inner), OldLocationToNewLocation);
 
 				// Get a new list of actions to execute now that the graph has been modified
 				TargetActionsToExecute = ActionGraph.GetActionsToExecute(PrerequisiteActions, CppDependencies, History, BuildConfiguration.bIgnoreOutdatedImportLibraries);
@@ -632,7 +632,7 @@ namespace UnrealBuildTool
 				}
 
 				// Now filter out all the hot reload files and update the state
-				foreach (Action Action in TargetActionsToExecute)
+				foreach (QueuedAction Action in TargetActionsToExecute)
 				{
 					foreach (FileItem ProducedItem in Action.ProducedItems)
 					{
