UBT: Duplicate actions for HotReload rather than mutating existing actions. We can probably make this much cleaner in future, just laying the groundwork.

diff --git a/Engine/Source/Programs/UnrealBuildTool/Modes/BuildMode.cs b/Engine/Source/Programs/UnrealBuildTool/Modes/BuildMode.cs
--- a/Engine/Source/Programs/UnrealBuildTool/Modes/BuildMode.cs
+++ b/Engine/Source/Programs/UnrealBuildTool/Modes/BuildMode.cs
@@ -308,21 +308,28 @@ namespace UnrealBuildTool
 					ActionGraph.CheckPathLengths(BuildConfiguration, Makefiles.SelectMany(x => x.Actions));
 				}
 
+				// Create a QueuedAction instance for each action in the makefiles
+				List<QueuedAction>[] QueuedActions = new List<QueuedAction>[Makefiles.Length];
+				for(int Idx = 0; Idx < Makefiles.Length; Idx++)
+				{
+					QueuedActions[Idx] = Makefiles[Idx].Actions.ConvertAll(x => new QueuedAction(x));
+				}
+
 				// Clean up any previous hot reload runs, and reapply the current state if it's already active
 				for (int TargetIdx = 0; TargetIdx < TargetDescriptors.Count; TargetIdx++)
 				{
-					HotReload.Setup(TargetDescriptors[TargetIdx], Makefiles[TargetIdx], BuildConfiguration);
+					HotReload.Setup(TargetDescriptors[TargetIdx], Makefiles[TargetIdx], QueuedActions[TargetIdx], BuildConfiguration);
 				}
 
 				// Merge the action graphs together
 				List<QueuedAction> MergedActions;
 				if (TargetDescriptors.Count == 1)
 				{
-					MergedActions = Makefiles[0].Actions.ConvertAll(x => new QueuedAction(x));
+					MergedActions = QueuedActions[0];
 				}
 				else
 				{
-					MergedActions = MergeActionGraphs(TargetDescriptors, Makefiles);
+					MergedActions = MergeActionGraphs(TargetDescriptors, QueuedActions);
 				}
 
 				// Gather all the prerequisite actions that are part of the targets
@@ -416,8 +423,10 @@ namespace UnrealBuildTool
 										FileItem ModuleOutput;
 										if (ModuleOutputs.TryGetValue(ImportedModule, out ModuleOutput))
 										{
-											PrerequisiteAction.Inner.PrerequisiteItems.Add(ModuleOutput);
-											PrerequisiteAction.Inner.CommandArguments += String.Format(" /reference \"{0}={1}\"", ImportedModule, ModuleOutput.FullName);
+											Action NewAction = new Action(PrerequisiteAction.Inner);
+											NewAction.PrerequisiteItems.Add(ModuleOutput);
+											NewAction.CommandArguments += String.Format(" /reference \"{0}={1}\"", ImportedModule, ModuleOutput.FullName);
+											PrerequisiteAction.Inner = NewAction;
 										}
 										else
 										{
@@ -433,7 +442,9 @@ namespace UnrealBuildTool
 									string ModuleName;
 									if(ModuleInputs.TryGetValue(PrerequisiteItem, out ModuleName))
 									{
-										PrerequisiteAction.Inner.CommandArguments += String.Format(" /reference \"{0}={1}\"", ModuleName, PrerequisiteItem.AbsolutePath);
+										Action NewAction = new Action(PrerequisiteAction.Inner);
+										NewAction.CommandArguments += String.Format(" /reference \"{0}={1}\"", ModuleName, PrerequisiteItem.AbsolutePath);
+										PrerequisiteAction.Inner = NewAction;
 									}
 								}
 							}
@@ -441,7 +452,9 @@ namespace UnrealBuildTool
 							{
 								if (PrerequisiteAction.ActionType != ActionType.GatherModuleDependencies)
 								{
-									PrerequisiteAction.Inner.PrerequisiteItems.AddRange(CompiledModuleInterfaces);
+									Action NewAction = new Action(PrerequisiteAction.Inner);
+									NewAction.PrerequisiteItems.AddRange(CompiledModuleInterfaces);
+									PrerequisiteAction.Inner = NewAction;
 								}
 							}
 						}
@@ -539,7 +552,7 @@ namespace UnrealBuildTool
 					// Write actions to an output file
 					using (Timeline.ScopeEvent("ActionGraph.WriteActions"))
 					{
-						ActionGraph.ExportJson(MergedActionsToExecute.ConvertAll(x => x.Inner), WriteOutdatedActionsFile);
+						ActionGraph.ExportJson(MergedActionsToExecute, WriteOutdatedActionsFile);
 					}
 				}
 				else
@@ -781,22 +794,22 @@ namespace UnrealBuildTool
 		/// Merge action graphs for multiple targets into a single set of actions. Sets group names on merged actions to indicate which target they belong to.
 		/// </summary>
 		/// <param name="TargetDescriptors">List of target descriptors</param>
-		/// <param name="Makefiles">The makefiles being built</param>
+		/// <param name="TargetActions">Array of actions for each target</param>
 		/// <returns>List of merged actions</returns>
-		static List<QueuedAction> MergeActionGraphs(List<TargetDescriptor> TargetDescriptors, TargetMakefile[] Makefiles)
+		static List<QueuedAction> MergeActionGraphs(List<TargetDescriptor> TargetDescriptors, List<QueuedAction>[] TargetActions)
 		{
 			// Set of all output items. Knowing that there are no conflicts in produced items, we use this to eliminate duplicate actions.
 			Dictionary<FileItem, QueuedAction> OutputItemToProducingAction = new Dictionary<FileItem, QueuedAction>();
 			for(int TargetIdx = 0; TargetIdx < TargetDescriptors.Count; TargetIdx++)
 			{
 				string GroupPrefix = String.Format("{0}-{1}-{2}", TargetDescriptors[TargetIdx].Name, TargetDescriptors[TargetIdx].Platform, TargetDescriptors[TargetIdx].Configuration);
-				foreach(Action Action in Makefiles[TargetIdx].Actions)
+				foreach(QueuedAction TargetAction in TargetActions[TargetIdx])
 				{
 					QueuedAction ExistingAction;
-					if(!OutputItemToProducingAction.TryGetValue(Action.ProducedItems[0], out ExistingAction))
+					if(!OutputItemToProducingAction.TryGetValue(TargetAction.ProducedItems[0], out ExistingAction))
 					{
-						ExistingAction = new QueuedAction(Action);
-						OutputItemToProducingAction[Action.ProducedItems[0]] = ExistingAction;
+						ExistingAction = new QueuedAction(TargetAction);
+						OutputItemToProducingAction[TargetAction.ProducedItems[0]] = ExistingAction;
 					}
 					ExistingAction.GroupNames.Add(GroupPrefix);
 				}
diff --git a/Engine/Source/Programs/UnrealBuildTool/System/Action.cs b/Engine/Source/Programs/UnrealBuildTool/System/Action.cs
--- a/Engine/Source/Programs/UnrealBuildTool/System/Action.cs
+++ b/Engine/Source/Programs/UnrealBuildTool/System/Action.cs
@@ -129,6 +129,12 @@ namespace UnrealBuildTool
 		/// True if any libraries produced by this action should be considered 'import libraries'
 		/// </summary>
 		bool bProducesImportLibrary { get; }
+
+		/// <summary>
+		/// Writes an action to a json file
+		/// </summary>
+		/// <param name="Writer">Writer to receive the output</param>
+		void ExportJson(JsonWriter Writer);
 	}
 
 	/// <summary>
@@ -246,6 +252,28 @@ namespace UnrealBuildTool
 			}
 		}
 
+		public Action(IAction InOther)
+		{
+			ActionType = InOther.ActionType;
+			PrerequisiteItems = new List<FileItem>(InOther.PrerequisiteItems);
+			ProducedItems = new List<FileItem>(InOther.ProducedItems);
+			DeleteItems = new List<FileItem>(InOther.DeleteItems);
+			DependencyListFile = InOther.DependencyListFile;
+			CompiledModuleInterfaceFile = InOther.CompiledModuleInterfaceFile;
+			TimingFile = InOther.TimingFile;
+			WorkingDirectory = InOther.WorkingDirectory;
+			bPrintDebugInfo = InOther.bPrintDebugInfo;
+			CommandPath = InOther.CommandPath;
+			CommandArguments = InOther.CommandArguments;
+			CommandDescription = InOther.CommandDescription;
+			StatusDescription = InOther.StatusDescription;
+			bCanExecuteRemotely = InOther.bCanExecuteRemotely;
+			bCanExecuteRemotelyWithSNDBS = InOther.bCanExecuteRemotelyWithSNDBS;
+			bIsGCCCompiler = InOther.bIsGCCCompiler;
+			bShouldOutputStatusDescription = InOther.bShouldOutputStatusDescription;
+			bProducesImportLibrary = InOther.bProducesImportLibrary;
+		}
+
 		public Action(BinaryArchiveReader Reader)
 		{
 			ActionType = (ActionType)Reader.ReadByte();
@@ -461,7 +489,7 @@ namespace UnrealBuildTool
 		/// <summary>
 		/// The inner action instance
 		/// </summary>
-		public Action Inner;
+		public IAction Inner;
 
 		/// <summary>
 		/// Set of other actions that this action depends on. This set is built when the action graph is linked.
@@ -499,13 +527,18 @@ namespace UnrealBuildTool
 		public bool bShouldOutputStatusDescription => Inner.bShouldOutputStatusDescription;
 		public bool bProducesImportLibrary => Inner.bProducesImportLibrary;
 
+		public void ExportJson(JsonWriter Writer)
+		{
+			Inner.ExportJson(Writer);
+		}
+
 		#endregion
 
 		/// <summary>
 		/// Constructor
 		/// </summary>
 		/// <param name="Inner">The inner action instance</param>
-		public QueuedAction(Action Inner)
+		public QueuedAction(IAction Inner)
 		{
 			this.Inner = Inner;
 		}
diff --git a/Engine/Source/Programs/UnrealBuildTool/System/ActionGraph.cs b/Engine/Source/Programs/UnrealBuildTool/System/ActionGraph.cs
--- a/Engine/Source/Programs/UnrealBuildTool/System/ActionGraph.cs
+++ b/Engine/Source/Programs/UnrealBuildTool/System/ActionGraph.cs
@@ -758,7 +758,7 @@ namespace UnrealBuildTool
 		/// </summary>
 		/// <param name="Actions">The actions to write</param>
 		/// <param name="OutputFile">Output file to write the actions to</param>
-		public static void ExportJson(List<Action> Actions, FileReference OutputFile)
+		public static void ExportJson(IEnumerable<IAction> Actions, FileReference OutputFile)
 		{
 			DirectoryReference.CreateDirectory(OutputFile.Directory);
 			using (JsonWriter Writer = new JsonWriter(OutputFile))
@@ -776,7 +776,7 @@ namespace UnrealBuildTool
 				Writer.WriteObjectEnd();
 
 				Writer.WriteArrayStart("Actions");
-				foreach (Action Action in Actions)
+				foreach (IAction Action in Actions)
 				{
 					Writer.WriteObjectStart();
 					Action.ExportJson(Writer);
diff --git a/Engine/Source/Programs/UnrealBuildTool/System/HotReload.cs b/Engine/Source/Programs/UnrealBuildTool/System/HotReload.cs
--- a/Engine/Source/Programs/UnrealBuildTool/System/HotReload.cs
+++ b/Engine/Source/Programs/UnrealBuildTool/System/HotReload.cs
@@ -158,8 +158,9 @@ namespace UnrealBuildTool
 		/// </summary>
 		/// <param name="TargetDescriptor">The target being built</param>
 		/// <param name="Makefile">Makefile for the targe</param>
+		/// <param name="Actions">Actions for this target</param>
 		/// <param name="BuildConfiguration">Global build configuration</param>
-		public static void Setup(TargetDescriptor TargetDescriptor, TargetMakefile Makefile, BuildConfiguration BuildConfiguration)
+		public static void Setup(TargetDescriptor TargetDescriptor, TargetMakefile Makefile, List<QueuedAction> Actions, BuildConfiguration BuildConfiguration)
 		{
 			// Get the hot-reload mode
 			if (TargetDescriptor.HotReloadMode == HotReloadMode.LiveCoding)
@@ -203,11 +204,11 @@ namespace UnrealBuildTool
 					HotReloadState HotReloadState = HotReloadState.Load(StateFile);
 
 					// Apply the old state to the makefile
-					HotReload.ApplyState(HotReloadState, Makefile);
+					HotReload.ApplyState(HotReloadState, Makefile, Actions);
 				}
 
 				// If we want a specific suffix on any modules, apply that now. We'll track the outputs later, but the suffix has to be forced (and is always out of date if it doesn't exist).
-				HotReload.PatchActionGraphWithNames(TargetDescriptor.HotReloadModuleNameToSuffix, Makefile);
+				HotReload.PatchActionGraphWithNames(TargetDescriptor.HotReloadModuleNameToSuffix, Makefile, Actions);
 			}
 		}
 
@@ -440,10 +441,11 @@ namespace UnrealBuildTool
 		/// </summary>
 		/// <param name="HotReloadState">The hot-reload state</param>
 		/// <param name="Makefile">Makefile to apply the state</param>
-		static void ApplyState(HotReloadState HotReloadState, TargetMakefile Makefile)
+		/// <param name="Actions">Actions for this makefile</param>
+		static void ApplyState(HotReloadState HotReloadState, TargetMakefile Makefile, List<QueuedAction> Actions)
 		{
 			// Update the action graph to produce these new files
-			HotReload.PatchActionGraph(Makefile.Actions, HotReloadState.OriginalFileToHotReloadFile);
+			HotReload.PatchActionGraph(Actions, HotReloadState.OriginalFileToHotReloadFile);
 
 			// Update the module to output file mapping
 			foreach(string HotReloadModuleName in Makefile.HotReloadModuleNames)
@@ -528,7 +530,7 @@ namespace UnrealBuildTool
 
 				// Update the action graph with these new paths
 				Dictionary<FileReference, FileReference> OriginalFileToPatchedFile = new Dictionary<FileReference, FileReference>();
-				HotReload.PatchActionGraphForLiveCoding(PrerequisiteActions.Select(x => x.Inner), OriginalFileToPatchedFile);
+				HotReload.PatchActionGraphForLiveCoding(PrerequisiteActions, OriginalFileToPatchedFile);
 
 				// Get a new list of actions to execute now that the graph has been modified
 				TargetActionsToExecute = ActionGraph.GetActionsToExecute(PrerequisiteActions, CppDependencies, History, BuildConfiguration.bIgnoreOutdatedImportLibraries);
@@ -610,7 +612,7 @@ namespace UnrealBuildTool
 				}
 
 				// Update the action graph with these new paths
-				HotReload.PatchActionGraph(PrerequisiteActions.Select(x => x.Inner), OldLocationToNewLocation);
+				HotReload.PatchActionGraph(PrerequisiteActions, OldLocationToNewLocation);
 
 				// Get a new list of actions to execute now that the graph has been modified
 				TargetActionsToExecute = ActionGraph.GetActionsToExecute(PrerequisiteActions, CppDependencies, History, BuildConfiguration.bIgnoreOutdatedImportLibraries);
@@ -745,9 +747,9 @@ namespace UnrealBuildTool
 		/// </summary>
 		/// <param name="Actions">Set of actions</param>
 		/// <param name="OriginalFileToPatchedFile">Dictionary that receives a map of original object file to patched object file</param>
-		public static void PatchActionGraphForLiveCoding(IEnumerable<Action> Actions, Dictionary<FileReference, FileReference> OriginalFileToPatchedFile)
+		public static void PatchActionGraphForLiveCoding(IEnumerable<QueuedAction> Actions, Dictionary<FileReference, FileReference> OriginalFileToPatchedFile)
 		{
-			foreach (Action Action in Actions)
+			foreach (QueuedAction Action in Actions)
 			{
 				if(Action.ActionType == ActionType.Compile)
 				{
@@ -781,13 +783,16 @@ namespace UnrealBuildTool
 						}
 					}
 
+					Action NewAction = new Action(Action.Inner);
+					Action.Inner = NewAction;
+
 					FileReference OldDependenciesFile = new FileReference(Arguments[DependenciesIdx].Substring(DependenciesPrefix.Length));
 					FileItem OldDependenciesFileItem = Action.ProducedItems.First(x => x.Location == OldDependenciesFile);
-					Action.ProducedItems.Remove(OldDependenciesFileItem);
+					NewAction.ProducedItems.Remove(OldDependenciesFileItem);
 
 					FileReference NewDependenciesFile = OldDependenciesFile.ChangeExtension(".lc.response");
 					FileItem NewDependenciesFileItem = FileItem.GetItemByFileReference(NewDependenciesFile);
-					Action.ProducedItems.Add(NewDependenciesFileItem);
+					NewAction.ProducedItems.Add(NewDependenciesFileItem);
 
 					Arguments[DependenciesIdx] = DependenciesPrefix + NewDependenciesFile.FullName;
 
@@ -818,11 +823,11 @@ namespace UnrealBuildTool
 						{
 							FileReference OldOutputFile = new FileReference(ResponseLine.Substring(3).Trim('\"'));
 							FileItem OldOutputFileItem = Action.ProducedItems.First(x => x.Location == OldOutputFile);
-							Action.ProducedItems.Remove(OldOutputFileItem);
+							NewAction.ProducedItems.Remove(OldOutputFileItem);
 
 							FileReference NewOutputFile = OldOutputFile.ChangeExtension(".lc.obj");
 							FileItem NewOutputFileItem = FileItem.GetItemByFileReference(NewOutputFile);
-							Action.ProducedItems.Add(NewOutputFileItem);
+							NewAction.ProducedItems.Add(NewOutputFileItem);
 
 							OriginalFileToPatchedFile[OldOutputFile] = NewOutputFile;
 
@@ -835,7 +840,7 @@ namespace UnrealBuildTool
 					Arguments[ResponseFileIdx] = "@" + NewResponseFile.FullName;
 
 					// Update the final arguments
-					Action.CommandArguments = Utils.FormatCommandLine(Arguments);
+					NewAction.CommandArguments = Utils.FormatCommandLine(Arguments);
 				}
 			}
 		}
@@ -843,7 +848,7 @@ namespace UnrealBuildTool
 		/// <summary>
 		/// Patch the action graph for hot reloading, mapping files according to the given dictionary.
 		/// </summary>
-		public static void PatchActionGraph(IEnumerable<Action> Actions, Dictionary<FileReference, FileReference> OriginalFileToHotReloadFile)
+		public static void PatchActionGraph(IEnumerable<QueuedAction> Actions, Dictionary<FileReference, FileReference> OriginalFileToHotReloadFile)
 		{
 			// Gather all of the response files for link actions.  We're going to need to patch 'em up after we figure out new
 			// names for all of the output files and import libraries
@@ -859,7 +864,7 @@ namespace UnrealBuildTool
 			// Finally, we'll keep track of any file items that we had to create counterparts for change file names, so we can fix those up too
 			Dictionary<FileItem, FileItem> AffectedOriginalFileItemAndNewFileItemMap = new Dictionary<FileItem, FileItem>();
 
-			foreach (Action Action in Actions.Where((Action) => Action.ActionType == ActionType.Link))
+			foreach (QueuedAction Action in Actions.Where((Action) => Action.ActionType == ActionType.Link))
 			{
 				// Assume that the first produced item (with no extension) is our output file name
 				FileReference HotReloadFile;
@@ -897,6 +902,10 @@ namespace UnrealBuildTool
 					ResponseFilePaths.Add(NewResponseFilePath);
 				}
 
+				// Duplicate the action
+				Action NewAction = new Action(Action);
+				Action.Inner = NewAction;
+
 				// Find the *.link.sh file in the command line.  We'll need to make a copy of it with our new file name.
 				// Only currently used on Linux
 				if (UEBuildPlatform.IsPlatformInGroup(BuildHostPlatform.Current.Platform, UnrealPlatformGroup.Unix))
@@ -934,7 +943,7 @@ namespace UnrealBuildTool
 						{
 							// OK, the prerequisite item's file name changed so we'll update it to point to our new file
 							FileItem NewPrerequisiteItem = FileItem.GetItemByPath(NewPrerequisiteItemFilePath);
-							Action.PrerequisiteItems[ItemIndex] = NewPrerequisiteItem;
+							NewAction.PrerequisiteItems[ItemIndex] = NewPrerequisiteItem;
 
 							// Keep track of it so we can fix up dependencies in a second pass afterwards
 							AffectedOriginalFileItemAndNewFileItemMap.Add(OriginalPrerequisiteItem, NewPrerequisiteItem);
@@ -967,7 +976,7 @@ namespace UnrealBuildTool
 					{
 						// OK, the produced item's file name changed so we'll update it to point to our new file
 						FileItem NewProducedItem = FileItem.GetItemByPath(NewProducedItemFilePath);
-						Action.ProducedItems[ItemIndex] = NewProducedItem;
+						NewAction.ProducedItems[ItemIndex] = NewProducedItem;
 
 						// Keep track of it so we can fix up dependencies in a second pass afterwards
 						AffectedOriginalFileItemAndNewFileItemMap.Add(OriginalProducedItem, NewProducedItem);
@@ -980,12 +989,12 @@ namespace UnrealBuildTool
 					FileItem NewItem;
 					if(AffectedOriginalFileItemAndNewFileItemMap.TryGetValue(Action.DeleteItems[Idx], out NewItem))
 					{
-						Action.DeleteItems[Idx] = NewItem;
+						NewAction.DeleteItems[Idx] = NewItem;
 					}
 				}
 
 				// The status description of the item has the file name, so we'll update it too
-				Action.StatusDescription = ReplaceBaseFileName(Action.StatusDescription, OriginalFileNameWithoutExtension, NewFileNameWithoutExtension);
+				NewAction.StatusDescription = ReplaceBaseFileName(Action.StatusDescription, OriginalFileNameWithoutExtension, NewFileNameWithoutExtension);
 
 				// Keep track of the file names, so we can fix up response files afterwards.
 				if(!OriginalFileNameAndNewFileNameList_NoExtensions.ContainsKey(OriginalFileNameWithoutExtension))
@@ -1000,7 +1009,7 @@ namespace UnrealBuildTool
 
 
 			// Do another pass and update any actions that depended on the original file names that we changed
-			foreach (Action Action in Actions)
+			foreach (QueuedAction Action in Actions)
 			{
 				for (int ItemIndex = 0; ItemIndex < Action.PrerequisiteItems.Count; ++ItemIndex)
 				{
@@ -1010,7 +1019,9 @@ namespace UnrealBuildTool
 					if (AffectedOriginalFileItemAndNewFileItemMap.TryGetValue(OriginalFileItem, out NewFileItem))
 					{
 						// OK, looks like we need to replace this file item because we've renamed the file
-						Action.PrerequisiteItems[ItemIndex] = NewFileItem;
+						Action NewAction = new Action(Action.Inner);
+						NewAction.PrerequisiteItems[ItemIndex] = NewFileItem;
+						Action.Inner = NewAction;
 					}
 				}
 			}
@@ -1019,14 +1030,16 @@ namespace UnrealBuildTool
 			if (OriginalFileNameAndNewFileNameList_NoExtensions.Count > 0)
 			{
 				// Update all the paths in link actions
-				foreach (Action Action in Actions.Where((Action) => Action.ActionType == ActionType.Link))
+				foreach (QueuedAction Action in Actions.Where((Action) => Action.ActionType == ActionType.Link))
 				{
 					foreach (KeyValuePair<string, string> FileNameTuple in OriginalFileNameAndNewFileNameList_NoExtensions)
 					{
 						string OriginalFileNameWithoutExtension = FileNameTuple.Key;
 						string NewFileNameWithoutExtension = FileNameTuple.Value;
 
-						Action.CommandArguments = ReplaceBaseFileName(Action.CommandArguments, OriginalFileNameWithoutExtension, NewFileNameWithoutExtension);
+						Action NewAction = new Action(Action.Inner);
+						NewAction.CommandArguments = ReplaceBaseFileName(Action.CommandArguments, OriginalFileNameWithoutExtension, NewFileNameWithoutExtension);
+						Action.Inner = NewAction;
 					}
 				}
 
@@ -1071,7 +1084,7 @@ namespace UnrealBuildTool
 			}
 
 			// Update the action that writes out the module manifests
-			foreach(Action Action in Actions)
+			foreach(QueuedAction Action in Actions)
 			{
 				if(Action.ActionType == ActionType.WriteMetadata)
 				{
@@ -1130,10 +1143,14 @@ namespace UnrealBuildTool
 						FileReference HotReloadTargetInfoFile = FileReference.Combine(TargetInfoFile.Directory, "Metadata-HotReload.dat");
 						BinaryFormatterUtils.SaveIfDifferent(HotReloadTargetInfoFile, TargetInfo);
 
-						Action.PrerequisiteItems.RemoveAll(x => x.Location == TargetInfoFile);
-						Action.PrerequisiteItems.Add(FileItem.GetItemByFileReference(HotReloadTargetInfoFile));
+						Action NewAction = new Action(Action.Inner);
+
+						NewAction.PrerequisiteItems.RemoveAll(x => x.Location == TargetInfoFile);
+						NewAction.PrerequisiteItems.Add(FileItem.GetItemByFileReference(HotReloadTargetInfoFile));
+
+						NewAction.CommandArguments = Arguments.Substring(0, FileNameIdx) + HotReloadTargetInfoFile + Arguments.Substring(FileNameEndIdx);
 
-						Action.CommandArguments = Arguments.Substring(0, FileNameIdx) + HotReloadTargetInfoFile + Arguments.Substring(FileNameEndIdx);
+						Action.Inner = NewAction;
 					}
 				}
 			}
@@ -1144,7 +1161,8 @@ namespace UnrealBuildTool
 		/// </summary>
 		/// <param name="ModuleNameToSuffix">Map of module name to suffix</param>
 		/// <param name="Makefile">Makefile for the target being built</param>
-		public static void PatchActionGraphWithNames(Dictionary<string, int> ModuleNameToSuffix, TargetMakefile Makefile)
+		/// <param name="Actions">Actions to be executed for this makefile</param>
+		public static void PatchActionGraphWithNames(Dictionary<string, int> ModuleNameToSuffix, TargetMakefile Makefile, List<QueuedAction> Actions)
 		{
 			if(ModuleNameToSuffix.Count > 0)
 			{
@@ -1163,7 +1181,7 @@ namespace UnrealBuildTool
 						}
 					}
 				}
-				HotReload.PatchActionGraph(Makefile.Actions, OldLocationToNewLocation);
+				HotReload.PatchActionGraph(Actions, OldLocationToNewLocation);
 			}
 		}
 
