UBT: Allow actions to implement custom serialization methods, and store abstract interfaces for actions in the makefile.

#rnx

diff --git a/Engine/Source/Programs/UnrealBuildTool/Platform/Windows/VCToolChain.cs b/Engine/Source/Programs/UnrealBuildTool/Platform/Windows/VCToolChain.cs
--- a/Engine/Source/Programs/UnrealBuildTool/Platform/Windows/VCToolChain.cs
+++ b/Engine/Source/Programs/UnrealBuildTool/Platform/Windows/VCToolChain.cs
@@ -1379,7 +1379,7 @@ namespace UnrealBuildTool
 		{
 			if (Target.bPrintToolChainTimingInfo || Target.WindowsPlatform.bCompilerTrace)
 			{
-				List<Action> ParseTimingActions = Makefile.Actions.Where(x => x.ActionType == ActionType.ParseTimingInfo).ToList();
+				List<IAction> ParseTimingActions = Makefile.Actions.Where(x => x.ActionType == ActionType.ParseTimingInfo).ToList();
 				List<FileItem> TimingJsonFiles = ParseTimingActions.SelectMany(a => a.ProducedItems.Where(i => i.HasExtension(".cta"))).ToList();
 				Makefile.OutputItems.AddRange(TimingJsonFiles);
 
diff --git a/Engine/Source/Programs/UnrealBuildTool/System/Action.cs b/Engine/Source/Programs/UnrealBuildTool/System/Action.cs
--- a/Engine/Source/Programs/UnrealBuildTool/System/Action.cs
+++ b/Engine/Source/Programs/UnrealBuildTool/System/Action.cs
@@ -481,6 +481,24 @@ namespace UnrealBuildTool
 		}
 	}
 
+	/// <summary>
+	/// Default serializer for <see cref="Action"/> instances
+	/// </summary>
+	class DefaultActionSerializer : ActionSerializerBase<Action>
+	{
+		/// <inheritdoc/>
+		public override Action Read(BinaryArchiveReader Reader)
+		{
+			return new Action(Reader);
+		}
+
+		/// <inheritdoc/>
+		public override void Write(BinaryArchiveWriter Writer, Action Action)
+		{
+			Action.Write(Writer);
+		}
+	}
+
 	/// <summary>
 	/// Information about an action queued to be executed
 	/// </summary>
diff --git a/Engine/Source/Programs/UnrealBuildTool/System/ActionGraph.cs b/Engine/Source/Programs/UnrealBuildTool/System/ActionGraph.cs
--- a/Engine/Source/Programs/UnrealBuildTool/System/ActionGraph.cs
+++ b/Engine/Source/Programs/UnrealBuildTool/System/ActionGraph.cs
@@ -55,16 +55,16 @@ namespace UnrealBuildTool
 		/// Checks a set of actions for conflicts (ie. different actions producing the same output items)
 		/// </summary>
 		/// <param name="Actions">The set of actions to check</param>
-		public static void CheckForConflicts(IEnumerable<Action> Actions)
+		public static void CheckForConflicts(IEnumerable<IAction> Actions)
 		{
 			bool bResult = true;
 
-			Dictionary<FileItem, Action> ItemToProducingAction = new Dictionary<FileItem, Action>();
-			foreach(Action Action in Actions)
+			Dictionary<FileItem, IAction> ItemToProducingAction = new Dictionary<FileItem, IAction>();
+			foreach(IAction Action in Actions)
 			{
 				foreach(FileItem ProducedItem in Action.ProducedItems)
 				{
-					Action ExistingAction;
+					IAction ExistingAction;
 					if(ItemToProducingAction.TryGetValue(ProducedItem, out ExistingAction))
 					{
 						bResult &= CheckForConflicts(ExistingAction, Action);
@@ -88,7 +88,7 @@ namespace UnrealBuildTool
 		/// <param name="A">The first action</param>
 		/// <param name="B">The second action</param>
 		/// <returns>True if any conflicts were found, false otherwise.</returns>
-		public static bool CheckForConflicts(Action A, Action B)
+		public static bool CheckForConflicts(IAction A, IAction B)
 		{
 			bool bResult = true;
 			if (A.ActionType != B.ActionType)
@@ -136,7 +136,7 @@ namespace UnrealBuildTool
 		/// <param name="Description">Description of the difference</param>
 		/// <param name="OldValue">Previous value for the field</param>
 		/// <param name="NewValue">Conflicting value for the field</param>
-		static void LogConflict(Action Action, string Description, string OldValue, string NewValue)
+		static void LogConflict(IAction Action, string Description, string OldValue, string NewValue)
 		{
 			Log.TraceError("Unable to merge actions producing {0}: {1}", Action.ProducedItems[0].Location.GetFileName(), Description);
 			Log.TraceLog("  Previous: {0}", OldValue);
@@ -164,7 +164,7 @@ namespace UnrealBuildTool
 		/// </summary>
 		/// <param name="BuildConfiguration">The build configuration</param>
 		/// <param name="Actions">List of actions in the graph</param>
-		public static void CheckPathLengths(BuildConfiguration BuildConfiguration, IEnumerable<Action> Actions)
+		public static void CheckPathLengths(BuildConfiguration BuildConfiguration, IEnumerable<IAction> Actions)
 		{
 			if (BuildHostPlatform.Current.Platform == UnrealTargetPlatform.Win64)
 			{
diff --git a/Engine/Source/Programs/UnrealBuildTool/System/ActionSerializer.cs b/Engine/Source/Programs/UnrealBuildTool/System/ActionSerializer.cs
new file mode 100644
--- /dev/null
+++ b/Engine/Source/Programs/UnrealBuildTool/System/ActionSerializer.cs
@@ -0,0 +1,150 @@
+using System;
+using System.Collections.Generic;
+using System.Linq;
+using System.Reflection;
+using System.Text;
+using System.Threading.Tasks;
+using Tools.DotNETCommon;
+
+namespace UnrealBuildTool
+{
+	/// <summary>
+	/// Interface for a class that can serialize an action type
+	/// </summary>
+	interface IActionSerializer
+	{
+		/// <summary>
+		/// The action type
+		/// </summary>
+		Type Type { get; }
+
+		/// <summary>
+		/// Read the action from an archive
+		/// </summary>
+		/// <param name="Reader">Reader for the action</param>
+		/// <returns>New action</returns>
+		IAction Read(BinaryArchiveReader Reader);
+
+		/// <summary>
+		/// Writes an action to an archive
+		/// </summary>
+		/// <param name="Writer">Writer for the archive</param>
+		/// <param name="Action">The action to write</param>
+		void Write(BinaryArchiveWriter Writer, IAction Action);
+	}
+
+	/// <summary>
+	/// Generic base class for an action serializer
+	/// </summary>
+	/// <typeparam name="TAction"></typeparam>
+	abstract class ActionSerializerBase<TAction> : IActionSerializer where TAction : IAction
+	{
+		/// <inheritdoc/>
+		public Type Type => typeof(TAction);
+
+		/// <inheritdoc/>
+		IAction IActionSerializer.Read(BinaryArchiveReader Reader) => Read(Reader);
+
+		/// <inheritdoc/>
+		void IActionSerializer.Write(BinaryArchiveWriter Writer, IAction Action) => Write(Writer, (TAction)Action);
+
+		/// <summary>
+		/// Read the action from an archive
+		/// </summary>
+		/// <param name="Reader">Reader for the action</param>
+		/// <returns>New action</returns>
+		public abstract TAction Read(BinaryArchiveReader Reader);
+
+		/// <summary>
+		/// Writes an action to an archive
+		/// </summary>
+		/// <param name="Writer">Writer for the archive</param>
+		/// <param name="Action">The action to write</param>
+		public abstract void Write(BinaryArchiveWriter Writer, TAction Action);
+	}
+
+	/// <summary>
+	/// Helper methods for registering serializers and serializing actions
+	/// </summary>
+	static class ActionSerialization
+	{
+		/// <summary>
+		/// Map from type name to deserializing constructor
+		/// </summary>
+		static Dictionary<Type, IActionSerializer> TypeToSerializer;
+
+		/// <summary>
+		/// Map from serializer name to instance
+		/// </summary>
+		static Dictionary<string, IActionSerializer> NameToSerializer;
+
+		/// <summary>
+		/// Creates a map of type name to constructor
+		/// </summary>
+		/// <returns></returns>
+		static ActionSerialization()
+		{
+			TypeToSerializer = new Dictionary<Type, IActionSerializer>();
+			NameToSerializer = new Dictionary<string, IActionSerializer>(StringComparer.Ordinal);
+
+			Type[] Types = Assembly.GetExecutingAssembly().GetTypes();
+			foreach (Type Type in Types)
+			{
+				if (Type.IsClass && !Type.IsAbstract && typeof(IActionSerializer).IsAssignableFrom(Type))
+				{
+					IActionSerializer Serializer = (IActionSerializer)Activator.CreateInstance(Type);
+					TypeToSerializer[Serializer.Type] = Serializer;
+					NameToSerializer[Type.Name] = Serializer;
+				}
+			}
+		}
+
+		/// <summary>
+		/// Read an action from the given archive
+		/// </summary>
+		/// <param name="Reader">Reader to deserialize from</param>
+		/// <returns>New action</returns>
+		public static IAction ReadAction(this BinaryArchiveReader Reader)
+		{
+			IActionSerializer Serializer = Reader.ReadObjectReference(() => ReadSerializer(Reader));
+			return Serializer.Read(Reader);
+		}
+
+		/// <summary>
+		/// Reads a type name and find its registered constructor from an archive
+		/// </summary>
+		/// <param name="Reader">Archive to read from</param>
+		/// <returns>New constructor info</returns>
+		static IActionSerializer ReadSerializer(BinaryArchiveReader Reader)
+		{
+			string Name = Reader.ReadString();
+
+			IActionSerializer Serializer;
+			if (!NameToSerializer.TryGetValue(Name, out Serializer))
+			{
+				throw new BuildException("Unable to find action type '{0}'", Name);
+			}
+
+			return Serializer;
+		}
+
+		/// <summary>
+		/// Writes an action to the given archive
+		/// </summary>
+		/// <param name="Writer">Writer to serialize the action to</param>
+		/// <param name="Action">Action to serialize</param>
+		public static void WriteAction(this BinaryArchiveWriter Writer, IAction Action)
+		{
+			Type Type = Action.GetType();
+
+			IActionSerializer Serializer;
+			if (!TypeToSerializer.TryGetValue(Type, out Serializer))
+			{
+				throw new BuildException("Unable to find serializer for action type '{0}'", Type.Name);
+			}
+
+			Writer.WriteObjectReference(Serializer, () => Writer.WriteString(Serializer.GetType().Name));
+			Serializer.Write(Writer, Action);
+		}
+	}
+}
diff --git a/Engine/Source/Programs/UnrealBuildTool/System/HotReload.cs b/Engine/Source/Programs/UnrealBuildTool/System/HotReload.cs
--- a/Engine/Source/Programs/UnrealBuildTool/System/HotReload.cs
+++ b/Engine/Source/Programs/UnrealBuildTool/System/HotReload.cs
@@ -1191,11 +1191,11 @@ namespace UnrealBuildTool
 		/// <param name="ManifestFile">File to write to</param>
 		/// <param name="Actions">List of actions that are part of the graph</param>
 		/// <param name="OriginalFileToPatchedFile">Map of original object files to patched object files</param>
-		public static void WriteLiveCodingManifest(FileReference ManifestFile, List<Action> Actions, Dictionary<FileReference, FileReference> OriginalFileToPatchedFile)
+		public static void WriteLiveCodingManifest(FileReference ManifestFile, List<IAction> Actions, Dictionary<FileReference, FileReference> OriginalFileToPatchedFile)
 		{
 			// Find all the output object files
 			HashSet<FileItem> ObjectFiles = new HashSet<FileItem>();
-			foreach(Action Action in Actions)
+			foreach(IAction Action in Actions)
 			{
 				if(Action.ActionType == ActionType.Compile)
 				{
@@ -1208,7 +1208,7 @@ namespace UnrealBuildTool
 			{
 				Writer.WriteObjectStart();
 
-				Action LinkAction = Actions.FirstOrDefault(x => x.ActionType == ActionType.Link && x.ProducedItems.Any(y => y.HasExtension(".exe") || y.HasExtension(".dll")));
+				IAction LinkAction = Actions.FirstOrDefault(x => x.ActionType == ActionType.Link && x.ProducedItems.Any(y => y.HasExtension(".exe") || y.HasExtension(".dll")));
 				if(LinkAction != null)
 				{
 					FileReference LinkerPath = LinkAction.CommandPath;
@@ -1235,7 +1235,7 @@ namespace UnrealBuildTool
 				Writer.WriteObjectEnd();
 
 				Writer.WriteArrayStart("Modules");
-				foreach(Action Action in Actions)
+				foreach(IAction Action in Actions)
 				{
 					if(Action.ActionType == ActionType.Link)
 					{
diff --git a/Engine/Source/Programs/UnrealBuildTool/System/TargetMakefile.cs b/Engine/Source/Programs/UnrealBuildTool/System/TargetMakefile.cs
--- a/Engine/Source/Programs/UnrealBuildTool/System/TargetMakefile.cs
+++ b/Engine/Source/Programs/UnrealBuildTool/System/TargetMakefile.cs
@@ -21,7 +21,7 @@ namespace UnrealBuildTool
 		/// <summary>
 		/// The version number to write
 		/// </summary>
-		public const int CurrentVersion = 20;
+		public const int CurrentVersion = 21;
 
 		/// <summary>
 		/// The time at which the makefile was created
@@ -96,7 +96,7 @@ namespace UnrealBuildTool
 		/// <summary>
 		/// Every action in the action graph
 		/// </summary>
-		public List<Action> Actions;
+		public List<IAction> Actions;
 
 		/// <summary>
 		/// Environment variables that we'll need in order to invoke the platform's compiler and linker
@@ -193,7 +193,7 @@ namespace UnrealBuildTool
 			this.ConfigValueTracker = ConfigValueTracker;
 			this.bDeployAfterCompile = bDeployAfterCompile;
 			this.bHasProjectScriptPlugin = bHasProjectScriptPlugin;
-			this.Actions = new List<Action>();
+			this.Actions = new List<IAction>();
 			this.OutputItems = new List<FileItem>();
 			this.ModuleNameToOutputItems = new Dictionary<string, FileItem[]>(StringComparer.OrdinalIgnoreCase);
 			this.HotReloadModuleNames = new HashSet<string>(StringComparer.OrdinalIgnoreCase);
@@ -229,7 +229,7 @@ namespace UnrealBuildTool
 			AdditionalArguments = Reader.ReadArray(() => Reader.ReadString());
 			PreBuildScripts = Reader.ReadArray(() => Reader.ReadFileReference());
 			PreBuildTargets = Reader.ReadArray(() => new TargetInfo(Reader));
-			Actions = Reader.ReadList(() => new Action(Reader));
+			Actions = Reader.ReadList(() => Reader.ReadAction());
 			EnvironmentVariables = Reader.ReadList(() => Tuple.Create(Reader.ReadString(), Reader.ReadString()));
 			OutputItems = Reader.ReadList(() => Reader.ReadFileItem());
 			ModuleNameToOutputItems = Reader.ReadDictionary(() => Reader.ReadString(), () => Reader.ReadArray(() => Reader.ReadFileItem()), StringComparer.OrdinalIgnoreCase);
@@ -264,7 +264,7 @@ namespace UnrealBuildTool
 			Writer.WriteArray(AdditionalArguments, Item => Writer.WriteString(Item));
 			Writer.WriteArray(PreBuildScripts, Item => Writer.WriteFileReference(Item));
 			Writer.WriteArray(PreBuildTargets, Item => Item.Write(Writer));
-			Writer.WriteList(Actions, Action => Action.Write(Writer));
+			Writer.WriteList(Actions, x => Writer.WriteAction(x));
 			Writer.WriteList(EnvironmentVariables, x => { Writer.WriteString(x.Item1); Writer.WriteString(x.Item2); });
 			Writer.WriteList(OutputItems, Item => Writer.WriteFileItem(Item));
 			Writer.WriteDictionary(ModuleNameToOutputItems, k => Writer.WriteString(k), v => Writer.WriteArray(v, e => Writer.WriteFileItem(e)));
