Missing UBT project file.

#rnx

diff --git a/Engine/Source/Programs/UnrealBuildTool/UnrealBuildTool.csproj b/Engine/Source/Programs/UnrealBuildTool/UnrealBuildTool.csproj
--- a/Engine/Source/Programs/UnrealBuildTool/UnrealBuildTool.csproj
+++ b/Engine/Source/Programs/UnrealBuildTool/UnrealBuildTool.csproj
@@ -174,6 +174,7 @@
     <Compile Include="System\Action.cs" />
     <Compile Include="System\ActionExecutor.cs" />
     <Compile Include="Modes\AggregateParsedTimingInfo.cs" />
+    <Compile Include="System\ActionSerializer.cs" />
     <Compile Include="System\Asn.cs" />
     <Compile Include="System\BuildVersion.cs" />
     <Compile Include="System\ConfigValueTracker.cs" />
