Add missing copyright notice.

#rnx

diff --git a/Engine/Source/Programs/UnrealBuildTool/System/ActionSerializer.cs b/Engine/Source/Programs/UnrealBuildTool/System/ActionSerializer.cs
--- a/Engine/Source/Programs/UnrealBuildTool/System/ActionSerializer.cs
+++ b/Engine/Source/Programs/UnrealBuildTool/System/ActionSerializer.cs
@@ -1,3 +1,5 @@
+// Copyright Epic Games, Inc. All Rights Reserved.
+
 using System;
 using System.Collections.Generic;
 using System.Linq;
