UBT: Convert a few more direct references to the Action class to IAction interfaces.

#rnx

diff --git a/Engine/Source/Programs/UnrealBuildTool/Platform/Windows/PVSToolChain.cs b/Engine/Source/Programs/UnrealBuildTool/Platform/Windows/PVSToolChain.cs
--- a/Engine/Source/Programs/UnrealBuildTool/Platform/Windows/PVSToolChain.cs
+++ b/Engine/Source/Programs/UnrealBuildTool/Platform/Windows/PVSToolChain.cs
@@ -397,19 +397,17 @@ namespace UnrealBuildTool
 
 		class ActionGraphCapture : ForwardingActionGraphBuilder
 		{
-			List<Action> Actions;
+			List<IAction> Actions;
 
-			public ActionGraphCapture(IActionGraphBuilder Inner, List<Action> Actions)
+			public ActionGraphCapture(IActionGraphBuilder Inner, List<IAction> Actions)
 				: base(Inner)
 			{
 				this.Actions = Actions;
 			}
 
-			public override Action CreateAction(ActionType Type)
+			public override void AddAction(IAction Action)
 			{
-				Action Action = base.CreateAction(Type);
 				Actions.Add(Action);
-				return Action;
 			}
 		}
 
@@ -424,12 +422,13 @@ namespace UnrealBuildTool
 			PreprocessCompileEnvironment.bEnableUndefinedIdentifierWarnings = false; // Not sure why THIRD_PARTY_INCLUDES_START doesn't pick this up; the _Pragma appears in the preprocessed output. Perhaps in preprocess-only mode the compiler doesn't respect these?
 			PreprocessCompileEnvironment.Definitions.Add("PVS_STUDIO");
 
-			List<Action> PreprocessActions = new List<Action>();
+			List<IAction> PreprocessActions = new List<IAction>();
 			CPPOutput Result = InnerToolChain.CompileCPPFiles(PreprocessCompileEnvironment, InputFiles, OutputDir, ModuleName, new ActionGraphCapture(Graph, PreprocessActions));
 
 			// Run the source files through PVS-Studio
-			foreach(Action PreprocessAction in PreprocessActions)
+			for(int Idx = 0; Idx < PreprocessActions.Count; Idx++)
 			{
+				IAction PreprocessAction = PreprocessActions[Idx];
 				if (PreprocessAction.ActionType != ActionType.Compile)
 				{
 					continue;
@@ -450,8 +449,10 @@ namespace UnrealBuildTool
 				}
 
 				// Disable a few warnings that seem to come from the preprocessor not respecting _Pragma
-				PreprocessAction.CommandArguments += " /wd4005"; // macro redefinition
-				PreprocessAction.CommandArguments += " /wd4828"; // file contains a character starting at offset xxxx that is illegal in the current source character set
+				Action NewPreprocessAction = new Action(PreprocessAction);
+				NewPreprocessAction.CommandArguments += " /wd4005"; // macro redefinition
+				NewPreprocessAction.CommandArguments += " /wd4828"; // file contains a character starting at offset xxxx that is illegal in the current source character set
+				PreprocessActions[Idx] = NewPreprocessAction;
 
 				// Write the PVS studio config file
 				StringBuilder ConfigFileContents = new StringBuilder();
diff --git a/Engine/Source/Programs/UnrealBuildTool/System/ActionGraphBuilder.cs b/Engine/Source/Programs/UnrealBuildTool/System/ActionGraphBuilder.cs
--- a/Engine/Source/Programs/UnrealBuildTool/System/ActionGraphBuilder.cs
+++ b/Engine/Source/Programs/UnrealBuildTool/System/ActionGraphBuilder.cs
@@ -16,11 +16,10 @@ namespace UnrealBuildTool
 	interface IActionGraphBuilder
 	{
 		/// <summary>
-		/// Creates a new action to be built as part of this target
+		/// Adds an action to this graph
 		/// </summary>
-		/// <param name="Type">Type of action to create</param>
-		/// <returns>New action</returns>
-		Action CreateAction(ActionType Type);
+		/// <param name="Action">Action to add</param>
+		void AddAction(IAction Action);
 
 		/// <summary>
 		/// Creates a response file for use in the action graph
@@ -75,9 +74,8 @@ namespace UnrealBuildTool
 	class NullActionGraphBuilder : IActionGraphBuilder
 	{
 		/// <inheritdoc/>
-		public Action CreateAction(ActionType Type)
+		public void AddAction(IAction Action)
 		{
-			return new Action(Type);
 		}
 
 		/// <inheritdoc/>
@@ -138,9 +136,9 @@ namespace UnrealBuildTool
 		}
 
 		/// <inheritdoc/>
-		public virtual Action CreateAction(ActionType Type)
+		public virtual void AddAction(IAction Action)
 		{
-			return Inner.CreateAction(Type);
+			Inner.AddAction(Action);
 		}
 
 		/// <inheritdoc/>
@@ -191,6 +189,19 @@ namespace UnrealBuildTool
 	/// </summary>
 	static class ActionGraphBuilderExtensions
 	{
+		/// <summary>
+		/// Creates a new action to be built as part of this target
+		/// </summary>
+		/// <param name="Graph">Graph to add the action to</param>
+		/// <param name="Type">Type of action to create</param>
+		/// <returns>New action</returns>
+		public static Action CreateAction(this IActionGraphBuilder Graph, ActionType Type)
+		{
+			Action Action = new Action(Type);
+			Graph.AddAction(Action);
+			return Action;
+		}
+
 		/// <summary>
 		/// Creates an action which copies a file from one location to another
 		/// </summary>
diff --git a/Engine/Source/Programs/UnrealBuildTool/System/TargetMakefile.cs b/Engine/Source/Programs/UnrealBuildTool/System/TargetMakefile.cs
--- a/Engine/Source/Programs/UnrealBuildTool/System/TargetMakefile.cs
+++ b/Engine/Source/Programs/UnrealBuildTool/System/TargetMakefile.cs
@@ -680,11 +680,9 @@ namespace UnrealBuildTool
 		}
 
 		/// <inheritdoc/>
-		public Action CreateAction(ActionType Type)
+		public void AddAction(IAction Action)
 		{
-			Action Action = new Action(Type);
 			Actions.Add(Action);
-			return Action;
 		}
 
 		/// <inheritdoc/>
