UBT: Change the produced/prerequisite item lists for actions to be IEnumerables, so we can generate them dynamically using iterators.

#rnx

diff --git a/Engine/Source/Programs/UnrealBuildTool/Executors/SNDBS.cs b/Engine/Source/Programs/UnrealBuildTool/Executors/SNDBS.cs
--- a/Engine/Source/Programs/UnrealBuildTool/Executors/SNDBS.cs
+++ b/Engine/Source/Programs/UnrealBuildTool/Executors/SNDBS.cs
@@ -76,7 +76,7 @@ namespace UnrealBuildTool
 						["run_locally"] = !(a.bCanExecuteRemotely && a.bCanExecuteRemotelyWithSNDBS)
 					};
 
-					if (a.PrerequisiteItems.Count > 0)
+					if (a.PrerequisiteItems.Count() > 0)
 					{
 						Job["explicit_input_files"] = a.PrerequisiteItems.Select(i => new Dictionary<string, object>()
 						{
diff --git a/Engine/Source/Programs/UnrealBuildTool/Modes/BuildMode.cs b/Engine/Source/Programs/UnrealBuildTool/Modes/BuildMode.cs
--- a/Engine/Source/Programs/UnrealBuildTool/Modes/BuildMode.cs
+++ b/Engine/Source/Programs/UnrealBuildTool/Modes/BuildMode.cs
@@ -805,11 +805,13 @@ namespace UnrealBuildTool
 				string GroupPrefix = String.Format("{0}-{1}-{2}", TargetDescriptors[TargetIdx].Name, TargetDescriptors[TargetIdx].Platform, TargetDescriptors[TargetIdx].Configuration);
 				foreach(QueuedAction TargetAction in TargetActions[TargetIdx])
 				{
+					FileItem ProducedItem = TargetAction.ProducedItems.First();
+
 					QueuedAction ExistingAction;
-					if(!OutputItemToProducingAction.TryGetValue(TargetAction.ProducedItems[0], out ExistingAction))
+					if(!OutputItemToProducingAction.TryGetValue(ProducedItem, out ExistingAction))
 					{
 						ExistingAction = new QueuedAction(TargetAction);
-						OutputItemToProducingAction[TargetAction.ProducedItems[0]] = ExistingAction;
+						OutputItemToProducingAction[ProducedItem] = ExistingAction;
 					}
 					ExistingAction.GroupNames.Add(GroupPrefix);
 				}
diff --git a/Engine/Source/Programs/UnrealBuildTool/System/Action.cs b/Engine/Source/Programs/UnrealBuildTool/System/Action.cs
--- a/Engine/Source/Programs/UnrealBuildTool/System/Action.cs
+++ b/Engine/Source/Programs/UnrealBuildTool/System/Action.cs
@@ -47,17 +47,17 @@ namespace UnrealBuildTool
 		/// <summary>
 		/// Every file this action depends on.  These files need to exist and be up to date in order for this action to even be considered
 		/// </summary>
-		IReadOnlyList<FileItem> PrerequisiteItems { get; }
+		IEnumerable<FileItem> PrerequisiteItems { get; }
 
 		/// <summary>
 		/// The files that this action produces after completing
 		/// </summary>
-		IReadOnlyList<FileItem> ProducedItems { get; }
+		IEnumerable<FileItem> ProducedItems { get; }
 
 		/// <summary>
 		/// Items that should be deleted before running this action
 		/// </summary>
-		IReadOnlyList<FileItem> DeleteItems { get; }
+		IEnumerable<FileItem> DeleteItems { get; }
 
 		/// <summary>
 		/// For C++ source files, specifies a dependency list file used to check changes to header files
@@ -129,12 +129,6 @@ namespace UnrealBuildTool
 		/// True if any libraries produced by this action should be considered 'import libraries'
 		/// </summary>
 		bool bProducesImportLibrary { get; }
-
-		/// <summary>
-		/// Writes an action to a json file
-		/// </summary>
-		/// <param name="Writer">Writer to receive the output</param>
-		void ExportJson(JsonWriter Writer);
 	}
 
 	/// <summary>
@@ -237,9 +231,9 @@ namespace UnrealBuildTool
 		/// </summary>
 		public bool bProducesImportLibrary { get; set; } = false;
 
-		IReadOnlyList<FileItem> IAction.PrerequisiteItems => PrerequisiteItems;
-		IReadOnlyList<FileItem> IAction.ProducedItems => ProducedItems;
-		IReadOnlyList<FileItem> IAction.DeleteItems => DeleteItems;
+		IEnumerable<FileItem> IAction.PrerequisiteItems => PrerequisiteItems;
+		IEnumerable<FileItem> IAction.ProducedItems => ProducedItems;
+		IEnumerable<FileItem> IAction.DeleteItems => DeleteItems;
 
 		public Action(ActionType InActionType)
 		{
@@ -420,65 +414,72 @@ namespace UnrealBuildTool
 			return Action;
 		}
 
+		public override string ToString()
+		{
+			string ReturnString = "";
+			if (CommandPath != null)
+			{
+				ReturnString += CommandPath + " - ";
+			}
+			if (CommandArguments != null)
+			{
+				ReturnString += CommandArguments;
+			}
+			return ReturnString;
+		}
+	}
+
+	/// <summary>
+	/// Extension methods for action classes
+	/// </summary>
+	static class ActionExtensions
+	{
 		/// <summary>
 		/// Writes an action to a json file
 		/// </summary>
+		/// <param name="Action">The action to write</param>
 		/// <param name="Writer">Writer to receive the output</param>
-		public void ExportJson(JsonWriter Writer)
+		public static void ExportJson(this IAction Action, JsonWriter Writer)
 		{
-			Writer.WriteEnumValue("Type", ActionType);
-			Writer.WriteValue("WorkingDirectory", WorkingDirectory.FullName);
-			Writer.WriteValue("CommandPath", CommandPath.FullName);
-			Writer.WriteValue("CommandArguments", CommandArguments);
-			Writer.WriteValue("CommandDescription", CommandDescription);
-			Writer.WriteValue("StatusDescription", StatusDescription);
-			Writer.WriteValue("bPrintDebugInfo", bPrintDebugInfo);
-			Writer.WriteValue("bCanExecuteRemotely", bCanExecuteRemotely);
-			Writer.WriteValue("bCanExecuteRemotelyWithSNDBS", bCanExecuteRemotelyWithSNDBS);
-			Writer.WriteValue("bIsGCCCompiler", bIsGCCCompiler);
-			Writer.WriteValue("bShouldOutputStatusDescription", bShouldOutputStatusDescription);
-			Writer.WriteValue("bProducesImportLibrary", bProducesImportLibrary);
+			Writer.WriteEnumValue("Type", Action.ActionType);
+			Writer.WriteValue("WorkingDirectory", Action.WorkingDirectory.FullName);
+			Writer.WriteValue("CommandPath", Action.CommandPath.FullName);
+			Writer.WriteValue("CommandArguments", Action.CommandArguments);
+			Writer.WriteValue("CommandDescription", Action.CommandDescription);
+			Writer.WriteValue("StatusDescription", Action.StatusDescription);
+			Writer.WriteValue("bPrintDebugInfo", Action.bPrintDebugInfo);
+			Writer.WriteValue("bCanExecuteRemotely", Action.bCanExecuteRemotely);
+			Writer.WriteValue("bCanExecuteRemotelyWithSNDBS", Action.bCanExecuteRemotelyWithSNDBS);
+			Writer.WriteValue("bIsGCCCompiler", Action.bIsGCCCompiler);
+			Writer.WriteValue("bShouldOutputStatusDescription", Action.bShouldOutputStatusDescription);
+			Writer.WriteValue("bProducesImportLibrary", Action.bProducesImportLibrary);
 
 			Writer.WriteArrayStart("PrerequisiteItems");
-			foreach(FileItem PrerequisiteItem in PrerequisiteItems)
+			foreach (FileItem PrerequisiteItem in Action.PrerequisiteItems)
 			{
 				Writer.WriteValue(PrerequisiteItem.AbsolutePath);
 			}
 			Writer.WriteArrayEnd();
 
 			Writer.WriteArrayStart("ProducedItems");
-			foreach(FileItem ProducedItem in ProducedItems)
+			foreach (FileItem ProducedItem in Action.ProducedItems)
 			{
 				Writer.WriteValue(ProducedItem.AbsolutePath);
 			}
 			Writer.WriteArrayEnd();
 
 			Writer.WriteArrayStart("DeleteItems");
-			foreach(FileItem DeleteItem in DeleteItems)
+			foreach (FileItem DeleteItem in Action.DeleteItems)
 			{
 				Writer.WriteValue(DeleteItem.AbsolutePath);
 			}
 			Writer.WriteArrayEnd();
 
-			if (DependencyListFile != null)
+			if (Action.DependencyListFile != null)
 			{
-				Writer.WriteValue("DependencyListFile", DependencyListFile.AbsolutePath);
+				Writer.WriteValue("DependencyListFile", Action.DependencyListFile.AbsolutePath);
 			}
 		}
-
-		public override string ToString()
-		{
-			string ReturnString = "";
-			if (CommandPath != null)
-			{
-				ReturnString += CommandPath + " - ";
-			}
-			if (CommandArguments != null)
-			{
-				ReturnString += CommandArguments;
-			}
-			return ReturnString;
-		}
 	}
 
 	/// <summary>
@@ -527,9 +528,9 @@ namespace UnrealBuildTool
 		#region Wrapper implementation of IAction
 
 		public ActionType ActionType => Inner.ActionType;
-		public IReadOnlyList<FileItem> PrerequisiteItems => Inner.PrerequisiteItems;
-		public IReadOnlyList<FileItem> ProducedItems => Inner.ProducedItems;
-		public IReadOnlyList<FileItem> DeleteItems => Inner.DeleteItems;
+		public IEnumerable<FileItem> PrerequisiteItems => Inner.PrerequisiteItems;
+		public IEnumerable<FileItem> ProducedItems => Inner.ProducedItems;
+		public IEnumerable<FileItem> DeleteItems => Inner.DeleteItems;
 		public FileItem DependencyListFile => Inner.DependencyListFile;
 		public FileItem CompiledModuleInterfaceFile => Inner.CompiledModuleInterfaceFile;
 		public FileItem TimingFile => Inner.TimingFile;
@@ -592,7 +593,7 @@ namespace UnrealBuildTool
 			// Secondary sort criteria is number of pre-requisites.
 			else
 			{
-				return Math.Sign(B.PrerequisiteItems.Count - A.PrerequisiteItems.Count);
+				return Math.Sign(B.PrerequisiteItems.Count() - A.PrerequisiteItems.Count());
 			}
 		}
 	}
diff --git a/Engine/Source/Programs/UnrealBuildTool/System/ActionGraph.cs b/Engine/Source/Programs/UnrealBuildTool/System/ActionGraph.cs
--- a/Engine/Source/Programs/UnrealBuildTool/System/ActionGraph.cs
+++ b/Engine/Source/Programs/UnrealBuildTool/System/ActionGraph.cs
@@ -138,7 +138,7 @@ namespace UnrealBuildTool
 		/// <param name="NewValue">Conflicting value for the field</param>
 		static void LogConflict(IAction Action, string Description, string OldValue, string NewValue)
 		{
-			Log.TraceError("Unable to merge actions producing {0}: {1}", Action.ProducedItems[0].Location.GetFileName(), Description);
+			Log.TraceError("Unable to merge actions producing {0}: {1}", Action.ProducedItems.First().Location.GetFileName(), Description);
 			Log.TraceLog("  Previous: {0}", OldValue);
 			Log.TraceLog("  Conflict: {0}", NewValue);
 		}
@@ -408,14 +408,15 @@ namespace UnrealBuildTool
 							{
 								if (CyclicActions.ContainsKey(CyclicPrerequisiteAction))
 								{
-									if (CyclicPrerequisiteAction.ProducedItems.Count == 1)
+									List<FileItem> CyclicProducedItems = CyclicPrerequisiteAction.ProducedItems.ToList();
+									if (CyclicProducedItems.Count == 1)
 									{
-										CycleDescription += string.Format("\t\t{0} (produces: {1})\n", ActionToIndex[CyclicPrerequisiteAction], CyclicPrerequisiteAction.ProducedItems[0].AbsolutePath);
+										CycleDescription += string.Format("\t\t{0} (produces: {1})\n", ActionToIndex[CyclicPrerequisiteAction], CyclicProducedItems[0].AbsolutePath);
 									}
 									else
 									{
 										CycleDescription += string.Format("\t\t{0}\n", ActionToIndex[CyclicPrerequisiteAction]);
-										foreach (FileItem CyclicProducedItem in CyclicPrerequisiteAction.ProducedItems)
+										foreach (FileItem CyclicProducedItem in CyclicProducedItems)
 										{
 											CycleDescription += string.Format("\t\t\tproduces:   {0}\n", CyclicProducedItem.AbsolutePath);
 										}
diff --git a/Engine/Source/Programs/UnrealBuildTool/System/HotReload.cs b/Engine/Source/Programs/UnrealBuildTool/System/HotReload.cs
--- a/Engine/Source/Programs/UnrealBuildTool/System/HotReload.cs
+++ b/Engine/Source/Programs/UnrealBuildTool/System/HotReload.cs
@@ -866,14 +866,16 @@ namespace UnrealBuildTool
 
 			foreach (QueuedAction Action in Actions.Where((Action) => Action.ActionType == ActionType.Link))
 			{
+				FileItem FirstProducedItem = Action.ProducedItems.First();
+
 				// Assume that the first produced item (with no extension) is our output file name
 				FileReference HotReloadFile;
-				if(!OriginalFileToHotReloadFile.TryGetValue(Action.ProducedItems[0].Location, out HotReloadFile))
+				if(!OriginalFileToHotReloadFile.TryGetValue(FirstProducedItem.Location, out HotReloadFile))
 				{
 					continue;
 				}
 
-				string OriginalFileNameWithoutExtension = Utils.GetFilenameWithoutAnyExtensions(Action.ProducedItems[0].AbsolutePath);
+				string OriginalFileNameWithoutExtension = Utils.GetFilenameWithoutAnyExtensions(FirstProducedItem.AbsolutePath);
 				string NewFileNameWithoutExtension = Utils.GetFilenameWithoutAnyExtensions(HotReloadFile.FullName);
 
 				// Find the response file in the command line.  We'll need to make a copy of it with our new file name.
@@ -934,9 +936,9 @@ namespace UnrealBuildTool
 					}
 
 					// Update this action's list of prerequisite items too
-					for (int ItemIndex = 0; ItemIndex < Action.PrerequisiteItems.Count; ++ItemIndex)
+					for (int ItemIndex = 0; ItemIndex < NewAction.PrerequisiteItems.Count; ++ItemIndex)
 					{
-						FileItem OriginalPrerequisiteItem = Action.PrerequisiteItems[ItemIndex];
+						FileItem OriginalPrerequisiteItem = NewAction.PrerequisiteItems[ItemIndex];
 						string NewPrerequisiteItemFilePath = ReplaceBaseFileName(OriginalPrerequisiteItem.AbsolutePath, OriginalFileNameWithoutExtension, NewFileNameWithoutExtension);
 
 						if (OriginalPrerequisiteItem.AbsolutePath != NewPrerequisiteItemFilePath)
@@ -967,9 +969,9 @@ namespace UnrealBuildTool
 				}
 
 				// Update this action's list of produced items too
-				for (int ItemIndex = 0; ItemIndex < Action.ProducedItems.Count; ++ItemIndex)
+				for (int ItemIndex = 0; ItemIndex < NewAction.ProducedItems.Count; ++ItemIndex)
 				{
-					FileItem OriginalProducedItem = Action.ProducedItems[ItemIndex];
+					FileItem OriginalProducedItem = NewAction.ProducedItems[ItemIndex];
 
 					string NewProducedItemFilePath = ReplaceBaseFileName(OriginalProducedItem.AbsolutePath, OriginalFileNameWithoutExtension, NewFileNameWithoutExtension);
 					if (OriginalProducedItem.AbsolutePath != NewProducedItemFilePath)
@@ -984,10 +986,10 @@ namespace UnrealBuildTool
 				}
 
 				// Fix up the list of items to delete too
-				for(int Idx = 0; Idx < Action.DeleteItems.Count; Idx++)
+				for(int Idx = 0; Idx < NewAction.DeleteItems.Count; Idx++)
 				{
 					FileItem NewItem;
-					if(AffectedOriginalFileItemAndNewFileItemMap.TryGetValue(Action.DeleteItems[Idx], out NewItem))
+					if(AffectedOriginalFileItemAndNewFileItemMap.TryGetValue(NewAction.DeleteItems[Idx], out NewItem))
 					{
 						NewAction.DeleteItems[Idx] = NewItem;
 					}
@@ -1011,19 +1013,19 @@ namespace UnrealBuildTool
 			// Do another pass and update any actions that depended on the original file names that we changed
 			foreach (QueuedAction Action in Actions)
 			{
-				for (int ItemIndex = 0; ItemIndex < Action.PrerequisiteItems.Count; ++ItemIndex)
+				Action NewAction = new Action(Action.Inner);
+				for (int ItemIndex = 0; ItemIndex < NewAction.PrerequisiteItems.Count; ++ItemIndex)
 				{
-					FileItem OriginalFileItem = Action.PrerequisiteItems[ItemIndex];
+					FileItem OriginalFileItem = NewAction.PrerequisiteItems[ItemIndex];
 
 					FileItem NewFileItem;
 					if (AffectedOriginalFileItemAndNewFileItemMap.TryGetValue(OriginalFileItem, out NewFileItem))
 					{
 						// OK, looks like we need to replace this file item because we've renamed the file
-						Action NewAction = new Action(Action.Inner);
 						NewAction.PrerequisiteItems[ItemIndex] = NewFileItem;
-						Action.Inner = NewAction;
 					}
 				}
+				Action.Inner = NewAction;
 			}
 
 
