UBT: Remove Action.bPrintDebugInfo. This is no longer read by anything.

#rnx

diff --git a/Engine/Source/Programs/UnrealBuildTool/Platform/HoloLens/HoloLensToolChain.cs b/Engine/Source/Programs/UnrealBuildTool/Platform/HoloLens/HoloLensToolChain.cs
--- a/Engine/Source/Programs/UnrealBuildTool/Platform/HoloLens/HoloLensToolChain.cs
+++ b/Engine/Source/Programs/UnrealBuildTool/Platform/HoloLens/HoloLensToolChain.cs
@@ -1009,12 +1009,6 @@ namespace UnrealBuildTool
 			LinkAction.PrerequisiteItems.AddRange(PrerequisiteItems);
 			LinkAction.StatusDescription = Path.GetFileName(OutputFile.AbsolutePath);
 
-			// ensure compiler timings are captured when we execute the action.
-			if (!WindowsPlatform.bAllowClangLinker && LinkEnvironment.bPrintTimingInfo)
-			{
-				LinkAction.bPrintDebugInfo = true;
-			}
-
 			// VS 15.3+ does not touch lib files if they do not contain any modifications, but we need to ensure the timestamps are updated to avoid repeatedly building them.
 			if (bBuildImportLibraryOnly || (LinkEnvironment.bHasExports && !bIsBuildingLibrary))
 			{
diff --git a/Engine/Source/Programs/UnrealBuildTool/Platform/Windows/VCToolChain.cs b/Engine/Source/Programs/UnrealBuildTool/Platform/Windows/VCToolChain.cs
--- a/Engine/Source/Programs/UnrealBuildTool/Platform/Windows/VCToolChain.cs
+++ b/Engine/Source/Programs/UnrealBuildTool/Platform/Windows/VCToolChain.cs
@@ -1060,12 +1060,6 @@ namespace UnrealBuildTool
 				Action CompileAction = Graph.CreateAction(ActionType.Compile);
 				CompileAction.CommandDescription = "Compile";
 
-				// ensure compiler timings are captured when we execute the action.
-				if (Target.WindowsPlatform.Compiler != WindowsCompiler.Clang && CompileEnvironment.bPrintTimingInfo)
-				{
-					CompileAction.bPrintDebugInfo = true;
-				}
-
 				List<string> FileArguments = new List<string>();
 				bool bIsPlainCFile = Path.GetExtension(SourceFile.AbsolutePath).ToUpperInvariant() == ".C";
 
@@ -1866,12 +1860,6 @@ namespace UnrealBuildTool
 			LinkAction.PrerequisiteItems.AddRange(PrerequisiteItems);
 			LinkAction.StatusDescription = Path.GetFileName(OutputFile.AbsolutePath);
 
-			// ensure compiler timings are captured when we execute the action.
-			if (Target.WindowsPlatform.Compiler != WindowsCompiler.Clang && LinkEnvironment.bPrintTimingInfo)
-			{
-				LinkAction.bPrintDebugInfo = true;
-			}
-
 			// VS 15.3+ does not touch lib files if they do not contain any modifications, but we need to ensure the timestamps are updated to avoid repeatedly building them.
 			if (bBuildImportLibraryOnly || (LinkEnvironment.bHasExports && !bIsBuildingLibraryOrImportLibrary))
 			{
diff --git a/Engine/Source/Programs/UnrealBuildTool/System/Action.cs b/Engine/Source/Programs/UnrealBuildTool/System/Action.cs
--- a/Engine/Source/Programs/UnrealBuildTool/System/Action.cs
+++ b/Engine/Source/Programs/UnrealBuildTool/System/Action.cs
@@ -79,11 +79,6 @@ namespace UnrealBuildTool
 		/// </summary>
 		DirectoryReference WorkingDirectory { get; }
 
-		/// <summary>
-		/// True if we should log extra information when we run a program to create produced items
-		/// </summary>
-		bool bPrintDebugInfo { get; }
-
 		/// <summary>
 		/// The command to run to create produced items
 		/// </summary>
@@ -180,11 +175,6 @@ namespace UnrealBuildTool
 		/// </summary>
 		public DirectoryReference WorkingDirectory { get; set; } = null;
 
-		/// <summary>
-		/// True if we should log extra information when we run a program to create produced items
-		/// </summary>
-		public bool bPrintDebugInfo { get; set; } = false;
-
 		/// <summary>
 		/// The command to run to create produced items
 		/// </summary>
@@ -256,7 +246,6 @@ namespace UnrealBuildTool
 			CompiledModuleInterfaceFile = InOther.CompiledModuleInterfaceFile;
 			TimingFile = InOther.TimingFile;
 			WorkingDirectory = InOther.WorkingDirectory;
-			bPrintDebugInfo = InOther.bPrintDebugInfo;
 			CommandPath = InOther.CommandPath;
 			CommandArguments = InOther.CommandArguments;
 			CommandDescription = InOther.CommandDescription;
@@ -272,7 +261,6 @@ namespace UnrealBuildTool
 		{
 			ActionType = (ActionType)Reader.ReadByte();
 			WorkingDirectory = Reader.ReadDirectoryReference();
-			bPrintDebugInfo = Reader.ReadBool();
 			CommandPath = Reader.ReadFileReference();
 			CommandArguments = Reader.ReadString();
 			CommandDescription = Reader.ReadString();
@@ -296,7 +284,6 @@ namespace UnrealBuildTool
 		{
 			Writer.WriteByte((byte)ActionType);
 			Writer.WriteDirectoryReference(WorkingDirectory);
-			Writer.WriteBool(bPrintDebugInfo);
 			Writer.WriteFileReference(CommandPath);
 			Writer.WriteString(CommandArguments);
 			Writer.WriteString(CommandDescription);
@@ -351,12 +338,6 @@ namespace UnrealBuildTool
 				Action.StatusDescription = StatusDescription;
 			}
 
-			bool bPrintDebugInfo;
-			if(Object.TryGetBoolField("bPrintDebugInfo", out bPrintDebugInfo))
-			{
-				Action.bPrintDebugInfo = bPrintDebugInfo;
-			}
-			
 			bool bCanExecuteRemotely;
 			if(Object.TryGetBoolField("bCanExecuteRemotely", out bCanExecuteRemotely))
 			{
@@ -447,7 +428,6 @@ namespace UnrealBuildTool
 			Writer.WriteValue("CommandArguments", Action.CommandArguments);
 			Writer.WriteValue("CommandDescription", Action.CommandDescription);
 			Writer.WriteValue("StatusDescription", Action.StatusDescription);
-			Writer.WriteValue("bPrintDebugInfo", Action.bPrintDebugInfo);
 			Writer.WriteValue("bCanExecuteRemotely", Action.bCanExecuteRemotely);
 			Writer.WriteValue("bCanExecuteRemotelyWithSNDBS", Action.bCanExecuteRemotelyWithSNDBS);
 			Writer.WriteValue("bIsGCCCompiler", Action.bIsGCCCompiler);
@@ -535,7 +515,6 @@ namespace UnrealBuildTool
 		public FileItem CompiledModuleInterfaceFile => Inner.CompiledModuleInterfaceFile;
 		public FileItem TimingFile => Inner.TimingFile;
 		public DirectoryReference WorkingDirectory => Inner.WorkingDirectory;
-		public bool bPrintDebugInfo => Inner.bPrintDebugInfo;
 		public FileReference CommandPath => Inner.CommandPath;
 		public string CommandArguments => Inner.CommandArguments;
 		public string CommandDescription => Inner.CommandDescription;
diff --git a/Engine/Source/Programs/UnrealBuildTool/System/TargetMakefile.cs b/Engine/Source/Programs/UnrealBuildTool/System/TargetMakefile.cs
--- a/Engine/Source/Programs/UnrealBuildTool/System/TargetMakefile.cs
+++ b/Engine/Source/Programs/UnrealBuildTool/System/TargetMakefile.cs
@@ -21,7 +21,7 @@ namespace UnrealBuildTool
 		/// <summary>
 		/// The version number to write
 		/// </summary>
-		public const int CurrentVersion = 21;
+		public const int CurrentVersion = 22;
 
 		/// <summary>
 		/// The time at which the makefile was created
