UBT: Add a custom compile action for Visual C++.

#rnx

diff --git a/Engine/Source/Programs/UnrealBuildTool/Platform/Windows/VCCompileAction.cs b/Engine/Source/Programs/UnrealBuildTool/Platform/Windows/VCCompileAction.cs
new file mode 100644
--- /dev/null
+++ b/Engine/Source/Programs/UnrealBuildTool/Platform/Windows/VCCompileAction.cs
@@ -0,0 +1,73 @@
+using System;
+using System.Collections.Generic;
+using System.Linq;
+using System.Text;
+using System.Threading.Tasks;
+using Tools.DotNETCommon;
+
+namespace UnrealBuildTool
+{
+	/// <summary>
+	/// Serializer which creates a portable object file and allows caching it
+	/// </summary>
+	class VCCompileAction : Action
+	{
+		/// <summary>
+		/// Compiler executable
+		/// </summary>
+		public FileItem Compiler { get; }
+
+		/// <summary>
+		/// For C++ source files, specifies a timing file used to track timing information.
+		/// </summary>
+		public FileItem TimingFile { get; set; }
+
+		/// <summary>
+		/// Constructor
+		/// </summary>
+		/// <param name="Environment">Compiler executable</param>
+		public VCCompileAction(VCEnvironment Environment)
+			: base(ActionType.Compile)
+		{
+			this.Compiler = FileItem.GetItemByFileReference(Environment.CompilerPath);
+		}
+
+		/// <summary>
+		/// Serialize a cache handler from an archive
+		/// </summary>
+		/// <param name="Reader">Reader to serialize from</param>
+		public VCCompileAction(BinaryArchiveReader Reader)
+			: base(Reader)
+		{
+			Compiler = Reader.ReadFileItem();
+			TimingFile = Reader.ReadFileItem();
+		}
+
+		/// <inheritdoc/>
+		public new void Write(BinaryArchiveWriter Writer)
+		{
+			base.Write(Writer);
+
+			Writer.WriteFileItem(Compiler);
+			Writer.WriteFileItem(TimingFile);
+		}
+	}
+
+	/// <summary>
+	/// Serializer for <see cref="VCCompileAction"/> instances
+	/// </summary>
+	class VCCompileActionSerializer : ActionSerializerBase<VCCompileAction>
+	{
+		/// <inheritdoc/>
+		public override VCCompileAction Read(BinaryArchiveReader Reader)
+		{
+			return new VCCompileAction(Reader);
+		}
+
+		/// <inheritdoc/>
+		public override void Write(BinaryArchiveWriter Writer, VCCompileAction Action)
+		{
+			Action.Write(Writer);
+		}
+	}
+}
diff --git a/Engine/Source/Programs/UnrealBuildTool/Platform/Windows/VCToolChain.cs b/Engine/Source/Programs/UnrealBuildTool/Platform/Windows/VCToolChain.cs
--- a/Engine/Source/Programs/UnrealBuildTool/Platform/Windows/VCToolChain.cs
+++ b/Engine/Source/Programs/UnrealBuildTool/Platform/Windows/VCToolChain.cs
@@ -1057,7 +1057,9 @@ namespace UnrealBuildTool
 			CPPOutput Result = new CPPOutput();
 			foreach (FileItem SourceFile in InputFiles)
 			{
-				Action CompileAction = Graph.CreateAction(ActionType.Compile);
+				VCCompileAction CompileAction = new VCCompileAction(EnvVars);
+				Graph.AddAction(CompileAction);
+
 				CompileAction.CommandDescription = "Compile";
 
 				List<string> FileArguments = new List<string>();
@@ -1317,7 +1319,7 @@ namespace UnrealBuildTool
 			return Result;
 		}
 
-		private Action GenerateParseTimingInfoAction(FileItem SourceFile, Action CompileAction, CPPOutput Result, IActionGraphBuilder Graph)
+		private Action GenerateParseTimingInfoAction(FileItem SourceFile, VCCompileAction CompileAction, CPPOutput Result, IActionGraphBuilder Graph)
 		{
 			FileItem TimingJsonFile = FileItem.GetItemByPath(Path.ChangeExtension(CompileAction.TimingFile.AbsolutePath, ".cta"));
 			Result.DebugDataFiles.Add(TimingJsonFile);
@@ -1340,7 +1342,7 @@ namespace UnrealBuildTool
 			return ParseTimingInfoAction;
 		}
 
-		private void GenerateDependenciesFile(ReadOnlyTargetRules Target, DirectoryReference OutputDir, CPPOutput Result, FileItem SourceFile, Action CompileAction, IActionGraphBuilder Graph)
+		private void GenerateDependenciesFile(ReadOnlyTargetRules Target, DirectoryReference OutputDir, CPPOutput Result, FileItem SourceFile, VCCompileAction CompileAction, IActionGraphBuilder Graph)
 		{
 			List<string> CommandArguments = new List<string>();
 			CompileAction.DependencyListFile = FileItem.GetItemByFileReference(FileReference.Combine(OutputDir, String.Format("{0}.txt", SourceFile.Location.GetFileName())));
diff --git a/Engine/Source/Programs/UnrealBuildTool/System/Action.cs b/Engine/Source/Programs/UnrealBuildTool/System/Action.cs
--- a/Engine/Source/Programs/UnrealBuildTool/System/Action.cs
+++ b/Engine/Source/Programs/UnrealBuildTool/System/Action.cs
@@ -69,11 +69,6 @@ namespace UnrealBuildTool
 		/// </summary>
 		FileItem CompiledModuleInterfaceFile { get; }
 
-		/// <summary>
-		/// For C++ source files, specifies a timing file used to track timing information.
-		/// </summary>
-		FileItem TimingFile { get; }
-
 		/// <summary>
 		/// Directory from which to execute the program to create produced items
 		/// </summary>
@@ -165,11 +160,6 @@ namespace UnrealBuildTool
 		/// </summary>
 		public FileItem CompiledModuleInterfaceFile { get; set; }
 
-		/// <summary>
-		/// For C++ source files, specifies a timing file used to track timing information.
-		/// </summary>
-		public FileItem TimingFile { get; set; }
-
 		/// <summary>
 		/// Directory from which to execute the program to create produced items
 		/// </summary>
@@ -244,7 +234,6 @@ namespace UnrealBuildTool
 			DeleteItems = new List<FileItem>(InOther.DeleteItems);
 			DependencyListFile = InOther.DependencyListFile;
 			CompiledModuleInterfaceFile = InOther.CompiledModuleInterfaceFile;
-			TimingFile = InOther.TimingFile;
 			WorkingDirectory = InOther.WorkingDirectory;
 			CommandPath = InOther.CommandPath;
 			CommandArguments = InOther.CommandArguments;
@@ -513,7 +502,6 @@ namespace UnrealBuildTool
 		public IEnumerable<FileItem> DeleteItems => Inner.DeleteItems;
 		public FileItem DependencyListFile => Inner.DependencyListFile;
 		public FileItem CompiledModuleInterfaceFile => Inner.CompiledModuleInterfaceFile;
-		public FileItem TimingFile => Inner.TimingFile;
 		public DirectoryReference WorkingDirectory => Inner.WorkingDirectory;
 		public FileReference CommandPath => Inner.CommandPath;
 		public string CommandArguments => Inner.CommandArguments;
diff --git a/Engine/Source/Programs/UnrealBuildTool/UnrealBuildTool.csproj b/Engine/Source/Programs/UnrealBuildTool/UnrealBuildTool.csproj
--- a/Engine/Source/Programs/UnrealBuildTool/UnrealBuildTool.csproj
+++ b/Engine/Source/Programs/UnrealBuildTool/UnrealBuildTool.csproj
@@ -132,6 +132,7 @@
     <Compile Include="Platform\Lumin\UEDeployLumin.cs" />
     <Compile Include="Platform\Mac\ApplePlatformSDK.cs" />
     <Compile Include="Platform\Windows\MicrosoftPlatformSDK.cs" />
+    <Compile Include="Platform\Windows\VCCompileAction.cs" />
     <Compile Include="Platform\Windows\VCManifestGenerator.cs" />
     <Compile Include="Platform\Windows\VCPreludeCache.cs" />
     <Compile Include="Preprocessor\Identifier.cs" />
