UBT: Start to move Windows specific compile functionality into the VCCompileAction class.

#rnx

diff --git a/Engine/Source/Programs/UnrealBuildTool/Modes/BuildMode.cs b/Engine/Source/Programs/UnrealBuildTool/Modes/BuildMode.cs
--- a/Engine/Source/Programs/UnrealBuildTool/Modes/BuildMode.cs
+++ b/Engine/Source/Programs/UnrealBuildTool/Modes/BuildMode.cs
@@ -399,14 +399,15 @@ namespace UnrealBuildTool
 						List<FileItem> CompiledModuleInterfaces = new List<FileItem>();
 						foreach(QueuedAction PrerequisiteAction in PrerequisiteActions)
 						{
-							if (PrerequisiteAction.ActionType == ActionType.CompileModuleInterface && PrerequisiteAction.CompiledModuleInterfaceFile != null)
+							ICppCompileAction CppModulesAction = PrerequisiteAction.Inner as ICppCompileAction;
+							if (CppModulesAction != null && CppModulesAction.CompiledModuleInterfaceFile != null)
 							{
 								string ProducedModule;
 								if (CppDependencies.TryGetProducedModule(PrerequisiteAction.DependencyListFile, out ProducedModule))
 								{
-									ModuleOutputs[ProducedModule] = PrerequisiteAction.CompiledModuleInterfaceFile;
+									ModuleOutputs[ProducedModule] = CppModulesAction.CompiledModuleInterfaceFile;
 								}
-								CompiledModuleInterfaces.Add(PrerequisiteAction.CompiledModuleInterfaceFile);
+								CompiledModuleInterfaces.Add(CppModulesAction.CompiledModuleInterfaceFile);
 							}
 						}
 
diff --git a/Engine/Source/Programs/UnrealBuildTool/Platform/Windows/VCCompileAction.cs b/Engine/Source/Programs/UnrealBuildTool/Platform/Windows/VCCompileAction.cs
--- a/Engine/Source/Programs/UnrealBuildTool/Platform/Windows/VCCompileAction.cs
+++ b/Engine/Source/Programs/UnrealBuildTool/Platform/Windows/VCCompileAction.cs
@@ -7,29 +7,84 @@ using Tools.DotNETCommon;
 
 namespace UnrealBuildTool
 {
+	/// <summary>
+	/// Interface for an action that compiles C++ source code
+	/// </summary>
+	interface ICppCompileAction : IAction
+	{
+		/// <summary>
+		/// Path to the compiled module interface file
+		/// </summary>
+		FileItem CompiledModuleInterfaceFile { get; }
+	}
+
 	/// <summary>
 	/// Serializer which creates a portable object file and allows caching it
 	/// </summary>
-	class VCCompileAction : Action
+	class VCCompileAction : ICppCompileAction
 	{
 		/// <summary>
-		/// Compiler executable
+		/// The dependency list file
 		/// </summary>
-		public FileItem Compiler { get; }
+		public FileItem DependencyListFile { get; set; }
+
+		/// <summary>
+		/// Compiled module interface
+		/// </summary>
+		public FileItem CompiledModuleInterfaceFile { get; set; }
 
 		/// <summary>
 		/// For C++ source files, specifies a timing file used to track timing information.
 		/// </summary>
 		public FileItem TimingFile { get; set; }
 
+		/// <summary>
+		/// Every file this action depends on.  These files need to exist and be up to date in order for this action to even be considered
+		/// </summary>
+		public List<FileItem> PrerequisiteItems { get; } = new List<FileItem>();
+
+		/// <summary>
+		/// The files that this action produces after completing
+		/// </summary>
+		public List<FileItem> ProducedItems { get; } = new List<FileItem>();
+
+		/// <summary>
+		/// Items that should be deleted before running this action
+		/// </summary>
+		public List<FileItem> DeleteItems { get; } = new List<FileItem>();
+
+		/// <inheritdoc/>
+		public FileReference CommandPath { get; set; }
+
+		/// <inheritdoc/>
+		public string CommandArguments { get; set; }
+
+		/// <inheritdoc/>
+		public string StatusDescription { get; set; }
+		
+		/// <inheritdoc/>
+		public bool bShouldOutputStatusDescription { get; set; }
+
+		#region Implemenation of IAction
+		ActionType IAction.ActionType => ActionType.Compile;
+		IEnumerable<FileItem> IAction.PrerequisiteItems => PrerequisiteItems;
+		IEnumerable<FileItem> IAction.ProducedItems => ProducedItems;
+		IEnumerable<FileItem> IAction.DeleteItems => DeleteItems;
+		public DirectoryReference WorkingDirectory => UnrealBuildTool.EngineSourceDirectory;
+		string IAction.CommandDescription => "Compile";
+		public bool bCanExecuteRemotely { get; set; }
+		public bool bCanExecuteRemotelyWithSNDBS { get; set; }
+		bool IAction.bIsGCCCompiler => false;
+		bool IAction.bProducesImportLibrary => false;
+		#endregion
+
 		/// <summary>
 		/// Constructor
 		/// </summary>
 		/// <param name="Environment">Compiler executable</param>
 		public VCCompileAction(VCEnvironment Environment)
-			: base(ActionType.Compile)
 		{
-			this.Compiler = FileItem.GetItemByFileReference(Environment.CompilerPath);
+			this.CommandPath = Environment.CompilerPath;
 		}
 
 		/// <summary>
@@ -37,19 +92,32 @@ namespace UnrealBuildTool
 		/// </summary>
 		/// <param name="Reader">Reader to serialize from</param>
 		public VCCompileAction(BinaryArchiveReader Reader)
-			: base(Reader)
 		{
-			Compiler = Reader.ReadFileItem();
+			DependencyListFile = Reader.ReadFileItem();
+			CompiledModuleInterfaceFile = Reader.ReadFileItem();
 			TimingFile = Reader.ReadFileItem();
+			PrerequisiteItems = Reader.ReadList(() => Reader.ReadFileItem());
+			ProducedItems = Reader.ReadList(() => Reader.ReadFileItem());
+			DeleteItems = Reader.ReadList(() => Reader.ReadFileItem());
+			CommandPath = Reader.ReadFileReference();
+			CommandArguments = Reader.ReadString();
+			StatusDescription = Reader.ReadString();
+			bShouldOutputStatusDescription = Reader.ReadBool();
 		}
 
 		/// <inheritdoc/>
-		public new void Write(BinaryArchiveWriter Writer)
+		public void Write(BinaryArchiveWriter Writer)
 		{
-			base.Write(Writer);
-
-			Writer.WriteFileItem(Compiler);
+			Writer.WriteFileItem(DependencyListFile);
+			Writer.WriteFileItem(CompiledModuleInterfaceFile);
 			Writer.WriteFileItem(TimingFile);
+			Writer.WriteList(PrerequisiteItems, Item => Writer.WriteFileItem(Item));
+			Writer.WriteList(ProducedItems, Item => Writer.WriteFileItem(Item));
+			Writer.WriteList(DeleteItems, Item => Writer.WriteFileItem(Item));
+			Writer.WriteFileReference(CommandPath);
+			Writer.WriteString(CommandArguments);
+			Writer.WriteString(StatusDescription);
+			Writer.WriteBool(bShouldOutputStatusDescription);
 		}
 	}
 
diff --git a/Engine/Source/Programs/UnrealBuildTool/Platform/Windows/VCToolChain.cs b/Engine/Source/Programs/UnrealBuildTool/Platform/Windows/VCToolChain.cs
--- a/Engine/Source/Programs/UnrealBuildTool/Platform/Windows/VCToolChain.cs
+++ b/Engine/Source/Programs/UnrealBuildTool/Platform/Windows/VCToolChain.cs
@@ -1060,8 +1060,6 @@ namespace UnrealBuildTool
 				VCCompileAction CompileAction = new VCCompileAction(EnvVars);
 				Graph.AddAction(CompileAction);
 
-				CompileAction.CommandDescription = "Compile";
-
 				List<string> FileArguments = new List<string>();
 				bool bIsPlainCFile = Path.GetExtension(SourceFile.AbsolutePath).ToUpperInvariant() == ".C";
 
@@ -1237,8 +1235,6 @@ namespace UnrealBuildTool
 					AppendCLArguments_CPP(CompileEnvironment, FileArguments);
 				}
 
-				CompileAction.WorkingDirectory = UnrealBuildTool.EngineSourceDirectory;
-				CompileAction.CommandPath = EnvVars.CompilerPath;
 				CompileAction.PrerequisiteItems.AddRange(CompileEnvironment.ForceIncludeFiles);
 				CompileAction.PrerequisiteItems.AddRange(CompileEnvironment.AdditionalPrerequisites);
 
@@ -1274,7 +1270,6 @@ namespace UnrealBuildTool
 					CompileDepsAction.bShouldOutputStatusDescription = false;
 
 					FileItem IfcFile = FileItem.GetItemByFileReference(FileReference.Combine(GetModuleInterfaceDir(OutputDir), SourceFile.Location.ChangeExtension(".ifc").GetFileName()));
-					CompileAction.ActionType = ActionType.CompileModuleInterface;
 					CompileAction.CommandArguments += String.Format(" /interface /ifcOutput \"{0}\"", IfcFile.Location);
 					CompileAction.PrerequisiteItems.Add(IfcDepsFile); // Force the dependencies file into the action graph
 					CompileAction.ProducedItems.Add(IfcFile);
diff --git a/Engine/Source/Programs/UnrealBuildTool/System/Action.cs b/Engine/Source/Programs/UnrealBuildTool/System/Action.cs
--- a/Engine/Source/Programs/UnrealBuildTool/System/Action.cs
+++ b/Engine/Source/Programs/UnrealBuildTool/System/Action.cs
@@ -64,11 +64,6 @@ namespace UnrealBuildTool
 		/// </summary>
 		FileItem DependencyListFile { get; }
 
-		/// <summary>
-		/// The compiled C++ module interface (IFC) file produced by this action
-		/// </summary>
-		FileItem CompiledModuleInterfaceFile { get; }
-
 		/// <summary>
 		/// Directory from which to execute the program to create produced items
 		/// </summary>
@@ -155,11 +150,6 @@ namespace UnrealBuildTool
 		/// </summary>
 		public FileItem DependencyListFile { get; set; }
 
-		/// <summary>
-		/// The compiled C++ module interface (IFC) file produced by this action
-		/// </summary>
-		public FileItem CompiledModuleInterfaceFile { get; set; }
-
 		/// <summary>
 		/// Directory from which to execute the program to create produced items
 		/// </summary>
@@ -233,7 +223,6 @@ namespace UnrealBuildTool
 			ProducedItems = new List<FileItem>(InOther.ProducedItems);
 			DeleteItems = new List<FileItem>(InOther.DeleteItems);
 			DependencyListFile = InOther.DependencyListFile;
-			CompiledModuleInterfaceFile = InOther.CompiledModuleInterfaceFile;
 			WorkingDirectory = InOther.WorkingDirectory;
 			CommandPath = InOther.CommandPath;
 			CommandArguments = InOther.CommandArguments;
@@ -263,7 +252,6 @@ namespace UnrealBuildTool
 			ProducedItems = Reader.ReadList(() => Reader.ReadFileItem());
 			DeleteItems = Reader.ReadList(() => Reader.ReadFileItem());
 			DependencyListFile = Reader.ReadFileItem();
-			CompiledModuleInterfaceFile = Reader.ReadFileItem();
 		}
 
 		/// <summary>
@@ -286,7 +274,6 @@ namespace UnrealBuildTool
 			Writer.WriteList(ProducedItems, Item => Writer.WriteFileItem(Item));
 			Writer.WriteList(DeleteItems, Item => Writer.WriteFileItem(Item));
 			Writer.WriteFileItem(DependencyListFile);
-			Writer.WriteFileItem(CompiledModuleInterfaceFile);
 		}
 
 		/// <summary>
@@ -501,7 +488,6 @@ namespace UnrealBuildTool
 		public IEnumerable<FileItem> ProducedItems => Inner.ProducedItems;
 		public IEnumerable<FileItem> DeleteItems => Inner.DeleteItems;
 		public FileItem DependencyListFile => Inner.DependencyListFile;
-		public FileItem CompiledModuleInterfaceFile => Inner.CompiledModuleInterfaceFile;
 		public DirectoryReference WorkingDirectory => Inner.WorkingDirectory;
 		public FileReference CommandPath => Inner.CommandPath;
 		public string CommandArguments => Inner.CommandArguments;
diff --git a/Engine/Source/Programs/UnrealBuildTool/System/ActionGraph.cs b/Engine/Source/Programs/UnrealBuildTool/System/ActionGraph.cs
--- a/Engine/Source/Programs/UnrealBuildTool/System/ActionGraph.cs
+++ b/Engine/Source/Programs/UnrealBuildTool/System/ActionGraph.cs
@@ -172,7 +172,7 @@ namespace UnrealBuildTool
 
 				List<FileReference> FailPaths = new List<FileReference>();
 				List<FileReference> WarnPaths = new List<FileReference>();
-				foreach (Action Action in Actions)
+				foreach (IAction Action in Actions)
 				{
 					foreach (FileItem PrerequisiteItem in Action.PrerequisiteItems)
 					{
