UBT: Rename QueuedAction to LinkedAction.

#rnx

diff --git a/Engine/Source/Programs/UnrealBuildTool/Executors/Experimental/FASTBuild.cs b/Engine/Source/Programs/UnrealBuildTool/Executors/Experimental/FASTBuild.cs
--- a/Engine/Source/Programs/UnrealBuildTool/Executors/Experimental/FASTBuild.cs
+++ b/Engine/Source/Programs/UnrealBuildTool/Executors/Experimental/FASTBuild.cs
@@ -269,13 +269,13 @@ namespace UnrealBuildTool
 
 		private ObjectIDGenerator objectIDGenerator = new ObjectIDGenerator();
 
-		private long GetActionID(QueuedAction Action)
+		private long GetActionID(LinkedAction Action)
 		{
 			bool bFirstTime = false;
 			return objectIDGenerator.GetId(Action, out bFirstTime);
 		}
 
-		private string ActionToActionString(QueuedAction Action)
+		private string ActionToActionString(LinkedAction Action)
 		{
 			return ActionToActionString(GetActionID(Action));
 		}
@@ -299,7 +299,7 @@ namespace UnrealBuildTool
 			return $"\t\t'{ActionToActionString(UniqueId)}', ;{StatusDescription}{ExtraInfoString}";
 		}
 
-		private string ActionToDependencyString(QueuedAction Action)
+		private string ActionToDependencyString(LinkedAction Action)
 		{
 			return ActionToDependencyString(GetActionID(Action), Action.StatusDescription, Action.CommandDescription, Action.ActionType);
 		}
@@ -317,33 +317,33 @@ namespace UnrealBuildTool
 
 		private FBBuildType BuildType = FBBuildType.Windows;
 
-		private readonly static Tuple<string, Func<QueuedAction, string>, FBBuildType>[] BuildTypeSearchParams = new Tuple<string, Func<QueuedAction, string>, FBBuildType>[]
+		private readonly static Tuple<string, Func<LinkedAction, string>, FBBuildType>[] BuildTypeSearchParams = new Tuple<string, Func<LinkedAction, string>, FBBuildType>[]
 		{
-			Tuple.Create<string, Func<QueuedAction, string>, FBBuildType>
+			Tuple.Create<string, Func<LinkedAction, string>, FBBuildType>
 			(
 				"Xcode",
 				Action => Action.CommandArguments,
 				FBBuildType.Apple
 			),
-			Tuple.Create<string, Func<QueuedAction, string>, FBBuildType>
+			Tuple.Create<string, Func<LinkedAction, string>, FBBuildType>
 			(
 				"apple",
 				Action => Action.CommandArguments.ToLower(),
 				FBBuildType.Apple
 			),
-			Tuple.Create<string, Func<QueuedAction, string>, FBBuildType>
+			Tuple.Create<string, Func<LinkedAction, string>, FBBuildType>
 			(
 				"/bin/sh",
 				Action => Action.CommandPath.FullName.ToLower(),
 				FBBuildType.Apple
 			),
-			Tuple.Create<string, Func<QueuedAction, string>, FBBuildType>
+			Tuple.Create<string, Func<LinkedAction, string>, FBBuildType>
 			(
 				"Windows",		// Not a great test
 				Action => Action.CommandPath.FullName,
 				FBBuildType.Windows
 			),
-			Tuple.Create<string, Func<QueuedAction, string>, FBBuildType>
+			Tuple.Create<string, Func<LinkedAction, string>, FBBuildType>
 			(
 				"Microsoft",	// Not a great test
 				Action => Action.CommandPath.FullName,
@@ -351,11 +351,11 @@ namespace UnrealBuildTool
 			),
 		};
 
-		private bool DetectBuildType(IEnumerable<QueuedAction> Actions)
+		private bool DetectBuildType(IEnumerable<LinkedAction> Actions)
 		{
-			foreach (QueuedAction Action in Actions)
+			foreach (LinkedAction Action in Actions)
 			{
-				foreach (Tuple<string, Func<QueuedAction, string>, FBBuildType> BuildTypeSearchParam in BuildTypeSearchParams)
+				foreach (Tuple<string, Func<LinkedAction, string>, FBBuildType> BuildTypeSearchParam in BuildTypeSearchParams)
 				{
 					if (BuildTypeSearchParam.Item2(Action).Contains(BuildTypeSearchParam.Item1))
 					{
@@ -367,7 +367,7 @@ namespace UnrealBuildTool
 			}
 
 			Log.TraceError("Couldn't detect build type from actions! Unsupported platform?");
-			foreach (QueuedAction Action in Actions)
+			foreach (LinkedAction Action in Actions)
 			{
 				PrintActionDetails(Action);
 			}
@@ -387,26 +387,26 @@ namespace UnrealBuildTool
 			}
 		}
 
-		public override bool ExecuteActions(List<QueuedAction> Actions)
+		public override bool ExecuteActions(List<LinkedAction> Actions)
 		{
 			if (Actions.Count <= 0)
 				return true;
 
-			IEnumerable<QueuedAction> CompileActions		= Actions.Where(Action => (Action.ActionType == ActionType.Compile && Action.bCanExecuteRemotely == true));
-			IEnumerable<QueuedAction> NonCompileActions	= Actions.Where(Action => (Action.ActionType != ActionType.Compile || Action.bCanExecuteRemotely == false));
+			IEnumerable<LinkedAction> CompileActions		= Actions.Where(Action => (Action.ActionType == ActionType.Compile && Action.bCanExecuteRemotely == true));
+			IEnumerable<LinkedAction> NonCompileActions	= Actions.Where(Action => (Action.ActionType != ActionType.Compile || Action.bCanExecuteRemotely == false));
 
 			///////////////////////////////////////////////////////////////
 			// Pre Compile Stage
 
 			// We want to complete any non-compile actions locally that are necessary for the distributed compile step
-			List<QueuedAction> PreCompileActions =
+			List<LinkedAction> PreCompileActions =
 				NonCompileActions
 				.Where(NonCompileAction => CompileActions.Any(CompileAction => CompileAction.PrerequisiteActions.Contains(NonCompileAction)))
 				.ToList();
 
 			// Precompile actions may have their own prerequisites which need to be executed along with them
-			List<QueuedAction> Prerequisites = new List<QueuedAction>();
-			foreach (QueuedAction PreCompileAction in PreCompileActions)
+			List<LinkedAction> Prerequisites = new List<LinkedAction>();
+			foreach (LinkedAction PreCompileAction in PreCompileActions)
 			{
 				Prerequisites.AddRange(PreCompileAction.PrerequisiteActions);
 			}
@@ -440,7 +440,7 @@ namespace UnrealBuildTool
 			///////////////////////////////////////////////////////////////
 			// Post Compile Stage
 
-			List<QueuedAction> PostCompileActions = NonCompileActions.Except(PreCompileActions).ToList();
+			List<LinkedAction> PostCompileActions = NonCompileActions.Except(PreCompileActions).ToList();
 
 			if (PostCompileActions.Any())
 			{
@@ -459,7 +459,7 @@ namespace UnrealBuildTool
 			bffOutputMemoryStream.Write(Info, 0, Info.Length);
 		}
 
-		private void AddPreBuildDependenciesText(IEnumerable<QueuedAction> PreBuildDependencies)
+		private void AddPreBuildDependenciesText(IEnumerable<LinkedAction> PreBuildDependencies)
 		{
 			if (!PreBuildDependencies.Any())
 				return;
@@ -717,7 +717,7 @@ namespace UnrealBuildTool
 			return ParsedCompilerOptions;
 		}
 
-		private string GetOptionValue(Dictionary<string, string> OptionsDictionary, string Key, QueuedAction Action, bool ProblemIfNotFound = false)
+		private string GetOptionValue(Dictionary<string, string> OptionsDictionary, string Key, LinkedAction Action, bool ProblemIfNotFound = false)
 		{
 			string Value = string.Empty;
 			if (OptionsDictionary.TryGetValue(Key, out Value))
@@ -994,7 +994,7 @@ namespace UnrealBuildTool
 			AddText("}\n\n"); //End Settings
 		}
 
-		private void AddCompileAction(QueuedAction Action, IEnumerable<QueuedAction> DependencyActions)
+		private void AddCompileAction(LinkedAction Action, IEnumerable<LinkedAction> DependencyActions)
 		{
 			string CompilerName = GetCompilerName();
 			if (Action.CommandPath.FullName.Contains("rc.exe"))
@@ -1109,7 +1109,7 @@ namespace UnrealBuildTool
 			AddText("}\n\n");
 		}
 
-		private void AddLinkAction(QueuedAction Action, IEnumerable<QueuedAction> DependencyActions)
+		private void AddLinkAction(LinkedAction Action, IEnumerable<LinkedAction> DependencyActions)
 		{
 			string[] SpecialLinkerOptions = IsApple() ? new string[] { "-o" } : new string[] { "/OUT:", "@", "-o" };
 
@@ -1149,13 +1149,13 @@ namespace UnrealBuildTool
 			string ResponseFilePath = GetOptionValue(ParsedLinkerOptions, "@", Action);
 			string OtherCompilerOptions = GetOptionValue(ParsedLinkerOptions, "OtherOptions", Action);
 
-			IEnumerable<QueuedAction> PrebuildDependencies = null;
+			IEnumerable<LinkedAction> PrebuildDependencies = null;
 
 			if (Action.CommandPath.FullName.Contains("lib.exe"))
 			{
 				if (DependencyActions.Any())
 				{
-					Func<QueuedAction, bool> DoesActionProducePCH = (QueuedAction ActionToCheck) =>
+					Func<LinkedAction, bool> DoesActionProducePCH = (LinkedAction ActionToCheck) =>
 					{
 						foreach (FileItem ProducedItem in ActionToCheck.ProducedItems)
 						{
@@ -1256,7 +1256,7 @@ namespace UnrealBuildTool
 				string InputFile = GetOptionValue(ParsedLinkerOptions, "InputFile", Action, ProblemIfNotFound: true);
 				if (!string.IsNullOrEmpty(InputFile))
 				{
-					QueuedAction InputFileAction = DependencyActions
+					LinkedAction InputFileAction = DependencyActions
 						.Where(ActionToInspect =>
 							ActionToInspect.ProducedItems.Any(Item => Item.AbsolutePath == InputFile)
 						).FirstOrDefault();
@@ -1304,7 +1304,7 @@ namespace UnrealBuildTool
 			}
 		}
 
-		private void PrintActionDetails(QueuedAction ActionToPrint)
+		private void PrintActionDetails(LinkedAction ActionToPrint)
 		{
 			Log.TraceInformation(ActionToActionString(ActionToPrint));
 			Log.TraceInformation($"Action Type: {ActionToPrint.ActionType.ToString()}");
@@ -1314,7 +1314,7 @@ namespace UnrealBuildTool
 
 		private MemoryStream bffOutputMemoryStream = null;
 
-		private bool CreateBffFile(IEnumerable<QueuedAction> Actions, string BffFilePath)
+		private bool CreateBffFile(IEnumerable<LinkedAction> Actions, string BffFilePath)
 		{
 			try
 			{
@@ -1326,11 +1326,11 @@ namespace UnrealBuildTool
 
 				WriteEnvironmentSetup(); //Compiler, environment variables and base paths
 
-				foreach (QueuedAction Action in Actions)
+				foreach (LinkedAction Action in Actions)
 				{
 					// Resolve the list of prerequisite items for this action to
 					// a list of actions which produce these prerequisites
-					IEnumerable<QueuedAction> DependencyActions = Action.PrerequisiteActions.Distinct();
+					IEnumerable<LinkedAction> DependencyActions = Action.PrerequisiteActions.Distinct();
 
 					AddText($";** Function for Action {GetActionID(Action)} **\n");
 					AddText($";** CommandPath: {Action.CommandPath.FullName}\n");
diff --git a/Engine/Source/Programs/UnrealBuildTool/Executors/HybridExecutor.cs b/Engine/Source/Programs/UnrealBuildTool/Executors/HybridExecutor.cs
--- a/Engine/Source/Programs/UnrealBuildTool/Executors/HybridExecutor.cs
+++ b/Engine/Source/Programs/UnrealBuildTool/Executors/HybridExecutor.cs
@@ -69,25 +69,25 @@ namespace UnrealBuildTool
 		/// </summary>
 		/// <param name="ActionsToExecute">Actions to be executed</param>
 		/// <returns>True if the build succeeded, false otherwise</returns>
-		public override bool ExecuteActions(List<QueuedAction> ActionsToExecute)
+		public override bool ExecuteActions(List<LinkedAction> ActionsToExecute)
 		{
 			// Find the number of dependants for each action
-			Dictionary<QueuedAction, int> ActionToNumDependents = ActionsToExecute.ToDictionary(x => x, x => 0);
-			foreach(QueuedAction Action in ActionsToExecute)
+			Dictionary<LinkedAction, int> ActionToNumDependents = ActionsToExecute.ToDictionary(x => x, x => 0);
+			foreach(LinkedAction Action in ActionsToExecute)
 			{
-				foreach(QueuedAction PrerequisiteAction in Action.PrerequisiteActions)
+				foreach(LinkedAction PrerequisiteAction in Action.PrerequisiteActions)
 				{
 					ActionToNumDependents[PrerequisiteAction]++;
 				}
 			}
 
 			// Build up a set of leaf actions in several iterations, ensuring that the number of leaf actions in each 
-			HashSet<QueuedAction> LeafActions = new HashSet<QueuedAction>();
+			HashSet<LinkedAction> LeafActions = new HashSet<LinkedAction>();
 			for(;;)
 			{
 				// Find all the leaf actions in the graph
-				List<QueuedAction> NewLeafActions = new List<QueuedAction>();
-				foreach(QueuedAction Action in ActionsToExecute)
+				List<LinkedAction> NewLeafActions = new List<LinkedAction>();
+				foreach(LinkedAction Action in ActionsToExecute)
 				{
 					if(ActionToNumDependents[Action] == 0 && !LeafActions.Contains(Action))
 					{
@@ -105,9 +105,9 @@ namespace UnrealBuildTool
 				LeafActions.UnionWith(NewLeafActions);
 
 				// Decrement the dependent counts for any of their prerequisites, so we can try and remove those from the tree in another iteration
-				foreach(QueuedAction NewLeafAction in NewLeafActions)
+				foreach(LinkedAction NewLeafAction in NewLeafActions)
 				{
-					foreach(QueuedAction PrerequisiteAction in NewLeafAction.PrerequisiteActions)
+					foreach(LinkedAction PrerequisiteAction in NewLeafAction.PrerequisiteActions)
 					{
 						ActionToNumDependents[PrerequisiteAction]--;
 					}
@@ -115,9 +115,9 @@ namespace UnrealBuildTool
 			}
 
 			// Split the list of actions into those which should be executed locally and remotely
-			List<QueuedAction> LocalActionsToExecute = new List<QueuedAction>(LeafActions.Count);
-			List<QueuedAction> RemoteActionsToExecute = new List<QueuedAction>(ActionsToExecute.Count - LeafActions.Count);
-			foreach(QueuedAction ActionToExecute in ActionsToExecute)
+			List<LinkedAction> LocalActionsToExecute = new List<LinkedAction>(LeafActions.Count);
+			List<LinkedAction> RemoteActionsToExecute = new List<LinkedAction>(ActionsToExecute.Count - LeafActions.Count);
+			foreach(LinkedAction ActionToExecute in ActionsToExecute)
 			{
 				if(LeafActions.Contains(ActionToExecute))
 				{
diff --git a/Engine/Source/Programs/UnrealBuildTool/Executors/LocalExecutor.cs b/Engine/Source/Programs/UnrealBuildTool/Executors/LocalExecutor.cs
--- a/Engine/Source/Programs/UnrealBuildTool/Executors/LocalExecutor.cs
+++ b/Engine/Source/Programs/UnrealBuildTool/Executors/LocalExecutor.cs
@@ -28,7 +28,7 @@ namespace UnrealBuildTool
 		/// <summary>
 		/// Cache the action that this thread is managing
 		/// </summary>
-		QueuedAction Action;
+		LinkedAction Action;
 
 		/// <summary>
 		/// For reporting status to the user
@@ -39,7 +39,7 @@ namespace UnrealBuildTool
 		/// <summary>
 		/// Constructor, takes the action to process
 		/// </summary>
-		public ActionThread(QueuedAction InAction, int InJobNumber, int InTotalJobs)
+		public ActionThread(LinkedAction InAction, int InJobNumber, int InTotalJobs)
 		{
 			Action = InAction;
 			JobNumber = InJobNumber;
@@ -296,7 +296,7 @@ namespace UnrealBuildTool
 		/// Executes the specified actions locally.
 		/// </summary>
 		/// <returns>True if all the tasks successfully executed, or false if any of them failed.</returns>
-		public override bool ExecuteActions(List<QueuedAction> Actions)
+		public override bool ExecuteActions(List<LinkedAction> Actions)
 		{
 			// Time to sleep after each iteration of the loop in order to not busy wait.
 			const float LoopSleepTime = 0.1f;
@@ -304,7 +304,7 @@ namespace UnrealBuildTool
 			// The number of actions to execute in parallel is trying to keep the CPU busy enough in presence of I/O stalls.
 			Log.TraceInformation("Performing {0} actions ({1} in parallel)", Actions.Count, NumParallelProcesses);
 
-			Dictionary<QueuedAction, ActionThread> ActionThreadDictionary = new Dictionary<QueuedAction, ActionThread>();
+			Dictionary<LinkedAction, ActionThread> ActionThreadDictionary = new Dictionary<LinkedAction, ActionThread>();
 			int JobNumber = 1;
 			using (ProgressWriter ProgressWriter = new ProgressWriter("Compiling C++ source code...", false))
 			{
@@ -314,7 +314,7 @@ namespace UnrealBuildTool
 					// Count the number of pending and still executing actions.
 					int NumUnexecutedActions = 0;
 					int NumExecutingActions = 0;
-					foreach (QueuedAction Action in Actions)
+					foreach (LinkedAction Action in Actions)
 					{
 						ActionThread ActionThread = null;
 						bool bFoundActionProcess = ActionThreadDictionary.TryGetValue(Action, out ActionThread);
@@ -348,7 +348,7 @@ namespace UnrealBuildTool
 
 					// If there are fewer actions executing than the maximum, look for pending actions that don't have any outdated
 					// prerequisites.
-					foreach (QueuedAction Action in Actions)
+					foreach (LinkedAction Action in Actions)
 					{
 						ActionThread ActionProcess = null;
 						bool bFoundActionProcess = ActionThreadDictionary.TryGetValue(Action, out ActionProcess);
@@ -359,7 +359,7 @@ namespace UnrealBuildTool
 								// Determine whether there are any prerequisites of the action that are outdated.
 								bool bHasOutdatedPrerequisites = false;
 								bool bHasFailedPrerequisites = false;
-								foreach (QueuedAction PrerequisiteAction in Action.PrerequisiteActions)
+								foreach (LinkedAction PrerequisiteAction in Action.PrerequisiteActions)
 								{
 									if (Actions.Contains(PrerequisiteAction))
 									{
@@ -422,9 +422,9 @@ namespace UnrealBuildTool
 
 			// Check whether any of the tasks failed and log action stats if wanted.
 			bool bSuccess = true;
-			foreach (KeyValuePair<QueuedAction, ActionThread> ActionProcess in ActionThreadDictionary)
+			foreach (KeyValuePair<LinkedAction, ActionThread> ActionProcess in ActionThreadDictionary)
 			{
-				QueuedAction Action = ActionProcess.Key;
+				LinkedAction Action = ActionProcess.Key;
 				ActionThread ActionThread = ActionProcess.Value;
 
 				// Check for pending actions, preemptive failure
diff --git a/Engine/Source/Programs/UnrealBuildTool/Executors/ParallelExecutor.cs b/Engine/Source/Programs/UnrealBuildTool/Executors/ParallelExecutor.cs
--- a/Engine/Source/Programs/UnrealBuildTool/Executors/ParallelExecutor.cs
+++ b/Engine/Source/Programs/UnrealBuildTool/Executors/ParallelExecutor.cs
@@ -20,7 +20,7 @@ namespace UnrealBuildTool
 		class BuildAction
 		{
 			public int SortIndex;
-			public QueuedAction Inner;
+			public LinkedAction Inner;
 
 			public HashSet<BuildAction> Dependencies = new HashSet<BuildAction>();
 			public int MissingDependencyCount;
@@ -96,7 +96,7 @@ namespace UnrealBuildTool
 		/// Executes the specified actions locally.
 		/// </summary>
 		/// <returns>True if all the tasks successfully executed, or false if any of them failed.</returns>
-		public override bool ExecuteActions(List<QueuedAction> InputActions)
+		public override bool ExecuteActions(List<LinkedAction> InputActions)
 		{
 			Log.TraceInformation("Building {0} {1} with {2} {3}...", InputActions.Count, (InputActions.Count == 1) ? "action" : "actions", NumParallelProcesses, (NumParallelProcesses == 1)? "process" : "processes");
 
diff --git a/Engine/Source/Programs/UnrealBuildTool/Executors/SNDBS.cs b/Engine/Source/Programs/UnrealBuildTool/Executors/SNDBS.cs
--- a/Engine/Source/Programs/UnrealBuildTool/Executors/SNDBS.cs
+++ b/Engine/Source/Programs/UnrealBuildTool/Executors/SNDBS.cs
@@ -44,7 +44,7 @@ namespace UnrealBuildTool
 			return ServiceController.GetServices().Any(s => s.ServiceName.StartsWith("SNDBS") && s.Status == ServiceControllerStatus.Running);
 		}
 
-		public override bool ExecuteActions(List<QueuedAction> Actions)
+		public override bool ExecuteActions(List<LinkedAction> Actions)
 		{
 			if (Actions.Count == 0)
 				return true;
diff --git a/Engine/Source/Programs/UnrealBuildTool/Executors/XGE.cs b/Engine/Source/Programs/UnrealBuildTool/Executors/XGE.cs
--- a/Engine/Source/Programs/UnrealBuildTool/Executors/XGE.cs
+++ b/Engine/Source/Programs/UnrealBuildTool/Executors/XGE.cs
@@ -301,7 +301,7 @@ namespace UnrealBuildTool
 		// precompile the Regex needed to parse the XGE output (the ones we want are of the form "File (Duration at +time)"
 		//private static Regex XGEDurationRegex = new Regex(@"(?<Filename>.*) *\((?<Duration>[0-9:\.]+) at [0-9\+:\.]+\)", RegexOptions.ExplicitCapture);
 
-		public static void ExportActions(List<QueuedAction> ActionsToExecute)
+		public static void ExportActions(List<LinkedAction> ActionsToExecute)
 		{
 			for(int FileNum = 0;;FileNum++)
 			{
@@ -314,22 +314,22 @@ namespace UnrealBuildTool
 			}
 		}
 
-		public static void ExportActions(List<QueuedAction> ActionsToExecute, string OutFile)
+		public static void ExportActions(List<LinkedAction> ActionsToExecute, string OutFile)
 		{
 			WriteTaskFile(ActionsToExecute, OutFile, ProgressWriter.bWriteMarkup, bXGEExport: true);
 			Log.TraceInformation("XGEEXPORT: Exported '{0}'", OutFile);
 		}
 
-		public override bool ExecuteActions(List<QueuedAction> ActionsToExecute)
+		public override bool ExecuteActions(List<LinkedAction> ActionsToExecute)
 		{
 			bool XGEResult = true;
 
 			// Batch up XGE execution by actions with the same output event handler.
-			List<QueuedAction> ActionBatch = new List<QueuedAction>();
+			List<LinkedAction> ActionBatch = new List<LinkedAction>();
 			ActionBatch.Add(ActionsToExecute[0]);
 			for (int ActionIndex = 1; ActionIndex < ActionsToExecute.Count && XGEResult; ++ActionIndex)
 			{
-				QueuedAction CurrentAction = ActionsToExecute[ActionIndex];
+				LinkedAction CurrentAction = ActionsToExecute[ActionIndex];
 				ActionBatch.Add(CurrentAction);
 			}
 			if (ActionBatch.Count > 0 && XGEResult)
@@ -341,7 +341,7 @@ namespace UnrealBuildTool
 			return XGEResult;
 		}
 
-		bool ExecuteActionBatch(List<QueuedAction> Actions)
+		bool ExecuteActionBatch(List<LinkedAction> Actions)
 		{
 			bool XGEResult = true;
 			if (Actions.Count > 0)
@@ -358,11 +358,11 @@ namespace UnrealBuildTool
 		/// <summary>
 		/// Writes a XGE task file containing the specified actions to the specified file path.
 		/// </summary>
-		static void WriteTaskFile(List<QueuedAction> InActions, string TaskFilePath, bool bProgressMarkup, bool bXGEExport)
+		static void WriteTaskFile(List<LinkedAction> InActions, string TaskFilePath, bool bProgressMarkup, bool bXGEExport)
 		{
 			Dictionary<string, string> ExportEnv = new Dictionary<string, string>();
 
-			List<QueuedAction> Actions = InActions;
+			List<LinkedAction> Actions = InActions;
 			if (bXGEExport)
 			{
 				IDictionary CurrentEnvironment = Environment.GetEnvironmentVariables();
@@ -413,7 +413,7 @@ namespace UnrealBuildTool
 
 			for (int ActionIndex = 0; ActionIndex < Actions.Count; ActionIndex++)
 			{
-				QueuedAction Action = Actions[ActionIndex];
+				LinkedAction Action = Actions[ActionIndex];
 
 				// <Tool ... />
 				XmlElement ToolElement = XGETaskDocument.CreateElement("Tool");
@@ -477,7 +477,7 @@ namespace UnrealBuildTool
 
 			for (int ActionIndex = 0; ActionIndex < Actions.Count; ActionIndex++)
 			{
-				QueuedAction Action = Actions[ActionIndex];
+				LinkedAction Action = Actions[ActionIndex];
 
 				// <Task ... />
 				XmlElement TaskElement = XGETaskDocument.CreateElement("Task");
@@ -500,7 +500,7 @@ namespace UnrealBuildTool
 
 				// Create a semi-colon separated list of the other tasks this task depends on the results of.
 				List<string> DependencyNames = new List<string>();
-				foreach(QueuedAction PrerequisiteAction in Action.PrerequisiteActions)
+				foreach(LinkedAction PrerequisiteAction in Action.PrerequisiteActions)
 				{
 					if (Actions.Contains(PrerequisiteAction))
 					{
diff --git a/Engine/Source/Programs/UnrealBuildTool/Modes/BuildMode.cs b/Engine/Source/Programs/UnrealBuildTool/Modes/BuildMode.cs
--- a/Engine/Source/Programs/UnrealBuildTool/Modes/BuildMode.cs
+++ b/Engine/Source/Programs/UnrealBuildTool/Modes/BuildMode.cs
@@ -309,10 +309,10 @@ namespace UnrealBuildTool
 				}
 
 				// Create a QueuedAction instance for each action in the makefiles
-				List<QueuedAction>[] QueuedActions = new List<QueuedAction>[Makefiles.Length];
+				List<LinkedAction>[] QueuedActions = new List<LinkedAction>[Makefiles.Length];
 				for(int Idx = 0; Idx < Makefiles.Length; Idx++)
 				{
-					QueuedActions[Idx] = Makefiles[Idx].Actions.ConvertAll(x => new QueuedAction(x));
+					QueuedActions[Idx] = Makefiles[Idx].Actions.ConvertAll(x => new LinkedAction(x));
 				}
 
 				// Clean up any previous hot reload runs, and reapply the current state if it's already active
@@ -322,7 +322,7 @@ namespace UnrealBuildTool
 				}
 
 				// Merge the action graphs together
-				List<QueuedAction> MergedActions;
+				List<LinkedAction> MergedActions;
 				if (TargetDescriptors.Count == 1)
 				{
 					MergedActions = QueuedActions[0];
@@ -343,7 +343,7 @@ namespace UnrealBuildTool
 				ActionGraph.Link(MergedActions);
 
 				// Get all the actions that are prerequisites for these targets. This forms the list of actions that we want executed.
-				List<QueuedAction> PrerequisiteActions = ActionGraph.GatherPrerequisiteActions(MergedActions, MergedOutputItems);
+				List<LinkedAction> PrerequisiteActions = ActionGraph.GatherPrerequisiteActions(MergedActions, MergedOutputItems);
 
 				// Create the action history
 				ActionHistory History = new ActionHistory();
@@ -360,13 +360,13 @@ namespace UnrealBuildTool
 				}
 
 				// Pre-process module interfaces to generate dependency files
-				List<QueuedAction> ModuleDependencyActions = PrerequisiteActions.Where(x => x.ActionType == ActionType.GatherModuleDependencies).ToList();
+				List<LinkedAction> ModuleDependencyActions = PrerequisiteActions.Where(x => x.ActionType == ActionType.GatherModuleDependencies).ToList();
 				if (ModuleDependencyActions.Count > 0)
 				{
-					Dictionary<QueuedAction, bool> PreprocessActionToOutdatedFlag = new Dictionary<QueuedAction, bool>();
+					Dictionary<LinkedAction, bool> PreprocessActionToOutdatedFlag = new Dictionary<LinkedAction, bool>();
 					ActionGraph.GatherAllOutdatedActions(ModuleDependencyActions, History, PreprocessActionToOutdatedFlag, null, BuildConfiguration.bIgnoreOutdatedImportLibraries);
 
-					List<QueuedAction> PreprocessActions = PreprocessActionToOutdatedFlag.Where(x => x.Value).Select(x => x.Key).ToList();
+					List<LinkedAction> PreprocessActions = PreprocessActionToOutdatedFlag.Where(x => x.Value).Select(x => x.Key).ToList();
 					if (PreprocessActions.Count > 0)
 					{
 						Log.TraceInformation("Updating module dependencies...");
@@ -380,7 +380,7 @@ namespace UnrealBuildTool
 				}
 
 				// Figure out which actions need to be built
-				Dictionary<QueuedAction, bool> ActionToOutdatedFlag = new Dictionary<QueuedAction, bool>();
+				Dictionary<LinkedAction, bool> ActionToOutdatedFlag = new Dictionary<LinkedAction, bool>();
 				for (int TargetIdx = 0; TargetIdx < TargetDescriptors.Count; TargetIdx++)
 				{
 					TargetDescriptor TargetDescriptor = TargetDescriptors[TargetIdx];
@@ -397,7 +397,7 @@ namespace UnrealBuildTool
 					if (ModuleDependencyActions.Count > 0)
 					{
 						List<FileItem> CompiledModuleInterfaces = new List<FileItem>();
-						foreach(QueuedAction PrerequisiteAction in PrerequisiteActions)
+						foreach(LinkedAction PrerequisiteAction in PrerequisiteActions)
 						{
 							ICppCompileAction CppModulesAction = PrerequisiteAction.Inner as ICppCompileAction;
 							if (CppModulesAction != null && CppModulesAction.CompiledModuleInterfaceFile != null)
@@ -412,7 +412,7 @@ namespace UnrealBuildTool
 						}
 
 						Dictionary<FileItem, string> ModuleInputs = ModuleOutputs.ToDictionary(x => x.Value, x => x.Key);
-						foreach (QueuedAction PrerequisiteAction in PrerequisiteActions)
+						foreach (LinkedAction PrerequisiteAction in PrerequisiteActions)
 						{
 							if (PrerequisiteAction.ActionType == ActionType.CompileModuleInterface)
 							{
@@ -470,7 +470,7 @@ namespace UnrealBuildTool
 					{
 						foreach (FileReference SpecificFile in TargetDescriptor.SpecificFilesToCompile)
 						{
-							foreach (QueuedAction PrerequisiteAction in PrerequisiteActions.Where(x => x.PrerequisiteItems.Any(y => y.Location == SpecificFile)))
+							foreach (LinkedAction PrerequisiteAction in PrerequisiteActions.Where(x => x.PrerequisiteItems.Any(y => y.Location == SpecificFile)))
 							{
 								ActionToOutdatedFlag[PrerequisiteAction] = true;
 							}
@@ -479,7 +479,7 @@ namespace UnrealBuildTool
 				}
 
 				// Link the action graph again to sort it
-				List<QueuedAction> MergedActionsToExecute = ActionToOutdatedFlag.Where(x => x.Value).Select(x => x.Key).ToList();
+				List<LinkedAction> MergedActionsToExecute = ActionToOutdatedFlag.Where(x => x.Value).Select(x => x.Key).ToList();
 				ActionGraph.Link(MergedActionsToExecute);
 
 				// Allow hot reload to override the actions
@@ -797,27 +797,27 @@ namespace UnrealBuildTool
 		/// <param name="TargetDescriptors">List of target descriptors</param>
 		/// <param name="TargetActions">Array of actions for each target</param>
 		/// <returns>List of merged actions</returns>
-		static List<QueuedAction> MergeActionGraphs(List<TargetDescriptor> TargetDescriptors, List<QueuedAction>[] TargetActions)
+		static List<LinkedAction> MergeActionGraphs(List<TargetDescriptor> TargetDescriptors, List<LinkedAction>[] TargetActions)
 		{
 			// Set of all output items. Knowing that there are no conflicts in produced items, we use this to eliminate duplicate actions.
-			Dictionary<FileItem, QueuedAction> OutputItemToProducingAction = new Dictionary<FileItem, QueuedAction>();
+			Dictionary<FileItem, LinkedAction> OutputItemToProducingAction = new Dictionary<FileItem, LinkedAction>();
 			for(int TargetIdx = 0; TargetIdx < TargetDescriptors.Count; TargetIdx++)
 			{
 				string GroupPrefix = String.Format("{0}-{1}-{2}", TargetDescriptors[TargetIdx].Name, TargetDescriptors[TargetIdx].Platform, TargetDescriptors[TargetIdx].Configuration);
-				foreach(QueuedAction TargetAction in TargetActions[TargetIdx])
+				foreach(LinkedAction TargetAction in TargetActions[TargetIdx])
 				{
 					FileItem ProducedItem = TargetAction.ProducedItems.First();
 
-					QueuedAction ExistingAction;
+					LinkedAction ExistingAction;
 					if(!OutputItemToProducingAction.TryGetValue(ProducedItem, out ExistingAction))
 					{
-						ExistingAction = new QueuedAction(TargetAction);
+						ExistingAction = new LinkedAction(TargetAction);
 						OutputItemToProducingAction[ProducedItem] = ExistingAction;
 					}
 					ExistingAction.GroupNames.Add(GroupPrefix);
 				}
 			}
-			return new List<QueuedAction>(OutputItemToProducingAction.Values);
+			return new List<LinkedAction>(OutputItemToProducingAction.Values);
 		}
 	}
 }
diff --git a/Engine/Source/Programs/UnrealBuildTool/Modes/ExecuteMode.cs b/Engine/Source/Programs/UnrealBuildTool/Modes/ExecuteMode.cs
--- a/Engine/Source/Programs/UnrealBuildTool/Modes/ExecuteMode.cs
+++ b/Engine/Source/Programs/UnrealBuildTool/Modes/ExecuteMode.cs
@@ -42,10 +42,10 @@ namespace UnrealBuildTool
 			Arguments.ApplyTo(BuildConfiguration);
 
 			// Read the actions file
-			List<QueuedAction> Actions;
+			List<LinkedAction> Actions;
 			using(Timeline.ScopeEvent("ActionGraph.ReadActions()"))
 			{
-				Actions = ActionGraph.ImportJson(ActionsFile).ConvertAll(x => new QueuedAction(x));
+				Actions = ActionGraph.ImportJson(ActionsFile).ConvertAll(x => new LinkedAction(x));
 			}
 
 			// Link the action graph
diff --git a/Engine/Source/Programs/UnrealBuildTool/Modes/JsonExportMode.cs b/Engine/Source/Programs/UnrealBuildTool/Modes/JsonExportMode.cs
--- a/Engine/Source/Programs/UnrealBuildTool/Modes/JsonExportMode.cs
+++ b/Engine/Source/Programs/UnrealBuildTool/Modes/JsonExportMode.cs
@@ -56,12 +56,12 @@ namespace UnrealBuildTool
 						// Create the makefile
 						const bool bIsAssemblingBuild = true;
 						TargetMakefile Makefile = Target.Build(BuildConfiguration, WorkingSet, bIsAssemblingBuild, TargetDescriptor.SpecificFilesToCompile);
-						List<QueuedAction> Actions = Makefile.Actions.ConvertAll(x => new QueuedAction(x));
+						List<LinkedAction> Actions = Makefile.Actions.ConvertAll(x => new LinkedAction(x));
 						ActionGraph.Link(Actions);
 
 						// Filter all the actions to execute
 						HashSet<FileItem> PrerequisiteItems = new HashSet<FileItem>(Makefile.Actions.SelectMany(x => x.ProducedItems).Where(x => x.HasExtension(".h") || x.HasExtension(".cpp")));
-						List<QueuedAction> PrerequisiteActions = ActionGraph.GatherPrerequisiteActions(Actions, PrerequisiteItems);
+						List<LinkedAction> PrerequisiteActions = ActionGraph.GatherPrerequisiteActions(Actions, PrerequisiteItems);
 
 						// Execute these actions
 						if (PrerequisiteActions.Count > 0)
diff --git a/Engine/Source/Programs/UnrealBuildTool/System/Action.cs b/Engine/Source/Programs/UnrealBuildTool/System/Action.cs
--- a/Engine/Source/Programs/UnrealBuildTool/System/Action.cs
+++ b/Engine/Source/Programs/UnrealBuildTool/System/Action.cs
@@ -459,7 +459,7 @@ namespace UnrealBuildTool
 	/// <summary>
 	/// Information about an action queued to be executed
 	/// </summary>
-	class QueuedAction : IAction
+	class LinkedAction : IAction
 	{
 		/// <summary>
 		/// The inner action instance
@@ -469,7 +469,7 @@ namespace UnrealBuildTool
 		/// <summary>
 		/// Set of other actions that this action depends on. This set is built when the action graph is linked.
 		/// </summary>
-		public HashSet<QueuedAction> PrerequisiteActions;
+		public HashSet<LinkedAction> PrerequisiteActions;
 
 		/// <summary>
 		/// Total number of actions depending on this one.
@@ -510,7 +510,7 @@ namespace UnrealBuildTool
 		/// Constructor
 		/// </summary>
 		/// <param name="Inner">The inner action instance</param>
-		public QueuedAction(IAction Inner)
+		public LinkedAction(IAction Inner)
 		{
 			this.Inner = Inner;
 		}
@@ -519,12 +519,12 @@ namespace UnrealBuildTool
 		/// Increment the number of dependents, recursively
 		/// </summary>
 		/// <param name="VisitedActions">Set of visited actions</param>
-		public void IncrementDependentCount(HashSet<QueuedAction> VisitedActions)
+		public void IncrementDependentCount(HashSet<LinkedAction> VisitedActions)
 		{
 			if (VisitedActions.Add(this))
 			{
 				NumTotalDependentActions++;
-				foreach (QueuedAction PrerequisiteAction in PrerequisiteActions)
+				foreach (LinkedAction PrerequisiteAction in PrerequisiteActions)
 				{
 					PrerequisiteAction.IncrementDependentCount(VisitedActions);
 				}
@@ -536,7 +536,7 @@ namespace UnrealBuildTool
 		/// </summary>
 		/// <param name="A">Action to compare</param>
 		/// <param name="B">Action to compare</param>
-		public static int Compare(QueuedAction A, QueuedAction B)
+		public static int Compare(LinkedAction A, LinkedAction B)
 		{
 			// Primary sort criteria is total number of dependent files, up to max depth.
 			if (B.NumTotalDependentActions != A.NumTotalDependentActions)
diff --git a/Engine/Source/Programs/UnrealBuildTool/System/ActionExecutor.cs b/Engine/Source/Programs/UnrealBuildTool/System/ActionExecutor.cs
--- a/Engine/Source/Programs/UnrealBuildTool/System/ActionExecutor.cs
+++ b/Engine/Source/Programs/UnrealBuildTool/System/ActionExecutor.cs
@@ -15,7 +15,7 @@ namespace UnrealBuildTool
 			get;
 		}
 
-		public abstract bool ExecuteActions(List<QueuedAction> ActionsToExecute);
+		public abstract bool ExecuteActions(List<LinkedAction> ActionsToExecute);
 	}
 
 }
diff --git a/Engine/Source/Programs/UnrealBuildTool/System/ActionGraph.cs b/Engine/Source/Programs/UnrealBuildTool/System/ActionGraph.cs
--- a/Engine/Source/Programs/UnrealBuildTool/System/ActionGraph.cs
+++ b/Engine/Source/Programs/UnrealBuildTool/System/ActionGraph.cs
@@ -18,11 +18,11 @@ namespace UnrealBuildTool
 		/// Links the actions together and sets up their dependencies
 		/// </summary>
 		/// <param name="Actions">List of actions in the graph</param>
-		public static void Link(List<QueuedAction> Actions)
+		public static void Link(List<LinkedAction> Actions)
 		{
 			// Build a map from item to its producing action
-			Dictionary<FileItem, QueuedAction> ItemToProducingAction = new Dictionary<FileItem, QueuedAction>();
-			foreach (QueuedAction Action in Actions)
+			Dictionary<FileItem, LinkedAction> ItemToProducingAction = new Dictionary<FileItem, LinkedAction>();
+			foreach (LinkedAction Action in Actions)
 			{
 				foreach (FileItem ProducedItem in Action.ProducedItems)
 				{
@@ -34,12 +34,12 @@ namespace UnrealBuildTool
 			DetectActionGraphCycles(Actions, ItemToProducingAction);
 
 			// Use this map to add all the prerequisite actions
-			foreach (QueuedAction Action in Actions)
+			foreach (LinkedAction Action in Actions)
 			{
-				Action.PrerequisiteActions = new HashSet<QueuedAction>();
+				Action.PrerequisiteActions = new HashSet<LinkedAction>();
 				foreach(FileItem PrerequisiteItem in Action.PrerequisiteItems)
 				{
-					QueuedAction PrerequisiteAction;
+					LinkedAction PrerequisiteAction;
 					if(ItemToProducingAction.TryGetValue(PrerequisiteItem, out PrerequisiteAction))
 					{
 						Action.PrerequisiteActions.Add(PrerequisiteAction);
@@ -146,12 +146,12 @@ namespace UnrealBuildTool
 		/// <summary>
 		/// Builds a list of actions that need to be executed to produce the specified output items.
 		/// </summary>
-		public static List<QueuedAction> GetActionsToExecute(List<QueuedAction> Actions, CppDependencyCache CppDependencies, ActionHistory History, bool bIgnoreOutdatedImportLibraries)
+		public static List<LinkedAction> GetActionsToExecute(List<LinkedAction> Actions, CppDependencyCache CppDependencies, ActionHistory History, bool bIgnoreOutdatedImportLibraries)
 		{
 			using (Timeline.ScopeEvent("ActionGraph.GetActionsToExecute()"))
 			{
 				// For all targets, build a set of all actions that are outdated.
-				Dictionary<QueuedAction, bool> OutdatedActionDictionary = new Dictionary<QueuedAction, bool>();
+				Dictionary<LinkedAction, bool> OutdatedActionDictionary = new Dictionary<LinkedAction, bool>();
 				GatherAllOutdatedActions(Actions, History, OutdatedActionDictionary, CppDependencies, bIgnoreOutdatedImportLibraries);
 
 				// Build a list of actions that are both needed for this target and outdated.
@@ -223,7 +223,7 @@ namespace UnrealBuildTool
 		/// <summary>
 		/// Executes a list of actions.
 		/// </summary>
-		public static void ExecuteActions(BuildConfiguration BuildConfiguration, List<QueuedAction> ActionsToExecute)
+		public static void ExecuteActions(BuildConfiguration BuildConfiguration, List<LinkedAction> ActionsToExecute)
 		{
 			if(ActionsToExecute.Count == 0)
 			{
@@ -267,7 +267,7 @@ namespace UnrealBuildTool
 				Log.TraceInformation("Total time in {0} executor: {1:0.00} seconds", Executor.Name, Timer.Elapsed.TotalSeconds);
 
 				// Reset the file info for all the produced items
-				foreach (QueuedAction BuildAction in ActionsToExecute)
+				foreach (LinkedAction BuildAction in ActionsToExecute)
 				{
 					foreach(FileItem ProducedItem in BuildAction.ProducedItems)
 					{
@@ -276,7 +276,7 @@ namespace UnrealBuildTool
 				}
 
 				// Verify the link outputs were created (seems to happen with Win64 compiles)
-				foreach (QueuedAction BuildAction in ActionsToExecute)
+				foreach (LinkedAction BuildAction in ActionsToExecute)
 				{
 					if (BuildAction.ActionType == ActionType.Link)
 					{
@@ -295,38 +295,38 @@ namespace UnrealBuildTool
 		/// <summary>
 		/// Sorts the action list for improved parallelism with local execution.
 		/// </summary>
-		static void SortActionList(List<QueuedAction> Actions)
+		static void SortActionList(List<LinkedAction> Actions)
 		{
 			// Clear the current dependent count
-			foreach(QueuedAction Action in Actions)
+			foreach(LinkedAction Action in Actions)
 			{
 				Action.NumTotalDependentActions = 0;
 			}
 
 			// Increment all the dependencies
-			foreach(QueuedAction Action in Actions)
+			foreach(LinkedAction Action in Actions)
 			{
-				Action.IncrementDependentCount(new HashSet<QueuedAction>());
+				Action.IncrementDependentCount(new HashSet<LinkedAction>());
 			}
 
 			// Sort actions by number of actions depending on them, descending. Secondary sort criteria is file size.
-			Actions.Sort(QueuedAction.Compare);
+			Actions.Sort(LinkedAction.Compare);
 		}
 
 		/// <summary>
 		/// Checks for cycles in the action graph.
 		/// </summary>
-		static void DetectActionGraphCycles(List<QueuedAction> Actions, Dictionary<FileItem, QueuedAction> ItemToProducingAction)
+		static void DetectActionGraphCycles(List<LinkedAction> Actions, Dictionary<FileItem, LinkedAction> ItemToProducingAction)
 		{
 			// Starting with actions that only depend on non-produced items, iteratively expand a set of actions that are only dependent on
 			// non-cyclical actions.
-			Dictionary<QueuedAction, bool> ActionIsNonCyclical = new Dictionary<QueuedAction, bool>();
-			Dictionary<QueuedAction, List<QueuedAction>> CyclicActions = new Dictionary<QueuedAction, List<QueuedAction>>();
+			Dictionary<LinkedAction, bool> ActionIsNonCyclical = new Dictionary<LinkedAction, bool>();
+			Dictionary<LinkedAction, List<LinkedAction>> CyclicActions = new Dictionary<LinkedAction, List<LinkedAction>>();
 			while (true)
 			{
 				bool bFoundNewNonCyclicalAction = false;
 
-				foreach (QueuedAction Action in Actions)
+				foreach (LinkedAction Action in Actions)
 				{
 					if (!ActionIsNonCyclical.ContainsKey(Action))
 					{
@@ -334,7 +334,7 @@ namespace UnrealBuildTool
 						bool bActionOnlyDependsOnNonCyclicalActions = true;
 						foreach (FileItem PrerequisiteItem in Action.PrerequisiteItems)
 						{
-							QueuedAction ProducingAction;
+							LinkedAction ProducingAction;
 							if (ItemToProducingAction.TryGetValue(PrerequisiteItem, out ProducingAction))
 							{
 								if (!ActionIsNonCyclical.ContainsKey(ProducingAction))
@@ -342,10 +342,10 @@ namespace UnrealBuildTool
 									bActionOnlyDependsOnNonCyclicalActions = false;
 									if (!CyclicActions.ContainsKey(Action))
 									{
-										CyclicActions.Add(Action, new List<QueuedAction>());
+										CyclicActions.Add(Action, new List<LinkedAction>());
 									}
 
-									List<QueuedAction> CyclicPrereq = CyclicActions[Action];
+									List<LinkedAction> CyclicPrereq = CyclicActions[Action];
 									if (!CyclicPrereq.Contains(ProducingAction))
 									{
 										CyclicPrereq.Add(ProducingAction);
@@ -379,7 +379,7 @@ namespace UnrealBuildTool
 			if (ActionIsNonCyclical.Count < Actions.Count)
 			{
 				// Find the index of each action
-				Dictionary<QueuedAction, int> ActionToIndex = new Dictionary<QueuedAction, int>();
+				Dictionary<LinkedAction, int> ActionToIndex = new Dictionary<LinkedAction, int>();
 				for(int Idx = 0; Idx < Actions.Count; Idx++)
 				{
 					ActionToIndex[Actions[Idx]] = Idx;
@@ -387,7 +387,7 @@ namespace UnrealBuildTool
 
 				// Describe the cyclical actions.
 				string CycleDescription = "";
-				foreach (QueuedAction Action in Actions)
+				foreach (LinkedAction Action in Actions)
 				{
 					if (!ActionIsNonCyclical.ContainsKey(Action))
 					{
@@ -404,7 +404,7 @@ namespace UnrealBuildTool
 						CycleDescription += string.Format("\tDepends on cyclic actions:\n");
 						if (CyclicActions.ContainsKey(Action))
 						{
-							foreach (QueuedAction CyclicPrerequisiteAction in CyclicActions[Action])
+							foreach (LinkedAction CyclicPrerequisiteAction in CyclicActions[Action])
 							{
 								if (CyclicActions.ContainsKey(CyclicPrerequisiteAction))
 								{
@@ -443,10 +443,10 @@ namespace UnrealBuildTool
 		/// <param name="Actions">All the actions in the graph</param>
 		/// <param name="OutputItems">Set of output items to be built</param>
 		/// <returns>Set of prerequisite actions</returns>
-		public static List<QueuedAction> GatherPrerequisiteActions(List<QueuedAction> Actions, HashSet<FileItem> OutputItems)
+		public static List<LinkedAction> GatherPrerequisiteActions(List<LinkedAction> Actions, HashSet<FileItem> OutputItems)
 		{
-			HashSet<QueuedAction> PrerequisiteActions = new HashSet<QueuedAction>();
-			foreach(QueuedAction Action in Actions)
+			HashSet<LinkedAction> PrerequisiteActions = new HashSet<LinkedAction>();
+			foreach(LinkedAction Action in Actions)
 			{
 				if(Action.ProducedItems.Any(x => OutputItems.Contains(x)))
 				{
@@ -461,11 +461,11 @@ namespace UnrealBuildTool
 		/// </summary>
 		/// <param name="Action">The root action to scan</param>
 		/// <param name="PrerequisiteActions">Set of prerequisite actions</param>
-		private static void GatherPrerequisiteActions(QueuedAction Action, HashSet<QueuedAction> PrerequisiteActions)
+		private static void GatherPrerequisiteActions(LinkedAction Action, HashSet<LinkedAction> PrerequisiteActions)
 		{
 			if(PrerequisiteActions.Add(Action))
 			{
-				foreach(QueuedAction PrerequisiteAction in Action.PrerequisiteActions)
+				foreach(LinkedAction PrerequisiteAction in Action.PrerequisiteActions)
 				{
 					GatherPrerequisiteActions(PrerequisiteAction, PrerequisiteActions);
 				}
@@ -482,7 +482,7 @@ namespace UnrealBuildTool
 		/// <param name="CppDependencies"></param>
 		/// <param name="bIgnoreOutdatedImportLibraries"></param>
 		/// <returns>true if outdated</returns>
-		public static bool IsActionOutdated(QueuedAction RootAction, Dictionary<QueuedAction, bool> OutdatedActionDictionary, ActionHistory ActionHistory, CppDependencyCache CppDependencies, bool bIgnoreOutdatedImportLibraries)
+		public static bool IsActionOutdated(LinkedAction RootAction, Dictionary<LinkedAction, bool> OutdatedActionDictionary, ActionHistory ActionHistory, CppDependencyCache CppDependencies, bool bIgnoreOutdatedImportLibraries)
 		{
 			// Only compute the outdated-ness for actions that don't aren't cached in the outdated action dictionary.
 			bool bIsOutdated = false;
@@ -543,7 +543,7 @@ namespace UnrealBuildTool
 			// Check if any of the prerequisite actions are out of date
 			if (!bIsOutdated)
 			{
-				foreach (QueuedAction PrerequisiteAction in RootAction.PrerequisiteActions)
+				foreach (LinkedAction PrerequisiteAction in RootAction.PrerequisiteActions)
 				{
 					if (IsActionOutdated(PrerequisiteAction, OutdatedActionDictionary, ActionHistory, CppDependencies, bIgnoreOutdatedImportLibraries))
 					{
@@ -632,7 +632,7 @@ namespace UnrealBuildTool
 		/// <param name="RootAction">The action to check</param>
 		/// <param name="PrerequisiteAction">The action that it depends on</param>
 		/// <returns>True if the only dependency between two actions is for an import library</returns>
-		static bool IsImportLibraryDependency(QueuedAction RootAction, QueuedAction PrerequisiteAction)
+		static bool IsImportLibraryDependency(LinkedAction RootAction, LinkedAction PrerequisiteAction)
 		{
 			if(PrerequisiteAction.bProducesImportLibrary)
 			{
@@ -650,11 +650,11 @@ namespace UnrealBuildTool
 		/// <param name="RootAction">The action to check</param>
 		/// <param name="PrerequisiteItem">The dependency that is out of date</param>
 		/// <returns>True if the only dependency between two actions is for an import library</returns>
-		static bool IsImportLibraryDependency(QueuedAction RootAction, FileItem PrerequisiteItem)
+		static bool IsImportLibraryDependency(LinkedAction RootAction, FileItem PrerequisiteItem)
 		{
 			if(PrerequisiteItem.Location.HasExtension(".lib"))
 			{
-				foreach(QueuedAction PrerequisiteAction in RootAction.PrerequisiteActions)
+				foreach(LinkedAction PrerequisiteAction in RootAction.PrerequisiteActions)
 				{
 					if(PrerequisiteAction.bProducesImportLibrary && PrerequisiteAction.ProducedItems.Contains(PrerequisiteItem))
 					{
@@ -669,12 +669,12 @@ namespace UnrealBuildTool
 		/// Builds a dictionary containing the actions from AllActions that are outdated by calling
 		/// IsActionOutdated.
 		/// </summary>
-		public static void GatherAllOutdatedActions(IEnumerable<QueuedAction> Actions, ActionHistory ActionHistory, Dictionary<QueuedAction, bool> OutdatedActions, CppDependencyCache CppDependencies, bool bIgnoreOutdatedImportLibraries)
+		public static void GatherAllOutdatedActions(IEnumerable<LinkedAction> Actions, ActionHistory ActionHistory, Dictionary<LinkedAction, bool> OutdatedActions, CppDependencyCache CppDependencies, bool bIgnoreOutdatedImportLibraries)
 		{
 			using(Timeline.ScopeEvent("Prefetching include dependencies"))
 			{
 				List<FileItem> Dependencies = new List<FileItem>();
-				foreach(QueuedAction Action in Actions)
+				foreach(LinkedAction Action in Actions)
 				{
 					if(Action.DependencyListFile != null)
 					{
@@ -693,9 +693,9 @@ namespace UnrealBuildTool
 		/// Deletes all the items produced by actions in the provided outdated action dictionary.
 		/// </summary>
 		/// <param name="OutdatedActions">List of outdated actions</param>
-		public static void DeleteOutdatedProducedItems(List<QueuedAction> OutdatedActions)
+		public static void DeleteOutdatedProducedItems(List<LinkedAction> OutdatedActions)
 		{
-			foreach(QueuedAction OutdatedAction in OutdatedActions)
+			foreach(LinkedAction OutdatedAction in OutdatedActions)
 			{
 				foreach (FileItem DeleteItem in OutdatedAction.DeleteItems)
 				{
@@ -712,10 +712,10 @@ namespace UnrealBuildTool
 		/// Creates directories for all the items produced by actions in the provided outdated action
 		/// dictionary.
 		/// </summary>
-		public static void CreateDirectoriesForProducedItems(List<QueuedAction> OutdatedActions)
+		public static void CreateDirectoriesForProducedItems(List<LinkedAction> OutdatedActions)
 		{
 			HashSet<DirectoryReference> OutputDirectories = new HashSet<DirectoryReference>();
-			foreach(QueuedAction OutdatedAction in OutdatedActions)
+			foreach(LinkedAction OutdatedAction in OutdatedActions)
 			{
 				foreach(FileItem ProducedItem in OutdatedAction.ProducedItems)
 				{
diff --git a/Engine/Source/Programs/UnrealBuildTool/System/HotReload.cs b/Engine/Source/Programs/UnrealBuildTool/System/HotReload.cs
--- a/Engine/Source/Programs/UnrealBuildTool/System/HotReload.cs
+++ b/Engine/Source/Programs/UnrealBuildTool/System/HotReload.cs
@@ -160,7 +160,7 @@ namespace UnrealBuildTool
 		/// <param name="Makefile">Makefile for the targe</param>
 		/// <param name="Actions">Actions for this target</param>
 		/// <param name="BuildConfiguration">Global build configuration</param>
-		public static void Setup(TargetDescriptor TargetDescriptor, TargetMakefile Makefile, List<QueuedAction> Actions, BuildConfiguration BuildConfiguration)
+		public static void Setup(TargetDescriptor TargetDescriptor, TargetMakefile Makefile, List<LinkedAction> Actions, BuildConfiguration BuildConfiguration)
 		{
 			// Get the hot-reload mode
 			if (TargetDescriptor.HotReloadMode == HotReloadMode.LiveCoding)
@@ -442,7 +442,7 @@ namespace UnrealBuildTool
 		/// <param name="HotReloadState">The hot-reload state</param>
 		/// <param name="Makefile">Makefile to apply the state</param>
 		/// <param name="Actions">Actions for this makefile</param>
-		static void ApplyState(HotReloadState HotReloadState, TargetMakefile Makefile, List<QueuedAction> Actions)
+		static void ApplyState(HotReloadState HotReloadState, TargetMakefile Makefile, List<LinkedAction> Actions)
 		{
 			// Update the action graph to produce these new files
 			HotReload.PatchActionGraph(Actions, HotReloadState.OriginalFileToHotReloadFile);
@@ -471,7 +471,7 @@ namespace UnrealBuildTool
 		/// <param name="PrerequisiteActions">The actions to execute</param>
 		/// <param name="TargetActionsToExecute">Actions to execute for this target</param>
 		/// <returns>Set of actions to execute</returns>
-		public static List<QueuedAction> PatchActionsForTarget(BuildConfiguration BuildConfiguration, TargetDescriptor TargetDescriptor, TargetMakefile Makefile, List<QueuedAction> PrerequisiteActions, List<QueuedAction> TargetActionsToExecute)
+		public static List<LinkedAction> PatchActionsForTarget(BuildConfiguration BuildConfiguration, TargetDescriptor TargetDescriptor, TargetMakefile Makefile, List<LinkedAction> PrerequisiteActions, List<LinkedAction> TargetActionsToExecute)
 		{
 			// Get the dependency history
 			CppDependencyCache CppDependencies = CppDependencyCache.CreateHierarchy(TargetDescriptor.ProjectFile, TargetDescriptor.Name, TargetDescriptor.Platform, TargetDescriptor.Configuration, Makefile.TargetType, TargetDescriptor.Architecture);
@@ -503,7 +503,7 @@ namespace UnrealBuildTool
 
 					// Find all the binaries that we're actually going to build
 					HashSet<FileReference> OutputFiles = new HashSet<FileReference>();
-					foreach (QueuedAction Action in TargetActionsToExecute)
+					foreach (LinkedAction Action in TargetActionsToExecute)
 					{
 						if (Action.ActionType == ActionType.Link)
 						{
@@ -525,7 +525,7 @@ namespace UnrealBuildTool
 				}
 
 				// Filter the prerequisite actions down to just the compile actions, then recompute all the actions to execute
-				PrerequisiteActions = new List<QueuedAction>(TargetActionsToExecute.Where(x => x.ActionType == ActionType.Compile));
+				PrerequisiteActions = new List<LinkedAction>(TargetActionsToExecute.Where(x => x.ActionType == ActionType.Compile));
 				TargetActionsToExecute = ActionGraph.GetActionsToExecute(PrerequisiteActions, CppDependencies, History, BuildConfiguration.bIgnoreOutdatedImportLibraries);
 
 				// Update the action graph with these new paths
@@ -585,7 +585,7 @@ namespace UnrealBuildTool
 				for (int LastNumFilesWithNewSuffix = 0; FilesRequiringSuffix.Count > LastNumFilesWithNewSuffix;)
 				{
 					LastNumFilesWithNewSuffix = FilesRequiringSuffix.Count;
-					foreach (QueuedAction PrerequisiteAction in PrerequisiteActions)
+					foreach (LinkedAction PrerequisiteAction in PrerequisiteActions)
 					{
 						if (!TargetActionsToExecute.Contains(PrerequisiteAction))
 						{
@@ -634,7 +634,7 @@ namespace UnrealBuildTool
 				}
 
 				// Now filter out all the hot reload files and update the state
-				foreach (QueuedAction Action in TargetActionsToExecute)
+				foreach (LinkedAction Action in TargetActionsToExecute)
 				{
 					foreach (FileItem ProducedItem in Action.ProducedItems)
 					{
@@ -747,9 +747,9 @@ namespace UnrealBuildTool
 		/// </summary>
 		/// <param name="Actions">Set of actions</param>
 		/// <param name="OriginalFileToPatchedFile">Dictionary that receives a map of original object file to patched object file</param>
-		public static void PatchActionGraphForLiveCoding(IEnumerable<QueuedAction> Actions, Dictionary<FileReference, FileReference> OriginalFileToPatchedFile)
+		public static void PatchActionGraphForLiveCoding(IEnumerable<LinkedAction> Actions, Dictionary<FileReference, FileReference> OriginalFileToPatchedFile)
 		{
-			foreach (QueuedAction Action in Actions)
+			foreach (LinkedAction Action in Actions)
 			{
 				if(Action.ActionType == ActionType.Compile)
 				{
@@ -848,7 +848,7 @@ namespace UnrealBuildTool
 		/// <summary>
 		/// Patch the action graph for hot reloading, mapping files according to the given dictionary.
 		/// </summary>
-		public static void PatchActionGraph(IEnumerable<QueuedAction> Actions, Dictionary<FileReference, FileReference> OriginalFileToHotReloadFile)
+		public static void PatchActionGraph(IEnumerable<LinkedAction> Actions, Dictionary<FileReference, FileReference> OriginalFileToHotReloadFile)
 		{
 			// Gather all of the response files for link actions.  We're going to need to patch 'em up after we figure out new
 			// names for all of the output files and import libraries
@@ -864,7 +864,7 @@ namespace UnrealBuildTool
 			// Finally, we'll keep track of any file items that we had to create counterparts for change file names, so we can fix those up too
 			Dictionary<FileItem, FileItem> AffectedOriginalFileItemAndNewFileItemMap = new Dictionary<FileItem, FileItem>();
 
-			foreach (QueuedAction Action in Actions.Where((Action) => Action.ActionType == ActionType.Link))
+			foreach (LinkedAction Action in Actions.Where((Action) => Action.ActionType == ActionType.Link))
 			{
 				FileItem FirstProducedItem = Action.ProducedItems.First();
 
@@ -1011,7 +1011,7 @@ namespace UnrealBuildTool
 
 
 			// Do another pass and update any actions that depended on the original file names that we changed
-			foreach (QueuedAction Action in Actions)
+			foreach (LinkedAction Action in Actions)
 			{
 				Action NewAction = new Action(Action.Inner);
 				for (int ItemIndex = 0; ItemIndex < NewAction.PrerequisiteItems.Count; ++ItemIndex)
@@ -1032,7 +1032,7 @@ namespace UnrealBuildTool
 			if (OriginalFileNameAndNewFileNameList_NoExtensions.Count > 0)
 			{
 				// Update all the paths in link actions
-				foreach (QueuedAction Action in Actions.Where((Action) => Action.ActionType == ActionType.Link))
+				foreach (LinkedAction Action in Actions.Where((Action) => Action.ActionType == ActionType.Link))
 				{
 					foreach (KeyValuePair<string, string> FileNameTuple in OriginalFileNameAndNewFileNameList_NoExtensions)
 					{
@@ -1086,7 +1086,7 @@ namespace UnrealBuildTool
 			}
 
 			// Update the action that writes out the module manifests
-			foreach(QueuedAction Action in Actions)
+			foreach(LinkedAction Action in Actions)
 			{
 				if(Action.ActionType == ActionType.WriteMetadata)
 				{
@@ -1164,7 +1164,7 @@ namespace UnrealBuildTool
 		/// <param name="ModuleNameToSuffix">Map of module name to suffix</param>
 		/// <param name="Makefile">Makefile for the target being built</param>
 		/// <param name="Actions">Actions to be executed for this makefile</param>
-		public static void PatchActionGraphWithNames(Dictionary<string, int> ModuleNameToSuffix, TargetMakefile Makefile, List<QueuedAction> Actions)
+		public static void PatchActionGraphWithNames(Dictionary<string, int> ModuleNameToSuffix, TargetMakefile Makefile, List<LinkedAction> Actions)
 		{
 			if(ModuleNameToSuffix.Count > 0)
 			{
