UBT: Move logic for building response files and generating inputs and outputs for MSVC compilation into the VCCompileAction class.

#rnx

diff --git a/Engine/Source/Programs/UnrealBuildTool/Configuration/UEBuildModuleCPP.cs b/Engine/Source/Programs/UnrealBuildTool/Configuration/UEBuildModuleCPP.cs
--- a/Engine/Source/Programs/UnrealBuildTool/Configuration/UEBuildModuleCPP.cs
+++ b/Engine/Source/Programs/UnrealBuildTool/Configuration/UEBuildModuleCPP.cs
@@ -457,6 +457,7 @@ namespace UnrealBuildTool
 			{
 				CPPOutput ModuleOutput = ToolChain.CompileCPPFiles(CompileEnvironment, InputFiles.IXXFiles, IntermediateDirectory, Name, Graph);
 				LinkInputFiles.AddRange(ModuleOutput.ObjectFiles);
+				CompileEnvironment.AdditionalPrerequisites.AddRange(ModuleOutput.CompiledModuleInterfaces);
 			}
 
 			// Configure the precompiled headers for this module
@@ -940,7 +941,6 @@ namespace UnrealBuildTool
 
 				// Merge output
 				OutputFiles.ObjectFiles.AddRange(AdaptiveOutput.ObjectFiles);
-				OutputFiles.DebugDataFiles.AddRange(AdaptiveOutput.DebugDataFiles);
 			}
 
 			return OutputFiles;
diff --git a/Engine/Source/Programs/UnrealBuildTool/Platform/HoloLens/HoloLensToolChain.cs b/Engine/Source/Programs/UnrealBuildTool/Platform/HoloLens/HoloLensToolChain.cs
--- a/Engine/Source/Programs/UnrealBuildTool/Platform/HoloLens/HoloLensToolChain.cs
+++ b/Engine/Source/Programs/UnrealBuildTool/Platform/HoloLens/HoloLensToolChain.cs
@@ -625,7 +625,6 @@ namespace UnrealBuildTool
 					if (CompileEnvironment.bUsePDBFiles && bActionProducesPDB)
 					{
 						CompileAction.ProducedItems.Add(PDBFile);
-						Result.DebugDataFiles.Add(PDBFile);
 					}
 				}
 
@@ -639,7 +638,6 @@ namespace UnrealBuildTool
 							)
 						);
 					CompileAction.ProducedItems.Add(AnalysisLogFile);
-					Result.DebugDataFiles.Add(AnalysisLogFile);
 					// Peform code analysis with results in a log file
 					FileArguments.AddFormat(" /analyze:log \"{0}\"", AnalysisLogFile.AbsolutePath);
 					// Suppress code analysis output
diff --git a/Engine/Source/Programs/UnrealBuildTool/Platform/Windows/VCCompileAction.cs b/Engine/Source/Programs/UnrealBuildTool/Platform/Windows/VCCompileAction.cs
--- a/Engine/Source/Programs/UnrealBuildTool/Platform/Windows/VCCompileAction.cs
+++ b/Engine/Source/Programs/UnrealBuildTool/Platform/Windows/VCCompileAction.cs
@@ -2,9 +2,6 @@
 
 using System;
 using System.Collections.Generic;
-using System.Linq;
-using System.Text;
-using System.Threading.Tasks;
 using Tools.DotNETCommon;
 
 namespace UnrealBuildTool
@@ -25,6 +22,31 @@ namespace UnrealBuildTool
 	/// </summary>
 	class VCCompileAction : ICppCompileAction
 	{
+		/// <summary>
+		/// Path to the compiler
+		/// </summary>
+		public FileItem CompilerExe { get; }
+
+		/// <summary>
+		/// The type of compiler being used
+		/// </summary>
+		public WindowsCompiler CompilerType { get; }
+
+		/// <summary>
+		/// Source file to compile
+		/// </summary>
+		public FileItem SourceFile { get; set; }
+
+		/// <summary>
+		/// The object file to output
+		/// </summary>
+		public FileItem ObjectFile { get; set; }
+
+		/// <summary>
+		/// The output preprocessed file
+		/// </summary>
+		public FileItem PreprocessedFile { get; set; }
+
 		/// <summary>
 		/// The dependency list file
 		/// </summary>
@@ -40,15 +62,67 @@ namespace UnrealBuildTool
 		/// </summary>
 		public FileItem TimingFile { get; set; }
 
+		/// <summary>
+		/// Response file for the compiler
+		/// </summary>
+		public FileItem ResponseFile { get; set; }
+
+		/// <summary>
+		/// The precompiled header file
+		/// </summary>
+		public FileItem CreatePchFile { get; set; }
+
+		/// <summary>
+		/// The precompiled header file
+		/// </summary>
+		public FileItem UsingPchFile { get; set; }
+
+		/// <summary>
+		/// The header which matches the PCH
+		/// </summary>
+		public FileItem PchThroughHeaderFile { get; set; }
+
+		/// <summary>
+		/// List of include paths
+		/// </summary>
+		public List<DirectoryReference> IncludePaths { get; } = new List<DirectoryReference>();
+
+		/// <summary>
+		/// List of system include paths
+		/// </summary>
+		public List<DirectoryReference> SystemIncludePaths { get; } = new List<DirectoryReference>();
+
+		/// <summary>
+		/// List of macro definitions
+		/// </summary>
+		public List<string> Definitions { get; } = new List<string>();
+
+		/// <summary>
+		/// List of force included files
+		/// </summary>
+		public List<FileItem> ForceIncludeFiles = new List<FileItem>();
+
 		/// <summary>
 		/// Every file this action depends on.  These files need to exist and be up to date in order for this action to even be considered
 		/// </summary>
-		public List<FileItem> PrerequisiteItems { get; } = new List<FileItem>();
+		public List<FileItem> AdditionalPrerequisiteItems { get; } = new List<FileItem>();
 
 		/// <summary>
 		/// The files that this action produces after completing
 		/// </summary>
-		public List<FileItem> ProducedItems { get; } = new List<FileItem>();
+		public List<FileItem> AdditionalProducedItems { get; } = new List<FileItem>();
+
+		/// <summary>
+		/// Arguments to pass to the compiler
+		/// </summary>
+		public List<string> Arguments { get; } = new List<string>();
+
+		/// <summary>
+		/// Whether to show included files to the console
+		/// </summary>
+		public bool bShowIncludes { get; set; }
+
+		#region Public IAction implementation
 
 		/// <summary>
 		/// Items that should be deleted before running this action
@@ -56,28 +130,108 @@ namespace UnrealBuildTool
 		public List<FileItem> DeleteItems { get; } = new List<FileItem>();
 
 		/// <inheritdoc/>
-		public FileReference CommandPath { get; set; }
+		public bool bCanExecuteRemotely { get; set; }
 
 		/// <inheritdoc/>
-		public string CommandArguments { get; set; }
+		public bool bCanExecuteRemotelyWithSNDBS { get; set; }
 
-		/// <inheritdoc/>
-		public string StatusDescription { get; set; }
-		
-		/// <inheritdoc/>
-		public bool bShouldOutputStatusDescription { get; set; }
+		#endregion
+
+		#region Implementation of IAction
 
-		#region Implemenation of IAction
 		ActionType IAction.ActionType => ActionType.Compile;
-		IEnumerable<FileItem> IAction.PrerequisiteItems => PrerequisiteItems;
-		IEnumerable<FileItem> IAction.ProducedItems => ProducedItems;
 		IEnumerable<FileItem> IAction.DeleteItems => DeleteItems;
 		public DirectoryReference WorkingDirectory => UnrealBuildTool.EngineSourceDirectory;
 		string IAction.CommandDescription => "Compile";
-		public bool bCanExecuteRemotely { get; set; }
-		public bool bCanExecuteRemotelyWithSNDBS { get; set; }
 		bool IAction.bIsGCCCompiler => false;
 		bool IAction.bProducesImportLibrary => false;
+		string IAction.StatusDescription => (SourceFile == null) ? "Compiling" : SourceFile.Location.GetFileName();
+		bool IAction.bShouldOutputStatusDescription => CompilerType == WindowsCompiler.Clang;
+
+		/// <inheritdoc/>
+		IEnumerable<FileItem> IAction.PrerequisiteItems
+		{
+			get
+			{
+				if (SourceFile != null)
+				{
+					yield return SourceFile;
+				}
+				if (UsingPchFile != null)
+				{
+					yield return UsingPchFile;
+				}
+				foreach(FileItem AdditionalPrerequisiteItem in AdditionalPrerequisiteItems)
+				{
+					yield return AdditionalPrerequisiteItem;
+				}
+			}
+		}
+
+		/// <inheritdoc/>
+		IEnumerable<FileItem> IAction.ProducedItems
+		{
+			get
+			{
+				if (ObjectFile != null)
+				{
+					yield return ObjectFile;
+				}
+				if (PreprocessedFile != null)
+				{
+					yield return PreprocessedFile;
+				}
+				if (DependencyListFile != null)
+				{
+					yield return DependencyListFile;
+				}
+				if (TimingFile != null)
+				{
+					yield return TimingFile;
+				}
+				if (CreatePchFile != null)
+				{
+					yield return CreatePchFile;
+				}
+				foreach (FileItem AdditionalProducedItem in AdditionalProducedItems)
+				{
+					yield return AdditionalProducedItem;
+				}
+			}
+		}
+
+		/// <inheritdoc/>
+		FileReference IAction.CommandPath
+		{
+			get
+			{
+				if (DependencyListFile != null)
+				{
+					return FileReference.Combine(UnrealBuildTool.EngineDirectory, "Build", "Windows", "cl-filter", "cl-filter.exe");
+				}
+				else
+				{
+					return CompilerExe.Location;
+				}
+			}
+		}
+
+		/// <inheritdoc/>
+		string IAction.CommandArguments
+		{
+			get
+			{
+				if (DependencyListFile != null)
+				{
+					return GetClFilterArguments();
+				}
+				else
+				{
+					return GetClArguments();
+				}
+			}
+		}
+
 		#endregion
 
 		/// <summary>
@@ -86,7 +240,38 @@ namespace UnrealBuildTool
 		/// <param name="Environment">Compiler executable</param>
 		public VCCompileAction(VCEnvironment Environment)
 		{
-			this.CommandPath = Environment.CompilerPath;
+			this.CompilerExe = FileItem.GetItemByFileReference(Environment.CompilerPath);
+			this.CompilerType = Environment.Compiler;
+		}
+
+		/// <summary>
+		/// Copy constructor
+		/// </summary>
+		/// <param name="InAction">Action to copy from</param>
+		public VCCompileAction(VCCompileAction InAction)
+		{
+			CompilerExe = InAction.CompilerExe;
+			CompilerType = InAction.CompilerType;
+			SourceFile = InAction.SourceFile;
+			ObjectFile = InAction.ObjectFile;
+			PreprocessedFile = InAction.PreprocessedFile;
+			DependencyListFile = InAction.DependencyListFile;
+			CompiledModuleInterfaceFile = InAction.CompiledModuleInterfaceFile;
+			TimingFile = InAction.TimingFile;
+			ResponseFile = InAction.ResponseFile;
+			CreatePchFile = InAction.CreatePchFile;
+			UsingPchFile = InAction.UsingPchFile;
+			PchThroughHeaderFile = InAction.PchThroughHeaderFile;
+			IncludePaths = new List<DirectoryReference>(InAction.IncludePaths);
+			SystemIncludePaths = new List<DirectoryReference>(InAction.SystemIncludePaths);
+			Definitions = new List<string>(InAction.Definitions);
+			ForceIncludeFiles = new List<FileItem>(InAction.ForceIncludeFiles);
+			Arguments = new List<string>(InAction.Arguments);
+			bShowIncludes = InAction.bShowIncludes;
+
+			AdditionalPrerequisiteItems = new List<FileItem>(InAction.AdditionalPrerequisiteItems);
+			AdditionalProducedItems = new List<FileItem>(InAction.AdditionalProducedItems);
+			DeleteItems = new List<FileItem>(InAction.DeleteItems);
 		}
 
 		/// <summary>
@@ -95,31 +280,160 @@ namespace UnrealBuildTool
 		/// <param name="Reader">Reader to serialize from</param>
 		public VCCompileAction(BinaryArchiveReader Reader)
 		{
+			CompilerExe = Reader.ReadFileItem();
+			CompilerType = (WindowsCompiler)Reader.ReadInt();
+			SourceFile = Reader.ReadFileItem();
+			ObjectFile = Reader.ReadFileItem();
+			PreprocessedFile = Reader.ReadFileItem();
 			DependencyListFile = Reader.ReadFileItem();
 			CompiledModuleInterfaceFile = Reader.ReadFileItem();
 			TimingFile = Reader.ReadFileItem();
-			PrerequisiteItems = Reader.ReadList(() => Reader.ReadFileItem());
-			ProducedItems = Reader.ReadList(() => Reader.ReadFileItem());
+			ResponseFile = Reader.ReadFileItem();
+			CreatePchFile = Reader.ReadFileItem();
+			UsingPchFile = Reader.ReadFileItem();
+			PchThroughHeaderFile = Reader.ReadFileItem();
+			IncludePaths = Reader.ReadList(() => Reader.ReadDirectoryReference());
+			SystemIncludePaths = Reader.ReadList(() => Reader.ReadDirectoryReference());
+			Definitions = Reader.ReadList(() => Reader.ReadString());
+			ForceIncludeFiles = Reader.ReadList(() => Reader.ReadFileItem());
+			Arguments = Reader.ReadList(() => Reader.ReadString());
+			bShowIncludes = Reader.ReadBool();
+
+			AdditionalPrerequisiteItems = Reader.ReadList(() => Reader.ReadFileItem());
+			AdditionalProducedItems = Reader.ReadList(() => Reader.ReadFileItem());
 			DeleteItems = Reader.ReadList(() => Reader.ReadFileItem());
-			CommandPath = Reader.ReadFileReference();
-			CommandArguments = Reader.ReadString();
-			StatusDescription = Reader.ReadString();
-			bShouldOutputStatusDescription = Reader.ReadBool();
 		}
 
 		/// <inheritdoc/>
 		public void Write(BinaryArchiveWriter Writer)
 		{
+			Writer.WriteFileItem(CompilerExe);
+			Writer.WriteInt((int)CompilerType);
+			Writer.WriteFileItem(SourceFile);
+			Writer.WriteFileItem(ObjectFile);
+			Writer.WriteFileItem(PreprocessedFile);
 			Writer.WriteFileItem(DependencyListFile);
 			Writer.WriteFileItem(CompiledModuleInterfaceFile);
 			Writer.WriteFileItem(TimingFile);
-			Writer.WriteList(PrerequisiteItems, Item => Writer.WriteFileItem(Item));
-			Writer.WriteList(ProducedItems, Item => Writer.WriteFileItem(Item));
+			Writer.WriteFileItem(ResponseFile);
+			Writer.WriteFileItem(CreatePchFile);
+			Writer.WriteFileItem(UsingPchFile);
+			Writer.WriteFileItem(PchThroughHeaderFile);
+			Writer.WriteList(IncludePaths, Item => Writer.WriteDirectoryReference(Item));
+			Writer.WriteList(SystemIncludePaths, Item => Writer.WriteDirectoryReference(Item));
+			Writer.WriteList(Definitions, Item => Writer.WriteString(Item));
+			Writer.WriteList(ForceIncludeFiles, Item => Writer.WriteFileItem(Item));
+			Writer.WriteList(Arguments, Item => Writer.WriteString(Item));
+			Writer.WriteBool(bShowIncludes);
+
+			Writer.WriteList(AdditionalPrerequisiteItems, Item => Writer.WriteFileItem(Item));
+			Writer.WriteList(AdditionalProducedItems, Item => Writer.WriteFileItem(Item));
 			Writer.WriteList(DeleteItems, Item => Writer.WriteFileItem(Item));
-			Writer.WriteFileReference(CommandPath);
-			Writer.WriteString(CommandArguments);
-			Writer.WriteString(StatusDescription);
-			Writer.WriteBool(bShouldOutputStatusDescription);
+		}
+
+		/// <summary>
+		/// Writes the response file with the action's arguments
+		/// </summary>
+		/// <param name="Graph">The graph builder</param>
+		public void WriteResponseFile(IActionGraphBuilder Graph)
+		{
+			if (ResponseFile != null)
+			{
+				Graph.CreateIntermediateTextFile(ResponseFile, GetFullClArgumentList());
+			}
+		}
+
+		List<string> GetFullClArgumentList()
+		{
+			List<string> Arguments = new List<string>();
+
+			if (SourceFile != null)
+			{
+				Arguments.Add(Utils.MakePathSafeToUseWithCommandLine(SourceFile.FullName));
+			}
+
+			foreach (DirectoryReference IncludePath in IncludePaths)
+			{
+				VCToolChain.AddIncludePath(Arguments, IncludePath, CompilerType, PreprocessedFile != null);
+			}
+
+			foreach (DirectoryReference SystemIncludePath in SystemIncludePaths)
+			{
+				VCToolChain.AddSystemIncludePath(Arguments, SystemIncludePath, CompilerType, PreprocessedFile != null);
+			}
+
+			foreach (string Definition in Definitions)
+			{
+				// Escape all quotation marks so that they get properly passed with the command line.
+				string DefinitionArgument = Definition.Contains("\"") ? Definition.Replace("\"", "\\\"") : Definition;
+				VCToolChain.AddDefinition(Arguments, DefinitionArgument);
+			}
+
+			foreach (FileItem ForceIncludeFile in ForceIncludeFiles)
+			{
+				Arguments.Add(String.Format("/FI\"{0}\"", ForceIncludeFile.Location));
+			}
+
+			if (CreatePchFile != null)
+			{
+				Arguments.Add(String.Format("/Yc\"{0}\"", PchThroughHeaderFile.FullName));
+				Arguments.Add(String.Format("/Fp\"{0}\"", CreatePchFile.FullName));
+			}
+			if (UsingPchFile != null)
+			{
+				Arguments.Add(String.Format("/Yu\"{0}\"", PchThroughHeaderFile.FullName));
+				Arguments.Add(String.Format("/Fp\"{0}\"", UsingPchFile.FullName));
+			}
+
+			if (PreprocessedFile != null)
+			{
+				Arguments.Add("/P"); // Preprocess
+				Arguments.Add("/C"); // Preserve comments when preprocessing
+				Arguments.Add(String.Format("/Fi\"{0}\"", PreprocessedFile)); // Preprocess to a file
+			}
+
+			if (ObjectFile != null)
+			{
+				Arguments.Add(String.Format("/Fo\"{0}\"", ObjectFile.AbsolutePath));
+			}
+
+			Arguments.AddRange(this.Arguments);
+			return Arguments;
+		}
+
+		string GetClArguments()
+		{
+			if (ResponseFile == null)
+			{
+				return String.Join(" ", Arguments);
+			}
+			else
+			{
+				return String.Format("@\"{0}\"", ResponseFile.Location);
+			}
+		}
+
+		string GetClFilterArguments()
+		{
+			List<string> Arguments = new List<string>();
+			Arguments.Add(String.Format("-dependencies={0}", Utils.MakePathSafeToUseWithCommandLine(DependencyListFile.Location)));
+
+			if (TimingFile != null)
+			{
+				Arguments.Add(String.Format("-timing={0}", Utils.MakePathSafeToUseWithCommandLine(TimingFile.Location)));
+			}
+			if (bShowIncludes)
+			{
+				Arguments.Add("-showincludes");
+			}
+
+			Arguments.Add(String.Format("-compiler={0}", Utils.MakePathSafeToUseWithCommandLine(CompilerExe.AbsolutePath)));
+			Arguments.Add("--");
+			Arguments.Add(Utils.MakePathSafeToUseWithCommandLine(CompilerExe.AbsolutePath));
+			Arguments.Add(GetClArguments());
+			Arguments.Add("/showIncludes");
+
+			return String.Join(" ", Arguments);
 		}
 	}
 
diff --git a/Engine/Source/Programs/UnrealBuildTool/Platform/Windows/VCToolChain.cs b/Engine/Source/Programs/UnrealBuildTool/Platform/Windows/VCToolChain.cs
--- a/Engine/Source/Programs/UnrealBuildTool/Platform/Windows/VCToolChain.cs
+++ b/Engine/Source/Programs/UnrealBuildTool/Platform/Windows/VCToolChain.cs
@@ -1008,25 +1008,31 @@ namespace UnrealBuildTool
 
 		public override CPPOutput CompileCPPFiles(CppCompileEnvironment CompileEnvironment, List<FileItem> InputFiles, DirectoryReference OutputDir, string ModuleName, IActionGraphBuilder Graph)
 		{
-			List<string> SharedArguments = new List<string>();
-			AppendCLArguments_Global(CompileEnvironment, SharedArguments);
+			VCCompileAction BaseCompileAction = new VCCompileAction(EnvVars);
+			AppendCLArguments_Global(CompileEnvironment, BaseCompileAction.Arguments);
 
 			// Add include paths to the argument list.
-			foreach (DirectoryReference IncludePath in CompileEnvironment.UserIncludePaths)
-			{
-				AddIncludePath(SharedArguments, IncludePath, Target.WindowsPlatform.Compiler, CompileEnvironment.bPreprocessOnly);
-			}
+			BaseCompileAction.IncludePaths.AddRange(CompileEnvironment.UserIncludePaths);
+			BaseCompileAction.SystemIncludePaths.AddRange(CompileEnvironment.SystemIncludePaths);
+			BaseCompileAction.SystemIncludePaths.AddRange(EnvVars.IncludePaths);
 
-			foreach (DirectoryReference IncludePath in CompileEnvironment.SystemIncludePaths)
-			{
-				AddSystemIncludePath(SharedArguments, IncludePath, Target.WindowsPlatform.Compiler, CompileEnvironment.bPreprocessOnly);
-			}
+			// Add preprocessor definitions to the argument list.
+			BaseCompileAction.Definitions.AddRange(CompileEnvironment.Definitions);
 
-			foreach (DirectoryReference IncludePath in EnvVars.IncludePaths)
+			// Add the force included headers
+			BaseCompileAction.ForceIncludeFiles.AddRange(CompileEnvironment.ForceIncludeFiles);
+
+			// If we're using precompiled headers, set that up now
+			if (CompileEnvironment.PrecompiledHeaderAction == PrecompiledHeaderAction.Include)
 			{
-				AddSystemIncludePath(SharedArguments, IncludePath, Target.WindowsPlatform.Compiler, CompileEnvironment.bPreprocessOnly);
+				FileItem IncludeHeader = FileItem.GetItemByFileReference(CompileEnvironment.PrecompiledHeaderIncludeFilename);
+				BaseCompileAction.ForceIncludeFiles.Insert(0, IncludeHeader);
+
+				BaseCompileAction.UsingPchFile = CompileEnvironment.PrecompiledHeaderFile;
+				BaseCompileAction.PchThroughHeaderFile = IncludeHeader;
 			}
 
+			// Generate the timing info
 			if (CompileEnvironment.bPrintTimingInfo || Target.WindowsPlatform.bCompilerTrace)
 			{
 				if (Target.WindowsPlatform.Compiler == WindowsCompiler.VisualStudio2015_DEPRECATED ||
@@ -1035,64 +1041,35 @@ namespace UnrealBuildTool
 				{
 					if (CompileEnvironment.bPrintTimingInfo)
 					{
-						SharedArguments.Add("/Bt+ /d2cgsummary");
+						BaseCompileAction.Arguments.Add("/Bt+ /d2cgsummary");
 					}
 
 					if(EnvVars.ToolChainVersion >= VersionNumber.Parse("14.14.26316"))
 					{
-						SharedArguments.Add("/d1reportTime");
+						BaseCompileAction.Arguments.Add("/d1reportTime");
 					}
 				}
 			}
 
-			// Add preprocessor definitions to the argument list.
-			foreach (string Definition in CompileEnvironment.Definitions)
-			{
-				// Escape all quotation marks so that they get properly passed with the command line.
-				string DefinitionArgument = Definition.Contains("\"") ? Definition.Replace("\"", "\\\"") : Definition;
-				AddDefinition(SharedArguments, DefinitionArgument);
-			}
-
 			// Create a compile action for each source file.
-			CPPOutput Result = new CPPOutput();
+			List<VCCompileAction> Actions = new List<VCCompileAction>();
 			foreach (FileItem SourceFile in InputFiles)
 			{
-				VCCompileAction CompileAction = new VCCompileAction(EnvVars);
-				Graph.AddAction(CompileAction);
+				VCCompileAction CompileAction = new VCCompileAction(BaseCompileAction);
+				CompileAction.SourceFile = SourceFile;
 
-				List<string> FileArguments = new List<string>();
 				bool bIsPlainCFile = Path.GetExtension(SourceFile.AbsolutePath).ToUpperInvariant() == ".C";
 
-				// Add the C++ source file and its included files to the prerequisite item list.
-				CompileAction.PrerequisiteItems.Add(SourceFile);
-
-				bool bEmitsObjectFile = true;
 				if (CompileEnvironment.PrecompiledHeaderAction == PrecompiledHeaderAction.Create)
 				{
 					// Generate a CPP File that just includes the precompiled header.
-					FileReference PCHCPPPath = CompileEnvironment.PrecompiledHeaderIncludeFilename.ChangeExtension(".cpp");
-					FileItem PCHCPPFile = Graph.CreateIntermediateTextFile(
-						PCHCPPPath,
-						string.Format("// Compiler: {0}\n#include \"{1}\"\r\n", EnvVars.CompilerVersion, CompileEnvironment.PrecompiledHeaderIncludeFilename.FullName.Replace('\\', '/'))
-						);
-
-					// Make sure the original source directory the PCH header file existed in is added as an include
-					// path -- it might be a private PCH header and we need to make sure that its found!
-					AddIncludePath(FileArguments, SourceFile.Location.Directory, Target.WindowsPlatform.Compiler, CompileEnvironment.bPreprocessOnly);
+					string PchCppFile = string.Format("// Compiler: {0}\n#include \"{1}\"\r\n", EnvVars.CompilerVersion, CompileEnvironment.PrecompiledHeaderIncludeFilename.FullName.Replace('\\', '/'));
+					CompileAction.SourceFile = FileItem.GetItemByFileReference(CompileEnvironment.PrecompiledHeaderIncludeFilename.ChangeExtension(".cpp"));
+					Graph.CreateIntermediateTextFile(CompileAction.SourceFile, PchCppFile);
 
 					// Add the precompiled header file to the produced items list.
-					FileItem PrecompiledHeaderFile = FileItem.GetItemByFileReference(
-						FileReference.Combine(
-							OutputDir,
-							Path.GetFileName(SourceFile.AbsolutePath) + ".pch"
-							)
-						);
-					CompileAction.ProducedItems.Add(PrecompiledHeaderFile);
-					Result.PrecompiledHeaderFile = PrecompiledHeaderFile;
-
-					// Add the parameters needed to compile the precompiled header file to the command-line.
-					FileArguments.Add(String.Format("/Yc\"{0}\"", CompileEnvironment.PrecompiledHeaderIncludeFilename));
-					FileArguments.Add(String.Format("/Fp\"{0}\"", PrecompiledHeaderFile.AbsolutePath));
+					CompileAction.CreatePchFile = FileItem.GetItemByFileReference(FileReference.Combine(OutputDir, SourceFile.Location.GetFileName() + ".pch"));
+					CompileAction.PchThroughHeaderFile = FileItem.GetItemByFileReference(CompileEnvironment.PrecompiledHeaderIncludeFilename);
 
 					// If we're creating a PCH that will be used to compile source files for a library, we need
 					// the compiled modules to retain a reference to PCH's module, so that debugging information
@@ -1103,60 +1080,24 @@ namespace UnrealBuildTool
 						// NOTE: The symbol name we use here is arbitrary, and all that matters is that it is
 						// unique per PCH module used in our library
 						string FakeUniquePCHSymbolName = CompileEnvironment.PrecompiledHeaderIncludeFilename.GetFileNameWithoutExtension();
-						FileArguments.Add(String.Format("/Yl{0}", FakeUniquePCHSymbolName));
+						CompileAction.Arguments.Add(String.Format("/Yl{0}", FakeUniquePCHSymbolName));
 					}
-
-					FileArguments.Add(String.Format("\"{0}\"", PCHCPPFile.AbsolutePath));
-
-					CompileAction.StatusDescription = PCHCPPPath.GetFileName();
-				}
-				else
-				{
-					if (CompileEnvironment.PrecompiledHeaderAction == PrecompiledHeaderAction.Include)
-					{
-						CompileAction.PrerequisiteItems.Add(CompileEnvironment.PrecompiledHeaderFile);
-
-						if (Target.WindowsPlatform.Compiler == WindowsCompiler.Clang)
-						{
-							FileArguments.Add(String.Format("/FI\"{0}\"", Path.ChangeExtension(CompileEnvironment.PrecompiledHeaderFile.AbsolutePath, null)));
-						}
-						else
-						{
-							FileArguments.Add(String.Format("/FI\"{0}\"", CompileEnvironment.PrecompiledHeaderIncludeFilename.FullName));
-							FileArguments.Add(String.Format("/Yu\"{0}\"", CompileEnvironment.PrecompiledHeaderIncludeFilename.FullName));
-							FileArguments.Add(String.Format("/Fp\"{0}\"", CompileEnvironment.PrecompiledHeaderFile.AbsolutePath));
-						}
-					}
-
-					// Add the source file path to the command-line.
-					FileArguments.Add(String.Format("\"{0}\"", SourceFile.AbsolutePath));
-
-					CompileAction.StatusDescription = Path.GetFileName(SourceFile.AbsolutePath);
-				}
-
-				foreach(FileItem ForceIncludeFile in CompileEnvironment.ForceIncludeFiles)
-				{
-					FileArguments.Add(String.Format("/FI\"{0}\"", ForceIncludeFile.Location));
 				}
 
 				if (CompileEnvironment.bPreprocessOnly)
 				{
-					FileItem PreprocessedFile = FileItem.GetItemByFileReference(FileReference.Combine(OutputDir, SourceFile.Location.GetFileName() + ".i"));
-
-					FileArguments.Add("/P"); // Preprocess
-					FileArguments.Add("/C"); // Preserve comments when preprocessing
-					FileArguments.Add(String.Format("/Fi\"{0}\"", PreprocessedFile)); // Preprocess to a file
-
-					CompileAction.ProducedItems.Add(PreprocessedFile);
-
-					bEmitsObjectFile = false;
+					CompileAction.PreprocessedFile = FileItem.GetItemByFileReference(FileReference.Combine(OutputDir, SourceFile.Location.GetFileName() + ".i"));
+					CompileAction.ResponseFile = FileItem.GetItemByPath(CompileAction.PreprocessedFile.FullName + ".response");
 				}
-
-				if (bEmitsObjectFile)
+				else
 				{
 					// Add the object file to the produced item list.
 					string ObjectLeafFilename = Path.GetFileName(SourceFile.AbsolutePath) + ".obj";
 					FileItem ObjectFile = FileItem.GetItemByFileReference(FileReference.Combine(OutputDir, ObjectLeafFilename));
+
+					CompileAction.ObjectFile = ObjectFile;
+					CompileAction.ResponseFile = FileItem.GetItemByPath(ObjectFile.FullName + ".response");
+
 					if (Target.WindowsPlatform.ObjSrcMapFile != null)
 					{
 						using (StreamWriter Writer = File.AppendText(Target.WindowsPlatform.ObjSrcMapFile))
@@ -1164,15 +1105,12 @@ namespace UnrealBuildTool
 							Writer.WriteLine(string.Format("\"{0}\" -> \"{1}\"", ObjectLeafFilename, SourceFile.AbsolutePath));
 						}
 					}
-					CompileAction.ProducedItems.Add(ObjectFile);
-					Result.ObjectFiles.Add(ObjectFile);
-					FileArguments.Add(String.Format("/Fo\"{0}\"", ObjectFile.AbsolutePath));
 
 					// Experimental: support for JSON output of timing data
 					if(Target.WindowsPlatform.Compiler == WindowsCompiler.Clang && Target.WindowsPlatform.bClangTimeTrace)
 					{
-						SharedArguments.Add("-Xclang -ftime-trace");
-						CompileAction.ProducedItems.Add(FileItem.GetItemByFileReference(ObjectFile.Location.ChangeExtension(".json")));
+						CompileAction.Arguments.Add("-Xclang -ftime-trace");
+						CompileAction.AdditionalProducedItems.Add(FileItem.GetItemByFileReference(ObjectFile.Location.ChangeExtension(".json")));
 					}
 				}
 
@@ -1192,7 +1130,7 @@ namespace UnrealBuildTool
 						PDBLocation = CompileEnvironment.PrecompiledHeaderFile.Location.ChangeExtension(".pdb");
 
 						// Enable synchronous file writes, since we'll be modifying the existing PDB
-						FileArguments.Add("/FS");
+						CompileAction.Arguments.Add("/FS");
 					}
 					else if (CompileEnvironment.PrecompiledHeaderAction == PrecompiledHeaderAction.Create)
 					{
@@ -1200,7 +1138,7 @@ namespace UnrealBuildTool
 						PDBLocation = FileReference.Combine(OutputDir, CompileEnvironment.PrecompiledHeaderIncludeFilename.GetFileName() + ".pdb");
 
 						// Enable synchronous file writes, since we'll be modifying the existing PDB
-						FileArguments.Add("/FS");
+						CompileAction.Arguments.Add("/FS");
 					}
 					else if (!bIsPlainCFile)
 					{
@@ -1214,11 +1152,7 @@ namespace UnrealBuildTool
 					}
 
 					// Specify the PDB file that the compiler should write to.
-					FileArguments.Add(String.Format("/Fd\"{0}\"", PDBLocation));
-
-					// Don't add the PDB as an output file because it's modified multiple times. This will break timestamp dependency tracking.
-					FileItem PDBFile = FileItem.GetItemByFileReference(PDBLocation);
-					Result.DebugDataFiles.Add(PDBFile);
+					CompileAction.Arguments.Add(String.Format("/Fd\"{0}\"", PDBLocation));
 
 					// Don't allow remote execution when PDB files are enabled; we need to modify the same files. XGE works around this by generating separate
 					// PDB files per agent, but this functionality is only available with the Visual C++ extension package (via the VCCompiler=true tool option).
@@ -1228,79 +1162,54 @@ namespace UnrealBuildTool
 				// Add C or C++ specific compiler arguments.
 				if (bIsPlainCFile)
 				{
-					AppendCLArguments_C(FileArguments);
+					AppendCLArguments_C(CompileAction.Arguments);
 				}
 				else
 				{
-					AppendCLArguments_CPP(CompileEnvironment, FileArguments);
+					AppendCLArguments_CPP(CompileEnvironment, CompileAction.Arguments);
 				}
 
-				CompileAction.PrerequisiteItems.AddRange(CompileEnvironment.ForceIncludeFiles);
-				CompileAction.PrerequisiteItems.AddRange(CompileEnvironment.AdditionalPrerequisites);
-
-				string[] AdditionalArguments = String.IsNullOrEmpty(CompileEnvironment.AdditionalArguments)? new string[0] : new string[] { CompileEnvironment.AdditionalArguments };
+				CompileAction.AdditionalPrerequisiteItems.AddRange(CompileEnvironment.AdditionalPrerequisites);
 
-				if (!ProjectFileGenerator.bGenerateProjectFiles
-					&& CompileAction.ProducedItems.Count > 0)
-				{
-					FileItem TargetFile = CompileAction.ProducedItems[0];
-					FileReference ResponseFileName = new FileReference(TargetFile.AbsolutePath + ".response");
-					FileItem ResponseFileItem = Graph.CreateIntermediateTextFile(ResponseFileName, SharedArguments.Concat(FileArguments).Concat(AdditionalArguments).Select(x => Utils.ExpandVariables(x)));
-					CompileAction.CommandArguments = " @\"" + ResponseFileName + "\"";
-					CompileAction.PrerequisiteItems.Add(ResponseFileItem);
-				}
-				else
+				if (!String.IsNullOrEmpty(CompileEnvironment.AdditionalArguments))
 				{
-					CompileAction.CommandArguments = String.Join(" ", SharedArguments.Concat(FileArguments).Concat(AdditionalArguments));
+					CompileAction.Arguments.Add(CompileEnvironment.AdditionalArguments);
 				}
 
 				if (SourceFile.HasExtension(".ixx"))
 				{
 					FileItem IfcDepsFile = FileItem.GetItemByFileReference(FileReference.Combine(OutputDir, SourceFile.Location.GetFileName() + ".md.json"));
 
-					Action CompileDepsAction = Graph.CreateAction(ActionType.GatherModuleDependencies);
-					CompileDepsAction.StatusDescription = SourceFile.Location.GetFileName();
-					CompileDepsAction.WorkingDirectory = CompileAction.WorkingDirectory;
-					CompileDepsAction.CommandPath = CompileAction.CommandPath;
-					CompileDepsAction.CommandArguments = String.Format("{0} /sourceDependencies:directives \"{1}\"", CompileAction.CommandArguments, IfcDepsFile.Location);
-					CompileDepsAction.PrerequisiteItems.Add(SourceFile);
-					CompileDepsAction.PrerequisiteItems.AddRange(CompileEnvironment.ForceIncludeFiles);
-					CompileDepsAction.PrerequisiteItems.AddRange(CompileEnvironment.AdditionalPrerequisites);
-					CompileDepsAction.ProducedItems.Add(IfcDepsFile);
-					CompileDepsAction.bShouldOutputStatusDescription = false;
+					VCCompileAction CompileDepsAction = new VCCompileAction(CompileAction);
+					CompileDepsAction.Arguments.Add(String.Format("/sourceDependencies:directives \"{0}\"", IfcDepsFile.Location));
+					CompileDepsAction.AdditionalPrerequisiteItems.Add(SourceFile);
+					CompileDepsAction.AdditionalPrerequisiteItems.AddRange(CompileEnvironment.ForceIncludeFiles);
+					CompileDepsAction.AdditionalPrerequisiteItems.AddRange(CompileEnvironment.AdditionalPrerequisites);
+					CompileDepsAction.AdditionalProducedItems.Add(IfcDepsFile);
 
 					FileItem IfcFile = FileItem.GetItemByFileReference(FileReference.Combine(GetModuleInterfaceDir(OutputDir), SourceFile.Location.ChangeExtension(".ifc").GetFileName()));
-					CompileAction.CommandArguments += String.Format(" /interface /ifcOutput \"{0}\"", IfcFile.Location);
-					CompileAction.PrerequisiteItems.Add(IfcDepsFile); // Force the dependencies file into the action graph
-					CompileAction.ProducedItems.Add(IfcFile);
+					CompileAction.Arguments.Add("/interface");
+					CompileAction.Arguments.Add(String.Format("/ifcOutput \"{0}\"", IfcFile.Location));
+					CompileAction.AdditionalPrerequisiteItems.Add(IfcDepsFile); // Force the dependencies file into the action graph
+					CompileAction.AdditionalProducedItems.Add(IfcFile);
 					CompileAction.DependencyListFile = IfcDepsFile;
 					CompileAction.CompiledModuleInterfaceFile = IfcFile;
-
-					Result.CompiledModuleInterfaces.Add(IfcFile);
 				}
 				else if (CompileEnvironment.bGenerateDependenciesFile)
 				{
-					GenerateDependenciesFile(Target, OutputDir, Result, SourceFile, CompileAction, Graph);
+					CompileAction.DependencyListFile = FileItem.GetItemByFileReference(FileReference.Combine(OutputDir, String.Format("{0}.txt", SourceFile.Location.GetFileName())));
+					CompileAction.bShowIncludes = Target.WindowsPlatform.bShowIncludes;
 				}
 
-				if (CompileEnvironment.PrecompiledHeaderAction == PrecompiledHeaderAction.Create)
-				{
-					Log.TraceVerbose("Creating PCH " + CompileEnvironment.PrecompiledHeaderIncludeFilename + ": \"" + CompileAction.CommandPath + "\"" + CompileAction.CommandArguments);
-				}
-				else
+				if (Target.bPrintToolChainTimingInfo || Target.WindowsPlatform.bCompilerTrace)
 				{
-					Log.TraceVerbose("   Compiling " + CompileAction.StatusDescription + ": \"" + CompileAction.CommandPath + "\"" + CompileAction.CommandArguments);
+					CompileAction.TimingFile = FileItem.GetItemByFileReference(FileReference.Combine(OutputDir, String.Format("{0}.timing", SourceFile.Location.GetFileName())));
+					GenerateParseTimingInfoAction(SourceFile, CompileAction.TimingFile, Graph);
 				}
 
-				if (Target.WindowsPlatform.Compiler == WindowsCompiler.Clang)
+				if (!ProjectFileGenerator.bGenerateProjectFiles)
 				{
-					// Clang doesn't print the file names by default, so we'll do it ourselves
-					CompileAction.bShouldOutputStatusDescription = true;
-				}
-				else
-				{
-					// VC++ always outputs the source file name being compiled, so we don't need to emit this ourselves
-					CompileAction.bShouldOutputStatusDescription = false;
+					CompileAction.WriteResponseFile(Graph);
 				}
 
 				// When compiling with SN-DBS, modules that contain a #import must be built locally
@@ -1309,18 +1218,24 @@ namespace UnrealBuildTool
 				{
 					CompileAction.bCanExecuteRemotelyWithSNDBS = false;
 				}
+
+				// Update the output
+				Graph.AddAction(CompileAction);
+				Actions.Add(CompileAction);
 			}
 
+			CPPOutput Result = new CPPOutput();
+			Result.ObjectFiles.AddRange(Actions.Select(x => x.ObjectFile).Where(x => x != null));
+			Result.CompiledModuleInterfaces.AddRange(Actions.Select(x => x.CompiledModuleInterfaceFile).Where(x => x != null));
+			Result.PrecompiledHeaderFile = Actions.Select(x => x.CreatePchFile).Where(x => x != null).FirstOrDefault();
 			return Result;
 		}
 
-		private Action GenerateParseTimingInfoAction(FileItem SourceFile, VCCompileAction CompileAction, CPPOutput Result, IActionGraphBuilder Graph)
+		private Action GenerateParseTimingInfoAction(FileItem SourceFile, FileItem TimingFile, IActionGraphBuilder Graph)
 		{
-			FileItem TimingJsonFile = FileItem.GetItemByPath(Path.ChangeExtension(CompileAction.TimingFile.AbsolutePath, ".cta"));
-			Result.DebugDataFiles.Add(TimingJsonFile);
-			Result.DebugDataFiles.Add(CompileAction.TimingFile);
+			FileItem TimingJsonFile = FileItem.GetItemByPath(Path.ChangeExtension(TimingFile.AbsolutePath, ".cta"));
 
-			string ParseTimingArguments = String.Format("-TimingFile=\"{0}\"", CompileAction.TimingFile);
+			string ParseTimingArguments = String.Format("-TimingFile=\"{0}\"", TimingFile);
 			if (Target.bParseTimingInfoForTracing)
 			{
 				ParseTimingArguments += " -Tracing";
@@ -1328,44 +1243,15 @@ namespace UnrealBuildTool
 
 			Action ParseTimingInfoAction = Graph.CreateRecursiveAction<ParseMsvcTimingInfoMode>(ActionType.ParseTimingInfo, ParseTimingArguments);
 			ParseTimingInfoAction.WorkingDirectory = UnrealBuildTool.EngineSourceDirectory;
-			ParseTimingInfoAction.StatusDescription = Path.GetFileName(CompileAction.TimingFile.AbsolutePath);
+			ParseTimingInfoAction.StatusDescription = Path.GetFileName(TimingFile.AbsolutePath);
 			ParseTimingInfoAction.bCanExecuteRemotely = true;
 			ParseTimingInfoAction.bCanExecuteRemotelyWithSNDBS = true;
 			ParseTimingInfoAction.PrerequisiteItems.Add(SourceFile);
-			ParseTimingInfoAction.PrerequisiteItems.Add(CompileAction.TimingFile);
+			ParseTimingInfoAction.PrerequisiteItems.Add(TimingFile);
 			ParseTimingInfoAction.ProducedItems.Add(TimingJsonFile);
 			return ParseTimingInfoAction;
 		}
 
-		private void GenerateDependenciesFile(ReadOnlyTargetRules Target, DirectoryReference OutputDir, CPPOutput Result, FileItem SourceFile, VCCompileAction CompileAction, IActionGraphBuilder Graph)
-		{
-			List<string> CommandArguments = new List<string>();
-			CompileAction.DependencyListFile = FileItem.GetItemByFileReference(FileReference.Combine(OutputDir, String.Format("{0}.txt", SourceFile.Location.GetFileName())));
-			CompileAction.ProducedItems.Add(CompileAction.DependencyListFile);
-			CommandArguments.Add(String.Format("-dependencies={0}", Utils.MakePathSafeToUseWithCommandLine(CompileAction.DependencyListFile.Location)));
-
-			if (Target.bPrintToolChainTimingInfo || Target.WindowsPlatform.bCompilerTrace)
-			{
-				CompileAction.TimingFile = FileItem.GetItemByFileReference(FileReference.Combine(OutputDir, String.Format("{0}.timing", SourceFile.Location.GetFileName())));
-				CompileAction.ProducedItems.Add(CompileAction.TimingFile);
-				GenerateParseTimingInfoAction(SourceFile, CompileAction, Result, Graph);
-				CommandArguments.Add(string.Format("-timing={0}", Utils.MakePathSafeToUseWithCommandLine(CompileAction.TimingFile.Location)));
-			}
-
-			if (Target.WindowsPlatform.bShowIncludes)
-			{
-				CommandArguments.Add("-showincludes");
-			}
-
-			CommandArguments.Add(String.Format("-compiler={0}", Utils.MakePathSafeToUseWithCommandLine(CompileAction.CommandPath)));
-			CommandArguments.Add("--");
-			CommandArguments.Add(Utils.MakePathSafeToUseWithCommandLine(CompileAction.CommandPath));
-			CommandArguments.Add(CompileAction.CommandArguments);
-			CommandArguments.Add("/showIncludes");
-			CompileAction.CommandArguments = string.Join(" ", CommandArguments);
-			CompileAction.CommandPath = FileReference.Combine(UnrealBuildTool.EngineDirectory, "Build", "Windows", "cl-filter", "cl-filter.exe");
-		}
-
 		public override void FinalizeOutput(ReadOnlyTargetRules Target, TargetMakefile Makefile)
 		{
 			if (Target.bPrintToolChainTimingInfo || Target.WindowsPlatform.bCompilerTrace)
diff --git a/Engine/Source/Programs/UnrealBuildTool/System/ActionGraphBuilder.cs b/Engine/Source/Programs/UnrealBuildTool/System/ActionGraphBuilder.cs
--- a/Engine/Source/Programs/UnrealBuildTool/System/ActionGraphBuilder.cs
+++ b/Engine/Source/Programs/UnrealBuildTool/System/ActionGraphBuilder.cs
@@ -27,7 +27,7 @@ namespace UnrealBuildTool
 		/// <param name="Location">Location of the response file</param>
 		/// <param name="Contents">Contents of the file</param>
 		/// <returns>New file item</returns>
-		FileItem CreateIntermediateTextFile(FileReference Location, string Contents);
+		void CreateIntermediateTextFile(FileItem Location, string Contents);
 
 		/// <summary>
 		/// Adds a file which is in the non-unity working set
@@ -79,10 +79,9 @@ namespace UnrealBuildTool
 		}
 
 		/// <inheritdoc/>
-		public virtual FileItem CreateIntermediateTextFile(FileReference Location, string Contents)
+		public virtual void CreateIntermediateTextFile(FileItem FileItem, string Contents)
 		{
-			Utils.WriteFileIfChanged(Location, Contents, StringComparison.OrdinalIgnoreCase);
-			return FileItem.GetItemByFileReference(Location);
+			Utils.WriteFileIfChanged(FileItem, Contents, StringComparison.OrdinalIgnoreCase);
 		}
 
 		/// <inheritdoc/>
@@ -142,9 +141,9 @@ namespace UnrealBuildTool
 		}
 
 		/// <inheritdoc/>
-		public virtual FileItem CreateIntermediateTextFile(FileReference Location, string Contents)
+		public virtual void CreateIntermediateTextFile(FileItem FileItem, string Contents)
 		{
-			return Inner.CreateIntermediateTextFile(Location, Contents);
+			Inner.CreateIntermediateTextFile(FileItem, Contents);
 		}
 
 		/// <inheritdoc/>
@@ -269,6 +268,33 @@ namespace UnrealBuildTool
 			return NewAction;
 		}
 
+		/// <summary>
+		/// Creates a text file with the given contents.  If the contents of the text file aren't changed, it won't write the new contents to
+		/// the file to avoid causing an action to be considered outdated.
+		/// </summary>
+		/// <param name="Graph">The action graph</param>
+		/// <param name="AbsolutePath">Path to the intermediate file to create</param>
+		/// <param name="Contents">Contents of the new file</param>
+		/// <returns>File item for the newly created file</returns>
+		public static FileItem CreateIntermediateTextFile(this IActionGraphBuilder Graph, FileReference AbsolutePath, string Contents)
+		{
+			FileItem FileItem = FileItem.GetItemByFileReference(AbsolutePath);
+			Graph.CreateIntermediateTextFile(FileItem, Contents);
+			return FileItem;
+		}
+
+		/// <summary>
+		/// Creates a text file with the given contents.  If the contents of the text file aren't changed, it won't write the new contents to
+		/// the file to avoid causing an action to be considered outdated.
+		/// </summary>
+		/// <param name="Graph">The action graph</param>
+		/// <param name="FileItem">Path to the intermediate file to create</param>
+		/// <param name="Contents">Contents of the new file</param>
+		public static void CreateIntermediateTextFile(this IActionGraphBuilder Graph, FileItem FileItem, IEnumerable<string> Contents)
+		{
+			Graph.CreateIntermediateTextFile(FileItem, string.Join(Environment.NewLine, Contents));
+		}
+
 		/// <summary>
 		/// Creates a text file with the given contents.  If the contents of the text file aren't changed, it won't write the new contents to
 		/// the file to avoid causing an action to be considered outdated.
diff --git a/Engine/Source/Programs/UnrealBuildTool/System/CppCompileEnvironment.cs b/Engine/Source/Programs/UnrealBuildTool/System/CppCompileEnvironment.cs
--- a/Engine/Source/Programs/UnrealBuildTool/System/CppCompileEnvironment.cs
+++ b/Engine/Source/Programs/UnrealBuildTool/System/CppCompileEnvironment.cs
@@ -78,7 +78,6 @@ namespace UnrealBuildTool
 	class CPPOutput
 	{
 		public List<FileItem> ObjectFiles = new List<FileItem>();
-		public List<FileItem> DebugDataFiles = new List<FileItem>();
 		public List<FileItem> CompiledModuleInterfaces = new List<FileItem>();
 		public List<FileItem> GeneratedHeaderFiles = new List<FileItem>();
 		public FileItem PrecompiledHeaderFile = null;
diff --git a/Engine/Source/Programs/UnrealBuildTool/System/TargetMakefile.cs b/Engine/Source/Programs/UnrealBuildTool/System/TargetMakefile.cs
--- a/Engine/Source/Programs/UnrealBuildTool/System/TargetMakefile.cs
+++ b/Engine/Source/Programs/UnrealBuildTool/System/TargetMakefile.cs
@@ -686,16 +686,14 @@ namespace UnrealBuildTool
 		}
 
 		/// <inheritdoc/>
-		public FileItem CreateIntermediateTextFile(FileReference Location, string Contents)
+		public void CreateIntermediateTextFile(FileItem FileItem, string Contents)
 		{
 			// Write the file
-			Utils.WriteFileIfChanged(Location, Contents, StringComparison.InvariantCultureIgnoreCase);
+			Utils.WriteFileIfChanged(FileItem.Location, Contents, StringComparison.InvariantCultureIgnoreCase);
 
 			// Reset the file info, in case it already knows about the old file
-			FileItem Item = FileItem.GetItemByFileReference(Location);
-			InternalDependencies.Add(Item);
-			Item.ResetCachedInfo();
-			return Item;
+			InternalDependencies.Add(FileItem);
+			FileItem.ResetCachedInfo();
 		}
 
 		/// <inheritdoc/>
diff --git a/Engine/Source/Programs/UnrealBuildTool/System/Utils.cs b/Engine/Source/Programs/UnrealBuildTool/System/Utils.cs
--- a/Engine/Source/Programs/UnrealBuildTool/System/Utils.cs
+++ b/Engine/Source/Programs/UnrealBuildTool/System/Utils.cs
@@ -1276,32 +1276,47 @@ namespace UnrealBuildTool
 		/// <param name="Location">Location of the file</param>
 		/// <param name="Contents">New contents of the file</param>
 		/// <param name="Comparison">The type of string comparison to use</param>
-		public static void WriteFileIfChanged(FileReference Location, string Contents, StringComparison Comparison)
+		internal static void WriteFileIfChanged(FileReference Location, string Contents, StringComparison Comparison)
+		{
+			FileItem FileItem = FileItem.GetItemByFileReference(Location);
+			WriteFileIfChanged(FileItem, Contents, Comparison);
+		}
+
+		/// <summary>
+		/// Writes a file if the contents have changed
+		/// </summary>
+		/// <param name="FileItem">Location of the file</param>
+		/// <param name="Contents">New contents of the file</param>
+		/// <param name="Comparison">The type of string comparison to use</param>
+		internal static void WriteFileIfChanged(FileItem FileItem, string Contents, StringComparison Comparison)
 		{
 			// Only write the file if its contents have changed.
-			if (!FileReference.Exists(Location))
+			FileReference Location = FileItem.Location;
+			if (!FileItem.Exists)
 			{
 				DirectoryReference.CreateDirectory(Location.Directory);
 				FileReference.WriteAllText(Location, Contents, GetEncodingForString(Contents));
+				FileItem.ResetCachedInfo();
 			}
 			else
 			{
-				string CurrentContents = Utils.ReadAllText(Location.FullName);
+				string CurrentContents = Utils.ReadAllText(FileItem.FullName);
 				if (!String.Equals(CurrentContents, Contents, Comparison))
 				{
-					FileReference BackupFile = new FileReference(Location.FullName + ".old");
+					FileReference BackupFile = new FileReference(FileItem.FullName + ".old");
 					try
 					{
-						Log.TraceLog("Updating {0}: contents have changed. Saving previous version to {1}.", Location, BackupFile);
+						Log.TraceLog("Updating {0}: contents have changed. Saving previous version to {1}.", FileItem, BackupFile);
 						FileReference.Delete(BackupFile);
 						FileReference.Move(Location, BackupFile);
 					}
 					catch (Exception Ex)
 					{
-						Log.TraceWarning("Unable to rename {0} to {1}", Location, BackupFile);
+						Log.TraceWarning("Unable to rename {0} to {1}", FileItem, BackupFile);
 						Log.TraceLog("{0}", ExceptionUtils.FormatExceptionDetails(Ex));
 					}
 					FileReference.WriteAllText(Location, Contents, GetEncodingForString(Contents));
+					FileItem.ResetCachedInfo();
 				}
 			}
 		}
