UBT: Remove support for outputting entire action graph for a target. This is not going to be valid with C++20 modules, since we require dependencies discovered at build time to feed back into the action graph execution.

#rnx

diff --git a/Engine/Source/Programs/DotNETCommon/DotNETUtilities/JsonWriter.cs b/Engine/Source/Programs/DotNETCommon/DotNETUtilities/JsonWriter.cs
--- a/Engine/Source/Programs/DotNETCommon/DotNETUtilities/JsonWriter.cs
+++ b/Engine/Source/Programs/DotNETCommon/DotNETUtilities/JsonWriter.cs
@@ -231,6 +231,20 @@ namespace Tools.DotNETCommon
 			WriteStringArrayField(Name, Values.Select(x => x.ToString()));
 		}
 
+		/// <summary>
+		/// Write a value with no field name, for the contents of an array
+		/// </summary>
+		/// <param name="Value">Value to write</param>
+		public void WriteValue(int Value)
+		{
+			WriteCommaNewline();
+
+			Writer.Write(Indent);
+			Writer.Write(Value);
+
+			bRequiresComma = true;
+		}
+
 		/// <summary>
 		/// Write a value with no field name, for the contents of an array
 		/// </summary>
diff --git a/Engine/Source/Programs/UnrealBuildTool/Configuration/TargetDescriptor.cs b/Engine/Source/Programs/UnrealBuildTool/Configuration/TargetDescriptor.cs
--- a/Engine/Source/Programs/UnrealBuildTool/Configuration/TargetDescriptor.cs
+++ b/Engine/Source/Programs/UnrealBuildTool/Configuration/TargetDescriptor.cs
@@ -52,12 +52,6 @@ namespace UnrealBuildTool
 		/// </summary>
 		public Dictionary<string, int> HotReloadModuleNameToSuffix = new Dictionary<string, int>(StringComparer.OrdinalIgnoreCase);
 
-		/// <summary>
-		/// Export the actions for the target to a file
-		/// </summary>
-		[CommandLine("-WriteActions=")]
-		public List<FileReference> WriteActionFiles = new List<FileReference>();
-
 		/// <summary>
 		/// Path to a file containing a list of modules that may be modified for live coding.
 		/// </summary>
diff --git a/Engine/Source/Programs/UnrealBuildTool/Modes/BuildMode.cs b/Engine/Source/Programs/UnrealBuildTool/Modes/BuildMode.cs
--- a/Engine/Source/Programs/UnrealBuildTool/Modes/BuildMode.cs
+++ b/Engine/Source/Programs/UnrealBuildTool/Modes/BuildMode.cs
@@ -285,17 +285,6 @@ namespace UnrealBuildTool
 		/// <returns>Result from the compilation</returns>
 		static void Build(TargetMakefile[] Makefiles, List<TargetDescriptor> TargetDescriptors, BuildConfiguration BuildConfiguration, ISourceFileWorkingSet WorkingSet, BuildOptions Options, FileReference WriteOutdatedActionsFile)
 		{
-			// Export the actions for each target
-			for (int TargetIdx = 0; TargetIdx < TargetDescriptors.Count; TargetIdx++)
-			{
-				TargetDescriptor TargetDescriptor = TargetDescriptors[TargetIdx];
-				foreach(FileReference WriteActionFile in TargetDescriptor.WriteActionFiles)
-				{
-					Log.TraceInformation("Writing actions to {0}", WriteActionFile);
-					ActionGraph.ExportJson(Makefiles[TargetIdx].Actions, WriteActionFile);
-				}
-			}
-
 			// Execute the build
 			if ((Options & BuildOptions.SkipBuild) == 0)
 			{
diff --git a/Engine/Source/Programs/UnrealBuildTool/System/Action.cs b/Engine/Source/Programs/UnrealBuildTool/System/Action.cs
--- a/Engine/Source/Programs/UnrealBuildTool/System/Action.cs
+++ b/Engine/Source/Programs/UnrealBuildTool/System/Action.cs
@@ -395,9 +395,11 @@ namespace UnrealBuildTool
 		/// Writes an action to a json file
 		/// </summary>
 		/// <param name="Action">The action to write</param>
+		/// <param name="LinkedActionToId">Map of action to unique id</param>
 		/// <param name="Writer">Writer to receive the output</param>
-		public static void ExportJson(this IAction Action, JsonWriter Writer)
+		public static void ExportJson(this LinkedAction Action, Dictionary<LinkedAction, int> LinkedActionToId, JsonWriter Writer)
 		{
+			Writer.WriteValue("Id", LinkedActionToId[Action]);
 			Writer.WriteEnumValue("Type", Action.ActionType);
 			Writer.WriteValue("WorkingDirectory", Action.WorkingDirectory.FullName);
 			Writer.WriteValue("CommandPath", Action.CommandPath.FullName);
@@ -410,10 +412,10 @@ namespace UnrealBuildTool
 			Writer.WriteValue("bShouldOutputStatusDescription", Action.bShouldOutputStatusDescription);
 			Writer.WriteValue("bProducesImportLibrary", Action.bProducesImportLibrary);
 
-			Writer.WriteArrayStart("PrerequisiteItems");
-			foreach (FileItem PrerequisiteItem in Action.PrerequisiteItems)
+			Writer.WriteArrayStart("PrerequisiteActions");
+			foreach (LinkedAction PrerequisiteAction in Action.PrerequisiteActions)
 			{
-				Writer.WriteValue(PrerequisiteItem.AbsolutePath);
+				Writer.WriteValue(LinkedActionToId[PrerequisiteAction]);
 			}
 			Writer.WriteArrayEnd();
 
@@ -499,11 +501,6 @@ namespace UnrealBuildTool
 		public bool bShouldOutputStatusDescription => Inner.bShouldOutputStatusDescription;
 		public bool bProducesImportLibrary => Inner.bProducesImportLibrary;
 
-		public void ExportJson(JsonWriter Writer)
-		{
-			Inner.ExportJson(Writer);
-		}
-
 		#endregion
 
 		/// <summary>
diff --git a/Engine/Source/Programs/UnrealBuildTool/System/ActionGraph.cs b/Engine/Source/Programs/UnrealBuildTool/System/ActionGraph.cs
--- a/Engine/Source/Programs/UnrealBuildTool/System/ActionGraph.cs
+++ b/Engine/Source/Programs/UnrealBuildTool/System/ActionGraph.cs
@@ -759,7 +759,7 @@ namespace UnrealBuildTool
 		/// </summary>
 		/// <param name="Actions">The actions to write</param>
 		/// <param name="OutputFile">Output file to write the actions to</param>
-		public static void ExportJson(IEnumerable<IAction> Actions, FileReference OutputFile)
+		public static void ExportJson(IEnumerable<LinkedAction> Actions, FileReference OutputFile)
 		{
 			DirectoryReference.CreateDirectory(OutputFile.Directory);
 			using (JsonWriter Writer = new JsonWriter(OutputFile))
@@ -776,11 +776,17 @@ namespace UnrealBuildTool
 				}
 				Writer.WriteObjectEnd();
 
+				Dictionary<LinkedAction, int> ActionToId = new Dictionary<LinkedAction, int>();
+				foreach (LinkedAction Action in Actions)
+				{
+					ActionToId[Action] = ActionToId.Count;
+				}
+
 				Writer.WriteArrayStart("Actions");
-				foreach (IAction Action in Actions)
+				foreach (LinkedAction Action in Actions)
 				{
 					Writer.WriteObjectStart();
-					Action.ExportJson(Writer);
+					Action.ExportJson(ActionToId, Writer);
 					Writer.WriteObjectEnd();
 				}
 				Writer.WriteArrayEnd();
