Changed the dsym generation to write a .sh file instead of running this rather complicated logic as a in line script. Resolves issues with escaping this complicated one line script in net core.

diff --git a/Engine/Source/Programs/UnrealBuildTool/Platform/Mac/MacToolChain.cs b/Engine/Source/Programs/UnrealBuildTool/Platform/Mac/MacToolChain.cs
--- a/Engine/Source/Programs/UnrealBuildTool/Platform/Mac/MacToolChain.cs
+++ b/Engine/Source/Programs/UnrealBuildTool/Platform/Mac/MacToolChain.cs
@@ -1179,11 +1179,22 @@ namespace UnrealBuildTool
 			// Note that the source and dest are switched from a copy command
 			string ExtraOptions;
 			string DsymutilPath = GetDsymutilPath(out ExtraOptions, bIsForLTOBuild: false);
-			GenDebugAction.CommandArguments = string.Format("-c 'rm -rf \"{2}\"; for i in {{1..30}}; do if [ -f \"{1}\" ] ; then break; else echo\"Waiting for {1} before generating dSYM file.\"; sleep 1; fi; done; \"{0}\" {3} -f \"{1}\" -o \"{2}\"'",
+
+			FileReference DsymShellFile = FileReference.Combine(LinkEnvironment.LocalShadowDirectory, Path.GetFileNameWithoutExtension(MachOBinary.FullName) + "GenerateDSym.sh");
+			FileItem DsymShellFileItem = FileItem.GetItemByFileReference(DsymShellFile);
+
+			string Contents = String.Format(
+				"#!/bin/sh\nrm -rf \"{2}\"\nfor i in {{1..30}}\ndo \nif [ -f \"{1}\" ] \nthen break\nelse \necho \"Waiting for {1} before generating dSYM file.\"\n\tsleep 1\nfi\ndone\n\"{0}\" {3} -f \"{1}\" -o \"{2}\"",
 				DsymutilPath,
 				MachOBinary.AbsolutePath,
 				OutputFile.AbsolutePath,
-				ExtraOptions);
+				ExtraOptions
+			);
+			Contents = Contents.Replace("\n", Environment.NewLine);
+
+			Utils.WriteFileIfChanged(DsymShellFile, Contents, StringComparison.InvariantCultureIgnoreCase);
+			GenDebugAction.CommandArguments = string.Format("-c 'chmod +x \"{0}\"; \"{0}\"'", DsymShellFileItem.AbsolutePath);
+
 			if (LinkEnvironment.bIsCrossReferenced)
 			{
 				GenDebugAction.PrerequisiteItems.Add(FixDylibOutputFile);
