UBT: Remove PVS studio dependencies on IAction.PrerequisiteItems, in preparation for them being computed by an async process.

#rnx

diff --git a/Engine/Source/Programs/UnrealBuildTool/Platform/Windows/PVSToolChain.cs b/Engine/Source/Programs/UnrealBuildTool/Platform/Windows/PVSToolChain.cs
--- a/Engine/Source/Programs/UnrealBuildTool/Platform/Windows/PVSToolChain.cs
+++ b/Engine/Source/Programs/UnrealBuildTool/Platform/Windows/PVSToolChain.cs
@@ -497,23 +497,23 @@ namespace UnrealBuildTool
 			// Run the source files through PVS-Studio
 			for(int Idx = 0; Idx < PreprocessActions.Count; Idx++)
 			{
-				IAction PreprocessAction = PreprocessActions[Idx];
-				if (PreprocessAction.ActionType != ActionType.Compile)
+				VCCompileAction PreprocessAction = PreprocessActions[Idx] as VCCompileAction;
+				if (PreprocessAction == null)
 				{
 					continue;
 				}
 
-				FileItem SourceFileItem = PreprocessAction.PrerequisiteItems.FirstOrDefault(x => x.HasExtension(".c") || x.HasExtension(".cc") || x.HasExtension(".cpp"));
+				FileItem SourceFileItem = PreprocessAction.SourceFile;
 				if (SourceFileItem == null)
 				{
-					Log.TraceWarning("Unable to find source file from command: {0} {1}", PreprocessAction.CommandArguments);
+					Log.TraceWarning("Unable to find source file from command producing: {0}", String.Join(", ", PreprocessActions[Idx].ProducedItems.Select(x => x.Location.GetFileName())));
 					continue;
 				}
 
-				FileItem PreprocessedFileItem = PreprocessAction.ProducedItems.FirstOrDefault(x => x.HasExtension(".i"));
+				FileItem PreprocessedFileItem = PreprocessAction.PreprocessedFile;
 				if (PreprocessedFileItem == null)
 				{
-					Log.TraceWarning("Unable to find preprocessed output file from command: {0} {1}", PreprocessAction.CommandArguments);
+					Log.TraceWarning("Unable to find preprocessed output file from {0}", SourceFileItem.Location.GetFileName());
 					continue;
 				}
 
@@ -585,7 +585,7 @@ namespace UnrealBuildTool
 				AnalyzeAction.StatusDescription = BaseFileName;
 				AnalyzeAction.WorkingDirectory = UnrealBuildTool.EngineSourceDirectory;
 				AnalyzeAction.CommandPath = AnalyzerFile;
-				AnalyzeAction.CommandArguments = String.Format("--cl-params \"{0}\" --source-file \"{1}\" --output-file \"{2}\" --cfg \"{3}\" --analysis-mode {4}", PreprocessAction.CommandArguments, SourceFileItem.AbsolutePath, OutputFileLocation, ConfigFileItem.AbsolutePath, (uint)Settings.ModeFlags);
+				AnalyzeAction.CommandArguments = String.Format("--cl-params \"{0}\" --source-file \"{1}\" --output-file \"{2}\" --cfg \"{3}\" --analysis-mode {4}", String.Join(" ", PreprocessAction.GetCompilerArguments()), SourceFileItem.AbsolutePath, OutputFileLocation, ConfigFileItem.AbsolutePath, (uint)Settings.ModeFlags);
 				if (LicenseFile != null)
 				{
 					AnalyzeAction.CommandArguments += String.Format(" --lic-file \"{0}\"", LicenseFile);
@@ -622,7 +622,7 @@ namespace UnrealBuildTool
 			List<FileReference> InputFiles = Makefile.OutputItems.Select(x => x.Location).Where(x => x.HasExtension(".pvslog")).ToList();
 
 			// Collect the prerequisite items off of the Compile action added in CompileCPPFiles so that in SingleFileCompile mode the PVSGather step is also not filtered out
-			List<FileItem> AnalyzeActionPrerequisiteItems = Makefile.Actions.Where(x => x.ActionType == ActionType.Compile).SelectMany(x => x.PrerequisiteItems).ToList();
+			List<FileItem> AnalyzeActionPrerequisiteItems = Makefile.Actions.SelectMany(x => x.ProducedItems).ToList();
 
 			FileItem InputFileListItem = Makefile.CreateIntermediateTextFile(OutputFile.ChangeExtension(".input"), InputFiles.Select(x => x.FullName));
 
diff --git a/Engine/Source/Programs/UnrealBuildTool/Platform/Windows/VCCompileAction.cs b/Engine/Source/Programs/UnrealBuildTool/Platform/Windows/VCCompileAction.cs
--- a/Engine/Source/Programs/UnrealBuildTool/Platform/Windows/VCCompileAction.cs
+++ b/Engine/Source/Programs/UnrealBuildTool/Platform/Windows/VCCompileAction.cs
@@ -339,11 +339,11 @@ namespace UnrealBuildTool
 		{
 			if (ResponseFile != null)
 			{
-				Graph.CreateIntermediateTextFile(ResponseFile, GetFullClArgumentList());
+				Graph.CreateIntermediateTextFile(ResponseFile, GetCompilerArguments());
 			}
 		}
 
-		List<string> GetFullClArgumentList()
+		public List<string> GetCompilerArguments()
 		{
 			List<string> Arguments = new List<string>();
 
diff --git a/Engine/Source/Programs/UnrealBuildTool/Platform/Windows/VCToolChain.cs b/Engine/Source/Programs/UnrealBuildTool/Platform/Windows/VCToolChain.cs
--- a/Engine/Source/Programs/UnrealBuildTool/Platform/Windows/VCToolChain.cs
+++ b/Engine/Source/Programs/UnrealBuildTool/Platform/Windows/VCToolChain.cs
@@ -559,7 +559,7 @@ namespace UnrealBuildTool
 				Arguments.Add("/wd4463"); // 4463 - overflow; assigning 1 to bit-field that can only hold values from -1 to 0
 			}
 
-			if (CompileEnvironment.bEnableUndefinedIdentifierWarnings)
+			if (CompileEnvironment.bEnableUndefinedIdentifierWarnings && !CompileEnvironment.bPreprocessOnly)
 			{
 				if (CompileEnvironment.bUndefinedIdentifierWarningsAsErrors)
 				{
@@ -1225,7 +1225,7 @@ namespace UnrealBuildTool
 			}
 
 			CPPOutput Result = new CPPOutput();
-			Result.ObjectFiles.AddRange(Actions.Select(x => x.ObjectFile).Where(x => x != null));
+			Result.ObjectFiles.AddRange(Actions.Select(x => x.ObjectFile ?? x.PreprocessedFile).Where(x => x != null));
 			Result.CompiledModuleInterfaces.AddRange(Actions.Select(x => x.CompiledModuleInterfaceFile).Where(x => x != null));
 			Result.PrecompiledHeaderFile = Actions.Select(x => x.CreatePchFile).Where(x => x != null).FirstOrDefault();
 			return Result;
