UBT: Fix PVS studio static analysis steps not preprocessing source files before analysis.

#rnx

diff --git a/Engine/Source/Programs/UnrealBuildTool/Platform/Windows/PVSToolChain.cs b/Engine/Source/Programs/UnrealBuildTool/Platform/Windows/PVSToolChain.cs
--- a/Engine/Source/Programs/UnrealBuildTool/Platform/Windows/PVSToolChain.cs
+++ b/Engine/Source/Programs/UnrealBuildTool/Platform/Windows/PVSToolChain.cs
@@ -454,6 +454,8 @@ namespace UnrealBuildTool
 
 			public override void AddAction(IAction Action)
 			{
+				base.AddAction(Action);
+
 				Actions.Add(Action);
 			}
 		}
