UBT: Remove dependencies on PrerequisiteItems from ParallelExecutor.

#rnx

diff --git a/Engine/Source/Programs/UnrealBuildTool/Executors/ParallelExecutor.cs b/Engine/Source/Programs/UnrealBuildTool/Executors/ParallelExecutor.cs
--- a/Engine/Source/Programs/UnrealBuildTool/Executors/ParallelExecutor.cs
+++ b/Engine/Source/Programs/UnrealBuildTool/Executors/ParallelExecutor.cs
@@ -116,23 +116,14 @@ namespace UnrealBuildTool
 				Actions.Add(Action);
 			}
 
-			// Build a map of items to their producing actions
-			Dictionary<FileItem, BuildAction> FileToProducingAction = new Dictionary<FileItem, BuildAction>();
-			foreach(BuildAction Action in Actions)
-			{
-				foreach(FileItem ProducedItem in Action.Inner.ProducedItems)
-				{
-					FileToProducingAction[ProducedItem] = Action;
-				}
-			}
-
 			// Update all the actions with all their dependencies
-			foreach(BuildAction Action in Actions)
+			Dictionary<LinkedAction, BuildAction> LinkedActionToBuildAction = Actions.ToDictionary(x => x.Inner, x => x);
+			foreach (BuildAction Action in Actions)
 			{
-				foreach(FileItem PrerequisiteItem in Action.Inner.PrerequisiteItems)
+				foreach(LinkedAction PrerequisiteAction in Action.Inner.PrerequisiteActions)
 				{
 					BuildAction Dependency;
-					if(FileToProducingAction.TryGetValue(PrerequisiteItem, out Dependency))
+					if(LinkedActionToBuildAction.TryGetValue(PrerequisiteAction, out Dependency))
 					{
 						Action.Dependencies.Add(Dependency);
 						Dependency.Dependants.Add(Action);
