UBT: Rename IAction to IExternalAction, to distinguish it from actions which only exist within UBT.

#rnx

diff --git a/Engine/Source/Programs/UnrealBuildTool/Platform/Windows/PVSToolChain.cs b/Engine/Source/Programs/UnrealBuildTool/Platform/Windows/PVSToolChain.cs
--- a/Engine/Source/Programs/UnrealBuildTool/Platform/Windows/PVSToolChain.cs
+++ b/Engine/Source/Programs/UnrealBuildTool/Platform/Windows/PVSToolChain.cs
@@ -444,15 +444,15 @@ namespace UnrealBuildTool
 
 		class ActionGraphCapture : ForwardingActionGraphBuilder
 		{
-			List<IAction> Actions;
+			List<IExternalAction> Actions;
 
-			public ActionGraphCapture(IActionGraphBuilder Inner, List<IAction> Actions)
+			public ActionGraphCapture(IActionGraphBuilder Inner, List<IExternalAction> Actions)
 				: base(Inner)
 			{
 				this.Actions = Actions;
 			}
 
-			public override void AddAction(IAction Action)
+			public override void AddAction(IExternalAction Action)
 			{
 				base.AddAction(Action);
 
@@ -493,7 +493,7 @@ namespace UnrealBuildTool
 			PreprocessCompileEnvironment.bEnableUndefinedIdentifierWarnings = false; // Not sure why THIRD_PARTY_INCLUDES_START doesn't pick this up; the _Pragma appears in the preprocessed output. Perhaps in preprocess-only mode the compiler doesn't respect these?
 			PreprocessCompileEnvironment.Definitions.Add("PVS_STUDIO");
 
-			List<IAction> PreprocessActions = new List<IAction>();
+			List<IExternalAction> PreprocessActions = new List<IExternalAction>();
 			CPPOutput Result = InnerToolChain.CompileCPPFiles(PreprocessCompileEnvironment, InputFiles, OutputDir, ModuleName, new ActionGraphCapture(Graph, PreprocessActions));
 
 			// Run the source files through PVS-Studio
diff --git a/Engine/Source/Programs/UnrealBuildTool/Platform/Windows/VCCompileAction.cs b/Engine/Source/Programs/UnrealBuildTool/Platform/Windows/VCCompileAction.cs
--- a/Engine/Source/Programs/UnrealBuildTool/Platform/Windows/VCCompileAction.cs
+++ b/Engine/Source/Programs/UnrealBuildTool/Platform/Windows/VCCompileAction.cs
@@ -9,7 +9,7 @@ namespace UnrealBuildTool
 	/// <summary>
 	/// Interface for an action that compiles C++ source code
 	/// </summary>
-	interface ICppCompileAction : IAction
+	interface ICppCompileAction : IExternalAction
 	{
 		/// <summary>
 		/// Path to the compiled module interface file
@@ -139,17 +139,17 @@ namespace UnrealBuildTool
 
 		#region Implementation of IAction
 
-		ActionType IAction.ActionType => ActionType.Compile;
-		IEnumerable<FileItem> IAction.DeleteItems => DeleteItems;
+		ActionType IExternalAction.ActionType => ActionType.Compile;
+		IEnumerable<FileItem> IExternalAction.DeleteItems => DeleteItems;
 		public DirectoryReference WorkingDirectory => UnrealBuildTool.EngineSourceDirectory;
-		string IAction.CommandDescription => "Compile";
-		bool IAction.bIsGCCCompiler => false;
-		bool IAction.bProducesImportLibrary => false;
-		string IAction.StatusDescription => (SourceFile == null) ? "Compiling" : SourceFile.Location.GetFileName();
-		bool IAction.bShouldOutputStatusDescription => CompilerType == WindowsCompiler.Clang;
+		string IExternalAction.CommandDescription => "Compile";
+		bool IExternalAction.bIsGCCCompiler => false;
+		bool IExternalAction.bProducesImportLibrary => false;
+		string IExternalAction.StatusDescription => (SourceFile == null) ? "Compiling" : SourceFile.Location.GetFileName();
+		bool IExternalAction.bShouldOutputStatusDescription => CompilerType == WindowsCompiler.Clang;
 
 		/// <inheritdoc/>
-		IEnumerable<FileItem> IAction.PrerequisiteItems
+		IEnumerable<FileItem> IExternalAction.PrerequisiteItems
 		{
 			get
 			{
@@ -169,7 +169,7 @@ namespace UnrealBuildTool
 		}
 
 		/// <inheritdoc/>
-		IEnumerable<FileItem> IAction.ProducedItems
+		IEnumerable<FileItem> IExternalAction.ProducedItems
 		{
 			get
 			{
@@ -201,7 +201,7 @@ namespace UnrealBuildTool
 		}
 
 		/// <inheritdoc/>
-		FileReference IAction.CommandPath
+		FileReference IExternalAction.CommandPath
 		{
 			get
 			{
@@ -217,7 +217,7 @@ namespace UnrealBuildTool
 		}
 
 		/// <inheritdoc/>
-		string IAction.CommandArguments
+		string IExternalAction.CommandArguments
 		{
 			get
 			{
diff --git a/Engine/Source/Programs/UnrealBuildTool/Platform/Windows/VCToolChain.cs b/Engine/Source/Programs/UnrealBuildTool/Platform/Windows/VCToolChain.cs
--- a/Engine/Source/Programs/UnrealBuildTool/Platform/Windows/VCToolChain.cs
+++ b/Engine/Source/Programs/UnrealBuildTool/Platform/Windows/VCToolChain.cs
@@ -1256,7 +1256,7 @@ namespace UnrealBuildTool
 		{
 			if (Target.bPrintToolChainTimingInfo || Target.WindowsPlatform.bCompilerTrace)
 			{
-				List<IAction> ParseTimingActions = Makefile.Actions.Where(x => x.ActionType == ActionType.ParseTimingInfo).ToList();
+				List<IExternalAction> ParseTimingActions = Makefile.Actions.Where(x => x.ActionType == ActionType.ParseTimingInfo).ToList();
 				List<FileItem> TimingJsonFiles = ParseTimingActions.SelectMany(a => a.ProducedItems.Where(i => i.HasExtension(".cta"))).ToList();
 				Makefile.OutputItems.AddRange(TimingJsonFiles);
 
diff --git a/Engine/Source/Programs/UnrealBuildTool/System/Action.cs b/Engine/Source/Programs/UnrealBuildTool/System/Action.cs
--- a/Engine/Source/Programs/UnrealBuildTool/System/Action.cs
+++ b/Engine/Source/Programs/UnrealBuildTool/System/Action.cs
@@ -37,7 +37,7 @@ namespace UnrealBuildTool
 		ParseTimingInfo,
 	}
 
-	interface IAction
+	interface IExternalAction
 	{
 		/// <summary>
 		/// The type of this action (for debugging purposes).
@@ -119,7 +119,7 @@ namespace UnrealBuildTool
 	/// <summary>
 	/// A build action.
 	/// </summary>
-	class Action : IAction
+	class Action : IExternalAction
 	{
 		///
 		/// Preparation and Assembly (serialized)
@@ -201,9 +201,9 @@ namespace UnrealBuildTool
 		/// </summary>
 		public bool bProducesImportLibrary { get; set; } = false;
 
-		IEnumerable<FileItem> IAction.PrerequisiteItems => PrerequisiteItems;
-		IEnumerable<FileItem> IAction.ProducedItems => ProducedItems;
-		IEnumerable<FileItem> IAction.DeleteItems => DeleteItems;
+		IEnumerable<FileItem> IExternalAction.PrerequisiteItems => PrerequisiteItems;
+		IEnumerable<FileItem> IExternalAction.ProducedItems => ProducedItems;
+		IEnumerable<FileItem> IExternalAction.DeleteItems => DeleteItems;
 
 		public Action(ActionType InActionType)
 		{
@@ -216,7 +216,7 @@ namespace UnrealBuildTool
 			}
 		}
 
-		public Action(IAction InOther)
+		public Action(IExternalAction InOther)
 		{
 			ActionType = InOther.ActionType;
 			PrerequisiteItems = new List<FileItem>(InOther.PrerequisiteItems);
@@ -461,12 +461,12 @@ namespace UnrealBuildTool
 	/// <summary>
 	/// Information about an action queued to be executed
 	/// </summary>
-	class LinkedAction : IAction
+	class LinkedAction : IExternalAction
 	{
 		/// <summary>
 		/// The inner action instance
 		/// </summary>
-		public IAction Inner;
+		public IExternalAction Inner;
 
 		/// <summary>
 		/// Set of other actions that this action depends on. This set is built when the action graph is linked.
@@ -507,7 +507,7 @@ namespace UnrealBuildTool
 		/// Constructor
 		/// </summary>
 		/// <param name="Inner">The inner action instance</param>
-		public LinkedAction(IAction Inner)
+		public LinkedAction(IExternalAction Inner)
 		{
 			this.Inner = Inner;
 		}
diff --git a/Engine/Source/Programs/UnrealBuildTool/System/ActionGraph.cs b/Engine/Source/Programs/UnrealBuildTool/System/ActionGraph.cs
--- a/Engine/Source/Programs/UnrealBuildTool/System/ActionGraph.cs
+++ b/Engine/Source/Programs/UnrealBuildTool/System/ActionGraph.cs
@@ -55,16 +55,16 @@ namespace UnrealBuildTool
 		/// Checks a set of actions for conflicts (ie. different actions producing the same output items)
 		/// </summary>
 		/// <param name="Actions">The set of actions to check</param>
-		public static void CheckForConflicts(IEnumerable<IAction> Actions)
+		public static void CheckForConflicts(IEnumerable<IExternalAction> Actions)
 		{
 			bool bResult = true;
 
-			Dictionary<FileItem, IAction> ItemToProducingAction = new Dictionary<FileItem, IAction>();
-			foreach(IAction Action in Actions)
+			Dictionary<FileItem, IExternalAction> ItemToProducingAction = new Dictionary<FileItem, IExternalAction>();
+			foreach(IExternalAction Action in Actions)
 			{
 				foreach(FileItem ProducedItem in Action.ProducedItems)
 				{
-					IAction ExistingAction;
+					IExternalAction ExistingAction;
 					if(ItemToProducingAction.TryGetValue(ProducedItem, out ExistingAction))
 					{
 						bResult &= CheckForConflicts(ExistingAction, Action);
@@ -88,7 +88,7 @@ namespace UnrealBuildTool
 		/// <param name="A">The first action</param>
 		/// <param name="B">The second action</param>
 		/// <returns>True if any conflicts were found, false otherwise.</returns>
-		public static bool CheckForConflicts(IAction A, IAction B)
+		public static bool CheckForConflicts(IExternalAction A, IExternalAction B)
 		{
 			bool bResult = true;
 			if (A.ActionType != B.ActionType)
@@ -136,7 +136,7 @@ namespace UnrealBuildTool
 		/// <param name="Description">Description of the difference</param>
 		/// <param name="OldValue">Previous value for the field</param>
 		/// <param name="NewValue">Conflicting value for the field</param>
-		static void LogConflict(IAction Action, string Description, string OldValue, string NewValue)
+		static void LogConflict(IExternalAction Action, string Description, string OldValue, string NewValue)
 		{
 			Log.TraceError("Unable to merge actions producing {0}: {1}", Action.ProducedItems.First().Location.GetFileName(), Description);
 			Log.TraceLog("  Previous: {0}", OldValue);
@@ -164,7 +164,7 @@ namespace UnrealBuildTool
 		/// </summary>
 		/// <param name="BuildConfiguration">The build configuration</param>
 		/// <param name="Actions">List of actions in the graph</param>
-		public static void CheckPathLengths(BuildConfiguration BuildConfiguration, IEnumerable<IAction> Actions)
+		public static void CheckPathLengths(BuildConfiguration BuildConfiguration, IEnumerable<IExternalAction> Actions)
 		{
 			if (BuildHostPlatform.Current.Platform == UnrealTargetPlatform.Win64)
 			{
@@ -172,7 +172,7 @@ namespace UnrealBuildTool
 
 				List<FileReference> FailPaths = new List<FileReference>();
 				List<FileReference> WarnPaths = new List<FileReference>();
-				foreach (IAction Action in Actions)
+				foreach (IExternalAction Action in Actions)
 				{
 					foreach (FileItem PrerequisiteItem in Action.PrerequisiteItems)
 					{
diff --git a/Engine/Source/Programs/UnrealBuildTool/System/ActionGraphBuilder.cs b/Engine/Source/Programs/UnrealBuildTool/System/ActionGraphBuilder.cs
--- a/Engine/Source/Programs/UnrealBuildTool/System/ActionGraphBuilder.cs
+++ b/Engine/Source/Programs/UnrealBuildTool/System/ActionGraphBuilder.cs
@@ -19,7 +19,7 @@ namespace UnrealBuildTool
 		/// Adds an action to this graph
 		/// </summary>
 		/// <param name="Action">Action to add</param>
-		void AddAction(IAction Action);
+		void AddAction(IExternalAction Action);
 
 		/// <summary>
 		/// Creates a response file for use in the action graph
@@ -74,7 +74,7 @@ namespace UnrealBuildTool
 	class NullActionGraphBuilder : IActionGraphBuilder
 	{
 		/// <inheritdoc/>
-		public void AddAction(IAction Action)
+		public void AddAction(IExternalAction Action)
 		{
 		}
 
@@ -135,7 +135,7 @@ namespace UnrealBuildTool
 		}
 
 		/// <inheritdoc/>
-		public virtual void AddAction(IAction Action)
+		public virtual void AddAction(IExternalAction Action)
 		{
 			Inner.AddAction(Action);
 		}
diff --git a/Engine/Source/Programs/UnrealBuildTool/System/ActionSerializer.cs b/Engine/Source/Programs/UnrealBuildTool/System/ActionSerializer.cs
--- a/Engine/Source/Programs/UnrealBuildTool/System/ActionSerializer.cs
+++ b/Engine/Source/Programs/UnrealBuildTool/System/ActionSerializer.cs
@@ -25,30 +25,30 @@ namespace UnrealBuildTool
 		/// </summary>
 		/// <param name="Reader">Reader for the action</param>
 		/// <returns>New action</returns>
-		IAction Read(BinaryArchiveReader Reader);
+		IExternalAction Read(BinaryArchiveReader Reader);
 
 		/// <summary>
 		/// Writes an action to an archive
 		/// </summary>
 		/// <param name="Writer">Writer for the archive</param>
 		/// <param name="Action">The action to write</param>
-		void Write(BinaryArchiveWriter Writer, IAction Action);
+		void Write(BinaryArchiveWriter Writer, IExternalAction Action);
 	}
 
 	/// <summary>
 	/// Generic base class for an action serializer
 	/// </summary>
 	/// <typeparam name="TAction"></typeparam>
-	abstract class ActionSerializerBase<TAction> : IActionSerializer where TAction : IAction
+	abstract class ActionSerializerBase<TAction> : IActionSerializer where TAction : IExternalAction
 	{
 		/// <inheritdoc/>
 		public Type Type => typeof(TAction);
 
 		/// <inheritdoc/>
-		IAction IActionSerializer.Read(BinaryArchiveReader Reader) => Read(Reader);
+		IExternalAction IActionSerializer.Read(BinaryArchiveReader Reader) => Read(Reader);
 
 		/// <inheritdoc/>
-		void IActionSerializer.Write(BinaryArchiveWriter Writer, IAction Action) => Write(Writer, (TAction)Action);
+		void IActionSerializer.Write(BinaryArchiveWriter Writer, IExternalAction Action) => Write(Writer, (TAction)Action);
 
 		/// <summary>
 		/// Read the action from an archive
@@ -106,7 +106,7 @@ namespace UnrealBuildTool
 		/// </summary>
 		/// <param name="Reader">Reader to deserialize from</param>
 		/// <returns>New action</returns>
-		public static IAction ReadAction(this BinaryArchiveReader Reader)
+		public static IExternalAction ReadAction(this BinaryArchiveReader Reader)
 		{
 			IActionSerializer Serializer = Reader.ReadObjectReference(() => ReadSerializer(Reader));
 			return Serializer.Read(Reader);
@@ -135,7 +135,7 @@ namespace UnrealBuildTool
 		/// </summary>
 		/// <param name="Writer">Writer to serialize the action to</param>
 		/// <param name="Action">Action to serialize</param>
-		public static void WriteAction(this BinaryArchiveWriter Writer, IAction Action)
+		public static void WriteAction(this BinaryArchiveWriter Writer, IExternalAction Action)
 		{
 			Type Type = Action.GetType();
 
diff --git a/Engine/Source/Programs/UnrealBuildTool/System/HotReload.cs b/Engine/Source/Programs/UnrealBuildTool/System/HotReload.cs
--- a/Engine/Source/Programs/UnrealBuildTool/System/HotReload.cs
+++ b/Engine/Source/Programs/UnrealBuildTool/System/HotReload.cs
@@ -1193,11 +1193,11 @@ namespace UnrealBuildTool
 		/// <param name="ManifestFile">File to write to</param>
 		/// <param name="Actions">List of actions that are part of the graph</param>
 		/// <param name="OriginalFileToPatchedFile">Map of original object files to patched object files</param>
-		public static void WriteLiveCodingManifest(FileReference ManifestFile, List<IAction> Actions, Dictionary<FileReference, FileReference> OriginalFileToPatchedFile)
+		public static void WriteLiveCodingManifest(FileReference ManifestFile, List<IExternalAction> Actions, Dictionary<FileReference, FileReference> OriginalFileToPatchedFile)
 		{
 			// Find all the output object files
 			HashSet<FileItem> ObjectFiles = new HashSet<FileItem>();
-			foreach(IAction Action in Actions)
+			foreach(IExternalAction Action in Actions)
 			{
 				if(Action.ActionType == ActionType.Compile)
 				{
@@ -1210,7 +1210,7 @@ namespace UnrealBuildTool
 			{
 				Writer.WriteObjectStart();
 
-				IAction LinkAction = Actions.FirstOrDefault(x => x.ActionType == ActionType.Link && x.ProducedItems.Any(y => y.HasExtension(".exe") || y.HasExtension(".dll")));
+				IExternalAction LinkAction = Actions.FirstOrDefault(x => x.ActionType == ActionType.Link && x.ProducedItems.Any(y => y.HasExtension(".exe") || y.HasExtension(".dll")));
 				if(LinkAction != null)
 				{
 					FileReference LinkerPath = LinkAction.CommandPath;
@@ -1237,7 +1237,7 @@ namespace UnrealBuildTool
 				Writer.WriteObjectEnd();
 
 				Writer.WriteArrayStart("Modules");
-				foreach(IAction Action in Actions)
+				foreach(IExternalAction Action in Actions)
 				{
 					if(Action.ActionType == ActionType.Link)
 					{
diff --git a/Engine/Source/Programs/UnrealBuildTool/System/TargetMakefile.cs b/Engine/Source/Programs/UnrealBuildTool/System/TargetMakefile.cs
--- a/Engine/Source/Programs/UnrealBuildTool/System/TargetMakefile.cs
+++ b/Engine/Source/Programs/UnrealBuildTool/System/TargetMakefile.cs
@@ -96,7 +96,7 @@ namespace UnrealBuildTool
 		/// <summary>
 		/// Every action in the action graph
 		/// </summary>
-		public List<IAction> Actions;
+		public List<IExternalAction> Actions;
 
 		/// <summary>
 		/// Environment variables that we'll need in order to invoke the platform's compiler and linker
@@ -193,7 +193,7 @@ namespace UnrealBuildTool
 			this.ConfigValueTracker = ConfigValueTracker;
 			this.bDeployAfterCompile = bDeployAfterCompile;
 			this.bHasProjectScriptPlugin = bHasProjectScriptPlugin;
-			this.Actions = new List<IAction>();
+			this.Actions = new List<IExternalAction>();
 			this.OutputItems = new List<FileItem>();
 			this.ModuleNameToOutputItems = new Dictionary<string, FileItem[]>(StringComparer.OrdinalIgnoreCase);
 			this.HotReloadModuleNames = new HashSet<string>(StringComparer.OrdinalIgnoreCase);
@@ -680,7 +680,7 @@ namespace UnrealBuildTool
 		}
 
 		/// <inheritdoc/>
-		public void AddAction(IAction Action)
+		public void AddAction(IExternalAction Action)
 		{
 			Actions.Add(Action);
 		}
