- Fixed debugging AutomationTool in Visual Studio with netCore changes

diff --git a/Engine/Source/Programs/UnrealBuildTool/ProjectFiles/VisualStudio/VCProjectFileGenerator.cs b/Engine/Source/Programs/UnrealBuildTool/ProjectFiles/VisualStudio/VCProjectFileGenerator.cs
--- a/Engine/Source/Programs/UnrealBuildTool/ProjectFiles/VisualStudio/VCProjectFileGenerator.cs
+++ b/Engine/Source/Programs/UnrealBuildTool/ProjectFiles/VisualStudio/VCProjectFileGenerator.cs
@@ -472,32 +472,9 @@ namespace UnrealBuildTool
 							)
 						)
 					)
-				).Save(FileReference.Combine(AutomationToolDir, "AutomationToolCore.csproj.References").FullName);
-			}
-	
-			// and a copy in the old traditional msbuild format
-			if (AutomationProjectFiles.Any())
-			{
-				XNamespace NS = XNamespace.Get("http://schemas.microsoft.com/developer/msbuild/2003");
-
-				DirectoryReference AutomationToolDir = DirectoryReference.Combine(UnrealBuildTool.EngineSourceDirectory, "Programs", "AutomationTool");
-				new XDocument(
-					new XElement(NS + "Project",
-						new XAttribute("ToolsVersion", VCProjectFileGenerator.GetProjectFileToolVersionString(Settings.ProjectFileFormat)),
-						new XAttribute("DefaultTargets", "Build"),
-						new XElement(NS + "ItemGroup",
-							from AutomationProject in AutomationProjectFiles
-							select new XElement(NS + "ProjectReference",
-								new XAttribute("Include", AutomationProject.ProjectFilePath.MakeRelativeTo(AutomationToolDir)),
-								new XElement(NS + "Project", (AutomationProject as VCSharpProjectFile).ProjectGUID.ToString("B")),
-								new XElement(NS + "Name", AutomationProject.ProjectFilePath.GetFileNameWithoutExtension()),
-								new XElement(NS + "Private", "false")
-							)
-						)
-					)
 				).Save(FileReference.Combine(AutomationToolDir, "AutomationTool.csproj.References").FullName);
 			}
-
+	
 			return true;
 		}
 
