UBT can now filter out any project which requires dotnet support from the solution, which we do for VS2017 as it does not support Dotnet Core 3+

diff --git a/Engine/Source/Programs/UnrealBuildTool/ProjectFiles/ProjectFileGenerator.cs b/Engine/Source/Programs/UnrealBuildTool/ProjectFiles/ProjectFileGenerator.cs
--- a/Engine/Source/Programs/UnrealBuildTool/ProjectFiles/ProjectFileGenerator.cs
+++ b/Engine/Source/Programs/UnrealBuildTool/ProjectFiles/ProjectFileGenerator.cs
@@ -408,6 +408,9 @@ namespace UnrealBuildTool
 				if (FileReference.Exists(ProjectFile))
 				{
 					VCSharpProjectFile Project = new VCSharpProjectFile(ProjectFile);
+					if (Project.IsDotNETCoreProject() && !SupportsDotnetCoreProjects())
+						continue;
+
 					Project.ShouldBuildForAllSolutionTargets = false;//true;
 					AddExistingProjectFile(Project, bForceDevelopmentConfiguration: true);
 					AutomationProjectFiles.Add( Project );
@@ -422,6 +425,11 @@ namespace UnrealBuildTool
 			}
 		}
 
+		protected virtual bool SupportsDotnetCoreProjects()
+		{
+			return true;
+		}
+
 		/// <summary>
 		/// Creates a .props file next to each automation project which specifies the path to the engine directory
 		/// </summary>
@@ -470,6 +478,9 @@ namespace UnrealBuildTool
 						{
 							VCSharpProjectFile Project = new VCSharpProjectFile(FoundProject);
 
+							if (Project.IsDotNETCoreProject() && !SupportsDotnetCoreProjects())
+								continue;
+
 							Project.ShouldBuildForAllSolutionTargets = true;
 							Project.ShouldBuildByDefaultForSolutionTargets = true;
 
@@ -895,7 +906,9 @@ namespace UnrealBuildTool
 					AddUnrealBuildToolProject( ProgramsFolder );
 
 					// Add AutomationTool to the master project
-					ProgramsFolder.ChildProjects.Add(AddSimpleCSharpProject("AutomationTool", bShouldBuildForAllSolutionTargets: true, bForceDevelopmentConfiguration: true));
+					VCSharpProjectFile AutomationToolProject = AddSimpleCSharpProject("AutomationTool", bShouldBuildForAllSolutionTargets: true, bForceDevelopmentConfiguration: true);
+					if (AutomationToolProject != null)
+						ProgramsFolder.ChildProjects.Add(AutomationToolProject);
 
 					// Add automation.csproj files to the master project
 					AddAutomationModules(AllGameProjects, ProgramsFolder);
@@ -1577,6 +1590,10 @@ namespace UnrealBuildTool
 		/// </summary>
 		private void AddUnrealBuildToolProject(MasterProjectFolder ProgramsFolder)
 		{
+			// UBT is a dotnet core project so we can not add it if the project generator does not support dotnet core projects
+			if (!SupportsDotnetCoreProjects())
+				return;
+
 			List<string> ProjectDirectoryNames = new List<string>();
 			ProjectDirectoryNames.Add("UnrealBuildTool");
 
@@ -1625,6 +1642,10 @@ namespace UnrealBuildTool
 			if( Info.Exists )
 			{
 				Project = new VCSharpProjectFile(ProjectFileName);
+
+				if (Project.IsDotNETCoreProject() && !SupportsDotnetCoreProjects())
+					return null;
+
 				Project.ShouldBuildForAllSolutionTargets = bShouldBuildForAllSolutionTargets;
 				Project.ShouldBuildByDefaultForSolutionTargets = bShouldBuildByDefaultForSolutionTargets;
 				AddExistingProjectFile(Project, bForceDevelopmentConfiguration: bForceDevelopmentConfiguration);
diff --git a/Engine/Source/Programs/UnrealBuildTool/ProjectFiles/VisualStudio/VCProjectFileGenerator.cs b/Engine/Source/Programs/UnrealBuildTool/ProjectFiles/VisualStudio/VCProjectFileGenerator.cs
--- a/Engine/Source/Programs/UnrealBuildTool/ProjectFiles/VisualStudio/VCProjectFileGenerator.cs
+++ b/Engine/Source/Programs/UnrealBuildTool/ProjectFiles/VisualStudio/VCProjectFileGenerator.cs
@@ -257,6 +257,25 @@ namespace UnrealBuildTool
 			return string.Empty;
 		}
 
+		protected override bool SupportsDotnetCoreProjects()
+		{
+			// we use dotnet 3 which requires visual studio 2019 or newer
+			switch (Settings.ProjectFileFormat)
+			{
+				case VCProjectFileFormat.VisualStudio2019:
+					return true;
+				case VCProjectFileFormat.Default:
+				case VCProjectFileFormat.VisualStudio2012:
+				case VCProjectFileFormat.VisualStudio2013:
+				case VCProjectFileFormat.VisualStudio2015:
+				case VCProjectFileFormat.VisualStudio2017:
+					return false;
+				default:
+					throw new ArgumentOutOfRangeException(nameof(ProjectFileFormat), Settings.ProjectFileFormat, null);
+			}
+		}
+
+
 		static public void AppendPlatformToolsetProperty(StringBuilder VCProjectFileContent, VCProjectFileFormat ProjectFileFormat)
 		{
 			VCProjectFileContent.AppendLine("    <PlatformToolset>{0}</PlatformToolset>", GetProjectFilePlatformToolsetVersionString(ProjectFileFormat));
