UBT: Fix PCH usage on Clang.

diff --git a/Engine/Source/Programs/UnrealBuildTool/Platform/Windows/VCToolChain.cs b/Engine/Source/Programs/UnrealBuildTool/Platform/Windows/VCToolChain.cs
--- a/Engine/Source/Programs/UnrealBuildTool/Platform/Windows/VCToolChain.cs
+++ b/Engine/Source/Programs/UnrealBuildTool/Platform/Windows/VCToolChain.cs
@@ -1025,11 +1025,18 @@ namespace UnrealBuildTool
 			// If we're using precompiled headers, set that up now
 			if (CompileEnvironment.PrecompiledHeaderAction == PrecompiledHeaderAction.Include)
 			{
-				FileItem IncludeHeader = FileItem.GetItemByFileReference(CompileEnvironment.PrecompiledHeaderIncludeFilename);
-				BaseCompileAction.ForceIncludeFiles.Insert(0, IncludeHeader);
+				if (Target.WindowsPlatform.Compiler == WindowsCompiler.Clang)
+				{
+					BaseCompileAction.ForceIncludeFiles.Insert(0, FileItem.GetItemByFileReference(CompileEnvironment.PrecompiledHeaderFile.Location.ChangeExtension(null)));
+				}
+				else
+				{
+					FileItem IncludeHeader = FileItem.GetItemByFileReference(CompileEnvironment.PrecompiledHeaderIncludeFilename);
+					BaseCompileAction.ForceIncludeFiles.Insert(0, IncludeHeader);
 
-				BaseCompileAction.UsingPchFile = CompileEnvironment.PrecompiledHeaderFile;
-				BaseCompileAction.PchThroughHeaderFile = IncludeHeader;
+					BaseCompileAction.UsingPchFile = CompileEnvironment.PrecompiledHeaderFile;
+					BaseCompileAction.PchThroughHeaderFile = IncludeHeader;
+				}
 			}
 
 			// Generate the timing info
