Fix PVS studio warnings due to character encodings.

diff --git a/Engine/Source/Programs/UnrealBuildTool/Platform/Windows/PVSToolChain.cs b/Engine/Source/Programs/UnrealBuildTool/Platform/Windows/PVSToolChain.cs
--- a/Engine/Source/Programs/UnrealBuildTool/Platform/Windows/PVSToolChain.cs
+++ b/Engine/Source/Programs/UnrealBuildTool/Platform/Windows/PVSToolChain.cs
@@ -520,10 +520,9 @@ namespace UnrealBuildTool
 				}
 
 				// Disable a few warnings that seem to come from the preprocessor not respecting _Pragma
-				Action NewPreprocessAction = new Action(PreprocessAction);
-				NewPreprocessAction.CommandArguments += " /wd4005"; // macro redefinition
-				NewPreprocessAction.CommandArguments += " /wd4828"; // file contains a character starting at offset xxxx that is illegal in the current source character set
-				PreprocessActions[Idx] = NewPreprocessAction;
+				PreprocessAction.Arguments.Add("/wd4005"); // macro redefinition
+				PreprocessAction.Arguments.Add("/wd4828"); // file contains a character starting at offset xxxx that is illegal in the current source character set
+				PreprocessAction.WriteResponseFile(Graph);
 
 				// Write the PVS studio config file
 				StringBuilder ConfigFileContents = new StringBuilder();
