- Added a new SDLPlatform platform group for future use

diff --git a/Engine/Source/Programs/UnrealBuildTool/Configuration/UEBuildTarget.cs b/Engine/Source/Programs/UnrealBuildTool/Configuration/UEBuildTarget.cs
--- a/Engine/Source/Programs/UnrealBuildTool/Configuration/UEBuildTarget.cs
+++ b/Engine/Source/Programs/UnrealBuildTool/Configuration/UEBuildTarget.cs
@@ -529,6 +529,10 @@ namespace UnrealBuildTool
 		/// </summary>
 		public static UnrealPlatformGroup Desktop = FindOrAddByName("Desktop");
 
+		/// <summary>
+		/// SDLPlatform is for platforms that use SDL for windows, cursors, etc
+		/// </summary>
+		public static UnrealPlatformGroup SDLPlatform = FindOrAddByName("SDLPlatform");
 	}
 
 	/// <summary>
