Fix missing documentation warnings.

diff --git a/Engine/Source/Programs/UnrealBuildTool/Configuration/TargetRules.cs b/Engine/Source/Programs/UnrealBuildTool/Configuration/TargetRules.cs
--- a/Engine/Source/Programs/UnrealBuildTool/Configuration/TargetRules.cs
+++ b/Engine/Source/Programs/UnrealBuildTool/Configuration/TargetRules.cs
@@ -858,6 +858,9 @@ namespace UnrealBuildTool
 			set { bUseUnityBuildOverride = value; }
 		}
 
+		/// <summary>
+		/// Whether to unify C++ code into larger files for faster compilation.
+		/// </summary>
 		[CommandLine("-DisableUnity", Value = "false")]
 		[XmlConfigFile(Category = "BuildConfiguration", Name = nameof(bUseUnityBuild))]
 		bool? bUseUnityBuildOverride = null;
