UBT: Fix compilation actions not being executed remotely when re-executing a makefile.

diff --git a/Engine/Source/Programs/UnrealBuildTool/Platform/Windows/VCCompileAction.cs b/Engine/Source/Programs/UnrealBuildTool/Platform/Windows/VCCompileAction.cs
--- a/Engine/Source/Programs/UnrealBuildTool/Platform/Windows/VCCompileAction.cs
+++ b/Engine/Source/Programs/UnrealBuildTool/Platform/Windows/VCCompileAction.cs
@@ -268,6 +268,8 @@ namespace UnrealBuildTool
 			ForceIncludeFiles = new List<FileItem>(InAction.ForceIncludeFiles);
 			Arguments = new List<string>(InAction.Arguments);
 			bShowIncludes = InAction.bShowIncludes;
+			bCanExecuteRemotely = InAction.bCanExecuteRemotely;
+			bCanExecuteRemotelyWithSNDBS = InAction.bCanExecuteRemotelyWithSNDBS;
 
 			AdditionalPrerequisiteItems = new List<FileItem>(InAction.AdditionalPrerequisiteItems);
 			AdditionalProducedItems = new List<FileItem>(InAction.AdditionalProducedItems);
@@ -298,6 +300,8 @@ namespace UnrealBuildTool
 			ForceIncludeFiles = Reader.ReadList(() => Reader.ReadFileItem());
 			Arguments = Reader.ReadList(() => Reader.ReadString());
 			bShowIncludes = Reader.ReadBool();
+			bCanExecuteRemotely = Reader.ReadBool();
+			bCanExecuteRemotelyWithSNDBS = Reader.ReadBool();
 
 			AdditionalPrerequisiteItems = Reader.ReadList(() => Reader.ReadFileItem());
 			AdditionalProducedItems = Reader.ReadList(() => Reader.ReadFileItem());
@@ -325,6 +329,8 @@ namespace UnrealBuildTool
 			Writer.WriteList(ForceIncludeFiles, Item => Writer.WriteFileItem(Item));
 			Writer.WriteList(Arguments, Item => Writer.WriteString(Item));
 			Writer.WriteBool(bShowIncludes);
+			Writer.WriteBool(bCanExecuteRemotely);
+			Writer.WriteBool(bCanExecuteRemotelyWithSNDBS);
 
 			Writer.WriteList(AdditionalPrerequisiteItems, Item => Writer.WriteFileItem(Item));
 			Writer.WriteList(AdditionalProducedItems, Item => Writer.WriteFileItem(Item));
