Fixed issues in generated VSCode project that wasnt correctly generating dotnet build commands to build UBT.

diff --git a/Engine/Source/Programs/UnrealBuildTool/ProjectFiles/VisualStudioCode/VSCodeProjectFileGenerator.cs b/Engine/Source/Programs/UnrealBuildTool/ProjectFiles/VisualStudioCode/VSCodeProjectFileGenerator.cs
--- a/Engine/Source/Programs/UnrealBuildTool/ProjectFiles/VisualStudioCode/VSCodeProjectFileGenerator.cs
+++ b/Engine/Source/Programs/UnrealBuildTool/ProjectFiles/VisualStudioCode/VSCodeProjectFileGenerator.cs
@@ -779,8 +779,7 @@ namespace UnrealBuildTool
 
 								OutFile.AddUnnamedField("--configuration");
 								OutFile.AddUnnamedField(BuildProduct.Config.ToString());
-								OutFile.AddUnnamedField("--output");
-								OutFile.AddUnnamedField(MakePathString(BuildProduct.OutputFile.Directory));
+								OutFile.AddUnnamedField(MakeUnquotedPathString(BuildProduct.CSharpInfo.ProjectPath, EPathType.Absolute));
 							}
 							OutFile.EndArray();
 						}
