UBT: Mark platforms and project file generation code as not nullable aware.

#rnx

diff --git a/Engine/Source/Programs/UnrealBuildTool/Configuration/UEBuildPlatform.cs b/Engine/Source/Programs/UnrealBuildTool/Configuration/UEBuildPlatform.cs
--- a/Engine/Source/Programs/UnrealBuildTool/Configuration/UEBuildPlatform.cs
+++ b/Engine/Source/Programs/UnrealBuildTool/Configuration/UEBuildPlatform.cs
@@ -428,19 +428,26 @@ namespace UnrealBuildTool
 		/// Retrieve the IUEBuildPlatform instance for the given TargetPlatform
 		/// </summary>
 		/// <param name="InPlatform">  The UnrealTargetPlatform being built</param>
-		/// <param name="bInAllowFailure"> If true, do not throw an exception and return null</param>
 		/// <returns>UEBuildPlatform  The instance of the build platform</returns>
-		public static UEBuildPlatform GetBuildPlatform(UnrealTargetPlatform InPlatform, bool bInAllowFailure = false)
+		public static UEBuildPlatform GetBuildPlatform(UnrealTargetPlatform InPlatform)
 		{
-			if (BuildPlatformDictionary.ContainsKey(InPlatform) == true)
+			UEBuildPlatform Platform;
+			if (!TryGetBuildPlatform(InPlatform, out Platform))
 			{
-				return BuildPlatformDictionary[InPlatform];
+				throw new BuildException("GetBuildPlatform: No BuildPlatform found for {0}", InPlatform.ToString());
 			}
-			if (bInAllowFailure == true)
-			{
-				return null;
-			}
-			throw new BuildException("GetBuildPlatform: No BuildPlatform found for {0}", InPlatform.ToString());
+			return Platform;
+		}
+
+		/// <summary>
+		/// Retrieve the IUEBuildPlatform instance for the given TargetPlatform
+		/// </summary>
+		/// <param name="InPlatform">  The UnrealTargetPlatform being built</param>
+		/// <param name="Platform">The build platform</param>
+		/// <returns>UEBuildPlatform  The instance of the build platform</returns>
+		public static bool TryGetBuildPlatform(UnrealTargetPlatform InPlatform, out UEBuildPlatform Platform)
+		{
+			return BuildPlatformDictionary.TryGetValue(InPlatform, out Platform);
 		}
 
 		/// <summary>
@@ -534,8 +541,8 @@ namespace UnrealBuildTool
 		public static bool PlatformRequiresMonolithicBuilds(UnrealTargetPlatform InPlatform, UnrealTargetConfiguration InConfiguration)
 		{
 			// Some platforms require monolithic builds...
-			UEBuildPlatform BuildPlatform = GetBuildPlatform(InPlatform, true);
-			if (BuildPlatform != null)
+			UEBuildPlatform BuildPlatform;
+			if (TryGetBuildPlatform(InPlatform, out BuildPlatform))
 			{
 				return BuildPlatform.ShouldCompileMonolithicBinary(InPlatform);
 			}
diff --git a/Engine/Source/Programs/UnrealBuildTool/Configuration/UEBuildTarget.cs b/Engine/Source/Programs/UnrealBuildTool/Configuration/UEBuildTarget.cs
--- a/Engine/Source/Programs/UnrealBuildTool/Configuration/UEBuildTarget.cs
+++ b/Engine/Source/Programs/UnrealBuildTool/Configuration/UEBuildTarget.cs
@@ -3031,7 +3031,7 @@ namespace UnrealBuildTool
 			List<UnrealTargetPlatform> MissingPlatforms = new List<UnrealTargetPlatform>();
 			foreach (UnrealTargetPlatform TargetPlatform in UnrealTargetPlatform.GetValidPlatforms())
 			{
-				if (UEBuildPlatform.GetBuildPlatform(TargetPlatform, true) == null)
+				if (!UEBuildPlatform.TryGetBuildPlatform(TargetPlatform, out _))
 				{
 					MissingPlatforms.Add(TargetPlatform);
 				}
diff --git a/Engine/Source/Programs/UnrealBuildTool/Modes/ValidatePlatformsMode.cs b/Engine/Source/Programs/UnrealBuildTool/Modes/ValidatePlatformsMode.cs
--- a/Engine/Source/Programs/UnrealBuildTool/Modes/ValidatePlatformsMode.cs
+++ b/Engine/Source/Programs/UnrealBuildTool/Modes/ValidatePlatformsMode.cs
@@ -53,7 +53,7 @@ namespace UnrealBuildTool
 			// Output a line for each registered platform
 			foreach (UnrealTargetPlatform Platform in Platforms)
 			{
-				UEBuildPlatform BuildPlatform = UEBuildPlatform.GetBuildPlatform(Platform, true);
+				UEBuildPlatform.TryGetBuildPlatform(Platform, out UEBuildPlatform BuildPlatform);
 				string PlatformSDKString = "";
 				if (bOutputSDKs)
 				{
diff --git a/Engine/Source/Programs/UnrealBuildTool/Platform/Android/AndroidAARHandler.cs b/Engine/Source/Programs/UnrealBuildTool/Platform/Android/AndroidAARHandler.cs
--- a/Engine/Source/Programs/UnrealBuildTool/Platform/Android/AndroidAARHandler.cs
+++ b/Engine/Source/Programs/UnrealBuildTool/Platform/Android/AndroidAARHandler.cs
@@ -12,6 +12,8 @@ using Tools.DotNETCommon;
 
 using Ionic.Zip;
 
+#nullable disable
+
 namespace UnrealBuildTool
 {
 	class AndroidAARHandler
diff --git a/Engine/Source/Programs/UnrealBuildTool/Platform/Android/AndroidExports.cs b/Engine/Source/Programs/UnrealBuildTool/Platform/Android/AndroidExports.cs
--- a/Engine/Source/Programs/UnrealBuildTool/Platform/Android/AndroidExports.cs
+++ b/Engine/Source/Programs/UnrealBuildTool/Platform/Android/AndroidExports.cs
@@ -7,6 +7,8 @@ using System.Text;
 using System.Threading.Tasks;
 using Tools.DotNETCommon;
 
+#nullable disable
+
 namespace UnrealBuildTool
 {
 	/// <summary>
diff --git a/Engine/Source/Programs/UnrealBuildTool/Platform/Android/AndroidPlatformSDK.cs b/Engine/Source/Programs/UnrealBuildTool/Platform/Android/AndroidPlatformSDK.cs
--- a/Engine/Source/Programs/UnrealBuildTool/Platform/Android/AndroidPlatformSDK.cs
+++ b/Engine/Source/Programs/UnrealBuildTool/Platform/Android/AndroidPlatformSDK.cs
@@ -7,6 +7,8 @@ using Tools.DotNETCommon;
 using System.Text.RegularExpressions;
 using System.Linq;
 
+#nullable disable
+
 namespace UnrealBuildTool
 {
 	class AndroidPlatformSDK : UEBuildPlatformSDK
diff --git a/Engine/Source/Programs/UnrealBuildTool/Platform/Android/AndroidProjectGenerator.cs b/Engine/Source/Programs/UnrealBuildTool/Platform/Android/AndroidProjectGenerator.cs
--- a/Engine/Source/Programs/UnrealBuildTool/Platform/Android/AndroidProjectGenerator.cs
+++ b/Engine/Source/Programs/UnrealBuildTool/Platform/Android/AndroidProjectGenerator.cs
@@ -7,6 +7,8 @@ using System.Diagnostics;
 using System.IO;
 using Tools.DotNETCommon;
 
+#nullable disable
+
 namespace UnrealBuildTool
 {
 	/// <summary>
diff --git a/Engine/Source/Programs/UnrealBuildTool/Platform/Android/AndroidToolChain.cs b/Engine/Source/Programs/UnrealBuildTool/Platform/Android/AndroidToolChain.cs
--- a/Engine/Source/Programs/UnrealBuildTool/Platform/Android/AndroidToolChain.cs
+++ b/Engine/Source/Programs/UnrealBuildTool/Platform/Android/AndroidToolChain.cs
@@ -9,6 +9,8 @@ using System.Linq;
 using Microsoft.Win32;
 using Tools.DotNETCommon;
 
+#nullable disable
+
 namespace UnrealBuildTool
 {
 	class AndroidToolChain : ISPCToolChain, IAndroidToolChain
diff --git a/Engine/Source/Programs/UnrealBuildTool/Platform/Android/UEBuildAndroid.cs b/Engine/Source/Programs/UnrealBuildTool/Platform/Android/UEBuildAndroid.cs
--- a/Engine/Source/Programs/UnrealBuildTool/Platform/Android/UEBuildAndroid.cs
+++ b/Engine/Source/Programs/UnrealBuildTool/Platform/Android/UEBuildAndroid.cs
@@ -6,6 +6,8 @@ using System.Text;
 using System.IO;
 using Tools.DotNETCommon;
 
+#nullable disable
+
 namespace UnrealBuildTool
 {
 	/// <summary>
diff --git a/Engine/Source/Programs/UnrealBuildTool/Platform/Android/UEDeployAndroid.cs b/Engine/Source/Programs/UnrealBuildTool/Platform/Android/UEDeployAndroid.cs
--- a/Engine/Source/Programs/UnrealBuildTool/Platform/Android/UEDeployAndroid.cs
+++ b/Engine/Source/Programs/UnrealBuildTool/Platform/Android/UEDeployAndroid.cs
@@ -12,6 +12,8 @@ using System.Xml.Linq;
 using Tools.DotNETCommon;
 using System.Security.Cryptography;
 
+#nullable disable
+
 namespace UnrealBuildTool
 {
 	class UEDeployAndroid : UEBuildDeploy, IAndroidDeploy
diff --git a/Engine/Source/Programs/UnrealBuildTool/Platform/HoloLens/HoloLensDeploy.cs b/Engine/Source/Programs/UnrealBuildTool/Platform/HoloLens/HoloLensDeploy.cs
--- a/Engine/Source/Programs/UnrealBuildTool/Platform/HoloLens/HoloLensDeploy.cs
+++ b/Engine/Source/Programs/UnrealBuildTool/Platform/HoloLens/HoloLensDeploy.cs
@@ -8,6 +8,8 @@ using System.Diagnostics;
 using System.Linq;
 using Tools.DotNETCommon;
 
+#nullable disable
+
 namespace UnrealBuildTool
 {
 	/// <summary>
diff --git a/Engine/Source/Programs/UnrealBuildTool/Platform/HoloLens/HoloLensExports.cs b/Engine/Source/Programs/UnrealBuildTool/Platform/HoloLens/HoloLensExports.cs
--- a/Engine/Source/Programs/UnrealBuildTool/Platform/HoloLens/HoloLensExports.cs
+++ b/Engine/Source/Programs/UnrealBuildTool/Platform/HoloLens/HoloLensExports.cs
@@ -7,6 +7,8 @@ using System.Text;
 using System.Threading.Tasks;
 using Tools.DotNETCommon;
 
+#nullable disable
+
 namespace UnrealBuildTool
 {
 	/// <summary>
diff --git a/Engine/Source/Programs/UnrealBuildTool/Platform/HoloLens/HoloLensProjectGenerator.cs b/Engine/Source/Programs/UnrealBuildTool/Platform/HoloLens/HoloLensProjectGenerator.cs
--- a/Engine/Source/Programs/UnrealBuildTool/Platform/HoloLens/HoloLensProjectGenerator.cs
+++ b/Engine/Source/Programs/UnrealBuildTool/Platform/HoloLens/HoloLensProjectGenerator.cs
@@ -8,6 +8,8 @@ using System.IO;
 using System.Xml.Linq;
 using Tools.DotNETCommon;
 
+#nullable disable
+
 namespace UnrealBuildTool
 {
 	/// <summary>
diff --git a/Engine/Source/Programs/UnrealBuildTool/Platform/HoloLens/HoloLensToolChain.cs b/Engine/Source/Programs/UnrealBuildTool/Platform/HoloLens/HoloLensToolChain.cs
--- a/Engine/Source/Programs/UnrealBuildTool/Platform/HoloLens/HoloLensToolChain.cs
+++ b/Engine/Source/Programs/UnrealBuildTool/Platform/HoloLens/HoloLensToolChain.cs
@@ -10,6 +10,8 @@ using System.Text.RegularExpressions;
 using Microsoft.Win32;
 using Tools.DotNETCommon;
 
+#nullable disable
+
 namespace UnrealBuildTool
 {
 	static class Extension
diff --git a/Engine/Source/Programs/UnrealBuildTool/Platform/HoloLens/ManifestGenerator.cs b/Engine/Source/Programs/UnrealBuildTool/Platform/HoloLens/ManifestGenerator.cs
--- a/Engine/Source/Programs/UnrealBuildTool/Platform/HoloLens/ManifestGenerator.cs
+++ b/Engine/Source/Programs/UnrealBuildTool/Platform/HoloLens/ManifestGenerator.cs
@@ -10,6 +10,8 @@ using System.Resources;
 using System.Xml;
 using Tools.DotNETCommon;
 
+#nullable disable
+
 namespace UnrealBuildTool
 {
 	/// <summary>
diff --git a/Engine/Source/Programs/UnrealBuildTool/Platform/HoloLens/UEBuildHoloLens.cs b/Engine/Source/Programs/UnrealBuildTool/Platform/HoloLens/UEBuildHoloLens.cs
--- a/Engine/Source/Programs/UnrealBuildTool/Platform/HoloLens/UEBuildHoloLens.cs
+++ b/Engine/Source/Programs/UnrealBuildTool/Platform/HoloLens/UEBuildHoloLens.cs
@@ -10,6 +10,8 @@ using Microsoft.Win32;
 using Tools.DotNETCommon;
 using System.Text.RegularExpressions;
 
+#nullable disable
+
 namespace UnrealBuildTool
 {
 	/// <summary>
diff --git a/Engine/Source/Programs/UnrealBuildTool/Platform/HoloLens/WinMDRegistrationInfo.cs b/Engine/Source/Programs/UnrealBuildTool/Platform/HoloLens/WinMDRegistrationInfo.cs
--- a/Engine/Source/Programs/UnrealBuildTool/Platform/HoloLens/WinMDRegistrationInfo.cs
+++ b/Engine/Source/Programs/UnrealBuildTool/Platform/HoloLens/WinMDRegistrationInfo.cs
@@ -8,6 +8,8 @@ using System.Reflection;
 using System.Runtime.InteropServices;
 using Tools.DotNETCommon;
 
+#nullable disable
+
 namespace UnrealBuildTool
 {
 	/// <summary>
diff --git a/Engine/Source/Programs/UnrealBuildTool/Platform/IOS/IOSExports.cs b/Engine/Source/Programs/UnrealBuildTool/Platform/IOS/IOSExports.cs
--- a/Engine/Source/Programs/UnrealBuildTool/Platform/IOS/IOSExports.cs
+++ b/Engine/Source/Programs/UnrealBuildTool/Platform/IOS/IOSExports.cs
@@ -10,6 +10,8 @@ using System.Threading.Tasks;
 using System.Xml.Linq;
 using Tools.DotNETCommon;
 
+#nullable disable
+
 namespace UnrealBuildTool
 {
 	/// <summary>
diff --git a/Engine/Source/Programs/UnrealBuildTool/Platform/IOS/IOSPostBuildSyncMode.cs b/Engine/Source/Programs/UnrealBuildTool/Platform/IOS/IOSPostBuildSyncMode.cs
--- a/Engine/Source/Programs/UnrealBuildTool/Platform/IOS/IOSPostBuildSyncMode.cs
+++ b/Engine/Source/Programs/UnrealBuildTool/Platform/IOS/IOSPostBuildSyncMode.cs
@@ -7,6 +7,8 @@ using System.Text;
 using System.Threading.Tasks;
 using Tools.DotNETCommon;
 
+#nullable disable
+
 namespace UnrealBuildTool
 {
 	[Serializable]
diff --git a/Engine/Source/Programs/UnrealBuildTool/Platform/IOS/IOSProjectGenerator.cs b/Engine/Source/Programs/UnrealBuildTool/Platform/IOS/IOSProjectGenerator.cs
--- a/Engine/Source/Programs/UnrealBuildTool/Platform/IOS/IOSProjectGenerator.cs
+++ b/Engine/Source/Programs/UnrealBuildTool/Platform/IOS/IOSProjectGenerator.cs
@@ -7,6 +7,8 @@ using System.Diagnostics;
 using System.IO;
 using Tools.DotNETCommon;
 
+#nullable disable
+
 namespace UnrealBuildTool
 {
 	/// <summary>
diff --git a/Engine/Source/Programs/UnrealBuildTool/Platform/IOS/IOSToolChain.cs b/Engine/Source/Programs/UnrealBuildTool/Platform/IOS/IOSToolChain.cs
--- a/Engine/Source/Programs/UnrealBuildTool/Platform/IOS/IOSToolChain.cs
+++ b/Engine/Source/Programs/UnrealBuildTool/Platform/IOS/IOSToolChain.cs
@@ -16,6 +16,8 @@ using Ionic.Zlib;
 using Tools.DotNETCommon;
 using System.Security.Cryptography.X509Certificates;
 
+#nullable disable
+
 namespace UnrealBuildTool
 {
 	class IOSToolChainSettings : AppleToolChainSettings
diff --git a/Engine/Source/Programs/UnrealBuildTool/Platform/IOS/MobileProvisionContents.cs b/Engine/Source/Programs/UnrealBuildTool/Platform/IOS/MobileProvisionContents.cs
--- a/Engine/Source/Programs/UnrealBuildTool/Platform/IOS/MobileProvisionContents.cs
+++ b/Engine/Source/Programs/UnrealBuildTool/Platform/IOS/MobileProvisionContents.cs
@@ -9,6 +9,8 @@ using System.Threading.Tasks;
 using System.Xml;
 using Tools.DotNETCommon;
 
+#nullable disable
+
 namespace UnrealBuildTool
 {
 	/// <summary>
diff --git a/Engine/Source/Programs/UnrealBuildTool/Platform/IOS/UEBuildIOS.cs b/Engine/Source/Programs/UnrealBuildTool/Platform/IOS/UEBuildIOS.cs
--- a/Engine/Source/Programs/UnrealBuildTool/Platform/IOS/UEBuildIOS.cs
+++ b/Engine/Source/Programs/UnrealBuildTool/Platform/IOS/UEBuildIOS.cs
@@ -8,6 +8,8 @@ using System.IO;
 using System.Linq;
 using Tools.DotNETCommon;
 
+#nullable disable
+
 namespace UnrealBuildTool
 {
 	/// <summary>
diff --git a/Engine/Source/Programs/UnrealBuildTool/Platform/IOS/UEDeployIOS.cs b/Engine/Source/Programs/UnrealBuildTool/Platform/IOS/UEDeployIOS.cs
--- a/Engine/Source/Programs/UnrealBuildTool/Platform/IOS/UEDeployIOS.cs
+++ b/Engine/Source/Programs/UnrealBuildTool/Platform/IOS/UEDeployIOS.cs
@@ -11,6 +11,8 @@ using System.Diagnostics;
 using System.Xml.Linq;
 using Tools.DotNETCommon;
 
+#nullable disable
+
 namespace UnrealBuildTool
 {
 	class UEDeployIOS : UEBuildDeploy
diff --git a/Engine/Source/Programs/UnrealBuildTool/Platform/Linux/LinuxCommon.cs b/Engine/Source/Programs/UnrealBuildTool/Platform/Linux/LinuxCommon.cs
--- a/Engine/Source/Programs/UnrealBuildTool/Platform/Linux/LinuxCommon.cs
+++ b/Engine/Source/Programs/UnrealBuildTool/Platform/Linux/LinuxCommon.cs
@@ -9,6 +9,8 @@ using System.Linq;
 using Microsoft.Win32;
 using Tools.DotNETCommon;
 
+#nullable disable
+
 namespace UnrealBuildTool
 {
 	class LinuxCommon
diff --git a/Engine/Source/Programs/UnrealBuildTool/Platform/Linux/LinuxExports.cs b/Engine/Source/Programs/UnrealBuildTool/Platform/Linux/LinuxExports.cs
--- a/Engine/Source/Programs/UnrealBuildTool/Platform/Linux/LinuxExports.cs
+++ b/Engine/Source/Programs/UnrealBuildTool/Platform/Linux/LinuxExports.cs
@@ -1,4 +1,4 @@
-﻿// Copyright Epic Games, Inc. All Rights Reserved.
+// Copyright Epic Games, Inc. All Rights Reserved.
 
 using System;
 using System.Collections.Generic;
@@ -7,6 +7,8 @@ using System.Text;
 using System.Threading.Tasks;
 using Tools.DotNETCommon;
 
+#nullable disable
+
 namespace UnrealBuildTool
 {
 	/// <summary>
diff --git a/Engine/Source/Programs/UnrealBuildTool/Platform/Linux/LinuxPlatformSDK.cs b/Engine/Source/Programs/UnrealBuildTool/Platform/Linux/LinuxPlatformSDK.cs
--- a/Engine/Source/Programs/UnrealBuildTool/Platform/Linux/LinuxPlatformSDK.cs
+++ b/Engine/Source/Programs/UnrealBuildTool/Platform/Linux/LinuxPlatformSDK.cs
@@ -6,6 +6,8 @@ using System.IO;
 using Tools.DotNETCommon;
 using System.Text.RegularExpressions;
 
+#nullable disable
+
 namespace UnrealBuildTool
 {
 	class LinuxPlatformSDK : UEBuildPlatformSDK
diff --git a/Engine/Source/Programs/UnrealBuildTool/Platform/Linux/LinuxProjectGenerator.cs b/Engine/Source/Programs/UnrealBuildTool/Platform/Linux/LinuxProjectGenerator.cs
--- a/Engine/Source/Programs/UnrealBuildTool/Platform/Linux/LinuxProjectGenerator.cs
+++ b/Engine/Source/Programs/UnrealBuildTool/Platform/Linux/LinuxProjectGenerator.cs
@@ -7,6 +7,8 @@ using System.Diagnostics;
 using System.IO;
 using Tools.DotNETCommon;
 
+#nullable disable
+
 namespace UnrealBuildTool
 {
 	/// <summary>
diff --git a/Engine/Source/Programs/UnrealBuildTool/Platform/Linux/LinuxToolChain.cs b/Engine/Source/Programs/UnrealBuildTool/Platform/Linux/LinuxToolChain.cs
--- a/Engine/Source/Programs/UnrealBuildTool/Platform/Linux/LinuxToolChain.cs
+++ b/Engine/Source/Programs/UnrealBuildTool/Platform/Linux/LinuxToolChain.cs
@@ -9,6 +9,8 @@ using System.Linq;
 using Microsoft.Win32;
 using Tools.DotNETCommon;
 
+#nullable disable
+
 namespace UnrealBuildTool
 {
 	/// <summary>
diff --git a/Engine/Source/Programs/UnrealBuildTool/Platform/Linux/UEBuildLinux.cs b/Engine/Source/Programs/UnrealBuildTool/Platform/Linux/UEBuildLinux.cs
--- a/Engine/Source/Programs/UnrealBuildTool/Platform/Linux/UEBuildLinux.cs
+++ b/Engine/Source/Programs/UnrealBuildTool/Platform/Linux/UEBuildLinux.cs
@@ -8,6 +8,8 @@ using System.IO;
 using Tools.DotNETCommon;
 using System.Text.RegularExpressions;
 
+#nullable disable
+
 namespace UnrealBuildTool
 {
 	/** Architecture as stored in the ini. */
diff --git a/Engine/Source/Programs/UnrealBuildTool/Platform/Lumin/LuminExports.cs b/Engine/Source/Programs/UnrealBuildTool/Platform/Lumin/LuminExports.cs
--- a/Engine/Source/Programs/UnrealBuildTool/Platform/Lumin/LuminExports.cs
+++ b/Engine/Source/Programs/UnrealBuildTool/Platform/Lumin/LuminExports.cs
@@ -7,6 +7,8 @@ using System.Text;
 using System.Threading.Tasks;
 using Tools.DotNETCommon;
 
+#nullable disable
+
 namespace UnrealBuildTool
 {
 	/// <summary>
diff --git a/Engine/Source/Programs/UnrealBuildTool/Platform/Lumin/LuminPlatformSDK.cs b/Engine/Source/Programs/UnrealBuildTool/Platform/Lumin/LuminPlatformSDK.cs
--- a/Engine/Source/Programs/UnrealBuildTool/Platform/Lumin/LuminPlatformSDK.cs
+++ b/Engine/Source/Programs/UnrealBuildTool/Platform/Lumin/LuminPlatformSDK.cs
@@ -6,6 +6,8 @@ using System.IO;
 using Tools.DotNETCommon;
 using System.Text.RegularExpressions;
 
+#nullable disable
+
 namespace UnrealBuildTool
 {
 	class LuminPlatformSDK : UEBuildPlatformSDK
diff --git a/Engine/Source/Programs/UnrealBuildTool/Platform/Lumin/LuminProjectGenerator.cs b/Engine/Source/Programs/UnrealBuildTool/Platform/Lumin/LuminProjectGenerator.cs
--- a/Engine/Source/Programs/UnrealBuildTool/Platform/Lumin/LuminProjectGenerator.cs
+++ b/Engine/Source/Programs/UnrealBuildTool/Platform/Lumin/LuminProjectGenerator.cs
@@ -7,6 +7,8 @@ using System.Diagnostics;
 using System.IO;
 using Tools.DotNETCommon;
 
+#nullable disable
+
 namespace UnrealBuildTool
 {
 	/// <summary>
diff --git a/Engine/Source/Programs/UnrealBuildTool/Platform/Lumin/LuminToolChain.cs b/Engine/Source/Programs/UnrealBuildTool/Platform/Lumin/LuminToolChain.cs
--- a/Engine/Source/Programs/UnrealBuildTool/Platform/Lumin/LuminToolChain.cs
+++ b/Engine/Source/Programs/UnrealBuildTool/Platform/Lumin/LuminToolChain.cs
@@ -9,6 +9,8 @@ using System.Linq;
 using Microsoft.Win32;
 using Tools.DotNETCommon;
 
+#nullable disable
+
 namespace UnrealBuildTool
 {
 	class LuminToolChain : AndroidToolChain
diff --git a/Engine/Source/Programs/UnrealBuildTool/Platform/Lumin/UEBuildLumin.cs b/Engine/Source/Programs/UnrealBuildTool/Platform/Lumin/UEBuildLumin.cs
--- a/Engine/Source/Programs/UnrealBuildTool/Platform/Lumin/UEBuildLumin.cs
+++ b/Engine/Source/Programs/UnrealBuildTool/Platform/Lumin/UEBuildLumin.cs
@@ -10,6 +10,8 @@ using System.Linq;
 using Tools.DotNETCommon;
 using System.Text.RegularExpressions;
 
+#nullable disable
+
 namespace UnrealBuildTool
 {
 	/// <summary>
diff --git a/Engine/Source/Programs/UnrealBuildTool/Platform/Lumin/UEDeployLumin.cs b/Engine/Source/Programs/UnrealBuildTool/Platform/Lumin/UEDeployLumin.cs
--- a/Engine/Source/Programs/UnrealBuildTool/Platform/Lumin/UEDeployLumin.cs
+++ b/Engine/Source/Programs/UnrealBuildTool/Platform/Lumin/UEDeployLumin.cs
@@ -12,6 +12,8 @@ using System.Xml;
 using System.Xml.Serialization;
 using System.Xml.Linq;
 
+#nullable disable
+
 namespace UnrealBuildTool
 {
 	class UEDeployLumin : UEBuildDeploy, ILuminDeploy
diff --git a/Engine/Source/Programs/UnrealBuildTool/Platform/Mac/ApplePlatformSDK.cs b/Engine/Source/Programs/UnrealBuildTool/Platform/Mac/ApplePlatformSDK.cs
--- a/Engine/Source/Programs/UnrealBuildTool/Platform/Mac/ApplePlatformSDK.cs
+++ b/Engine/Source/Programs/UnrealBuildTool/Platform/Mac/ApplePlatformSDK.cs
@@ -8,6 +8,8 @@ using System.Text.RegularExpressions;
 using Microsoft.Win32;
 using System.Diagnostics;
 
+#nullable disable
+
 namespace UnrealBuildTool
 {
 	internal class ApplePlatformSDK : UEBuildPlatformSDK
diff --git a/Engine/Source/Programs/UnrealBuildTool/Platform/Mac/MacExports.cs b/Engine/Source/Programs/UnrealBuildTool/Platform/Mac/MacExports.cs
--- a/Engine/Source/Programs/UnrealBuildTool/Platform/Mac/MacExports.cs
+++ b/Engine/Source/Programs/UnrealBuildTool/Platform/Mac/MacExports.cs
@@ -1,4 +1,4 @@
-﻿// Copyright Epic Games, Inc. All Rights Reserved.
+// Copyright Epic Games, Inc. All Rights Reserved.
 
 using System;
 using System.Collections.Generic;
@@ -7,6 +7,8 @@ using System.Text;
 using System.Threading.Tasks;
 using Tools.DotNETCommon;
 
+#nullable disable
+
 namespace UnrealBuildTool
 {
 	/// <summary>
diff --git a/Engine/Source/Programs/UnrealBuildTool/Platform/Mac/MacProjectGenerator.cs b/Engine/Source/Programs/UnrealBuildTool/Platform/Mac/MacProjectGenerator.cs
--- a/Engine/Source/Programs/UnrealBuildTool/Platform/Mac/MacProjectGenerator.cs
+++ b/Engine/Source/Programs/UnrealBuildTool/Platform/Mac/MacProjectGenerator.cs
@@ -7,6 +7,8 @@ using System.Diagnostics;
 using System.IO;
 using Tools.DotNETCommon;
 
+#nullable disable
+
 namespace UnrealBuildTool
 {
 	/// <summary>
diff --git a/Engine/Source/Programs/UnrealBuildTool/Platform/Mac/MacToolChain.cs b/Engine/Source/Programs/UnrealBuildTool/Platform/Mac/MacToolChain.cs
--- a/Engine/Source/Programs/UnrealBuildTool/Platform/Mac/MacToolChain.cs
+++ b/Engine/Source/Programs/UnrealBuildTool/Platform/Mac/MacToolChain.cs
@@ -13,6 +13,8 @@ using Tools.DotNETCommon;
 using System.Globalization;
 using System.Text.RegularExpressions;
 
+#nullable disable
+
 namespace UnrealBuildTool
 {
 	class MacToolChainSettings : AppleToolChainSettings
diff --git a/Engine/Source/Programs/UnrealBuildTool/Platform/Mac/UEBuildMac.cs b/Engine/Source/Programs/UnrealBuildTool/Platform/Mac/UEBuildMac.cs
--- a/Engine/Source/Programs/UnrealBuildTool/Platform/Mac/UEBuildMac.cs
+++ b/Engine/Source/Programs/UnrealBuildTool/Platform/Mac/UEBuildMac.cs
@@ -6,6 +6,8 @@ using System.Collections.Generic;
 using Tools.DotNETCommon;
 using System.Linq;
 
+#nullable disable
+
 namespace UnrealBuildTool
 {
 	/// <summary>
diff --git a/Engine/Source/Programs/UnrealBuildTool/Platform/Mac/UEDeployMac.cs b/Engine/Source/Programs/UnrealBuildTool/Platform/Mac/UEDeployMac.cs
--- a/Engine/Source/Programs/UnrealBuildTool/Platform/Mac/UEDeployMac.cs
+++ b/Engine/Source/Programs/UnrealBuildTool/Platform/Mac/UEDeployMac.cs
@@ -10,6 +10,8 @@ using System.IO;
 using Ionic.Zip;
 using Tools.DotNETCommon;
 
+#nullable disable
+
 namespace UnrealBuildTool
 {
 	class UEDeployMac : UEBuildDeploy
diff --git a/Engine/Source/Programs/UnrealBuildTool/Platform/TVOS/TVOSExports.cs b/Engine/Source/Programs/UnrealBuildTool/Platform/TVOS/TVOSExports.cs
--- a/Engine/Source/Programs/UnrealBuildTool/Platform/TVOS/TVOSExports.cs
+++ b/Engine/Source/Programs/UnrealBuildTool/Platform/TVOS/TVOSExports.cs
@@ -7,6 +7,8 @@ using System.Text;
 using System.Threading.Tasks;
 using Tools.DotNETCommon;
 
+#nullable disable
+
 namespace UnrealBuildTool
 {
 	/// <summary>
diff --git a/Engine/Source/Programs/UnrealBuildTool/Platform/TVOS/UEBuildTVOS.cs b/Engine/Source/Programs/UnrealBuildTool/Platform/TVOS/UEBuildTVOS.cs
--- a/Engine/Source/Programs/UnrealBuildTool/Platform/TVOS/UEBuildTVOS.cs
+++ b/Engine/Source/Programs/UnrealBuildTool/Platform/TVOS/UEBuildTVOS.cs
@@ -4,6 +4,8 @@ using System;
 using System.Collections.Generic;
 using Tools.DotNETCommon;
 
+#nullable disable
+
 namespace UnrealBuildTool
 {
 	class TVOSProjectSettings : IOSProjectSettings
diff --git a/Engine/Source/Programs/UnrealBuildTool/Platform/TVOS/UEDeployTVOS.cs b/Engine/Source/Programs/UnrealBuildTool/Platform/TVOS/UEDeployTVOS.cs
--- a/Engine/Source/Programs/UnrealBuildTool/Platform/TVOS/UEDeployTVOS.cs
+++ b/Engine/Source/Programs/UnrealBuildTool/Platform/TVOS/UEDeployTVOS.cs
@@ -10,6 +10,8 @@ using System.Diagnostics;
 using System.Xml.Linq;
 using Tools.DotNETCommon;
 
+#nullable disable
+
 namespace UnrealBuildTool
 {
 	class UEDeployTVOS : UEDeployIOS
diff --git a/Engine/Source/Programs/UnrealBuildTool/Platform/Windows/MicrosoftPlatformSDK.cs b/Engine/Source/Programs/UnrealBuildTool/Platform/Windows/MicrosoftPlatformSDK.cs
--- a/Engine/Source/Programs/UnrealBuildTool/Platform/Windows/MicrosoftPlatformSDK.cs
+++ b/Engine/Source/Programs/UnrealBuildTool/Platform/Windows/MicrosoftPlatformSDK.cs
@@ -7,6 +7,8 @@ using System.Linq;
 using System.Text.RegularExpressions;
 using Microsoft.Win32;
 
+#nullable disable
+
 namespace UnrealBuildTool
 {
 	internal class MicrosoftPlatformSDK : UEBuildPlatformSDK
diff --git a/Engine/Source/Programs/UnrealBuildTool/Platform/Windows/PVSToolChain.cs b/Engine/Source/Programs/UnrealBuildTool/Platform/Windows/PVSToolChain.cs
--- a/Engine/Source/Programs/UnrealBuildTool/Platform/Windows/PVSToolChain.cs
+++ b/Engine/Source/Programs/UnrealBuildTool/Platform/Windows/PVSToolChain.cs
@@ -13,6 +13,8 @@ using System.Xml;
 using System.Xml.Serialization;
 using Tools.DotNETCommon;
 
+#nullable disable
+
 namespace UnrealBuildTool
 {
 	/// <summary>
diff --git a/Engine/Source/Programs/UnrealBuildTool/Platform/Windows/UEBuildWindows.cs b/Engine/Source/Programs/UnrealBuildTool/Platform/Windows/UEBuildWindows.cs
--- a/Engine/Source/Programs/UnrealBuildTool/Platform/Windows/UEBuildWindows.cs
+++ b/Engine/Source/Programs/UnrealBuildTool/Platform/Windows/UEBuildWindows.cs
@@ -11,6 +11,8 @@ using Tools.DotNETCommon;
 using Microsoft.VisualStudio.Setup.Configuration;
 using System.Runtime.InteropServices;
 
+#nullable disable
+
 namespace UnrealBuildTool
 {
 	/// <summary>
diff --git a/Engine/Source/Programs/UnrealBuildTool/Platform/Windows/VCCompileAction.cs b/Engine/Source/Programs/UnrealBuildTool/Platform/Windows/VCCompileAction.cs
--- a/Engine/Source/Programs/UnrealBuildTool/Platform/Windows/VCCompileAction.cs
+++ b/Engine/Source/Programs/UnrealBuildTool/Platform/Windows/VCCompileAction.cs
@@ -4,6 +4,8 @@ using System;
 using System.Collections.Generic;
 using Tools.DotNETCommon;
 
+#nullable disable
+
 namespace UnrealBuildTool
 {
 	/// <summary>
diff --git a/Engine/Source/Programs/UnrealBuildTool/Platform/Windows/VCEnvironment.cs b/Engine/Source/Programs/UnrealBuildTool/Platform/Windows/VCEnvironment.cs
--- a/Engine/Source/Programs/UnrealBuildTool/Platform/Windows/VCEnvironment.cs
+++ b/Engine/Source/Programs/UnrealBuildTool/Platform/Windows/VCEnvironment.cs
@@ -10,6 +10,8 @@ using Microsoft.Win32;
 using System.Text;
 using Tools.DotNETCommon;
 
+#nullable disable
+
 namespace UnrealBuildTool
 {
 	/// <summary>
diff --git a/Engine/Source/Programs/UnrealBuildTool/Platform/Windows/VCManifestGenerator.cs b/Engine/Source/Programs/UnrealBuildTool/Platform/Windows/VCManifestGenerator.cs
--- a/Engine/Source/Programs/UnrealBuildTool/Platform/Windows/VCManifestGenerator.cs
+++ b/Engine/Source/Programs/UnrealBuildTool/Platform/Windows/VCManifestGenerator.cs
@@ -11,6 +11,8 @@ using System.Xml.Linq;
 using System.Text;
 using System.Diagnostics;
 
+#nullable disable
+
 namespace UnrealBuildTool
 {
     abstract class VCManifestGenerator
diff --git a/Engine/Source/Programs/UnrealBuildTool/Platform/Windows/VCPreludeCache.cs b/Engine/Source/Programs/UnrealBuildTool/Platform/Windows/VCPreludeCache.cs
--- a/Engine/Source/Programs/UnrealBuildTool/Platform/Windows/VCPreludeCache.cs
+++ b/Engine/Source/Programs/UnrealBuildTool/Platform/Windows/VCPreludeCache.cs
@@ -7,6 +7,8 @@ using System.IO;
 using System.Linq;
 using Tools.DotNETCommon;
 
+#nullable disable
+
 namespace UnrealBuildTool
 {
 	/// <summary>
diff --git a/Engine/Source/Programs/UnrealBuildTool/Platform/Windows/VCToolChain.cs b/Engine/Source/Programs/UnrealBuildTool/Platform/Windows/VCToolChain.cs
--- a/Engine/Source/Programs/UnrealBuildTool/Platform/Windows/VCToolChain.cs
+++ b/Engine/Source/Programs/UnrealBuildTool/Platform/Windows/VCToolChain.cs
@@ -10,6 +10,8 @@ using Microsoft.Win32;
 using System.Text;
 using Tools.DotNETCommon;
 
+#nullable disable
+
 namespace UnrealBuildTool
 {
 	class VCToolChain : ISPCToolChain
diff --git a/Engine/Source/Programs/UnrealBuildTool/ProjectFiles/CMake/CMakefileGenerator.cs b/Engine/Source/Programs/UnrealBuildTool/ProjectFiles/CMake/CMakefileGenerator.cs
--- a/Engine/Source/Programs/UnrealBuildTool/ProjectFiles/CMake/CMakefileGenerator.cs
+++ b/Engine/Source/Programs/UnrealBuildTool/ProjectFiles/CMake/CMakefileGenerator.cs
@@ -6,6 +6,8 @@ using System.Text;
 using System.IO;
 using Tools.DotNETCommon;
 
+#nullable disable
+
 namespace UnrealBuildTool
 {
 	/// <summary>
diff --git a/Engine/Source/Programs/UnrealBuildTool/ProjectFiles/CodeLite/CodeLiteGenerator.cs b/Engine/Source/Programs/UnrealBuildTool/ProjectFiles/CodeLite/CodeLiteGenerator.cs
--- a/Engine/Source/Programs/UnrealBuildTool/ProjectFiles/CodeLite/CodeLiteGenerator.cs
+++ b/Engine/Source/Programs/UnrealBuildTool/ProjectFiles/CodeLite/CodeLiteGenerator.cs
@@ -8,6 +8,8 @@ using System.Xml;
 using System.Xml.Linq;
 using Tools.DotNETCommon;
 
+#nullable disable
+
 namespace UnrealBuildTool
 {
 	// Represents a folder within the master project. TODO Not using at the moment.
diff --git a/Engine/Source/Programs/UnrealBuildTool/ProjectFiles/CodeLite/CodeLiteProject.cs b/Engine/Source/Programs/UnrealBuildTool/ProjectFiles/CodeLite/CodeLiteProject.cs
--- a/Engine/Source/Programs/UnrealBuildTool/ProjectFiles/CodeLite/CodeLiteProject.cs
+++ b/Engine/Source/Programs/UnrealBuildTool/ProjectFiles/CodeLite/CodeLiteProject.cs
@@ -8,6 +8,8 @@ using System.Xml;
 using System.Xml.Linq;
 using Tools.DotNETCommon;
 
+#nullable disable
+
 namespace UnrealBuildTool
 {
 	class CodeLiteProject : ProjectFile
diff --git a/Engine/Source/Programs/UnrealBuildTool/ProjectFiles/Eddie/EddieGenerator.cs b/Engine/Source/Programs/UnrealBuildTool/ProjectFiles/Eddie/EddieGenerator.cs
--- a/Engine/Source/Programs/UnrealBuildTool/ProjectFiles/Eddie/EddieGenerator.cs
+++ b/Engine/Source/Programs/UnrealBuildTool/ProjectFiles/Eddie/EddieGenerator.cs
@@ -6,6 +6,8 @@ using System.Text;
 using System.IO;
 using Tools.DotNETCommon;
 
+#nullable disable
+
 namespace UnrealBuildTool
 {
 
diff --git a/Engine/Source/Programs/UnrealBuildTool/ProjectFiles/Eddie/EddieProject.cs b/Engine/Source/Programs/UnrealBuildTool/ProjectFiles/Eddie/EddieProject.cs
--- a/Engine/Source/Programs/UnrealBuildTool/ProjectFiles/Eddie/EddieProject.cs
+++ b/Engine/Source/Programs/UnrealBuildTool/ProjectFiles/Eddie/EddieProject.cs
@@ -10,6 +10,8 @@ using System.Linq;
 using System.Text;
 using Tools.DotNETCommon;
 
+#nullable disable
+
 namespace UnrealBuildTool
 {
 
diff --git a/Engine/Source/Programs/UnrealBuildTool/ProjectFiles/KDevelop/KDevelopGenerator.cs b/Engine/Source/Programs/UnrealBuildTool/ProjectFiles/KDevelop/KDevelopGenerator.cs
--- a/Engine/Source/Programs/UnrealBuildTool/ProjectFiles/KDevelop/KDevelopGenerator.cs
+++ b/Engine/Source/Programs/UnrealBuildTool/ProjectFiles/KDevelop/KDevelopGenerator.cs
@@ -7,6 +7,8 @@ using System.IO;
 using System.Linq;
 using Tools.DotNETCommon;
 
+#nullable disable
+
 namespace UnrealBuildTool
 {
 	/// <summary>
diff --git a/Engine/Source/Programs/UnrealBuildTool/ProjectFiles/Make/MakefileGenerator.cs b/Engine/Source/Programs/UnrealBuildTool/ProjectFiles/Make/MakefileGenerator.cs
--- a/Engine/Source/Programs/UnrealBuildTool/ProjectFiles/Make/MakefileGenerator.cs
+++ b/Engine/Source/Programs/UnrealBuildTool/ProjectFiles/Make/MakefileGenerator.cs
@@ -6,6 +6,8 @@ using System.Text;
 using System.IO;
 using Tools.DotNETCommon;
 
+#nullable disable
+
 namespace UnrealBuildTool
 {
 	/// <summary>
diff --git a/Engine/Source/Programs/UnrealBuildTool/ProjectFiles/PlatformProjectGenerator.cs b/Engine/Source/Programs/UnrealBuildTool/ProjectFiles/PlatformProjectGenerator.cs
--- a/Engine/Source/Programs/UnrealBuildTool/ProjectFiles/PlatformProjectGenerator.cs
+++ b/Engine/Source/Programs/UnrealBuildTool/ProjectFiles/PlatformProjectGenerator.cs
@@ -7,6 +7,8 @@ using System.Diagnostics;
 using System.IO;
 using Tools.DotNETCommon;
 
+#nullable disable
+
 namespace UnrealBuildTool
 {
 	/// <summary>
diff --git a/Engine/Source/Programs/UnrealBuildTool/ProjectFiles/PlatformProjectGeneratorCollection.cs b/Engine/Source/Programs/UnrealBuildTool/ProjectFiles/PlatformProjectGeneratorCollection.cs
--- a/Engine/Source/Programs/UnrealBuildTool/ProjectFiles/PlatformProjectGeneratorCollection.cs
+++ b/Engine/Source/Programs/UnrealBuildTool/ProjectFiles/PlatformProjectGeneratorCollection.cs
@@ -8,6 +8,8 @@ using System.Text;
 using System.Threading.Tasks;
 using Tools.DotNETCommon;
 
+#nullable disable
+
 namespace UnrealBuildTool
 {
 	/// <summary>
@@ -25,8 +27,8 @@ namespace UnrealBuildTool
 		public void RegisterPlatformProjectGenerator(UnrealTargetPlatform InPlatform, PlatformProjectGenerator InProjectGenerator)
 		{
 			// Make sure the build platform is legal
-			UEBuildPlatform BuildPlatform = UEBuildPlatform.GetBuildPlatform(InPlatform, true);
-			if (BuildPlatform != null)
+			UEBuildPlatform BuildPlatform;
+			if(UEBuildPlatform.TryGetBuildPlatform(InPlatform, out BuildPlatform))
 			{
 				if (ProjectGeneratorDictionary.ContainsKey(InPlatform) == true)
 				{
diff --git a/Engine/Source/Programs/UnrealBuildTool/ProjectFiles/Project.cs b/Engine/Source/Programs/UnrealBuildTool/ProjectFiles/Project.cs
--- a/Engine/Source/Programs/UnrealBuildTool/ProjectFiles/Project.cs
+++ b/Engine/Source/Programs/UnrealBuildTool/ProjectFiles/Project.cs
@@ -10,6 +10,8 @@ using System.Linq;
 using System.Text;
 using Tools.DotNETCommon;
 
+#nullable disable
+
 namespace UnrealBuildTool
 {
 	/// <summary>
diff --git a/Engine/Source/Programs/UnrealBuildTool/ProjectFiles/ProjectFileGenerator.cs b/Engine/Source/Programs/UnrealBuildTool/ProjectFiles/ProjectFileGenerator.cs
--- a/Engine/Source/Programs/UnrealBuildTool/ProjectFiles/ProjectFileGenerator.cs
+++ b/Engine/Source/Programs/UnrealBuildTool/ProjectFiles/ProjectFileGenerator.cs
@@ -12,6 +12,8 @@ using Tools.DotNETCommon;
 using System.Xml.Linq;
 using System.Xml;
 
+#nullable disable
+
 namespace UnrealBuildTool
 {
 	/// <summary>
@@ -1810,8 +1812,8 @@ namespace UnrealBuildTool
 				if (bInProjectPlatformsList && (IncludeAllPlatforms || IsRequiredPlatform))
 				{
 					// If there is a build platform present, add it to the SupportedPlatforms list
-					UEBuildPlatform BuildPlatform = UEBuildPlatform.GetBuildPlatform( Platform, true );
-					if( BuildPlatform != null )
+					UEBuildPlatform BuildPlatform;
+					if(UEBuildPlatform.TryGetBuildPlatform(Platform, out BuildPlatform))
 					{
 						if (InstalledPlatformInfo.IsValidPlatform(Platform, EProjectType.Code))
 						{
@@ -2338,7 +2340,7 @@ namespace UnrealBuildTool
 							TargetFilePath = TargetFilePath,
 							ProjectFilePath = ProjectFilePath,
 							UnrealProjectFilePath = CheckProjectFile,
-							SupportedPlatforms = TargetRulesObject.GetSupportedPlatforms().Where(x => UEBuildPlatform.GetBuildPlatform(x, true) != null).ToArray(),
+							SupportedPlatforms = TargetRulesObject.GetSupportedPlatforms().Where(x => UEBuildPlatform.TryGetBuildPlatform(x, out _)).ToArray(),
 							CreateRulesDelegate = (Platform, Configuration) => RulesAssembly.CreateTargetRules(TargetName, Platform, Configuration, "", CheckProjectFile, new CommandLineArguments(GetTargetArguments(Arguments)))
                         };
 
diff --git a/Engine/Source/Programs/UnrealBuildTool/ProjectFiles/QMake/QMakefileGenerator.cs b/Engine/Source/Programs/UnrealBuildTool/ProjectFiles/QMake/QMakefileGenerator.cs
--- a/Engine/Source/Programs/UnrealBuildTool/ProjectFiles/QMake/QMakefileGenerator.cs
+++ b/Engine/Source/Programs/UnrealBuildTool/ProjectFiles/QMake/QMakefileGenerator.cs
@@ -1,4 +1,4 @@
-﻿// Copyright Epic Games, Inc. All Rights Reserved.
+// Copyright Epic Games, Inc. All Rights Reserved.
 
 using System;
 using System.Collections.Generic;
@@ -7,6 +7,8 @@ using System.IO;
 using System.Linq;
 using Tools.DotNETCommon;
 
+#nullable disable
+
 namespace UnrealBuildTool
 {
 	/// <summary>
diff --git a/Engine/Source/Programs/UnrealBuildTool/ProjectFiles/Rider/RiderProjectFile.cs b/Engine/Source/Programs/UnrealBuildTool/ProjectFiles/Rider/RiderProjectFile.cs
--- a/Engine/Source/Programs/UnrealBuildTool/ProjectFiles/Rider/RiderProjectFile.cs
+++ b/Engine/Source/Programs/UnrealBuildTool/ProjectFiles/Rider/RiderProjectFile.cs
@@ -7,6 +7,8 @@ using System.IO;
 using System.Linq;
 using Tools.DotNETCommon;
 
+#nullable disable
+
 namespace UnrealBuildTool
 {
 	internal class RiderProjectFile : ProjectFile
diff --git a/Engine/Source/Programs/UnrealBuildTool/ProjectFiles/Rider/RiderProjectFileGenerator.cs b/Engine/Source/Programs/UnrealBuildTool/ProjectFiles/Rider/RiderProjectFileGenerator.cs
--- a/Engine/Source/Programs/UnrealBuildTool/ProjectFiles/Rider/RiderProjectFileGenerator.cs
+++ b/Engine/Source/Programs/UnrealBuildTool/ProjectFiles/Rider/RiderProjectFileGenerator.cs
@@ -6,6 +6,8 @@ using System.IO;
 using System.Linq;
 using Tools.DotNETCommon;
 
+#nullable disable
+
 namespace UnrealBuildTool
 {
 	class RiderMasterProjectFolder : MasterProjectFolder
@@ -390,7 +392,7 @@ namespace UnrealBuildTool
 						ProjectFilePath = ProjectFilePath,
 						UnrealProjectFilePath = CheckProjectFile,
 						SupportedPlatforms = TargetRulesObject.GetSupportedPlatforms()
-							.Where(x => UEBuildPlatform.GetBuildPlatform(x, true) != null).ToArray(),
+							.Where(x => UEBuildPlatform.TryGetBuildPlatform(x, out _)).ToArray(),
 						CreateRulesDelegate = (Platform, Configuration) =>
 							RulesAssembly.CreateTargetRules(TargetName, Platform, Configuration, "", CheckProjectFile,
 								null)
diff --git a/Engine/Source/Programs/UnrealBuildTool/ProjectFiles/Rider/ToolchainInfo.cs b/Engine/Source/Programs/UnrealBuildTool/ProjectFiles/Rider/ToolchainInfo.cs
--- a/Engine/Source/Programs/UnrealBuildTool/ProjectFiles/Rider/ToolchainInfo.cs
+++ b/Engine/Source/Programs/UnrealBuildTool/ProjectFiles/Rider/ToolchainInfo.cs
@@ -5,6 +5,8 @@ using System.Collections.Generic;
 using System.Linq;
 using System.Reflection;
 
+#nullable disable
+
 namespace UnrealBuildTool
 {
 	sealed class ToolchainInfo : IEquatable<ToolchainInfo>
diff --git a/Engine/Source/Programs/UnrealBuildTool/ProjectFiles/VisualStudio/VCProject.cs b/Engine/Source/Programs/UnrealBuildTool/ProjectFiles/VisualStudio/VCProject.cs
--- a/Engine/Source/Programs/UnrealBuildTool/ProjectFiles/VisualStudio/VCProject.cs
+++ b/Engine/Source/Programs/UnrealBuildTool/ProjectFiles/VisualStudio/VCProject.cs
@@ -13,6 +13,8 @@ using System.Diagnostics;
 using System.Security;
 using Tools.DotNETCommon;
 
+#nullable disable
+
 namespace UnrealBuildTool
 {
 	/// <summary>
@@ -178,8 +180,8 @@ namespace UnrealBuildTool
 				}
 			}
 
-			UEBuildPlatform BuildPlatform = UEBuildPlatform.GetBuildPlatform(Platform, true);
-			if (BuildPlatform == null)
+			UEBuildPlatform BuildPlatform;
+			if (!UEBuildPlatform.TryGetBuildPlatform(Platform, out BuildPlatform))
 			{
 				return false;
 			}
@@ -620,8 +622,8 @@ namespace UnrealBuildTool
 							{
 								continue;
 							}
-							UEBuildPlatform BuildPlatform = UEBuildPlatform.GetBuildPlatform(Platform, true);
-							if ((BuildPlatform != null) && (BuildPlatform.HasRequiredSDKsInstalled() == SDKStatus.Valid))
+							UEBuildPlatform BuildPlatform;
+							if (UEBuildPlatform.TryGetBuildPlatform(Platform, out BuildPlatform) && (BuildPlatform.HasRequiredSDKsInstalled() == SDKStatus.Valid))
 							{
 								// Now go through all of the target types for this project
 								if (ProjectTargets.Count == 0)
diff --git a/Engine/Source/Programs/UnrealBuildTool/ProjectFiles/VisualStudio/VCProjectFileGenerator.cs b/Engine/Source/Programs/UnrealBuildTool/ProjectFiles/VisualStudio/VCProjectFileGenerator.cs
--- a/Engine/Source/Programs/UnrealBuildTool/ProjectFiles/VisualStudio/VCProjectFileGenerator.cs
+++ b/Engine/Source/Programs/UnrealBuildTool/ProjectFiles/VisualStudio/VCProjectFileGenerator.cs
@@ -9,6 +9,8 @@ using System.Text;
 using System.Xml.Linq;
 using Tools.DotNETCommon;
 
+#nullable disable
+
 namespace UnrealBuildTool
 {
 	/// <summary>
@@ -355,8 +357,8 @@ namespace UnrealBuildTool
 				// Allow the SDKs to override
 				foreach (UnrealTargetPlatform SupportedPlatform in SupportedPlatforms)
 				{
-					UEBuildPlatform BuildPlatform = UEBuildPlatform.GetBuildPlatform(SupportedPlatform, true);
-					if (BuildPlatform != null)
+					UEBuildPlatform BuildPlatform;
+					if (UEBuildPlatform.TryGetBuildPlatform(SupportedPlatform, out BuildPlatform))
 					{
 						// Don't worry about platforms that we're missing SDKs for
 						if (BuildPlatform.HasRequiredSDKsInstalled() == SDKStatus.Valid)
diff --git a/Engine/Source/Programs/UnrealBuildTool/ProjectFiles/VisualStudio/VCSolutionOptions.cs b/Engine/Source/Programs/UnrealBuildTool/ProjectFiles/VisualStudio/VCSolutionOptions.cs
--- a/Engine/Source/Programs/UnrealBuildTool/ProjectFiles/VisualStudio/VCSolutionOptions.cs
+++ b/Engine/Source/Programs/UnrealBuildTool/ProjectFiles/VisualStudio/VCSolutionOptions.cs
@@ -14,6 +14,8 @@ using System.Diagnostics;
 using STATSTG = System.Runtime.InteropServices.ComTypes.STATSTG;
 using FILETIME = System.Runtime.InteropServices.ComTypes.FILETIME;
 
+#nullable disable
+
 namespace UnrealBuildTool
 {
 	class VCBinarySetting
diff --git a/Engine/Source/Programs/UnrealBuildTool/ProjectFiles/VisualStudioCode/VSCodeProjectFileGenerator.cs b/Engine/Source/Programs/UnrealBuildTool/ProjectFiles/VisualStudioCode/VSCodeProjectFileGenerator.cs
--- a/Engine/Source/Programs/UnrealBuildTool/ProjectFiles/VisualStudioCode/VSCodeProjectFileGenerator.cs
+++ b/Engine/Source/Programs/UnrealBuildTool/ProjectFiles/VisualStudioCode/VSCodeProjectFileGenerator.cs
@@ -8,6 +8,8 @@ using Tools.DotNETCommon;
 using System.Diagnostics;
 using System.Linq;
 
+#nullable disable
+
 namespace UnrealBuildTool
 {
 	class VSCodeProjectFolder : MasterProjectFolder
@@ -461,7 +463,7 @@ namespace UnrealBuildTool
 
 						foreach (UnrealTargetPlatform Platform in Platforms)
 						{
-							var BuildPlatform = UEBuildPlatform.GetBuildPlatform(Platform, true);
+							UEBuildPlatform.TryGetBuildPlatform(Platform, out UEBuildPlatform BuildPlatform);
 							if (SupportedPlatforms.Contains(Platform) && (BuildPlatform != null) && (BuildPlatform.HasRequiredSDKsInstalled() == SDKStatus.Valid))
 							{
 								foreach (UnrealTargetConfiguration Config in Configs)
diff --git a/Engine/Source/Programs/UnrealBuildTool/ProjectFiles/VisualStudioMac/VisualStudioMacProjectFileGenerator.cs b/Engine/Source/Programs/UnrealBuildTool/ProjectFiles/VisualStudioMac/VisualStudioMacProjectFileGenerator.cs
--- a/Engine/Source/Programs/UnrealBuildTool/ProjectFiles/VisualStudioMac/VisualStudioMacProjectFileGenerator.cs
+++ b/Engine/Source/Programs/UnrealBuildTool/ProjectFiles/VisualStudioMac/VisualStudioMacProjectFileGenerator.cs
@@ -8,6 +8,8 @@ using System.Linq;
 using System.Xml.Linq;
 using Tools.DotNETCommon;
 
+#nullable disable
+
 namespace UnrealBuildTool
 {
     
diff --git a/Engine/Source/Programs/UnrealBuildTool/ProjectFiles/Xcode/XcodeFrameworkWrapperProject.cs b/Engine/Source/Programs/UnrealBuildTool/ProjectFiles/Xcode/XcodeFrameworkWrapperProject.cs
--- a/Engine/Source/Programs/UnrealBuildTool/ProjectFiles/Xcode/XcodeFrameworkWrapperProject.cs
+++ b/Engine/Source/Programs/UnrealBuildTool/ProjectFiles/Xcode/XcodeFrameworkWrapperProject.cs
@@ -7,6 +7,8 @@ using System.Linq;
 using System.Collections.Generic;
 using Tools.DotNETCommon;
 
+#nullable disable
+
 namespace UnrealBuildTool.ProjectFiles.Xcode
 {
 	/// <summary>
diff --git a/Engine/Source/Programs/UnrealBuildTool/ProjectFiles/Xcode/XcodeProject.cs b/Engine/Source/Programs/UnrealBuildTool/ProjectFiles/Xcode/XcodeProject.cs
--- a/Engine/Source/Programs/UnrealBuildTool/ProjectFiles/Xcode/XcodeProject.cs
+++ b/Engine/Source/Programs/UnrealBuildTool/ProjectFiles/Xcode/XcodeProject.cs
@@ -10,6 +10,8 @@ using System.Linq;
 using System.Text;
 using Tools.DotNETCommon;
 
+#nullable disable
+
 namespace UnrealBuildTool
 {
 	/// <summary>
@@ -1441,8 +1443,8 @@ namespace UnrealBuildTool
 					{
 						if (InstalledPlatformInfo.IsValidPlatform(Platform, EProjectType.Code) && (Platform == UnrealTargetPlatform.Mac || Platform == UnrealTargetPlatform.IOS || Platform == UnrealTargetPlatform.TVOS)) // @todo support other platforms
 						{
-							UEBuildPlatform BuildPlatform = UEBuildPlatform.GetBuildPlatform(Platform, true);
-							if ((BuildPlatform != null) && (BuildPlatform.HasRequiredSDKsInstalled() == SDKStatus.Valid))
+							UEBuildPlatform BuildPlatform;
+							if (UEBuildPlatform.TryGetBuildPlatform(Platform, out BuildPlatform) && (BuildPlatform.HasRequiredSDKsInstalled() == SDKStatus.Valid))
 							{
 								// Now go through all of the target types for this project
 								if (ProjectTargets.Count == 0)
diff --git a/Engine/Source/Programs/UnrealBuildTool/ProjectFiles/Xcode/XcodeProjectFileGenerator.cs b/Engine/Source/Programs/UnrealBuildTool/ProjectFiles/Xcode/XcodeProjectFileGenerator.cs
--- a/Engine/Source/Programs/UnrealBuildTool/ProjectFiles/Xcode/XcodeProjectFileGenerator.cs
+++ b/Engine/Source/Programs/UnrealBuildTool/ProjectFiles/Xcode/XcodeProjectFileGenerator.cs
@@ -7,6 +7,8 @@ using System.IO;
 using Tools.DotNETCommon;
 using System.Linq;
 
+#nullable disable
+
 namespace UnrealBuildTool
 {
 	/// <summary>
diff --git a/Engine/Source/Programs/UnrealBuildTool/System/ExternalExecution.cs b/Engine/Source/Programs/UnrealBuildTool/System/ExternalExecution.cs
--- a/Engine/Source/Programs/UnrealBuildTool/System/ExternalExecution.cs
+++ b/Engine/Source/Programs/UnrealBuildTool/System/ExternalExecution.cs
@@ -587,7 +587,7 @@ namespace UnrealBuildTool
 			UHTModuleInfo[] ModuleInfoArray = new UHTModuleInfo[ModulesSortedByType.Count];
 			using(ThreadPoolWorkQueue Queue = new ThreadPoolWorkQueue())
 			{
-				ReadOnlyHashSet<string> ExcludedFolders = UEBuildPlatform.GetBuildPlatform(Platform, true).GetExcludedFolderNames();
+				ReadOnlyHashSet<string> ExcludedFolders = UEBuildPlatform.GetBuildPlatform(Platform).GetExcludedFolderNames();
 				for(int Idx = 0; Idx < ModulesSortedByType.Count; Idx++)
 				{
 					UEBuildModuleCPP Module = ModulesSortedByType[Idx];
diff --git a/Engine/Source/Programs/UnrealBuildTool/System/JunkDeleter.cs b/Engine/Source/Programs/UnrealBuildTool/System/JunkDeleter.cs
--- a/Engine/Source/Programs/UnrealBuildTool/System/JunkDeleter.cs
+++ b/Engine/Source/Programs/UnrealBuildTool/System/JunkDeleter.cs
@@ -64,7 +64,7 @@ namespace UnrealBuildTool
 										// if the platform is valid, then we want to keep the files, which means that we don't want to apply the junk line
 										if (UnrealTargetPlatform.TryParse(InnerTokens[1], out ParsedPlatform))
 										{
-											if (UEBuildPlatform.GetBuildPlatform(ParsedPlatform, bInAllowFailure: true) != null)
+											if (UEBuildPlatform.TryGetBuildPlatform(ParsedPlatform, out _))
 											{
 												// this is a good platform, so don't delete any files!
 												bIsValidJunkLine = false;
diff --git a/Engine/Source/Programs/UnrealBuildTool/System/PlatformExports.cs b/Engine/Source/Programs/UnrealBuildTool/System/PlatformExports.cs
--- a/Engine/Source/Programs/UnrealBuildTool/System/PlatformExports.cs
+++ b/Engine/Source/Programs/UnrealBuildTool/System/PlatformExports.cs
@@ -52,7 +52,7 @@ namespace UnrealBuildTool
 		/// <returns>True if the project uses the default build configuration</returns>
 		public static bool HasDefaultBuildConfig(FileReference ProjectFile, UnrealTargetPlatform Platform)
 		{
-			UEBuildPlatform BuildPlat = UEBuildPlatform.GetBuildPlatform(Platform, true);
+			UEBuildPlatform.TryGetBuildPlatform(Platform, out UEBuildPlatform BuildPlat);
 			return (BuildPlat == null)? true : BuildPlat.HasDefaultBuildConfig(Platform, ProjectFile.Directory);
 		}
 
@@ -64,7 +64,7 @@ namespace UnrealBuildTool
 		/// <returns>True if the project requires a build for the platform</returns>
 		public static bool RequiresBuild(FileReference ProjectFile, UnrealTargetPlatform Platform)
 		{
-			UEBuildPlatform BuildPlat = UEBuildPlatform.GetBuildPlatform(Platform, true);
+			UEBuildPlatform.TryGetBuildPlatform(Platform, out UEBuildPlatform BuildPlat);
 			return (BuildPlat == null) ? false : BuildPlat.RequiresBuild(Platform, ProjectFile.Directory);
 		}
 
@@ -84,7 +84,7 @@ namespace UnrealBuildTool
 		/// <returns>All platform folder names</returns>
 		public static string[] GetIncludedFolderNames(UnrealTargetPlatform Platform)
 		{
-			UEBuildPlatform BuildPlatform = UEBuildPlatform.GetBuildPlatform(Platform, false);
+			UEBuildPlatform BuildPlatform = UEBuildPlatform.GetBuildPlatform(Platform);
 			return BuildPlatform.GetIncludedFolderNames().ToArray();
 		}
 
@@ -95,7 +95,7 @@ namespace UnrealBuildTool
 		/// <returns>Array of folder names</returns>
 		public static string[] GetExcludedFolderNames(UnrealTargetPlatform Platform)
 		{
-			UEBuildPlatform BuildPlatform = UEBuildPlatform.GetBuildPlatform(Platform, false);
+			UEBuildPlatform BuildPlatform = UEBuildPlatform.GetBuildPlatform(Platform);
 			return BuildPlatform.GetExcludedFolderNames().ToArray();
 		}
 
