UBT: Convert UnrealBuildTool to use nullable references.

#rnx

diff --git a/Engine/Source/Programs/UnrealBuildTool/Configuration/ModuleRules.cs b/Engine/Source/Programs/UnrealBuildTool/Configuration/ModuleRules.cs
--- a/Engine/Source/Programs/UnrealBuildTool/Configuration/ModuleRules.cs
+++ b/Engine/Source/Programs/UnrealBuildTool/Configuration/ModuleRules.cs
@@ -2,6 +2,7 @@
 
 using System;
 using System.Collections.Generic;
+using System.Diagnostics.CodeAnalysis;
 using System.IO;
 using System.Linq;
 using System.Text;
@@ -213,7 +214,7 @@ namespace UnrealBuildTool
 			/// <summary>
 			/// The initial location for this file. It will be copied to Path at build time, ready for staging.
 			/// </summary>
-			public string SourcePath;
+			public string? SourcePath;
 
 			/// <summary>
 			/// How to stage this file.
@@ -361,7 +362,7 @@ namespace UnrealBuildTool
 			/// <summary>
 			/// 
 			/// </summary>
-			internal string CopyBundledAssets = null;
+			internal string? CopyBundledAssets = null;
 
 			/// <summary>
 			/// Copy the framework to the target's Framework directory
@@ -375,7 +376,7 @@ namespace UnrealBuildTool
 			/// <param name="Path">Path to a zip file containing the framework or a framework on disk</param>
 			/// <param name="CopyBundledAssets"></param>
 			/// <param name="bCopyFramework">Copy the framework to the target's Framework directory</param>
-			public Framework(string Name, string Path, string CopyBundledAssets = null, bool bCopyFramework = false)
+			public Framework(string Name, string Path, string? CopyBundledAssets = null, bool bCopyFramework = false)
 			{
 				this.Name = Name;
 				this.Path = Path;
@@ -404,7 +405,7 @@ namespace UnrealBuildTool
 			/// <param name="Name">Name of the framework</param>
 			/// <param name="ZipPath">Path to a zip file containing the framework</param>
 			/// <param name="CopyBundledAssets"></param>
-			public UEBuildFramework(string Name, string ZipPath, string CopyBundledAssets = null)
+			public UEBuildFramework(string Name, string ZipPath, string? CopyBundledAssets = null)
 				: base(Name, ZipPath, CopyBundledAssets)
 			{
 			}
@@ -418,12 +419,12 @@ namespace UnrealBuildTool
 			/// <summary>
 			/// 
 			/// </summary>
-			public string ResourcePath = null;
+			public string? ResourcePath = null;
 
 			/// <summary>
 			/// 
 			/// </summary>
-			public string BundleContentsSubdir = null;
+			public string? BundleContentsSubdir = null;
 
 			/// <summary>
 			/// 
@@ -508,18 +509,18 @@ namespace UnrealBuildTool
 		/// <summary>
 		/// File containing this module
 		/// </summary>
-		internal FileReference File;
+		internal FileReference File { get; set; }
 
 		/// <summary>
 		/// Directory containing this module
 		/// </summary>
-		internal DirectoryReference Directory;
+		internal DirectoryReference Directory { get; set; }
 
 		/// <summary>
 		/// Additional directories that contribute to this module (likely in UnrealBuildTool.EnginePlatformExtensionsDirectory). 
 		/// The dictionary tracks module subclasses 
 		/// </summary>
-		internal Dictionary<Type, DirectoryReference> DirectoriesForModuleSubClasses;
+		internal Dictionary<Type, DirectoryReference>? DirectoriesForModuleSubClasses;
 
 		/// <summary>
 		/// Additional directories that contribute to this module but are not based on a subclass (NotForLicensees, etc)
@@ -529,7 +530,7 @@ namespace UnrealBuildTool
 		/// <summary>
 		/// Plugin containing this module
 		/// </summary>
-		internal PluginInfo Plugin;
+		internal PluginInfo? Plugin;
 
 		/// <summary>
 		/// True if a Plugin contains this module
@@ -545,7 +546,7 @@ namespace UnrealBuildTool
 		/// <summary>
 		/// The rules context for this instance
 		/// </summary>
-		internal ModuleRulesContext Context;
+		internal ModuleRulesContext Context { get; set; }
 
 		/// <summary>
 		/// Rules for the target that this module belongs to
@@ -626,18 +627,18 @@ namespace UnrealBuildTool
 		/// <summary>
 		/// Explicit private PCH for this module. Implies that this module will not use a shared PCH.
 		/// </summary>
-		public string PrivatePCHHeaderFile;
+		public string? PrivatePCHHeaderFile;
 
 		/// <summary>
 		/// Header file name for a shared PCH provided by this module.  Must be a valid relative path to a public C++ header file.
 		/// This should only be set for header files that are included by a significant number of other C++ modules.
 		/// </summary>
-		public string SharedPCHHeaderFile;
-		
+		public string? SharedPCHHeaderFile;
+
 		/// <summary>
 		/// Specifies an alternate name for intermediate directories and files for intermediates of this module. Useful when hitting path length limitations.
 		/// </summary>
-		public string ShortName = null;
+		public string? ShortName = null;
 
 		/// <summary>
 		/// Precompiled header usage for this module
@@ -651,16 +652,16 @@ namespace UnrealBuildTool
 					// Use the override
 					return PCHUsagePrivate.Value;
 				}
-				else if(Target.bEnableCppModules)
+				else if (Target.bEnableCppModules)
 				{
 					return PCHUsageMode.NoPCHs;
 				}
-				else if(Target.bIWYU || DefaultBuildSettings >= BuildSettingsVersion.V2)
+				else if (Target.bIWYU || DefaultBuildSettings >= BuildSettingsVersion.V2)
 				{
 					// Use shared or explicit PCHs, and enable IWYU
 					return PCHUsageMode.UseExplicitOrSharedPCHs;
 				}
-				else if(Plugin != null)
+				else if (Plugin != null)
 				{
 					// Older plugins use shared PCHs by default, since they aren't typically large enough to warrant their own PCH.
 					return PCHUsageMode.UseSharedPCHs;
@@ -699,7 +700,7 @@ namespace UnrealBuildTool
 		public bool bUseBackwardsCompatibleDefaults
 		{
 			get { return DefaultBuildSettings != BuildSettingsVersion.Latest; }
-			set { DefaultBuildSettings = bUseBackwardsCompatibleDefaults? BuildSettingsVersion.V1 : BuildSettingsVersion.Latest; }
+			set { DefaultBuildSettings = bUseBackwardsCompatibleDefaults ? BuildSettingsVersion.V1 : BuildSettingsVersion.Latest; }
 		}
 
 		/// <summary>
@@ -1015,7 +1016,7 @@ namespace UnrealBuildTool
 		public void AppendStringToPublicDefinition(string Definition, string Text)
 		{
 			string WithEquals = Definition + "=";
-			for (int Index=0; Index < PublicDefinitions.Count; Index++)
+			for (int Index = 0; Index < PublicDefinitions.Count; Index++)
 			{
 				if (PublicDefinitions[Index].StartsWith(WithEquals))
 				{
@@ -1070,7 +1071,7 @@ namespace UnrealBuildTool
 		/// <summary>
 		/// Subclass rules files which invalidate the makefile if modified.
 		/// </summary>
-		public List<string> SubclassRules;
+		public List<string>? SubclassRules;
 
 		/// <summary>
 		/// Whether this module requires the IMPLEMENT_MODULE macro to be implemented. Most UE4 modules require this, since we use the IMPLEMENT_MODULE macro
@@ -1085,7 +1086,7 @@ namespace UnrealBuildTool
 		public bool bLegacyPublicIncludePaths
 		{
 			set { bLegacyPublicIncludePathsPrivate = value; }
-			get { return bLegacyPublicIncludePathsPrivate ?? ((DefaultBuildSettings < BuildSettingsVersion.V2)? Target.bLegacyPublicIncludePaths : false); }
+			get { return bLegacyPublicIncludePathsPrivate ?? ((DefaultBuildSettings < BuildSettingsVersion.V2) ? Target.bLegacyPublicIncludePaths : false); }
 		}
 		private bool? bLegacyPublicIncludePathsPrivate;
 
@@ -1102,11 +1103,11 @@ namespace UnrealBuildTool
 		/// <summary>
 		/// The AutoSDK directory for the active host platform
 		/// </summary>
-		public string AutoSdkDirectory
+		public string? AutoSdkDirectory
 		{
 			get
 			{
-				DirectoryReference AutoSdkDir;
+				DirectoryReference? AutoSdkDir;
 				return UEBuildPlatformSDK.TryGetHostPlatformAutoSDKDir(out AutoSdkDir) ? AutoSdkDir.FullName : null;
 			}
 		}
@@ -1129,7 +1130,7 @@ namespace UnrealBuildTool
 		{
 			get
 			{
-				if(Plugin == null)
+				if (Plugin == null)
 				{
 					throw new BuildException("Module '{0}' does not belong to a plugin; PluginDirectory property is invalid.", Name);
 				}
@@ -1159,6 +1160,11 @@ namespace UnrealBuildTool
 		public ModuleRules(ReadOnlyTargetRules Target)
 		{
 			this.Target = Target;
+
+			Name = Name!;
+			File = File!;
+			Directory = Directory!;
+			Context = Context!;
 		}
 
 		/// <summary>
@@ -1406,14 +1412,14 @@ namespace UnrealBuildTool
 		/// </summary>
 		/// <param name="Type">typeof the subclass</param>
 		/// <returns>Directory where the subclass's .Build.cs lives, or null if not found</returns>
-		public DirectoryReference GetModuleDirectoryForSubClass(Type Type)
+		public DirectoryReference? GetModuleDirectoryForSubClass(Type Type)
 		{
 			if (DirectoriesForModuleSubClasses == null)
 			{
 				return null;
 			}
 
-			DirectoryReference Directory;
+			DirectoryReference? Directory;
 			if (DirectoriesForModuleSubClasses.TryGetValue(Type, out Directory))
 			{
 				return Directory;
diff --git a/Engine/Source/Programs/UnrealBuildTool/Configuration/ModuleRulesContext.cs b/Engine/Source/Programs/UnrealBuildTool/Configuration/ModuleRulesContext.cs
--- a/Engine/Source/Programs/UnrealBuildTool/Configuration/ModuleRulesContext.cs
+++ b/Engine/Source/Programs/UnrealBuildTool/Configuration/ModuleRulesContext.cs
@@ -27,7 +27,7 @@ namespace UnrealBuildTool
 		/// <summary>
 		/// The plugin that this module belongs to
 		/// </summary>
-		public PluginInfo Plugin;
+		public PluginInfo? Plugin;
 
 		/// <summary>
 		/// Whether this module should be included in the default hot reload set
diff --git a/Engine/Source/Programs/UnrealBuildTool/Configuration/PlatformTargetRules.cs b/Engine/Source/Programs/UnrealBuildTool/Configuration/PlatformTargetRules.cs
--- a/Engine/Source/Programs/UnrealBuildTool/Configuration/PlatformTargetRules.cs
+++ b/Engine/Source/Programs/UnrealBuildTool/Configuration/PlatformTargetRules.cs
@@ -1,4 +1,4 @@
-﻿// Copyright Epic Games, Inc. All Rights Reserved.
+// Copyright Epic Games, Inc. All Rights Reserved.
 
 using System;
 using System.Collections.Generic;
diff --git a/Engine/Source/Programs/UnrealBuildTool/Configuration/TargetDescriptor.cs b/Engine/Source/Programs/UnrealBuildTool/Configuration/TargetDescriptor.cs
--- a/Engine/Source/Programs/UnrealBuildTool/Configuration/TargetDescriptor.cs
+++ b/Engine/Source/Programs/UnrealBuildTool/Configuration/TargetDescriptor.cs
@@ -1,6 +1,7 @@
 // Copyright Epic Games, Inc. All Rights Reserved.
 using System;
 using System.Collections.Generic;
+using System.Diagnostics.CodeAnalysis;
 using System.IO;
 using System.Linq;
 using System.Text;
@@ -14,7 +15,7 @@ namespace UnrealBuildTool
 	/// </summary>
 	class TargetDescriptor
 	{
-		public FileReference ProjectFile;
+		public FileReference? ProjectFile;
 		public string Name;
 		public UnrealTargetPlatform Platform;
 		public UnrealTargetConfiguration Configuration;
@@ -25,7 +26,7 @@ namespace UnrealBuildTool
 		/// Foreign plugin to compile against this target
 		/// </summary>
 		[CommandLine("-Plugin=")]
-		public FileReference ForeignPlugin = null;
+		public FileReference? ForeignPlugin = null;
 
 		/// <summary>
 		/// Set of module names to compile.
@@ -56,13 +57,13 @@ namespace UnrealBuildTool
 		/// Path to a file containing a list of modules that may be modified for live coding.
 		/// </summary>
 		[CommandLine("-LiveCodingModules=")]
-		public FileReference LiveCodingModules = null;
+		public FileReference? LiveCodingModules = null;
 
 		/// <summary>
 		/// Path to the manifest for passing info about the output to live coding
 		/// </summary>
 		[CommandLine("-LiveCodingManifest=")]
-		public FileReference LiveCodingManifest = null;
+		public FileReference? LiveCodingManifest = null;
 
 		/// <summary>
 		/// If a non-zero value, a live coding request will be terminated if more than the given number of actions are required.
@@ -85,7 +86,7 @@ namespace UnrealBuildTool
 		/// <param name="Configuration">Configuration to build</param>
 		/// <param name="Architecture">Architecture to build for</param>
 		/// <param name="Arguments">Other command-line arguments for the target</param>
-		public TargetDescriptor(FileReference ProjectFile, string TargetName, UnrealTargetPlatform Platform, UnrealTargetConfiguration Configuration, string Architecture, CommandLineArguments Arguments)
+		public TargetDescriptor(FileReference? ProjectFile, string TargetName, UnrealTargetPlatform Platform, UnrealTargetConfiguration Configuration, string Architecture, CommandLineArguments? Arguments)
 		{
 			this.ProjectFile = ProjectFile;
 			this.Name = TargetName;
@@ -208,7 +209,7 @@ namespace UnrealBuildTool
 			List<UnrealTargetPlatform> Platforms = new List<UnrealTargetPlatform>();
 			List<UnrealTargetConfiguration> Configurations = new List<UnrealTargetConfiguration>();
 			List<string> TargetNames = new List<string>();
-			FileReference ProjectFile = Arguments.GetFileReferenceOrDefault("-Project=", null);
+			FileReference? ProjectFile = Arguments.GetFileReferenceOrDefault("-Project=", null);
 
 			// Settings for creating/using static libraries for the engine
 			for (int ArgumentIndex = 0; ArgumentIndex < Arguments.Count; ArgumentIndex++)
@@ -363,7 +364,7 @@ namespace UnrealBuildTool
 		/// <param name="Arguments">The command line arguments</param>
 		/// <param name="ProjectFile">The project file that was parsed</param>
 		/// <returns>True if the project file was parsed, false otherwise</returns>
-		public static bool TryParseProjectFileArgument(CommandLineArguments Arguments, out FileReference ProjectFile)
+		public static bool TryParseProjectFileArgument(CommandLineArguments Arguments, out FileReference? ProjectFile)
 		{
 			FileReference ExplicitProjectFile;
 			if(Arguments.TryGetValue("-Project=", out ExplicitProjectFile))
@@ -399,7 +400,7 @@ namespace UnrealBuildTool
 		/// <param name="Prefix">The argument prefix, eg. "-Foo="</param>
 		/// <param name="Value">Receives the value of the argument</param>
 		/// <returns>True if the argument could be parsed, false otherwise</returns>
-		private static bool ParseArgumentValue(string Argument, string Prefix, out string Value)
+		private static bool ParseArgumentValue(string Argument, string Prefix, [NotNullWhen(true)] out string? Value)
 		{
 			if(Argument.StartsWith(Prefix, StringComparison.InvariantCultureIgnoreCase))
 			{
@@ -438,16 +439,16 @@ namespace UnrealBuildTool
 
 		public override int GetHashCode()
 		{
-			return ProjectFile.GetHashCode() + 
+			return String.GetHashCode(ProjectFile?.FullName) + 
 				Name.GetHashCode() + 
 				Platform.GetHashCode() + 
 				Configuration.GetHashCode() + 
 				Architecture.GetHashCode();
 		}
 
-		public override bool Equals(object Obj) 
+		public override bool Equals(object? Obj) 
 		{
-			TargetDescriptor Other = Obj as TargetDescriptor;
+			TargetDescriptor? Other = Obj as TargetDescriptor;
 			if (Other != null)
 			{
 				return
diff --git a/Engine/Source/Programs/UnrealBuildTool/Configuration/TargetInfo.cs b/Engine/Source/Programs/UnrealBuildTool/Configuration/TargetInfo.cs
--- a/Engine/Source/Programs/UnrealBuildTool/Configuration/TargetInfo.cs
+++ b/Engine/Source/Programs/UnrealBuildTool/Configuration/TargetInfo.cs
@@ -38,7 +38,7 @@ namespace UnrealBuildTool
 		/// <summary>
 		/// The project containing the target
 		/// </summary>
-		public readonly FileReference ProjectFile;
+		public readonly FileReference? ProjectFile;
 
 		/// <summary>
 		/// The current build version
@@ -51,7 +51,7 @@ namespace UnrealBuildTool
 		/// <summary>
 		/// Additional command line arguments for this target
 		/// </summary>
-		public CommandLineArguments Arguments;
+		public CommandLineArguments? Arguments;
 
 		/// <summary>
 		/// Constructs a TargetInfo for passing to the TargetRules constructor.
@@ -62,7 +62,7 @@ namespace UnrealBuildTool
 		/// <param name="Architecture">The architecture being built for</param>
 		/// <param name="ProjectFile">Path to the project file containing the target</param>
 		/// <param name="Arguments">Additional command line arguments for this target</param>
-		public TargetInfo(string Name, UnrealTargetPlatform Platform, UnrealTargetConfiguration Configuration, string Architecture, FileReference ProjectFile, CommandLineArguments Arguments)
+		public TargetInfo(string Name, UnrealTargetPlatform Platform, UnrealTargetConfiguration Configuration, string Architecture, FileReference? ProjectFile, CommandLineArguments? Arguments)
 		{
 			this.Name = Name;
 			this.Platform = Platform;
@@ -105,7 +105,7 @@ namespace UnrealBuildTool
 			Writer.WriteString(Configuration.ToString());
 			Writer.WriteString(Architecture);
 			Writer.WriteFileReference(ProjectFile);
-			Writer.WriteArray(Arguments.GetRawArray(), Item => Writer.WriteString(Item));
+			Writer.WriteArray(Arguments?.GetRawArray(), Item => Writer.WriteString(Item));
 		}
 	}
 }
diff --git a/Engine/Source/Programs/UnrealBuildTool/Configuration/TargetRules.cs b/Engine/Source/Programs/UnrealBuildTool/Configuration/TargetRules.cs
--- a/Engine/Source/Programs/UnrealBuildTool/Configuration/TargetRules.cs
+++ b/Engine/Source/Programs/UnrealBuildTool/Configuration/TargetRules.cs
@@ -182,14 +182,14 @@ namespace UnrealBuildTool
 		/// </summary>
 		public bool IsNameOverriden() { return !String.IsNullOrEmpty(NameOverride); }
 
-		private string NameOverride;
+		private string? NameOverride;
 
 		private readonly string DefaultName;
 
 		/// <summary>
 		/// File containing this target
 		/// </summary>
-		internal FileReference File;
+		internal FileReference File { get; set; }
 
 		/// <summary>
 		/// Platform that this target is being built for.
@@ -209,7 +209,7 @@ namespace UnrealBuildTool
 		/// <summary>
 		/// Path to the project file for the project containing this target.
 		/// </summary>
-		public readonly FileReference ProjectFile;
+		public readonly FileReference? ProjectFile;
 
 		/// <summary>
 		/// The current build version
@@ -1187,7 +1187,7 @@ namespace UnrealBuildTool
 		/// Bundle version for Mac apps.
 		/// </summary>
 		[CommandLine("-BundleVersion")]
-		public string BundleVersion = null;
+		public string? BundleVersion = null;
 
 		/// <summary>
 		/// Whether to deploy the executable after compilation on platforms that require deployment.
@@ -1253,7 +1253,7 @@ namespace UnrealBuildTool
 		/// Allows overriding the toolchain to be created for this target. This must match the name of a class declared in the UnrealBuildTool assembly.
 		/// </summary>
 		[CommandLine("-ToolChain")]
-		public string ToolChainName = null;
+		public string? ToolChainName = null;
 
 		/// <summary>
 		/// Whether to allow engine configuration to determine if we can load unverified certificates.
@@ -1298,7 +1298,7 @@ namespace UnrealBuildTool
 		/// The build version string
 		/// </summary>
 		[CommandLine("-BuildVersion")]
-		public string BuildVersion;
+		public string? BuildVersion;
 
 		/// <summary>
 		/// Specifies how to link modules in this target (monolithic or modular). This is currently protected for backwards compatibility. Call the GetLinkType() accessor
@@ -1340,7 +1340,7 @@ namespace UnrealBuildTool
 		/// <summary>
 		/// Specifies the name of the launch module. For modular builds, this is the module that is compiled into the target's executable.
 		/// </summary>
-		public string LaunchModuleName
+		public string? LaunchModuleName
 		{
 			get
 			{
@@ -1355,12 +1355,12 @@ namespace UnrealBuildTool
 		/// <summary>
 		/// Backing storage for the LaunchModuleName property.
 		/// </summary>
-		private string LaunchModuleNamePrivate;
+		private string? LaunchModuleNamePrivate;
 
 		/// <summary>
 		/// Specifies the path to write a header containing public definitions for this target. Useful when building a DLL to be consumed by external build processes.
 		/// </summary>
-		public string ExportPublicHeader;
+		public string? ExportPublicHeader;
 
 		/// <summary>
 		/// List of additional modules to be compiled into the target.
@@ -1451,38 +1451,38 @@ namespace UnrealBuildTool
 		/// </summary>
 		[RequiresUniqueBuildEnvironment]
 		[CommandLine("-CompilerArguments=")]
-		public string AdditionalCompilerArguments;
+		public string? AdditionalCompilerArguments;
 
 		/// <summary>
 		/// Additional arguments to pass to the linker
 		/// </summary>
 		[RequiresUniqueBuildEnvironment]
 		[CommandLine("-LinkerArguments=")]
-		public string AdditionalLinkerArguments;
+		public string? AdditionalLinkerArguments;
 
 
 		/// <summary>
 		/// List of modules to disable unity builds for
 		/// </summary>
 		[XmlConfigFile(Category = "ModuleConfiguration", Name = "DisableUnityBuild")]
-		public string[] DisableUnityBuildForModules = null;
+		public string[]? DisableUnityBuildForModules = null;
 
 		/// <summary>
 		///  List of modules to enable optimizations for
 		/// </summary>
 		[XmlConfigFile(Category = "ModuleConfiguration", Name = "EnableOptimizeCode")]
-		public string[] EnableOptimizeCodeForModules = null;
+		public string[]? EnableOptimizeCodeForModules = null;
 
 		/// <summary>
 		/// List of modules to disable optimizations for
 		/// </summary>
 		[XmlConfigFile(Category = "ModuleConfiguration", Name = "DisableOptimizeCode")]
-		public string[] DisableOptimizeCodeForModules = null;
+		public string[]? DisableOptimizeCodeForModules = null;
 
 		/// <summary>
 		/// When generating project files, specifies the name of the project file to use when there are multiple targets of the same type.
 		/// </summary>
-		public string GeneratedProjectName;
+		public string? GeneratedProjectName;
 
 		/// <summary>
 		/// Android-specific target settings.
@@ -1559,6 +1559,9 @@ namespace UnrealBuildTool
 			this.WindowsPlatform = new WindowsTargetRules(this);
 			this.HoloLensPlatform = new HoloLensTargetRules(Target);
 
+			// Suppress nullable warnings (initialized before constructor is run)
+			File = File!;
+
 			// Read settings from config files
 			foreach (object ConfigurableObject in GetConfigurableObjects())
 			{
@@ -1592,7 +1595,7 @@ namespace UnrealBuildTool
 			// Get the directory to use for crypto settings. We can build engine targets (eg. UHT) with 
 			// a project file, but we can't use that to determine crypto settings without triggering
 			// constant rebuilds of UHT.
-			DirectoryReference CryptoSettingsDir = DirectoryReference.FromFile(ProjectFile);
+			DirectoryReference? CryptoSettingsDir = DirectoryReference.FromFile(ProjectFile);
 			if (CryptoSettingsDir != null && File != null && !File.IsUnderDirectory(CryptoSettingsDir))
 			{
 				CryptoSettingsDir = null;
@@ -1661,7 +1664,7 @@ namespace UnrealBuildTool
 		internal UnrealTargetPlatform[] GetSupportedPlatforms()
 		{
 			// Otherwise take the SupportedPlatformsAttribute from the first type in the inheritance chain that supports it
-			for (Type CurrentType = GetType(); CurrentType != null; CurrentType = CurrentType.BaseType)
+			for (Type? CurrentType = GetType(); CurrentType != null; CurrentType = CurrentType.BaseType)
 			{
 				object[] Attributes = CurrentType.GetCustomAttributes(typeof(SupportedPlatformsAttribute), false);
 				if (Attributes.Length > 0)
@@ -1692,7 +1695,7 @@ namespace UnrealBuildTool
 		internal UnrealTargetConfiguration[] GetSupportedConfigurations()
 		{
 			// Otherwise take the SupportedConfigurationsAttribute from the first type in the inheritance chain that supports it
-			for (Type CurrentType = GetType(); CurrentType != null; CurrentType = CurrentType.BaseType)
+			for (Type? CurrentType = GetType(); CurrentType != null; CurrentType = CurrentType.BaseType)
 			{
 				object[] Attributes = CurrentType.GetCustomAttributes(typeof(SupportedConfigurationsAttribute), false);
 				if (Attributes.Length > 0)
@@ -1724,7 +1727,11 @@ namespace UnrealBuildTool
 			{
 				if(Field.GetCustomAttribute<ConfigSubObjectAttribute>() != null)
 				{
-					yield return Field.GetValue(this);
+					object? Value = Field.GetValue(this);
+					if(Value != null)
+					{
+						yield return Value;
+					}
 				}
 			}
 		}
@@ -1849,7 +1856,7 @@ namespace UnrealBuildTool
 			get { return Inner.Architecture; }
 		}
 
-		public FileReference ProjectFile
+		public FileReference? ProjectFile
 		{
 			get { return Inner.ProjectFile; }
 		}
@@ -2493,7 +2500,7 @@ namespace UnrealBuildTool
 			get { return Inner.bAllowRuntimeSymbolFiles; }
 		}
 
-		public string BundleVersion
+		public string? BundleVersion
 		{
 			get { return Inner.BundleVersion; }
 		}
@@ -2554,7 +2561,7 @@ namespace UnrealBuildTool
 			get { return Inner.bPublicSymbolsByDefault; }
 		}
 
-		public string ToolChainName
+		public string? ToolChainName
 		{
 			get { return Inner.ToolChainName; }
 		}
@@ -2574,7 +2581,7 @@ namespace UnrealBuildTool
 			get { return Inner.bNoManifestChanges; }
 		}
 
-		public string BuildVersion
+		public string? BuildVersion
 		{
 			get { return Inner.BuildVersion; }
 		}
@@ -2594,12 +2601,12 @@ namespace UnrealBuildTool
 			get { return Inner.ProjectDefinitions.AsReadOnly(); }
 		}
 
-		public string LaunchModuleName
+		public string? LaunchModuleName
 		{
 			get { return Inner.LaunchModuleName; }
 		}
 
-		public string ExportPublicHeader
+		public string? ExportPublicHeader
 		{
 			get { return Inner.ExportPublicHeader; }
 		}
@@ -2649,17 +2656,17 @@ namespace UnrealBuildTool
 			get { return Inner.AdditionalBuildProducts; }
 		}
 
-		public string AdditionalCompilerArguments
+		public string? AdditionalCompilerArguments
 		{
 			get { return Inner.AdditionalCompilerArguments; }
 		}
 
-		public string AdditionalLinkerArguments
+		public string? AdditionalLinkerArguments
 		{
 			get { return Inner.AdditionalLinkerArguments; }
 		}
 
-		public string GeneratedProjectName
+		public string? GeneratedProjectName
 		{
 			get { return Inner.GeneratedProjectName; }
 		}
@@ -2740,17 +2747,17 @@ namespace UnrealBuildTool
 		}
 
 
-		public IReadOnlyList<string> DisableUnityBuildForModules
+		public IReadOnlyList<string>? DisableUnityBuildForModules
 		{
 			get { return Inner.DisableUnityBuildForModules; }
 		}
 
-		public IReadOnlyList<string> EnableOptimizeCodeForModules
+		public IReadOnlyList<string>? EnableOptimizeCodeForModules
 		{
 			get { return Inner.EnableOptimizeCodeForModules; }
 		}
 
-		public IReadOnlyList<string> DisableOptimizeCodeForModules
+		public IReadOnlyList<string>? DisableOptimizeCodeForModules
 		{
 			get { return Inner.DisableOptimizeCodeForModules; }
 		}
diff --git a/Engine/Source/Programs/UnrealBuildTool/Configuration/UEBuildBinary.cs b/Engine/Source/Programs/UnrealBuildTool/Configuration/UEBuildBinary.cs
--- a/Engine/Source/Programs/UnrealBuildTool/Configuration/UEBuildBinary.cs
+++ b/Engine/Source/Programs/UnrealBuildTool/Configuration/UEBuildBinary.cs
@@ -117,7 +117,7 @@ namespace UnrealBuildTool
 		/// <summary>
 		/// Cached list of dependent link libraries.
 		/// </summary>
-		private List<FileReference> DependentLinkLibraries;
+		private List<FileReference>? DependentLinkLibraries;
 
 		/// <summary>
 		/// Create an instance of the class with the given configuration data
@@ -302,7 +302,7 @@ namespace UnrealBuildTool
 						// Add actions to copy everything
 						foreach(KeyValuePair<FileReference, FileReference> Pair in Mapping)
 						{
-							FileReference ExistingSourceFile;
+							FileReference? ExistingSourceFile;
 							if(!TargetFileToSourceFile.TryGetValue(Pair.Key, out ExistingSourceFile))
 							{
 								TargetFileToSourceFile[Pair.Key] = Pair.Value;
@@ -397,7 +397,7 @@ namespace UnrealBuildTool
 		/// </summary>
 		/// <param name="ReferencedBy">Map of module to the module that referenced it</param>
 		/// <returns>List of all referenced modules</returns>
-		public void FindModuleReferences(Dictionary<UEBuildModule, UEBuildModule> ReferencedBy)
+		public void FindModuleReferences(Dictionary<UEBuildModule, UEBuildModule?> ReferencedBy)
 		{
 			List<UEBuildModule> ReferencedModules = new List<UEBuildModule>();
 			foreach(UEBuildModule Module in Modules)
@@ -596,7 +596,7 @@ namespace UnrealBuildTool
 		public bool CheckRestrictedFolders(List<DirectoryReference> RootDirectories, Dictionary<UEBuildModule, Dictionary<RestrictedFolder, DirectoryReference>> ModuleRestrictedFolderCache)
 		{
 			// Find all the modules we depend on
-			Dictionary<UEBuildModule, UEBuildModule> ModuleReferencedBy = new Dictionary<UEBuildModule, UEBuildModule>();
+			Dictionary<UEBuildModule, UEBuildModule?> ModuleReferencedBy = new Dictionary<UEBuildModule, UEBuildModule?>();
 			FindModuleReferences(ModuleReferencedBy);
 
 			// Loop through each of the output binaries and check them separately
@@ -616,7 +616,7 @@ namespace UnrealBuildTool
 				List<RestrictedFolder> AliasedBinaryFolders = new List<RestrictedFolder>();
 				foreach (RestrictedFolder BinaryFolder in BinaryFolders)
 				{
-					string Alias;
+					string? Alias;
 					if (PrimaryModule.AliasRestrictedFolders.TryGetValue(BinaryFolder.ToString(), out Alias))
 					{
 						foreach(RestrictedFolder Folder in RestrictedFolder.GetValues())
@@ -634,7 +634,7 @@ namespace UnrealBuildTool
 				foreach(UEBuildModule Module in ModuleReferencedBy.Keys)
 				{
 					// Find the restricted folders for this module
-					Dictionary<RestrictedFolder, DirectoryReference> ModuleRestrictedFolders;
+					Dictionary<RestrictedFolder, DirectoryReference>? ModuleRestrictedFolders;
 					if (!ModuleRestrictedFolderCache.TryGetValue(Module, out ModuleRestrictedFolders))
 					{
 						ModuleRestrictedFolders = Module.FindRestrictedFolderReferences(RootDirectories);
@@ -647,7 +647,7 @@ namespace UnrealBuildTool
 						if(!BinaryFolders.Contains(Pair.Key))
 						{
 							List<string> ReferenceChain = new List<string>();
-							for(UEBuildModule ReferencedModule = Module; ReferencedModule != null; ReferencedModule = ModuleReferencedBy[ReferencedModule])
+							for(UEBuildModule? ReferencedModule = Module; ReferencedModule != null; ReferencedModule = ModuleReferencedBy[ReferencedModule])
 							{
 								ReferenceChain.Insert(0, ReferencedModule.Name);
 							}
@@ -737,7 +737,7 @@ namespace UnrealBuildTool
 					}
 
 					// Force a reference to initialize module for this binary
-					if(Module.Rules.bRequiresImplementModule.Value)
+					if(Module.Rules.bRequiresImplementModule ?? true)
 					{
 						BinaryLinkEnvironment.IncludeFunctions.Add(String.Format("IMPLEMENT_MODULE_{0}", Module.Name));
 					}
diff --git a/Engine/Source/Programs/UnrealBuildTool/Configuration/UEBuildBundleResource.cs b/Engine/Source/Programs/UnrealBuildTool/Configuration/UEBuildBundleResource.cs
--- a/Engine/Source/Programs/UnrealBuildTool/Configuration/UEBuildBundleResource.cs
+++ b/Engine/Source/Programs/UnrealBuildTool/Configuration/UEBuildBundleResource.cs
@@ -1,4 +1,4 @@
-﻿// Copyright Epic Games, Inc. All Rights Reserved.
+// Copyright Epic Games, Inc. All Rights Reserved.
 
 using System;
 using System.Collections.Generic;
@@ -27,12 +27,12 @@ namespace UnrealBuildTool
 		/// <summary>
 		/// 
 		/// </summary>
-		public string ResourcePath = null;
+		public string? ResourcePath = null;
 
 		/// <summary>
 		/// 
 		/// </summary>
-		public string BundleContentsSubdir = null;
+		public string? BundleContentsSubdir = null;
 
 		/// <summary>
 		/// 
diff --git a/Engine/Source/Programs/UnrealBuildTool/Configuration/UEBuildDeploy.cs b/Engine/Source/Programs/UnrealBuildTool/Configuration/UEBuildDeploy.cs
--- a/Engine/Source/Programs/UnrealBuildTool/Configuration/UEBuildDeploy.cs
+++ b/Engine/Source/Programs/UnrealBuildTool/Configuration/UEBuildDeploy.cs
@@ -1,4 +1,4 @@
-﻿// Copyright Epic Games, Inc. All Rights Reserved.
+// Copyright Epic Games, Inc. All Rights Reserved.
 
 using System;
 using System.Collections.Generic;
diff --git a/Engine/Source/Programs/UnrealBuildTool/Configuration/UEBuildFramework.cs b/Engine/Source/Programs/UnrealBuildTool/Configuration/UEBuildFramework.cs
--- a/Engine/Source/Programs/UnrealBuildTool/Configuration/UEBuildFramework.cs
+++ b/Engine/Source/Programs/UnrealBuildTool/Configuration/UEBuildFramework.cs
@@ -22,17 +22,17 @@ namespace UnrealBuildTool
 		/// <summary>
 		/// Path to a zip file containing the framework. May be null.
 		/// </summary>
-		public readonly FileReference ZipFile;
+		public readonly FileReference? ZipFile;
 
 		/// <summary>
 		/// Path to the framework on disk.
 		/// </summary>
-		public readonly DirectoryReference FrameworkDirectory;
+		public readonly DirectoryReference? FrameworkDirectory;
 
 		/// <summary>
 		/// 
 		/// </summary>
-		public readonly string CopyBundledAssets;
+		public readonly string? CopyBundledAssets;
 
 		/// <summary>
 		/// Copy the framework to the target's Framework directory
@@ -42,14 +42,14 @@ namespace UnrealBuildTool
 		/// <summary>
 		/// File created after the framework has been extracted. Used to add dependencies into the action graph.
 		/// </summary>
-		public FileItem ExtractedTokenFile;
+		public FileItem? ExtractedTokenFile;
 
 		/// <summary>
 		/// Constructor
 		/// </summary>
 		/// <param name="Name">The framework name</param>
 		/// <param name="CopyBundledAssets"></param>
-		public UEBuildFramework(string Name, string CopyBundledAssets = null)
+		public UEBuildFramework(string Name, string? CopyBundledAssets = null)
 		{
 			this.Name = Name;
 			this.CopyBundledAssets = CopyBundledAssets;
@@ -63,7 +63,7 @@ namespace UnrealBuildTool
 		/// <param name="OutputDirectory">Path for the extracted zip file</param>
 		/// <param name="CopyBundledAssets"></param>
 		/// <param name="bCopyFramework">Copy the framework to the target's Framework directory</param>
-		public UEBuildFramework(string Name, FileReference ZipFile, DirectoryReference OutputDirectory, string CopyBundledAssets, bool bCopyFramework)
+		public UEBuildFramework(string Name, FileReference? ZipFile, DirectoryReference? OutputDirectory, string? CopyBundledAssets, bool bCopyFramework)
 		{
 			this.Name = Name;
 			this.ZipFile = ZipFile;
@@ -79,7 +79,7 @@ namespace UnrealBuildTool
 		/// <param name="FrameworkDirectory">Path for the framework on disk</param>
 		/// <param name="CopyBundledAssets"></param>
 		/// <param name="bCopyFramework">Copy the framework to the target's Framework directory</param>
-		public UEBuildFramework(String Name, DirectoryReference FrameworkDirectory, string CopyBundledAssets, bool bCopyFramework)
+		public UEBuildFramework(String Name, DirectoryReference? FrameworkDirectory, string? CopyBundledAssets, bool bCopyFramework)
 		{
 			this.Name = Name;
 			this.FrameworkDirectory = FrameworkDirectory;
diff --git a/Engine/Source/Programs/UnrealBuildTool/Configuration/UEBuildModule.cs b/Engine/Source/Programs/UnrealBuildTool/Configuration/UEBuildModule.cs
--- a/Engine/Source/Programs/UnrealBuildTool/Configuration/UEBuildModule.cs
+++ b/Engine/Source/Programs/UnrealBuildTool/Configuration/UEBuildModule.cs
@@ -60,7 +60,7 @@ namespace UnrealBuildTool
 		/// <summary>
 		/// The binary the module will be linked into for the current target.  Only set after UEBuildBinary.BindModules is called.
 		/// </summary>
-		public UEBuildBinary Binary = null;
+		public UEBuildBinary Binary = null!;
 
 		/// <summary>
 		/// The name of the _API define for this module
@@ -130,12 +130,12 @@ namespace UnrealBuildTool
 		/// <summary>
 		/// Names of modules with header files that this module's public interface needs access to.
 		/// </summary>
-		public List<UEBuildModule> PublicIncludePathModules;
+		public List<UEBuildModule>? PublicIncludePathModules;
 
 		/// <summary>
 		/// Names of modules that this module's public interface depends on.
 		/// </summary>
-		public List<UEBuildModule> PublicDependencyModules;
+		public List<UEBuildModule>? PublicDependencyModules;
 
 		/// <summary>
 		/// Names of DLLs that this module should delay load
@@ -145,17 +145,17 @@ namespace UnrealBuildTool
 		/// <summary>
 		/// Names of modules with header files that this module's private implementation needs access to.
 		/// </summary>
-		public List<UEBuildModule> PrivateIncludePathModules;
+		public List<UEBuildModule>? PrivateIncludePathModules;
 
 		/// <summary>
 		/// Names of modules that this module's private implementation depends on.
 		/// </summary>
-		public List<UEBuildModule> PrivateDependencyModules;
+		public List<UEBuildModule>? PrivateDependencyModules;
 
 		/// <summary>
 		/// Extra modules this module may require at run time
 		/// </summary>
-		public List<UEBuildModule> DynamicallyLoadedModules;
+		public List<UEBuildModule>? DynamicallyLoadedModules;
 
 		/// <summary>
 		/// Set of all whitelisted restricted folder references
@@ -266,16 +266,16 @@ namespace UnrealBuildTool
 		public HashSet<UEBuildModule> GetDependencies(bool bWithIncludePathModules, bool bWithDynamicallyLoadedModules)
 		{
 			HashSet<UEBuildModule> Modules = new HashSet<UEBuildModule>();
-			Modules.UnionWith(PublicDependencyModules);
-			Modules.UnionWith(PrivateDependencyModules);
+			Modules.UnionWith(PublicDependencyModules!);
+			Modules.UnionWith(PrivateDependencyModules!);
 			if(bWithIncludePathModules)
 			{
-				Modules.UnionWith(PublicIncludePathModules);
-				Modules.UnionWith(PrivateIncludePathModules);
+				Modules.UnionWith(PublicIncludePathModules!);
+				Modules.UnionWith(PrivateIncludePathModules!);
 			}
 			if(bWithDynamicallyLoadedModules)
 			{
-				Modules.UnionWith(DynamicallyLoadedModules);
+				Modules.UnionWith(DynamicallyLoadedModules!);
 			}
 			return Modules;
         }
@@ -428,11 +428,11 @@ namespace UnrealBuildTool
 		protected void FindModulesInPrivateCompileEnvironment(Dictionary<UEBuildModule, bool> ModuleToIncludePathsOnlyFlag)
 		{
 			// Add in all the modules that are only in the private compile environment
-			foreach (UEBuildModule PrivateDependencyModule in PrivateDependencyModules)
+			foreach (UEBuildModule PrivateDependencyModule in PrivateDependencyModules!)
 			{
 				PrivateDependencyModule.FindModulesInPublicCompileEnvironment(ModuleToIncludePathsOnlyFlag);
 			}
-			foreach (UEBuildModule PrivateIncludePathModule in PrivateIncludePathModules)
+			foreach (UEBuildModule PrivateIncludePathModule in PrivateIncludePathModules!)
 			{
 				PrivateIncludePathModule.FindIncludePathModulesInPublicCompileEnvironment(ModuleToIncludePathsOnlyFlag);
 			}
@@ -456,13 +456,13 @@ namespace UnrealBuildTool
 
 			ModuleToIncludePathsOnlyFlag[this] = false;
 
-			foreach (UEBuildModule DependencyModule in PublicDependencyModules)
+			foreach (UEBuildModule DependencyModule in PublicDependencyModules!)
 			{
 				DependencyModule.FindModulesInPublicCompileEnvironment(ModuleToIncludePathsOnlyFlag);
 			}
 
 			// Now add an include paths from modules with header files that we need access to, but won't necessarily be importing
-			foreach (UEBuildModule IncludePathModule in PublicIncludePathModules)
+			foreach (UEBuildModule IncludePathModule in PublicIncludePathModules!)
 			{
 				IncludePathModule.FindIncludePathModulesInPublicCompileEnvironment(ModuleToIncludePathsOnlyFlag);
 			}
@@ -480,7 +480,7 @@ namespace UnrealBuildTool
 				ModuleToIncludePathsOnlyFlag.Add(this, true);
 
 				// Include any of its public include path modules in the compile environment too
-				foreach (UEBuildModule IncludePathModule in PublicIncludePathModules)
+				foreach (UEBuildModule IncludePathModule in PublicIncludePathModules!)
 				{
 					IncludePathModule.FindIncludePathModulesInPublicCompileEnvironment(ModuleToIncludePathsOnlyFlag);
 				}
@@ -500,7 +500,7 @@ namespace UnrealBuildTool
 		/// Sets up the environment for compiling any module that includes the public interface of this module.
 		/// </summary>
 		public virtual void AddModuleToCompileEnvironment(
-			UEBuildBinary SourceBinary,
+			UEBuildBinary? SourceBinary,
 			HashSet<DirectoryReference> IncludePaths,
 			HashSet<DirectoryReference> SystemIncludePaths,
 			HashSet<DirectoryReference> ModuleInterfacePaths,
@@ -604,7 +604,7 @@ namespace UnrealBuildTool
 		/// <param name="BinaryOutputDir">Directory containing the binary that links this module. May be mull.</param>
 		/// <param name="TargetOutputDir">Directory containing the output executable. May be null.</param>
 		/// <returns>The path with variables expanded</returns>
-		public string ExpandPathVariables(string Path, DirectoryReference BinaryOutputDir, DirectoryReference TargetOutputDir)
+		public string ExpandPathVariables(string Path, DirectoryReference? BinaryOutputDir, DirectoryReference? TargetOutputDir)
 		{
 			if(Path.StartsWith("$(", StringComparison.Ordinal))
 			{
@@ -647,7 +647,7 @@ namespace UnrealBuildTool
 						else
 						{
 							string Name = Path.Substring(StartIdx, EndIdx - StartIdx);
-							string Value = Environment.GetEnvironmentVariable(Name);
+							string? Value = Environment.GetEnvironmentVariable(Name);
 							if(String.IsNullOrEmpty(Value))
 							{
 								throw new BuildException("Environment variable '{0}' is not defined (referenced by {1})", Name, Rules.File);
@@ -681,7 +681,7 @@ namespace UnrealBuildTool
 		/// <param name="BinaryDir">Directory containing the binary that links this module. May be mull.</param>
 		/// <param name="ExeDir">Directory containing the output executable. May be null.</param>
 		/// <returns>The path with variables expanded</returns>
-		private IEnumerable<string> ExpandPathVariables(IEnumerable<string> Paths, DirectoryReference BinaryDir, DirectoryReference ExeDir)
+		private IEnumerable<string> ExpandPathVariables(IEnumerable<string> Paths, DirectoryReference? BinaryDir, DirectoryReference? ExeDir)
 		{
 			foreach(string Path in Paths)
 			{
@@ -693,7 +693,7 @@ namespace UnrealBuildTool
 		/// Sets up the environment for linking any module that includes the public interface of this module.
 		/// </summary>
 		protected virtual void SetupPublicLinkEnvironment(
-			UEBuildBinary SourceBinary,
+			UEBuildBinary? SourceBinary,
 			List<FileReference> Libraries,
 			List<DirectoryReference> SystemLibraryPaths,
 			List<string> SystemLibraries,
@@ -728,8 +728,8 @@ namespace UnrealBuildTool
 				{
 					// Gather all dependencies and recursively call SetupPublicLinkEnvironmnet
 					List<UEBuildModule> AllDependencyModules = new List<UEBuildModule>();
-					AllDependencyModules.AddRange(PrivateDependencyModules);
-					AllDependencyModules.AddRange(PublicDependencyModules);
+					AllDependencyModules.AddRange(PrivateDependencyModules!);
+					AllDependencyModules.AddRange(PublicDependencyModules!);
 
 					foreach (UEBuildModule DependencyModule in AllDependencyModules)
 					{
@@ -747,7 +747,7 @@ namespace UnrealBuildTool
 				Libraries.AddRange(PublicLibraries);
 				SystemLibraryPaths.AddRange(PublicSystemLibraryPaths);
 				SystemLibraries.AddRange(PublicSystemLibraries);
-				RuntimeLibraryPaths.AddRange(ExpandPathVariables(Rules.PublicRuntimeLibraryPaths, SourceBinary.OutputDir, ExeDir));
+				RuntimeLibraryPaths.AddRange(ExpandPathVariables(Rules.PublicRuntimeLibraryPaths, SourceBinary?.OutputDir, ExeDir));
 				Frameworks.AddRange(PublicFrameworks);
 				WeakFrameworks.AddRange(PublicWeakFrameworks);
 				AdditionalBundleResources.AddRange(PublicAdditionalBundleResources);
@@ -776,8 +776,8 @@ namespace UnrealBuildTool
 
 			// Also allow the module's public and private dependencies to modify the link environment.
 			List<UEBuildModule> AllDependencyModules = new List<UEBuildModule>();
-			AllDependencyModules.AddRange(PrivateDependencyModules);
-			AllDependencyModules.AddRange(PublicDependencyModules);
+			AllDependencyModules.AddRange(PrivateDependencyModules!);
+			AllDependencyModules.AddRange(PublicDependencyModules!);
 
 			foreach (UEBuildModule DependencyModule in AllDependencyModules)
 			{
@@ -880,8 +880,8 @@ namespace UnrealBuildTool
 				// (Note: it's *bad* that it is doing this, but it's even worse not to flag these
 				// cycles when they're introduced and have a way of tracking them!)
 
-				string GuiltyModule = null;
-				string VictimModule = null;
+				string? GuiltyModule = null;
+				string? VictimModule = null;
 
 				for (int i = 0; i < ReferenceStack.Count() - 1; i++)
 				{
@@ -935,7 +935,7 @@ namespace UnrealBuildTool
 			ReferenceStack.RemoveAt(ReferenceStack.Count - 1);
 		}
 
-		private static void RecursivelyCreateModulesByName(List<string> ModuleNames, ref List<UEBuildModule> Modules, CreateModuleDelegate CreateModule, string ReferenceChain, List<UEBuildModule> ReferenceStack)
+		private static void RecursivelyCreateModulesByName(List<string> ModuleNames, ref List<UEBuildModule>? Modules, CreateModuleDelegate CreateModule, string ReferenceChain, List<UEBuildModule> ReferenceStack)
 		{
 			// Check whether the module list is already set. We set this immediately (via the ref) to avoid infinite recursion.
 			if (Modules == null)
@@ -953,7 +953,7 @@ namespace UnrealBuildTool
 			}
 		}
 
-		private static void RecursivelyCreateIncludePathModulesByName(List<string> ModuleNames, ref List<UEBuildModule> Modules, CreateModuleDelegate CreateModule, string ReferenceChain)
+		private static void RecursivelyCreateIncludePathModulesByName(List<string> ModuleNames, ref List<UEBuildModule>? Modules, CreateModuleDelegate CreateModule, string ReferenceChain)
 		{
 			// Check whether the module list is already set. We set this immediately (via the ref) to avoid infinite recursion.
 			if (Modules == null)
@@ -1053,7 +1053,7 @@ namespace UnrealBuildTool
 		/// <param name="Writer">Writer for the array data</param>
 		/// <param name="ArrayName">Name of the array property</param>
 		/// <param name="Modules">Sequence of modules to write. May be null.</param>
-		void ExportJsonModuleArray(JsonWriter Writer, string ArrayName, IEnumerable<UEBuildModule> Modules)
+		void ExportJsonModuleArray(JsonWriter Writer, string ArrayName, IEnumerable<UEBuildModule>? Modules)
 		{
 			Writer.WriteArrayStart(ArrayName);
 			if (Modules != null)
@@ -1072,7 +1072,7 @@ namespace UnrealBuildTool
 		/// <param name="Writer">Writer for the array data</param>
 		/// <param name="ArrayName">Name of the array property</param>
 		/// <param name="Strings">Sequence of strings to write. May be null.</param>
-		void ExportJsonStringArray(JsonWriter Writer, string ArrayName, IEnumerable<string> Strings)
+		void ExportJsonStringArray(JsonWriter Writer, string ArrayName, IEnumerable<string>? Strings)
 		{
 			Writer.WriteArrayStart(ArrayName);
 			if (Strings != null)
diff --git a/Engine/Source/Programs/UnrealBuildTool/Configuration/UEBuildModuleCPP.cs b/Engine/Source/Programs/UnrealBuildTool/Configuration/UEBuildModuleCPP.cs
--- a/Engine/Source/Programs/UnrealBuildTool/Configuration/UEBuildModuleCPP.cs
+++ b/Engine/Source/Programs/UnrealBuildTool/Configuration/UEBuildModuleCPP.cs
@@ -36,7 +36,7 @@ namespace UnrealBuildTool
 		/// <summary>
 		/// The directory for this module's generated code
 		/// </summary>
-		public readonly DirectoryReference GeneratedCodeDirectory;
+		public readonly DirectoryReference? GeneratedCodeDirectory;
 
 		/// <summary>
 		/// Set for modules that have generated code
@@ -46,17 +46,17 @@ namespace UnrealBuildTool
 		/// <summary>
 		/// Paths containing *.gen.cpp files for this module.  If this is null then this module doesn't have any generated code.
 		/// </summary>
-		public List<string> GeneratedCppDirectories;
+		public List<string>? GeneratedCppDirectories;
 
 		/// <summary>
 		/// List of invalid include directives. These are buffered up and output before we start compiling.
 		/// </summary>
-		public List<string> InvalidIncludeDirectiveMessages;
+		public List<string>? InvalidIncludeDirectiveMessages;
 
 		/// <summary>
 		/// Set of source directories referenced during a build
 		/// </summary>
-		HashSet<DirectoryReference> SourceDirectories;
+		HashSet<DirectoryReference>? SourceDirectories;
 
 		protected override void GetReferencedDirectories(HashSet<DirectoryReference> Directories)
 		{
@@ -170,7 +170,7 @@ namespace UnrealBuildTool
         };
 
 
-		public UEBuildModuleCPP(ModuleRules Rules, DirectoryReference IntermediateDirectory, DirectoryReference GeneratedCodeDirectory)
+		public UEBuildModuleCPP(ModuleRules Rules, DirectoryReference IntermediateDirectory, DirectoryReference? GeneratedCodeDirectory)
 			: base(Rules, IntermediateDirectory)
 		{
 			this.GeneratedCodeDirectory = GeneratedCodeDirectory;
@@ -283,7 +283,7 @@ namespace UnrealBuildTool
 		/// Sets up the environment for compiling any module that includes the public interface of this module.
 		/// </summary>
 		public override void AddModuleToCompileEnvironment(
-			UEBuildBinary SourceBinary,
+			UEBuildBinary? SourceBinary,
 			HashSet<DirectoryReference> IncludePaths,
 			HashSet<DirectoryReference> SystemIncludePaths,
 			HashSet<DirectoryReference> ModuleInterfacePaths,
@@ -297,7 +297,7 @@ namespace UnrealBuildTool
 			// to search it so we can pick up generated code definitions after UHT is run for the first time.
 			if(bAddGeneratedCodeIncludePath || (ProjectFileGenerator.bGenerateProjectFiles && GeneratedCodeDirectory != null))
 			{
-				IncludePaths.Add(GeneratedCodeDirectory);
+				IncludePaths.Add(GeneratedCodeDirectory!);
 
 				if (Rules.AdditionalCodeGenDirectories != null)
 				{
@@ -641,7 +641,7 @@ namespace UnrealBuildTool
 		/// <param name="ModuleCompileEnvironment"></param>
 		/// <param name="Graph">List of actions to be executed. Additional actions will be added to this list.</param>
 		/// <returns>The created PCH instance.</returns>
-		private PrecompiledHeaderInstance CreatePrivatePCH(UEToolChain ToolChain, FileItem HeaderFile, CppCompileEnvironment ModuleCompileEnvironment, IActionGraphBuilder Graph)
+		private PrecompiledHeaderInstance CreatePrivatePCH(UEToolChain? ToolChain, FileItem HeaderFile, CppCompileEnvironment ModuleCompileEnvironment, IActionGraphBuilder Graph)
 		{
 			// Create the wrapper file, which sets all the definitions needed to compile it
 			FileReference WrapperLocation = FileReference.Combine(IntermediateDirectory, String.Format("PCH.{0}.h", Name));
@@ -675,9 +675,9 @@ namespace UnrealBuildTool
 		/// <param name="ModuleCompileEnvironment">Compile environment for the current module</param>
 		/// <param name="Graph">List of actions to be executed. Additional actions will be added to this list.</param>
 		/// <returns>Instance of a PCH</returns>
-		public PrecompiledHeaderInstance FindOrCreateSharedPCH(UEToolChain ToolChain, PrecompiledHeaderTemplate Template, CppCompileEnvironment ModuleCompileEnvironment, IActionGraphBuilder Graph)
+		public PrecompiledHeaderInstance FindOrCreateSharedPCH(UEToolChain? ToolChain, PrecompiledHeaderTemplate Template, CppCompileEnvironment ModuleCompileEnvironment, IActionGraphBuilder Graph)
 		{
-			PrecompiledHeaderInstance Instance = Template.Instances.Find(x => IsCompatibleForSharedPCH(x.CompileEnvironment, ModuleCompileEnvironment));
+			PrecompiledHeaderInstance? Instance = Template.Instances.Find(x => IsCompatibleForSharedPCH(x.CompileEnvironment, ModuleCompileEnvironment));
 			if(Instance == null)
 			{
 				// Create a suffix to distinguish this shared PCH variant from any others. Currently only optimized and non-optimized shared PCHs are supported.
@@ -1023,7 +1023,7 @@ namespace UnrealBuildTool
 		/// <param name="CompileEnvironment">The current compile environment</param>
 		/// <param name="LinkInputFiles">List of files that will be linked for the target</param>
 		/// <param name="Graph">List of build actions</param>
-		CppCompileEnvironment SetupPrecompiledHeaders(ReadOnlyTargetRules Target, UEToolChain ToolChain, CppCompileEnvironment CompileEnvironment, List<FileItem> LinkInputFiles, IActionGraphBuilder Graph)
+		CppCompileEnvironment SetupPrecompiledHeaders(ReadOnlyTargetRules Target, UEToolChain? ToolChain, CppCompileEnvironment CompileEnvironment, List<FileItem> LinkInputFiles, IActionGraphBuilder Graph)
 		{
 			if (Target.bUsePCHFiles && Rules.PCHUsage != ModuleRules.PCHUsageMode.NoPCHs)
 			{
@@ -1099,7 +1099,7 @@ namespace UnrealBuildTool
 		/// <param name="IntermediateDirectory">Directory to create the intermediate file</param>
 		/// <param name="HeaderSuffix">Suffix for the included file</param>
 		/// <param name="Graph">The action graph being built</param>
-		static void CreateHeaderForDefinitions(CppCompileEnvironment CompileEnvironment, DirectoryReference IntermediateDirectory, string HeaderSuffix, IActionGraphBuilder Graph)
+		static void CreateHeaderForDefinitions(CppCompileEnvironment CompileEnvironment, DirectoryReference IntermediateDirectory, string? HeaderSuffix, IActionGraphBuilder Graph)
 		{
 			if(CompileEnvironment.Definitions.Count > 0)
 			{
@@ -1219,14 +1219,14 @@ namespace UnrealBuildTool
 					{
 						foreach (FileItem CppFile in CppFiles)
 						{
-							string FirstInclude = ModuleCompileEnvironment.MetadataCache.GetFirstInclude(CppFile);
+							string? FirstInclude = ModuleCompileEnvironment.MetadataCache.GetFirstInclude(CppFile);
 							if(FirstInclude != null)
 							{
 								string IncludeName = Path.GetFileNameWithoutExtension(FirstInclude);
 								string ExpectedName = CppFile.Location.GetFileNameWithoutExtension();
 								if (String.Compare(IncludeName, ExpectedName, StringComparison.OrdinalIgnoreCase) != 0)
 								{
-									FileReference HeaderFile;
+									FileReference? HeaderFile;
 									if (NameToHeaderFile.TryGetValue(ExpectedName, out HeaderFile) && !IgnoreMismatchedHeader(ExpectedName))
 									{
 										InvalidIncludeDirectiveMessages.Add(String.Format("{0}(1): error: Expected {1} to be first header included.", CppFile.Location, HeaderFile.GetFileName()));
@@ -1415,11 +1415,11 @@ namespace UnrealBuildTool
 		public override void GetAllDependencyModules(List<UEBuildModule> ReferencedModules, HashSet<UEBuildModule> IgnoreReferencedModules, bool bIncludeDynamicallyLoaded, bool bForceCircular, bool bOnlyDirectDependencies)
 		{
 			List<UEBuildModule> AllDependencyModules = new List<UEBuildModule>();
-			AllDependencyModules.AddRange(PrivateDependencyModules);
-			AllDependencyModules.AddRange(PublicDependencyModules);
+			AllDependencyModules.AddRange(PrivateDependencyModules!);
+			AllDependencyModules.AddRange(PublicDependencyModules!);
 			if (bIncludeDynamicallyLoaded)
 			{
-				AllDependencyModules.AddRange(DynamicallyLoadedModules);
+				AllDependencyModules.AddRange(DynamicallyLoadedModules!);
 			}
 
 			foreach (UEBuildModule DependencyModule in AllDependencyModules)
diff --git a/Engine/Source/Programs/UnrealBuildTool/Configuration/UEBuildPlatform.cs b/Engine/Source/Programs/UnrealBuildTool/Configuration/UEBuildPlatform.cs
--- a/Engine/Source/Programs/UnrealBuildTool/Configuration/UEBuildPlatform.cs
+++ b/Engine/Source/Programs/UnrealBuildTool/Configuration/UEBuildPlatform.cs
@@ -8,6 +8,7 @@ using System.IO;
 using System.Linq;
 using Tools.DotNETCommon;
 using System.Reflection;
+using System.Diagnostics.CodeAnalysis;
 
 namespace UnrealBuildTool
 {
@@ -26,17 +27,17 @@ namespace UnrealBuildTool
 		/// <summary>
 		/// All the platform folder names
 		/// </summary>
-		private static string[] CachedPlatformFolderNames;
+		private static string[]? CachedPlatformFolderNames;
 
 		/// <summary>
 		/// Cached copy of the list of folders to include for this platform
 		/// </summary>
-		private ReadOnlyHashSet<string> CachedIncludedFolderNames;
+		private ReadOnlyHashSet<string>? CachedIncludedFolderNames;
 
 		/// <summary>
 		/// Cached copy of the list of folders to exclude for this platform
 		/// </summary>
-		private ReadOnlyHashSet<string> CachedExcludedFolderNames;
+		private ReadOnlyHashSet<string>? CachedExcludedFolderNames;
 
 		/// <summary>
 		/// Constructor.
@@ -49,7 +50,7 @@ namespace UnrealBuildTool
 
 			// check DDPI to see if the platform is enabled on this host platform
 			string IniPlatformName = ConfigHierarchy.GetIniPlatformName(Platform);
-			bool bIsEnabled = DataDrivenPlatformInfo.GetDataDrivenInfoForPlatform(IniPlatformName).bIsEnabled;
+			bool bIsEnabled = DataDrivenPlatformInfo.GetDataDrivenInfoForPlatform(IniPlatformName)?.bIsEnabled ?? true;
 
 			// set up the SDK if the platform is enabled
 			UEBuildPlatformSDK.RegisterSDKForPlatform(SDK, Platform.ToString(), bIsEnabled);
@@ -89,7 +90,7 @@ namespace UnrealBuildTool
 						Log.TraceVerbose("    Registering build platform: {0}", CheckType.ToString());
 						using(Timeline.ScopeEvent(CheckType.Name))
 						{
-							UEBuildPlatformFactory TempInst = (UEBuildPlatformFactory)Activator.CreateInstance(CheckType);
+							UEBuildPlatformFactory TempInst = (UEBuildPlatformFactory)Activator.CreateInstance(CheckType)!;
 							
 							if(bHostPlatformOnly && TempInst.TargetPlatform != BuildHostPlatform.Current.Platform)
 							{
@@ -191,7 +192,7 @@ namespace UnrealBuildTool
 		/// Get the default architecture for a project. This may be overriden on the command line to UBT.
 		/// </summary>
 		/// <param name="ProjectFile">Optional project to read settings from </param>
-		public virtual string GetDefaultArchitecture(FileReference ProjectFile)
+		public virtual string GetDefaultArchitecture(FileReference? ProjectFile)
 		{
 			// by default, use an empty architecture (which is really just a modifer to the platform for some paths/names)
 			return "";
@@ -354,7 +355,7 @@ namespace UnrealBuildTool
 		/// <param name="Rules">The target rules</param>
 		/// <param name="OutputFiles">List of executable output files</param>
 		/// <returns>Path to the bundle directory</returns>
-		public virtual DirectoryReference GetBundleDirectory(ReadOnlyTargetRules Rules, List<FileReference> OutputFiles)
+		public virtual DirectoryReference? GetBundleDirectory(ReadOnlyTargetRules Rules, List<FileReference> OutputFiles)
 		{
 			return null;
 		}
@@ -395,7 +396,7 @@ namespace UnrealBuildTool
 		public static void RegisterPlatformWithGroup(UnrealTargetPlatform InPlatform, UnrealPlatformGroup InGroup)
 		{
 			// find or add the list of groups for this platform
-			List<UnrealTargetPlatform> Platforms;
+			List<UnrealTargetPlatform>? Platforms;
 			if(!PlatformGroupDictionary.TryGetValue(InGroup, out Platforms))
 			{
 				Platforms = new List<UnrealTargetPlatform>();
@@ -409,8 +410,11 @@ namespace UnrealBuildTool
 		/// </summary>
 		public static List<UnrealTargetPlatform> GetPlatformsInGroup(UnrealPlatformGroup InGroup)
 		{
-			List<UnrealTargetPlatform> PlatformList;
-			PlatformGroupDictionary.TryGetValue(InGroup, out PlatformList);
+			List<UnrealTargetPlatform>? PlatformList;
+			if (!PlatformGroupDictionary.TryGetValue(InGroup, out PlatformList))
+			{
+				PlatformList = new List<UnrealTargetPlatform>();
+			}
 			return PlatformList;
 		}
 
@@ -431,22 +435,22 @@ namespace UnrealBuildTool
 		/// <returns>UEBuildPlatform  The instance of the build platform</returns>
 		public static UEBuildPlatform GetBuildPlatform(UnrealTargetPlatform InPlatform)
 		{
-			UEBuildPlatform Platform;
-			if (!TryGetBuildPlatform(InPlatform, out Platform))
+			UEBuildPlatform? Platform;
+			if(!TryGetBuildPlatform(InPlatform, out Platform))
 			{
 				throw new BuildException("GetBuildPlatform: No BuildPlatform found for {0}", InPlatform.ToString());
 			}
 			return Platform;
-		}
+			}
 
 		/// <summary>
 		/// Retrieve the IUEBuildPlatform instance for the given TargetPlatform
 		/// </summary>
 		/// <param name="InPlatform">  The UnrealTargetPlatform being built</param>
-		/// <param name="Platform">The build platform</param>
+		/// <param name="Platform"></param>
 		/// <returns>UEBuildPlatform  The instance of the build platform</returns>
-		public static bool TryGetBuildPlatform(UnrealTargetPlatform InPlatform, out UEBuildPlatform Platform)
-		{
+		public static bool TryGetBuildPlatform(UnrealTargetPlatform InPlatform, [NotNullWhen(true)] out UEBuildPlatform? Platform)
+			{
 			return BuildPlatformDictionary.TryGetValue(InPlatform, out Platform);
 		}
 
@@ -541,7 +545,7 @@ namespace UnrealBuildTool
 		public static bool PlatformRequiresMonolithicBuilds(UnrealTargetPlatform InPlatform, UnrealTargetConfiguration InConfiguration)
 		{
 			// Some platforms require monolithic builds...
-			UEBuildPlatform BuildPlatform;
+			UEBuildPlatform? BuildPlatform;
 			if (TryGetBuildPlatform(InPlatform, out BuildPlatform))
 			{
 				return BuildPlatform.ShouldCompileMonolithicBinary(InPlatform);
@@ -604,7 +608,7 @@ namespace UnrealBuildTool
 		/// For platforms that need to output multiple files per binary (ie Android "fat" binaries)
 		/// this will emit multiple paths. By default, it simply makes an array from the input
 		/// </summary>
-		public virtual List<FileReference> FinalizeBinaryPaths(FileReference BinaryName, FileReference ProjectFile, ReadOnlyTargetRules Target)
+		public virtual List<FileReference> FinalizeBinaryPaths(FileReference BinaryName, FileReference? ProjectFile, ReadOnlyTargetRules Target)
 		{
 			List<FileReference> TempList = new List<FileReference>() { BinaryName };
 			return TempList;
@@ -629,10 +633,10 @@ namespace UnrealBuildTool
 			return Configurations;
 		}
 
-		protected static bool DoProjectSettingsMatchDefault(UnrealTargetPlatform Platform, DirectoryReference ProjectDirectoryName, string Section, string[] BoolKeys, string[] IntKeys, string[] StringKeys)
+		protected static bool DoProjectSettingsMatchDefault(UnrealTargetPlatform Platform, DirectoryReference ProjectDirectoryName, string Section, string[] BoolKeys, string[]? IntKeys, string[]? StringKeys)
 		{
 			ConfigHierarchy ProjIni = ConfigCache.ReadHierarchy(ConfigHierarchyType.Engine, ProjectDirectoryName, Platform);
-			ConfigHierarchy DefaultIni = ConfigCache.ReadHierarchy(ConfigHierarchyType.Engine, (DirectoryReference)null, Platform);
+			ConfigHierarchy DefaultIni = ConfigCache.ReadHierarchy(ConfigHierarchyType.Engine, (DirectoryReference?)null, Platform);
 
 			// look at all bool values
 			if (BoolKeys != null) foreach (string Key in BoolKeys)
@@ -663,7 +667,7 @@ namespace UnrealBuildTool
 			// look for all string values
 			if (StringKeys != null) foreach (string Key in StringKeys)
 				{
-					string Default = "", Project = "";
+					string? Default = "", Project = "";
 					DefaultIni.GetString(Section, Key, out Default);
 					ProjIni.GetString(Section, Key, out Project);
 					if (Default != Project)
@@ -779,7 +783,7 @@ namespace UnrealBuildTool
 		/// </summary>
 		/// <param name="ProjectFile">The project file being built</param>
 		/// <returns>String describing the current build metadata</returns>
-		public string GetExternalBuildMetadata(FileReference ProjectFile)
+		public string GetExternalBuildMetadata(FileReference? ProjectFile)
 		{
 			StringBuilder Result = new StringBuilder();
 			GetExternalBuildMetadata(ProjectFile, Result);
@@ -791,7 +795,7 @@ namespace UnrealBuildTool
 		/// </summary>
 		/// <param name="ProjectFile">The project file being built</param>
 		/// <param name="Metadata">String builder to contain build metadata</param>
-		public virtual void GetExternalBuildMetadata(FileReference ProjectFile, StringBuilder Metadata)
+		public virtual void GetExternalBuildMetadata(FileReference? ProjectFile, StringBuilder Metadata)
 		{
 		}
 
@@ -817,7 +821,7 @@ namespace UnrealBuildTool
 		/// <returns>True if platform is part of a platform group</returns>
 		internal static bool IsPlatformInGroup(UnrealTargetPlatform Platform, UnrealPlatformGroup PlatformGroup)
 		{
-			List<UnrealTargetPlatform> Platforms = UEBuildPlatform.GetPlatformsInGroup(PlatformGroup);
+			List<UnrealTargetPlatform>? Platforms = UEBuildPlatform.GetPlatformsInGroup(PlatformGroup);
 			if (Platforms != null)
 			{
 				return Platforms.Contains(Platform);
diff --git a/Engine/Source/Programs/UnrealBuildTool/Configuration/UEBuildPlugin.cs b/Engine/Source/Programs/UnrealBuildTool/Configuration/UEBuildPlugin.cs
--- a/Engine/Source/Programs/UnrealBuildTool/Configuration/UEBuildPlugin.cs
+++ b/Engine/Source/Programs/UnrealBuildTool/Configuration/UEBuildPlugin.cs
@@ -27,7 +27,7 @@ namespace UnrealBuildTool
 		/// <summary>
 		/// Recursive
 		/// </summary>
-		public HashSet<UEBuildPlugin> Dependencies;
+		public HashSet<UEBuildPlugin>? Dependencies;
 
 		/// <summary>
 		/// Whether the descriptor for this plugin is needed at runtime; because it has modules or content which is used, or because it references another module that does.
diff --git a/Engine/Source/Programs/UnrealBuildTool/Configuration/UEBuildTarget.cs b/Engine/Source/Programs/UnrealBuildTool/Configuration/UEBuildTarget.cs
--- a/Engine/Source/Programs/UnrealBuildTool/Configuration/UEBuildTarget.cs
+++ b/Engine/Source/Programs/UnrealBuildTool/Configuration/UEBuildTarget.cs
@@ -28,7 +28,7 @@ namespace UnrealBuildTool
 		private int Id;
 
 		// shared string instance registry - pass in a delegate to create a new one with a name that wasn't made yet
-		private static UniqueStringRegistry StringRegistry;
+		private static UniqueStringRegistry? StringRegistry;
 
 		// #jira UE-88908 if parts of a partial struct have each static member variables, their initialization order does not appear guaranteed
 		// here this means initializing "StringRegistry" directly to "new UniqueStringRegistry()" may not be executed before FindOrAddByName() has been called as part of initializing a static member variable of another part of the partial struct
@@ -68,7 +68,7 @@ namespace UnrealBuildTool
 		/// <param name="Context"></param>
 		public UnrealTargetPlatform(SerializationInfo Info, StreamingContext Context)
 		{
-			Id = GetUniqueStringRegistry().FindOrAddByName((string)Info.GetValue("Name", typeof(string)));
+			Id = GetUniqueStringRegistry().FindOrAddByName((string)Info.GetValue("Name", typeof(string))!);
 		}
 
 		/// <summary>
@@ -106,8 +106,13 @@ namespace UnrealBuildTool
 		/// </summary>
 		/// <param name="B"></param>
 		/// <returns></returns>
-		public override bool Equals(object B)
+		public override bool Equals(object? B)
 		{
+			if(Object.ReferenceEquals(B, null))
+			{
+				return false;
+			}
+
 			return Id == ((UnrealTargetPlatform)B).Id;
 		}
 		/// <summary>
@@ -343,7 +348,7 @@ namespace UnrealBuildTool
 		private int Id;
 
 		// shared string instance registry - pass in a delegate to create a new one with a name that wasn't made yet
-		private static UniqueStringRegistry StringRegistry;
+		private static UniqueStringRegistry? StringRegistry;
 
 		// #jira UE-88908 (see above)
 		private static UniqueStringRegistry GetUniqueStringRegistry()
@@ -382,7 +387,7 @@ namespace UnrealBuildTool
 		/// <param name="Context"></param>
 		public UnrealPlatformGroup(SerializationInfo Info, StreamingContext Context)
 		{
-			Id = GetUniqueStringRegistry().FindOrAddByName((string)Info.GetValue("Name", typeof(string)));
+			Id = GetUniqueStringRegistry().FindOrAddByName((string)Info.GetValue("Name", typeof(string))!);
 		}
 
 		/// <summary>
@@ -420,8 +425,13 @@ namespace UnrealBuildTool
 		/// </summary>
 		/// <param name="B"></param>
 		/// <returns></returns>
-		public override bool Equals(object B)
+		public override bool Equals(object? B)
 		{
+			if(Object.ReferenceEquals(B, null))
+			{
+				return false;
+			}
+
 			return Id == ((UnrealPlatformGroup)B).Id;
 		}
 		/// <summary>
@@ -836,8 +846,8 @@ namespace UnrealBuildTool
 				{
 					if(Field.GetCustomAttribute<RequiresUniqueBuildEnvironmentAttribute>() != null)
 					{
-						object ThisValue = Field.GetValue(ThisObjects[Idx]);
-						object BaseValue = Field.GetValue(BaseObjects[Idx]);
+						object? ThisValue = Field.GetValue(ThisObjects[Idx]);
+						object? BaseValue = Field.GetValue(BaseObjects[Idx]);
 						CheckValuesMatch(ThisRules.GetType(), ThisTargetName, BaseTargetName, Field.Name, Field.FieldType, ThisValue, BaseValue);
 					}
 				}
@@ -845,8 +855,8 @@ namespace UnrealBuildTool
 				{
 					if(Property.GetCustomAttribute<RequiresUniqueBuildEnvironmentAttribute>() != null)
 					{
-						object ThisValue = Property.GetValue(ThisObjects[Idx]);
-						object BaseValue = Property.GetValue(BaseObjects[Idx]);
+						object? ThisValue = Property.GetValue(ThisObjects[Idx]);
+						object? BaseValue = Property.GetValue(BaseObjects[Idx]);
 						CheckValuesMatch(ThisRules.GetType(), ThisTargetName, BaseTargetName, Property.Name, Property.PropertyType, ThisValue, BaseValue);
 					}
 				}
@@ -863,7 +873,7 @@ namespace UnrealBuildTool
 		/// <summary>
 		/// Check that two values match between a base and derived rules type
 		/// </summary>
-		static void CheckValuesMatch(Type RulesType, string ThisTargetName, string BaseTargetName, string FieldName, Type ValueType, object ThisValue, object BaseValue)
+		static void CheckValuesMatch(Type RulesType, string ThisTargetName, string BaseTargetName, string FieldName, Type ValueType, object? ThisValue, object? BaseValue)
 		{
 			// Check if the fields match, treating lists of strings (eg. definitions) differently to value types.
 			bool bFieldsMatch;
@@ -905,12 +915,12 @@ namespace UnrealBuildTool
 		/// <summary>
 		/// The project file for this target
 		/// </summary>
-		public FileReference ProjectFile;
+		public FileReference? ProjectFile;
 
 		/// <summary>
 		/// The project descriptor for this target
 		/// </summary>
-		public ProjectDescriptor ProjectDescriptor;
+		public ProjectDescriptor? ProjectDescriptor;
 
 		/// <summary>
 		/// Type of target
@@ -975,17 +985,17 @@ namespace UnrealBuildTool
 		/// <summary>
 		/// All plugins which are built for this target
 		/// </summary>
-		public List<UEBuildPlugin> BuildPlugins;
+		public List<UEBuildPlugin>? BuildPlugins;
 
 		/// <summary>
 		/// All plugin dependencies for this target. This differs from the list of plugins that is built for Launcher, where we build everything, but link in only the enabled plugins.
 		/// </summary>
-		public List<UEBuildPlugin> EnabledPlugins;
+		public List<UEBuildPlugin>? EnabledPlugins;
 
 		/// <summary>
 		/// Specifies the path to a specific plugin to compile.
 		/// </summary>
-		public FileReference ForeignPlugin;
+		public FileReference? ForeignPlugin;
 
 		/// <summary>
 		/// All application binaries; may include binaries not built by this target.
@@ -1161,7 +1171,7 @@ namespace UnrealBuildTool
 			// Find all the runtime dependency files in their original location
 			foreach(RuntimeDependency RuntimeDependency in RuntimeDependencies)
 			{
-				FileReference SourceFile;
+				FileReference? SourceFile;
 				if(!RuntimeDependencyTargetFileToSourceFile.TryGetValue(RuntimeDependency.Path, out SourceFile))
 				{
 					SourceFile = RuntimeDependency.Path;
@@ -1369,7 +1379,7 @@ namespace UnrealBuildTool
 						bool bIsGameBinary = Binary.PrimaryModule.Rules.Context.bCanBuildDebugGame;
 						FileReference ManifestFileName = FileReference.Combine(DirectoryName, ModuleManifest.GetStandardFileName(AppName, Platform, Configuration, Architecture, bIsGameBinary));
 
-						ModuleManifest Manifest;
+						ModuleManifest? Manifest;
 						if (!FileNameToModuleManifest.TryGetValue(ManifestFileName, out Manifest))
 						{
 							Manifest = new ModuleManifest("");
@@ -1395,7 +1405,7 @@ namespace UnrealBuildTool
 		TargetReceipt PrepareReceipt(UEToolChain ToolChain, List<KeyValuePair<FileReference, BuildProductType>> BuildProducts, List<RuntimeDependency> RuntimeDependencies)
 		{
 			// Read the version file
-			BuildVersion Version;
+			BuildVersion? Version;
 			if (!BuildVersion.TryRead(BuildVersion.GetDefaultFileName(), out Version))
 			{
 				Version = new BuildVersion();
@@ -1455,7 +1465,7 @@ namespace UnrealBuildTool
 			}
 
 			// Add the descriptors for all enabled plugins
-			foreach(UEBuildPlugin EnabledPlugin in EnabledPlugins)
+			foreach(UEBuildPlugin EnabledPlugin in EnabledPlugins!)
 			{
 				if(EnabledPlugin.bDescriptorNeededAtRuntime || EnabledPlugin.bDescriptorReferencedExplicitly)
 				{
@@ -1612,7 +1622,7 @@ namespace UnrealBuildTool
 
 			// Build a mapping from module to its plugin
 			Dictionary<UEBuildModule, UEBuildPlugin> ModuleToPlugin = new Dictionary<UEBuildModule, UEBuildPlugin>();
-			foreach(UEBuildPlugin Plugin in BuildPlugins)
+			foreach(UEBuildPlugin Plugin in BuildPlugins!)
 			{
 				foreach(UEBuildModule Module in Plugin.Modules)
 				{
@@ -1647,8 +1657,8 @@ namespace UnrealBuildTool
 					HashSet<UEBuildModule> DependencyModules = Module.GetDependencies(bWithIncludePathModules: true, bWithDynamicallyLoadedModules: true);
 					foreach(UEBuildModule DependencyModule in DependencyModules)
 					{
-						UEBuildPlugin DependencyPlugin;
-						if(ModuleToPlugin.TryGetValue(DependencyModule, out DependencyPlugin) && DependencyPlugin != Plugin && !Plugin.Dependencies.Contains(DependencyPlugin))
+						UEBuildPlugin? DependencyPlugin;
+						if(ModuleToPlugin.TryGetValue(DependencyModule, out DependencyPlugin) && DependencyPlugin != Plugin && !Plugin.Dependencies!.Contains(DependencyPlugin))
 						{
 							Log.TraceWarning("Warning: Plugin '{0}' does not list plugin '{1}' as a dependency, but module '{2}' depends on '{3}'.", Plugin.Name, DependencyPlugin.Name, Module.Name, DependencyModule.Name);
 						}
@@ -1735,7 +1745,7 @@ namespace UnrealBuildTool
 			List<UEBuildBinary> BuildBinaries = Binaries;
 			foreach (FileReference SpecificFile in SpecificFilesToCompile)
 			{
-				UEBuildBinary Binary = Binaries.Find(x => x.Modules.Any(y => y.ContainsFile(SpecificFile)));
+				UEBuildBinary? Binary = Binaries.Find(x => x.Modules.Any(y => y.ContainsFile(SpecificFile)));
 
 				if (Binary == null)
 				{
@@ -1786,7 +1796,7 @@ namespace UnrealBuildTool
 					{
 						foreach (UEBuildModule Module in Plugin.Modules)
 						{
-							FileItem[] ModuleOutputItems;
+							FileItem[]? ModuleOutputItems;
 							if(Makefile.ModuleNameToOutputItems.TryGetValue(Module.Name, out ModuleOutputItems))
 							{
 								RetainOutputItems.UnionWith(ModuleOutputItems);
@@ -1841,7 +1851,7 @@ namespace UnrealBuildTool
 				}
 
 				// Get the path to the version file unless this is a formal build (where it will be compiled in)
-				FileReference VersionFile = null;
+				FileReference? VersionFile = null;
 				if(Rules.LinkType != TargetLinkType.Monolithic && Binaries[0].Type == UEBuildBinaryType.Executable)
 				{
 					UnrealTargetConfiguration VersionConfig = Configuration;
@@ -2097,12 +2107,12 @@ namespace UnrealBuildTool
 			}
 			else
 			{
-				Type ToolchainType = Assembly.GetExecutingAssembly().GetType(String.Format("UnrealBuildTool.{0}", Rules.ToolChainName), false, true);
+				Type? ToolchainType = Assembly.GetExecutingAssembly().GetType(String.Format("UnrealBuildTool.{0}", Rules.ToolChainName), false, true);
 				if (ToolchainType == null)
 				{
 					throw new BuildException("Unable to create toolchain '{0}'. Check that the name is correct.", Rules.ToolChainName);
 				}
-				return (UEToolChain)Activator.CreateInstance(ToolchainType, Rules);
+				return (UEToolChain)Activator.CreateInstance(ToolchainType, Rules)!;
 			}
 		}
 
@@ -2135,7 +2145,7 @@ namespace UnrealBuildTool
 				{
 					foreach (FileReference ExistingFile in DirectoryReference.EnumerateFiles(OutputDirectory))
 					{
-						FileReference OutputFile;
+						FileReference? OutputFile;
 						if (OutputNameToLocation.TryGetValue(ExistingFile.GetFileName(), out OutputFile) && !OutputFiles.Contains(ExistingFile))
 						{
 							Log.TraceInformation("Deleting '{0}' to avoid ambiguity with '{1}'", ExistingFile, OutputFile);
@@ -2204,7 +2214,7 @@ namespace UnrealBuildTool
 				foreach (UEBuildModule Module in Modules.Values)
 				{
 					Writer.WriteObjectStart(Module.Name);
-					Module.ExportJson(Module.Binary?.OutputDir, GetExecutableDir(), Writer);
+					Module.ExportJson(Module.Binary!.OutputDir, GetExecutableDir(), Writer);
 					Writer.WriteObjectEnd();
 				}
 				Writer.WriteObjectEnd();
@@ -2341,7 +2351,7 @@ namespace UnrealBuildTool
 			SetupPlugins();
 
 			// Add the plugin binaries to the build
-			foreach (UEBuildPlugin Plugin in BuildPlugins)
+			foreach (UEBuildPlugin Plugin in BuildPlugins!)
 			{
 				foreach(UEBuildModuleCPP Module in Plugin.Modules)
 				{
@@ -2408,7 +2418,7 @@ namespace UnrealBuildTool
 					{
 						foreach (string CircularlyReferencedModuleName in Module.Rules.CircularlyReferencedDependentModules)
 						{
-							UEBuildModule CircularlyReferencedModule;
+							UEBuildModule? CircularlyReferencedModule;
 							if (Modules.TryGetValue(CircularlyReferencedModuleName, out CircularlyReferencedModule) && CircularlyReferencedModule.Binary != null)
 							{
 								CircularlyReferencedModule.Binary.SetCreateImportLibrarySeparately(true);
@@ -2425,18 +2435,18 @@ namespace UnrealBuildTool
 		public FileReference[] CreatePreBuildScripts()
 		{
 			// Find all the pre-build steps
-			List<Tuple<string[], UEBuildPlugin>> PreBuildCommandBatches = new List<Tuple<string[], UEBuildPlugin>>();
+			List<Tuple<string[], UEBuildPlugin?>> PreBuildCommandBatches = new List<Tuple<string[], UEBuildPlugin?>>();
 			if(ProjectDescriptor != null && ProjectDescriptor.PreBuildSteps != null)
 			{
 				AddCustomBuildSteps(ProjectDescriptor.PreBuildSteps, null, PreBuildCommandBatches);
 			}
 			if(Rules.PreBuildSteps.Count > 0)
 			{
-				PreBuildCommandBatches.Add(new Tuple<string[], UEBuildPlugin>(Rules.PreBuildSteps.ToArray(), null));
+				PreBuildCommandBatches.Add(new Tuple<string[], UEBuildPlugin?>(Rules.PreBuildSteps.ToArray(), null));
 			}
 			foreach(UEBuildPlugin BuildPlugin in BuildPlugins.Where(x => x.Descriptor.PreBuildSteps != null))
 			{
-				AddCustomBuildSteps(BuildPlugin.Descriptor.PreBuildSteps, BuildPlugin, PreBuildCommandBatches);
+				AddCustomBuildSteps(BuildPlugin.Descriptor.PreBuildSteps!, BuildPlugin, PreBuildCommandBatches);
 			}
 			return WriteCustomBuildStepScripts(BuildHostPlatform.Current.Platform, ProjectIntermediateDirectory, "PreBuild", PreBuildCommandBatches);
 		}
@@ -2448,7 +2458,7 @@ namespace UnrealBuildTool
 		private FileReference[] CreatePostBuildScripts()
 		{
 			// Find all the post-build steps
-			List<Tuple<string[], UEBuildPlugin>> PostBuildCommandBatches = new List<Tuple<string[], UEBuildPlugin>>();
+			List<Tuple<string[], UEBuildPlugin?>> PostBuildCommandBatches = new List<Tuple<string[], UEBuildPlugin?>>();
 			if(!Rules.bDisableLinking)
 			{
 				if(ProjectDescriptor != null && ProjectDescriptor.PostBuildSteps != null)
@@ -2457,11 +2467,11 @@ namespace UnrealBuildTool
 				}
 				if(Rules.PostBuildSteps.Count > 0)
 				{
-					PostBuildCommandBatches.Add(new Tuple<string[], UEBuildPlugin>(Rules.PostBuildSteps.ToArray(), null));
+					PostBuildCommandBatches.Add(new Tuple<string[], UEBuildPlugin?>(Rules.PostBuildSteps.ToArray(), null));
 				}
 				foreach(UEBuildPlugin BuildPlugin in BuildPlugins.Where(x => x.Descriptor.PostBuildSteps != null))
 				{
-					AddCustomBuildSteps(BuildPlugin.Descriptor.PostBuildSteps, BuildPlugin, PostBuildCommandBatches);
+					AddCustomBuildSteps(BuildPlugin.Descriptor.PostBuildSteps!, BuildPlugin, PostBuildCommandBatches);
 				}
 			}
 			return WriteCustomBuildStepScripts(BuildHostPlatform.Current.Platform, ProjectIntermediateDirectory, "PostBuild", PostBuildCommandBatches);
@@ -2473,9 +2483,9 @@ namespace UnrealBuildTool
 		/// <param name="BuildSteps">The custom build steps</param>
 		/// <param name="Plugin">The plugin to associate with these commands</param>
 		/// <param name="CommandBatches">List to receive the command batches</param>
-		private void AddCustomBuildSteps(CustomBuildSteps BuildSteps, UEBuildPlugin Plugin, List<Tuple<string[], UEBuildPlugin>> CommandBatches)
+		private void AddCustomBuildSteps(CustomBuildSteps BuildSteps, UEBuildPlugin? Plugin, List<Tuple<string[], UEBuildPlugin?>> CommandBatches)
 		{
-			string[] Commands;
+			string[]? Commands;
 			if(BuildSteps.TryGetCommands(BuildHostPlatform.Current.Platform, out Commands))
 			{
 				CommandBatches.Add(Tuple.Create(Commands, Plugin));
@@ -2487,7 +2497,7 @@ namespace UnrealBuildTool
 		/// </summary>
 		/// <param name="Plugin">The current plugin</param>
 		/// <returns>Map of variable names to values</returns>
-		private Dictionary<string, string> GetTargetVariables(UEBuildPlugin Plugin)
+		private Dictionary<string, string> GetTargetVariables(UEBuildPlugin? Plugin)
 		{
 			Dictionary<string, string> Variables = new Dictionary<string,string>();
 			Variables.Add("RootDir", UnrealBuildTool.RootDirectory.FullName);
@@ -2516,10 +2526,10 @@ namespace UnrealBuildTool
 		/// <param name="FilePrefix">Bare prefix for all the created script files</param>
 		/// <param name="CommandBatches">List of custom build steps, and their matching PluginInfo (if appropriate)</param>
 		/// <returns>List of created script files</returns>
-		private FileReference[] WriteCustomBuildStepScripts(UnrealTargetPlatform HostPlatform, DirectoryReference Directory, string FilePrefix, List<Tuple<string[], UEBuildPlugin>> CommandBatches)
+		private FileReference[] WriteCustomBuildStepScripts(UnrealTargetPlatform HostPlatform, DirectoryReference Directory, string FilePrefix, List<Tuple<string[], UEBuildPlugin?>> CommandBatches)
 		{
 			List<FileReference> ScriptFiles = new List<FileReference>();
-			foreach(Tuple<string[], UEBuildPlugin> CommandBatch in CommandBatches)
+			foreach(Tuple<string[], UEBuildPlugin?> CommandBatch in CommandBatches)
 			{
 				// Find all the standard variables
 				Dictionary<string, string> Variables = GetTargetVariables(CommandBatch.Item2);
@@ -2700,16 +2710,16 @@ namespace UnrealBuildTool
 				// Find all the modules that are not part of the standard set
 				foreach (string ModuleName in ModuleNames)
 				{
-					FileReference ModuleFileName = RulesAssembly.GetModuleFileName(ModuleName);
+					FileReference ModuleFileName = RulesAssembly.GetModuleFileName(ModuleName)!;
 					if (Directories.Any(BaseDir => ModuleFileName.IsUnderDirectory(BaseDir)))
 					{
-						Type RulesType = RulesAssembly.GetModuleRulesType(ModuleName);
+						Type RulesType = RulesAssembly.GetModuleRulesType(ModuleName)!;
 
 						// Skip platform extension modules. We only care about the base modules, not the platform overrides.
 						// The platform overrides get applied at a later stage when we actually come to build the module.
 						if (!UEBuildPlatform.GetPlatformFolderNames().Any(Name => RulesType.Name.EndsWith("_" + Name)))
 						{
-							SupportedPlatformsAttribute SupportedPlatforms = RulesType.GetCustomAttribute<SupportedPlatformsAttribute>();
+							SupportedPlatformsAttribute? SupportedPlatforms = RulesType.GetCustomAttribute<SupportedPlatformsAttribute>();
 							if(SupportedPlatforms != null)
 							{
 								if(SupportedPlatforms.Platforms.Contains(Platform))
@@ -2768,7 +2778,7 @@ namespace UnrealBuildTool
 				{
 					if (ModuleDescriptor.IsCompiledInConfiguration(Platform, Configuration, TargetName, TargetType, Rules.bBuildDeveloperTools, Rules.bBuildRequiresCookedData))
 					{
-						FileReference ModuleFileName = RulesAssembly.GetModuleFileName(ModuleDescriptor.Name);
+						FileReference? ModuleFileName = RulesAssembly.GetModuleFileName(ModuleDescriptor.Name);
 						if(ModuleFileName == null)
 						{
 							throw new BuildException("Unable to find module '{0}' referenced by {1}", ModuleDescriptor.Name, Plugin.File);
@@ -2785,7 +2795,7 @@ namespace UnrealBuildTool
 			foreach(string FilteredModuleName in FilteredModuleNames)
 			{
 				// Try to create the rules object, but catch any exceptions if it fails. Some modules (eg. SQLite) may determine that they are unavailable in the constructor.
-				ModuleRules ModuleRules;
+				ModuleRules? ModuleRules;
 				try
 				{
 					ModuleRules = RulesAssembly.CreateModuleRules(FilteredModuleName, this.Rules, "all modules option");
@@ -2902,7 +2912,7 @@ namespace UnrealBuildTool
 			DirectoryReference IntermediateDirectory = DirectoryReference.Combine(BaseOutputDirectory, PlatformIntermediateFolder, AppName, ModuleConfiguration.ToString());
 
 			// Append a subdirectory if the module rules specifies one
-			if (ModuleRules != null && !String.IsNullOrEmpty(ModuleRules.BinariesSubFolder))
+			if (!String.IsNullOrEmpty(ModuleRules.BinariesSubFolder))
 			{
 				IntermediateDirectory = DirectoryReference.Combine(IntermediateDirectory, ModuleRules.BinariesSubFolder);
 			}
@@ -2922,7 +2932,7 @@ namespace UnrealBuildTool
 			DirectoryReference OutputDirectory = DirectoryReference.Combine(BaseOutputDirectory, "Binaries", Platform.ToString());
 
 			// Append a subdirectory if the module rules specifies one
-			if (Module.Rules != null && !String.IsNullOrEmpty(Module.Rules.BinariesSubFolder))
+			if (!String.IsNullOrEmpty(Module.Rules.BinariesSubFolder))
 			{
 				OutputDirectory = DirectoryReference.Combine(OutputDirectory, Module.Rules.BinariesSubFolder);
 			}
@@ -3002,7 +3012,7 @@ namespace UnrealBuildTool
 		/// <param name="ProjectFile">The project file containing the target being built</param>
 		/// <param name="Rules">Rules for the target being built</param>
 		/// <returns>List of executable paths for this target</returns>
-		public static List<FileReference> MakeBinaryPaths(DirectoryReference BaseDirectory, string BinaryName, UnrealTargetPlatform Platform, UnrealTargetConfiguration Configuration, UEBuildBinaryType BinaryType, string Architecture, UnrealTargetConfiguration UndecoratedConfiguration, bool bIncludesGameModules, string ExeSubFolder, FileReference ProjectFile, ReadOnlyTargetRules Rules)
+		public static List<FileReference> MakeBinaryPaths(DirectoryReference BaseDirectory, string BinaryName, UnrealTargetPlatform Platform, UnrealTargetConfiguration Configuration, UEBuildBinaryType BinaryType, string Architecture, UnrealTargetConfiguration UndecoratedConfiguration, bool bIncludesGameModules, string ExeSubFolder, FileReference? ProjectFile, ReadOnlyTargetRules Rules)
 		{
 			// Get the configuration for the executable. If we're building DebugGame, and this executable only contains engine modules, use the same name as development.
 			UnrealTargetConfiguration ExeConfiguration = Configuration;
@@ -3081,7 +3091,7 @@ namespace UnrealBuildTool
 				bAllowEnginePluginsEnabledByDefault = !ProjectDescriptor.DisableEnginePluginsByDefault;
 				if (ProjectDescriptor.Plugins != null)
 				{
-					string ProjectReferenceChain = ProjectFile.GetFileName();
+					string ProjectReferenceChain = ProjectFile!.GetFileName();
 					foreach (PluginReferenceDescriptor PluginReference in ProjectDescriptor.Plugins)
 					{
 						if (!Rules.EnablePlugins.Contains(PluginReference.Name, StringComparer.InvariantCultureIgnoreCase) && !Rules.DisablePlugins.Contains(PluginReference.Name, StringComparer.InvariantCultureIgnoreCase))
@@ -3123,7 +3133,7 @@ namespace UnrealBuildTool
 			{
 				ConfigHierarchy EngineConfig = ConfigCache.ReadHierarchy(ConfigHierarchyType.Engine, DirectoryReference.Combine(UnrealBuildTool.EngineDirectory, "Programs", TargetName), Platform);
 
-				List<string> PluginNames;
+				List<string>? PluginNames;
 				if(EngineConfig.GetArray("Plugins", "ProgramEnabledPlugins", out PluginNames))
 				{
 					foreach(string PluginName in PluginNames)
@@ -3163,7 +3173,7 @@ namespace UnrealBuildTool
 		/// <param name="NameToInstance">Map from plugin name to instance of it</param>
 		/// <param name="NameToInfo">Map from plugin name to information</param>
 		/// <returns>Instance of the plugin, or null if it should not be used</returns>
-		private UEBuildPlugin AddPlugin(PluginReferenceDescriptor Reference, string ReferenceChain, string[] ExcludeFolders, Dictionary<string, UEBuildPlugin> NameToInstance, Dictionary<string, PluginInfo> NameToInfo)
+		private UEBuildPlugin? AddPlugin(PluginReferenceDescriptor Reference, string ReferenceChain, string[] ExcludeFolders, Dictionary<string, UEBuildPlugin> NameToInstance, Dictionary<string, PluginInfo> NameToInfo)
 		{
 			// Ignore disabled references
 			if(!Reference.bEnabled)
@@ -3172,7 +3182,7 @@ namespace UnrealBuildTool
 			}
 
 			// Try to get an existing reference to this plugin
-			UEBuildPlugin Instance;
+			UEBuildPlugin? Instance;
 			if(NameToInstance.TryGetValue(Reference.Name, out Instance))
 			{
 				// If this is a non-optional reference, make sure that and every referenced dependency is staged
@@ -3206,7 +3216,7 @@ namespace UnrealBuildTool
 				}
 
 				// Find the plugin being enabled
-				PluginInfo Info;
+				PluginInfo? Info;
 				if(!NameToInfo.TryGetValue(Reference.Name, out Info))
 				{
 					if (Reference.bOptional)
@@ -3283,7 +3293,7 @@ namespace UnrealBuildTool
 				{
 					foreach(PluginReferenceDescriptor NextReference in Info.Descriptor.Plugins)
 					{
-						UEBuildPlugin NextInstance = AddPlugin(NextReference, PluginReferenceChain, ExcludeFolders, NameToInstance, NameToInfo);
+						UEBuildPlugin? NextInstance = AddPlugin(NextReference, PluginReferenceChain, ExcludeFolders, NameToInstance, NameToInfo);
 						if(NextInstance != null)
 						{
 							Dependencies.Add(NextInstance);
@@ -3429,7 +3439,7 @@ namespace UnrealBuildTool
 			GlobalCompileEnvironment.bCompileISPC = Rules.bCompileISPC;
 			GlobalCompileEnvironment.bHideSymbolsByDefault = !Rules.bPublicSymbolsByDefault;
 			GlobalCompileEnvironment.CppStandard = Rules.CppStandard;
-			GlobalCompileEnvironment.AdditionalArguments = Rules.AdditionalCompilerArguments;
+			GlobalCompileEnvironment.AdditionalArguments = Rules.AdditionalCompilerArguments ?? String.Empty;
 
 			GlobalLinkEnvironment.bIsBuildingConsoleApplication = Rules.bIsBuildingConsoleApplication;
 			GlobalLinkEnvironment.bOptimizeForSize = Rules.bCompileForSize;
@@ -3446,7 +3456,7 @@ namespace UnrealBuildTool
 			GlobalLinkEnvironment.bUseIncrementalLinking = Rules.bUseIncrementalLinking;
 			GlobalLinkEnvironment.bUseFastPDBLinking = Rules.bUseFastPDBLinking ?? false;
 			GlobalLinkEnvironment.bPrintTimingInfo = Rules.bPrintToolChainTimingInfo;
-			GlobalLinkEnvironment.AdditionalArguments = Rules.AdditionalLinkerArguments;
+			GlobalLinkEnvironment.AdditionalArguments = Rules.AdditionalLinkerArguments ?? String.Empty;
 
 			if (Rules.bPGOOptimize && Rules.bPGOProfile)
 			{
@@ -3847,7 +3857,7 @@ namespace UnrealBuildTool
 					FileReference CppFile = DirectoryReference.EnumerateFiles(RulesObject.Directory, "*.cpp", SearchOption.AllDirectories).FirstOrDefault();
 					if(CppFile != null)
 					{
-						string IncludeFile = MetadataCache.GetFirstInclude(FileItem.GetItemByFileReference(CppFile));
+						string? IncludeFile = MetadataCache.GetFirstInclude(FileItem.GetItemByFileReference(CppFile));
 						if(IncludeFile != null)
 						{
 							FileReference PchIncludeFile = DirectoryReference.EnumerateFiles(RulesObject.Directory, Path.GetFileName(IncludeFile), SearchOption.AllDirectories).FirstOrDefault();
@@ -3891,7 +3901,7 @@ namespace UnrealBuildTool
 		/// <param name="ReferenceChain">Chain of references causing this module to be instantiated, for display in error messages</param>
 		public UEBuildModule FindOrCreateModuleByName(string ModuleName, string ReferenceChain)
 		{
-			UEBuildModule Module;
+			UEBuildModule? Module;
 			if (!Modules.TryGetValue(ModuleName, out Module))
 			{
 				// @todo projectfiles: Cross-platform modules can appear here during project generation, but they may have already
@@ -3920,7 +3930,7 @@ namespace UnrealBuildTool
 
 				// Get the generated code directory. Plugins always write to their own intermediate directory so they can be copied between projects, shared engine
 				// intermediates go in the engine intermediate folder, and anything else goes in the project folder.
-				DirectoryReference GeneratedCodeDirectory = null;
+				DirectoryReference? GeneratedCodeDirectory = null;
 				if (RulesObject.Type != ModuleRules.ModuleType.External)
 				{
 					// Get the base directory
@@ -3957,7 +3967,7 @@ namespace UnrealBuildTool
 					}
 					else
 					{
-						BaseSourceDirectory = DirectoryReference.Combine(ProjectFile.Directory, "Source");
+						BaseSourceDirectory = DirectoryReference.Combine(ProjectFile!.Directory, "Source");
 					}
 
 					// If it's a game module (plugin or otherwise), add the root source directory to the include paths.
@@ -4005,7 +4015,7 @@ namespace UnrealBuildTool
 		/// <returns>New C++ module</returns>
 		public UEBuildModuleCPP FindOrCreateCppModuleByName(string ModuleName, string ReferenceChain)
 		{
-			UEBuildModuleCPP CppModule = FindOrCreateModuleByName(ModuleName, ReferenceChain) as UEBuildModuleCPP;
+			UEBuildModuleCPP? CppModule = FindOrCreateModuleByName(ModuleName, ReferenceChain) as UEBuildModuleCPP;
 			if(CppModule == null)
 			{
 				throw new BuildException("'{0}' is not a C++ module (referenced via {1})", ModuleName, ReferenceChain);
@@ -4015,7 +4025,7 @@ namespace UnrealBuildTool
 
 		protected UEBuildModule InstantiateModule(
 			ModuleRules RulesObject,
-			DirectoryReference GeneratedCodeDirectory)
+			DirectoryReference? GeneratedCodeDirectory)
 		{
 			switch (RulesObject.Type)
 			{
@@ -4047,7 +4057,7 @@ namespace UnrealBuildTool
 		/// </summary>
 		public UEBuildModule GetModuleByName(string Name)
 		{
-			UEBuildModule Result;
+			UEBuildModule? Result;
 			if (Modules.TryGetValue(Name, out Result))
 			{
 				return Result;
diff --git a/Engine/Source/Programs/UnrealBuildTool/Executors/Experimental/FASTBuild.cs b/Engine/Source/Programs/UnrealBuildTool/Executors/Experimental/FASTBuild.cs
--- a/Engine/Source/Programs/UnrealBuildTool/Executors/Experimental/FASTBuild.cs
+++ b/Engine/Source/Programs/UnrealBuildTool/Executors/Experimental/FASTBuild.cs
@@ -21,6 +21,8 @@ using System.Linq;
 using Tools.DotNETCommon;
 using System.Runtime.Serialization;
 
+#nullable disable
+
 namespace UnrealBuildTool
 {
 
diff --git a/Engine/Source/Programs/UnrealBuildTool/Executors/LocalExecutor.cs b/Engine/Source/Programs/UnrealBuildTool/Executors/LocalExecutor.cs
--- a/Engine/Source/Programs/UnrealBuildTool/Executors/LocalExecutor.cs
+++ b/Engine/Source/Programs/UnrealBuildTool/Executors/LocalExecutor.cs
@@ -11,6 +11,8 @@ using System.Threading;
 using System.Linq;
 using Tools.DotNETCommon;
 
+#nullable disable
+
 namespace UnrealBuildTool
 {
 	class ActionThread
diff --git a/Engine/Source/Programs/UnrealBuildTool/Modes/AggregateParsedTimingInfo.cs b/Engine/Source/Programs/UnrealBuildTool/Modes/AggregateParsedTimingInfo.cs
--- a/Engine/Source/Programs/UnrealBuildTool/Modes/AggregateParsedTimingInfo.cs
+++ b/Engine/Source/Programs/UnrealBuildTool/Modes/AggregateParsedTimingInfo.cs
@@ -11,6 +11,8 @@ using System.Threading;
 using System.Threading.Tasks;
 using Tools.DotNETCommon;
 
+#nullable disable
+
 namespace UnrealBuildTool
 {
 	/// <summary>
@@ -33,12 +35,12 @@ namespace UnrealBuildTool
 			string[] ParsedFileNames = FileReference.ReadAllLines(ManifestFile);
 
 			// Load all the file timing data and summarize them for the aggregate.
-			FileTimingData = new TimingData() { Name = "Files", Type = TimingDataType.Summary };
+			FileTimingData = new TimingData("Files", TimingDataType.Summary);
 			Parallel.ForEach(ParsedFileNames, new ParallelOptions() { MaxDegreeOfParallelism = 4 }, ParseTimingDataFile);
 
 			// Create aggregate summary. Duration is the duration of the files in the aggregate.
 			string AggregateName = Arguments.GetString("-Name=");
-			TimingData AggregateData = new TimingData() { Name = AggregateName, Type = TimingDataType.Aggregate };
+			TimingData AggregateData = new TimingData(AggregateName, TimingDataType.Aggregate);
 			AggregateData.AddChild(FileTimingData);
 
 			// Group the includes, classes, and functions by name and sum them up then add to the aggregate.
@@ -91,11 +93,11 @@ namespace UnrealBuildTool
 		private void GroupTimingDataOnName(TimingData AggregateData, string TimingName, IEnumerable<TimingData> UngroupedData)
 		{
 			string GroupName = TimingName;
-			TimingData GroupedTimingData = new TimingData() { Name = GroupName, Type = TimingDataType.Summary };
+			TimingData GroupedTimingData = new TimingData(GroupName, TimingDataType.Summary);
 			IEnumerable<IGrouping<string, TimingData>> Groups = UngroupedData.GroupBy(i => i.Name).OrderByDescending(g => g.Sum(d => d.ExclusiveDuration)).ToList();
 			foreach (IGrouping<string, TimingData> Group in Groups)
 			{
-				TimingData GroupedData = new TimingData() { Name = Group.Key, ExclusiveDuration = Group.Sum(d => d.ExclusiveDuration), Count = Group.Sum(d => d.Count) };
+				TimingData GroupedData = new TimingData(Group.Key, TimingDataType.None){ ExclusiveDuration = Group.Sum(d => d.ExclusiveDuration), Count = Group.Sum(d => d.Count) };
 				GroupedTimingData.Children.Add(Group.Key, GroupedData);
 			}
 
@@ -126,7 +128,7 @@ namespace UnrealBuildTool
 					}
 				});
 
-				TimingData SummarizedTimingData = new TimingData() { Name = ParsedTimingData.Name, Parent = FileTimingData, Type = TimingDataType.Summary };
+				TimingData SummarizedTimingData = new TimingData(ParsedTimingData.Name, TimingDataType.Summary){ Parent = FileTimingData };
 				foreach (TimingData Include in ParsedTimingData.Children["IncludeTimings"].Children.Values)
 				{
 					SummarizedTimingData.AddChild(Include.Clone());
@@ -153,7 +155,7 @@ namespace UnrealBuildTool
 					IEnumerable<TimingData> CollapsedClasses = GroupChildren(ParsedTimingData.Children["ClassTimings"].Children.Values, TimingDataType.Class);
 					foreach (TimingData Class in CollapsedClasses)
 					{
-						AggregateClasses.Add(new TimingData() { Name = Class.Name, Type = TimingDataType.Class, ExclusiveDuration = Class.InclusiveDuration });
+						AggregateClasses.Add(new TimingData(Class.Name, TimingDataType.Class){ ExclusiveDuration = Class.InclusiveDuration });
 					}
 				});
 
@@ -163,7 +165,7 @@ namespace UnrealBuildTool
 					IEnumerable<TimingData> CollapsedFunctions = GroupChildren(ParsedTimingData.Children["FunctionTimings"].Children.Values, TimingDataType.Function);
 					foreach (TimingData Function in CollapsedFunctions)
 					{
-						AggregateFunctions.Add(new TimingData() { Name = Function.Name, Type = TimingDataType.Function, ExclusiveDuration = Function.InclusiveDuration });
+						AggregateFunctions.Add(new TimingData(Function.Name, TimingDataType.Function){ ExclusiveDuration = Function.InclusiveDuration });
 					}
 				});
 
@@ -201,7 +203,7 @@ namespace UnrealBuildTool
 						}
 					}
 
-					Includes.Add(Include.Name, new TimingData() { Name = Include.Name, Type = TimingDataType.Include, ExclusiveDuration = Include.ExclusiveDuration });
+					Includes.Add(Include.Name, new TimingData(Include.Name, TimingDataType.Include){ ExclusiveDuration = Include.ExclusiveDuration });
 				}
 
 				FlattenIncludes(Includes, Include.Children.Values);
@@ -250,15 +252,11 @@ namespace UnrealBuildTool
 					continue;
 				}
 
-				TimingData NewViewModel = new TimingData()
-				{
-					Name = Group.Key,
-					Type = Type,
-				};
+				TimingData NewViewModel = new TimingData(Group.Key, Type);
 
 				foreach (TimingData Child in Group.Value)
 				{
-					TimingData NewChild = new TimingData() { Name = Child.Name, Type = Child.Type, Parent = NewViewModel, ExclusiveDuration = Child.ExclusiveDuration };
+					TimingData NewChild = new TimingData(Child.Name, Child.Type){ Parent = NewViewModel, ExclusiveDuration = Child.ExclusiveDuration };
 					NewViewModel.AddChild(NewChild);
 				}
 
