UBT: Convert UnrealBuildTool to use nullable references.

diff --git a/Engine/Source/Programs/UnrealBuildTool/Modes/BuildMode.cs b/Engine/Source/Programs/UnrealBuildTool/Modes/BuildMode.cs
--- a/Engine/Source/Programs/UnrealBuildTool/Modes/BuildMode.cs
+++ b/Engine/Source/Programs/UnrealBuildTool/Modes/BuildMode.cs
@@ -49,7 +49,7 @@ namespace UnrealBuildTool
 		/// Specifies the file to use for logging.
 		/// </summary>
 		[XmlConfigFile(Category = "BuildConfiguration")]
-		public string BaseLogFileName;
+		public string? BaseLogFileName;
 
 		/// <summary>
 		/// Whether to skip checking for files identified by the junk manifest.
@@ -86,7 +86,7 @@ namespace UnrealBuildTool
 		/// Whether we should just export the outdated actions list
 		/// </summary>
 		[CommandLine("-WriteOutdatedActions=")]
-		public FileReference WriteOutdatedActionsFile = null;
+		public FileReference? WriteOutdatedActionsFile = null;
 
 		/// <summary>
 		/// Main entry point
@@ -248,7 +248,7 @@ namespace UnrealBuildTool
 		/// <param name="WriteOutdatedActionsFile">Files to write the list of outdated actions to (rather than building them)</param>
 		/// <param name="bSkipPreBuildTargets">If true then only the current target descriptors will be built.</param>
 		/// <returns>Result from the compilation</returns>
-		public static void Build(List<TargetDescriptor> TargetDescriptors, BuildConfiguration BuildConfiguration, ISourceFileWorkingSet WorkingSet, BuildOptions Options, FileReference WriteOutdatedActionsFile, bool bSkipPreBuildTargets = false)
+		public static void Build(List<TargetDescriptor> TargetDescriptors, BuildConfiguration BuildConfiguration, ISourceFileWorkingSet WorkingSet, BuildOptions Options, FileReference? WriteOutdatedActionsFile, bool bSkipPreBuildTargets = false)
 		{
 
 			List<TargetMakefile> TargetMakefiles = new List<TargetMakefile>();
@@ -270,7 +270,7 @@ namespace UnrealBuildTool
 				}
 			}
 
-			Build(TargetMakefiles.ToArray(), TargetDescriptors, BuildConfiguration, WorkingSet, Options, WriteOutdatedActionsFile);
+			Build(TargetMakefiles.ToArray(), TargetDescriptors, BuildConfiguration, Options, WriteOutdatedActionsFile);
 		}
 
 		/// <summary>
@@ -279,11 +279,10 @@ namespace UnrealBuildTool
 		/// <param name="Makefiles">Makefiles created with CreateMakefiles</param>
 		/// <param name="TargetDescriptors">Target descriptors</param>
 		/// <param name="BuildConfiguration">Current build configuration</param>
-		/// <param name="WorkingSet">The source file working set</param>
 		/// <param name="Options">Additional options for the build</param>
 		/// <param name="WriteOutdatedActionsFile">Files to write the list of outdated actions to (rather than building them)</param>
 		/// <returns>Result from the compilation</returns>
-		static void Build(TargetMakefile[] Makefiles, List<TargetDescriptor> TargetDescriptors, BuildConfiguration BuildConfiguration, ISourceFileWorkingSet WorkingSet, BuildOptions Options, FileReference WriteOutdatedActionsFile)
+		static void Build(TargetMakefile[] Makefiles, List<TargetDescriptor> TargetDescriptors, BuildConfiguration BuildConfiguration, BuildOptions Options, FileReference? WriteOutdatedActionsFile)
 		{
 			// Execute the build
 			if ((Options & BuildOptions.SkipBuild) == 0)
@@ -388,10 +387,10 @@ namespace UnrealBuildTool
 						List<FileItem> CompiledModuleInterfaces = new List<FileItem>();
 						foreach(LinkedAction PrerequisiteAction in PrerequisiteActions)
 						{
-							ICppCompileAction CppModulesAction = PrerequisiteAction.Inner as ICppCompileAction;
+							ICppCompileAction? CppModulesAction = PrerequisiteAction.Inner as ICppCompileAction;
 							if (CppModulesAction != null && CppModulesAction.CompiledModuleInterfaceFile != null)
 							{
-								string ProducedModule;
+								string? ProducedModule;
 								if (CppDependencies.TryGetProducedModule(PrerequisiteAction.DependencyListFile, out ProducedModule))
 								{
 									ModuleOutputs[ProducedModule] = CppModulesAction.CompiledModuleInterfaceFile;
@@ -405,12 +404,12 @@ namespace UnrealBuildTool
 						{
 							if (PrerequisiteAction.ActionType == ActionType.CompileModuleInterface)
 							{
-								List<string> ImportedModules;
+								List<string>? ImportedModules;
 								if(CppDependencies.TryGetImportedModules(PrerequisiteAction.DependencyListFile, out ImportedModules))
 								{
 									foreach (string ImportedModule in ImportedModules)
 									{
-										FileItem ModuleOutput;
+										FileItem? ModuleOutput;
 										if (ModuleOutputs.TryGetValue(ImportedModule, out ModuleOutput))
 										{
 											Action NewAction = new Action(PrerequisiteAction.Inner);
@@ -429,7 +428,7 @@ namespace UnrealBuildTool
 							{
 								foreach(FileItem PrerequisiteItem in PrerequisiteAction.PrerequisiteItems)
 								{
-									string ModuleName;
+									string? ModuleName;
 									if(ModuleInputs.TryGetValue(PrerequisiteItem, out ModuleName))
 									{
 										Action NewAction = new Action(PrerequisiteAction.Inner);
@@ -630,19 +629,19 @@ namespace UnrealBuildTool
 		static TargetMakefile CreateMakefile(BuildConfiguration BuildConfiguration, TargetDescriptor TargetDescriptor, ISourceFileWorkingSet WorkingSet)
 		{
 			// Get the path to the makefile for this target
-			FileReference MakefileLocation = null;
+			FileReference? MakefileLocation = null;
 			if(BuildConfiguration.bUseUBTMakefiles && TargetDescriptor.SpecificFilesToCompile.Count == 0)
 			{
 				MakefileLocation = TargetMakefile.GetLocation(TargetDescriptor.ProjectFile, TargetDescriptor.Name, TargetDescriptor.Platform, TargetDescriptor.Configuration);
 			}
 
 			// Try to load an existing makefile
-			TargetMakefile Makefile = null;
+			TargetMakefile? Makefile = null;
 			if(MakefileLocation != null)
 			{
 				using(Timeline.ScopeEvent("TargetMakefile.Load()"))
 				{
-					string ReasonNotLoaded;
+					string? ReasonNotLoaded;
 					Makefile = TargetMakefile.Load(MakefileLocation, TargetDescriptor.ProjectFile, TargetDescriptor.Platform, TargetDescriptor.AdditionalArguments.GetRawArray(), out ReasonNotLoaded);
 					if (Makefile == null)
 					{
@@ -666,7 +665,7 @@ namespace UnrealBuildTool
 				bHasRunPreBuildScripts = true;
 
 				// Check that the makefile is still valid
-				string Reason;
+				string? Reason;
 				if(!TargetMakefile.IsValidForSourceFiles(Makefile, TargetDescriptor.ProjectFile, TargetDescriptor.Platform, WorkingSet, out Reason))
 				{
 					Log.TraceInformation("Invalidating makefile for {0} ({1})", TargetDescriptor.Name, Reason);
@@ -708,9 +707,12 @@ namespace UnrealBuildTool
 				Makefile.AdditionalArguments = TargetDescriptor.AdditionalArguments.GetRawArray();
 
 				// Save the environment variables
-				foreach (System.Collections.DictionaryEntry EnvironmentVariable in Environment.GetEnvironmentVariables())
+				foreach (System.Collections.DictionaryEntry? EnvironmentVariable in Environment.GetEnvironmentVariables())
 				{
-					Makefile.EnvironmentVariables.Add(Tuple.Create((string)EnvironmentVariable.Key, (string)EnvironmentVariable.Value));
+					if (EnvironmentVariable != null)
+					{
+						Makefile.EnvironmentVariables.Add(Tuple.Create((string)EnvironmentVariable.Value.Key, (string)EnvironmentVariable.Value.Value!));
+					}
 				}
 
 				// Save the makefile for next time
@@ -765,7 +767,7 @@ namespace UnrealBuildTool
 				// Find the output items for this module
 				foreach(string OnlyModuleName in TargetDescriptor.OnlyModuleNames)
 				{
-					FileItem[] OutputItemsForModule;
+					FileItem[]? OutputItemsForModule;
 					if(!Makefile.ModuleNameToOutputItems.TryGetValue(OnlyModuleName, out OutputItemsForModule))
 					{
 						throw new BuildException("Unable to find output items for module '{0}'", OnlyModuleName);
@@ -797,7 +799,7 @@ namespace UnrealBuildTool
 				{
 					FileItem ProducedItem = TargetAction.ProducedItems.First();
 
-					LinkedAction ExistingAction;
+					LinkedAction? ExistingAction;
 					if(!OutputItemToProducingAction.TryGetValue(ProducedItem, out ExistingAction))
 					{
 						ExistingAction = new LinkedAction(TargetAction);
diff --git a/Engine/Source/Programs/UnrealBuildTool/Modes/ParseMsvcTimingInfo.cs b/Engine/Source/Programs/UnrealBuildTool/Modes/ParseMsvcTimingInfo.cs
--- a/Engine/Source/Programs/UnrealBuildTool/Modes/ParseMsvcTimingInfo.cs
+++ b/Engine/Source/Programs/UnrealBuildTool/Modes/ParseMsvcTimingInfo.cs
@@ -7,6 +7,8 @@ using System.Linq;
 using System.Text.RegularExpressions;
 using Tools.DotNETCommon;
 
+#nullable disable
+
 namespace UnrealBuildTool
 {
 	/// <summary>
@@ -99,7 +101,7 @@ namespace UnrealBuildTool
 			}
 
 			// Build the summary.
-			TimingData Summary = new TimingData() { Name = InputFile.FullName.Replace(".timing.txt", string.Empty), Type = TimingDataType.Summary };
+			TimingData Summary = new TimingData(InputFile.FullName.Replace(".timing.txt", string.Empty), TimingDataType.Summary);
 			Summary.AddChild(SummarizeParsedTimingData("IncludeTimings", TimingDataType.Include, Includes));
 			Summary.AddChild(SummarizeParsedTimingData("ClassTimings", TimingDataType.Class, Classes));
 			Summary.AddChild(SummarizeParsedTimingData("FunctionTimings", TimingDataType.Function, Functions));
@@ -115,7 +117,7 @@ namespace UnrealBuildTool
 
 		TimingData SummarizeParsedTimingData(string SummaryName, TimingDataType TimingType,  IEnumerable<string> Lines)
 		{
-			TimingData Summary = new TimingData() { Name = SummaryName, Type = TimingDataType.Summary };
+			TimingData Summary = new TimingData(SummaryName, TimingDataType.Summary);
 			List<TimingData> ParsedTimingData = ParseTimingDataFromLines(TimingType, Lines);
 			foreach (TimingData Data in ParsedTimingData)
 			{
@@ -196,12 +198,8 @@ namespace UnrealBuildTool
 
 			LineDepth = TimingDataMatch.Groups["Indent"].Success ? TimingDataMatch.Groups["Indent"].Value.Count() : 0;
 
-			TimingData ParsedTimingData = new TimingData()
-			{
-				Name = TimingDataMatch.Groups["Name"].Value,
-				Type = TimingType,
-				ExclusiveDuration = float.Parse(TimingDataMatch.Groups["Duration"].Value),
-			};
+			TimingData ParsedTimingData = new TimingData(TimingDataMatch.Groups["Name"].Value, TimingType);
+			ParsedTimingData.ExclusiveDuration = float.Parse(TimingDataMatch.Groups["Duration"].Value);
 
 			return ParsedTimingData;
 		}
diff --git a/Engine/Source/Programs/UnrealBuildTool/Modes/QueryTargetsMode.cs b/Engine/Source/Programs/UnrealBuildTool/Modes/QueryTargetsMode.cs
--- a/Engine/Source/Programs/UnrealBuildTool/Modes/QueryTargetsMode.cs
+++ b/Engine/Source/Programs/UnrealBuildTool/Modes/QueryTargetsMode.cs
@@ -19,13 +19,13 @@ namespace UnrealBuildTool
 		/// Path to the project file to query
 		/// </summary>
 		[CommandLine("-Project=")]
-		FileReference ProjectFile = null;
+		FileReference? ProjectFile = null;
 
 		/// <summary>
 		/// Path to the output file to receive information about the targets
 		/// </summary>
 		[CommandLine("-Output=")]
-		FileReference OutputFile = null;
+		FileReference? OutputFile = null;
 
 		/// <summary>
 		/// Execute the mode
@@ -69,7 +69,7 @@ namespace UnrealBuildTool
 		/// </summary>
 		/// <param name="ProjectFile">Project file being queried</param>
 		/// <returns>Path to the output file</returns>
-		public static FileReference GetDefaultOutputFile(FileReference ProjectFile)
+		public static FileReference GetDefaultOutputFile(FileReference? ProjectFile)
 		{
 			if (ProjectFile == null)
 			{
@@ -88,7 +88,7 @@ namespace UnrealBuildTool
 		/// <param name="Assembly">The rules assembly for this target</param>
 		/// <param name="OutputFile">Output file to write to</param>
 		/// <param name="Arguments"></param>
-		public static void WriteTargetInfo(FileReference ProjectFile, RulesAssembly Assembly, FileReference OutputFile, CommandLineArguments Arguments)
+		public static void WriteTargetInfo(FileReference? ProjectFile, RulesAssembly Assembly, FileReference OutputFile, CommandLineArguments Arguments)
 		{
 			// Construct all the targets in this assembly
 			List<string> TargetNames = new List<string>();
@@ -126,7 +126,7 @@ namespace UnrealBuildTool
 					// Write the target info
 					Writer.WriteObjectStart();
 					Writer.WriteValue("Name", TargetName);
-					Writer.WriteValue("Path", Assembly.GetTargetFileName(TargetName).ToString());
+					Writer.WriteValue("Path", Assembly.GetTargetFileName(TargetName)?.ToString());
 					Writer.WriteValue("Type", TargetRules.Type.ToString());
 					Writer.WriteObjectEnd();
 				}
diff --git a/Engine/Source/Programs/UnrealBuildTool/Modes/ValidatePlatformsMode.cs b/Engine/Source/Programs/UnrealBuildTool/Modes/ValidatePlatformsMode.cs
--- a/Engine/Source/Programs/UnrealBuildTool/Modes/ValidatePlatformsMode.cs
+++ b/Engine/Source/Programs/UnrealBuildTool/Modes/ValidatePlatformsMode.cs
@@ -53,7 +53,7 @@ namespace UnrealBuildTool
 			// Output a line for each registered platform
 			foreach (UnrealTargetPlatform Platform in Platforms)
 			{
-				UEBuildPlatform.TryGetBuildPlatform(Platform, out UEBuildPlatform BuildPlatform);
+				UEBuildPlatform.TryGetBuildPlatform(Platform, out UEBuildPlatform? BuildPlatform);
 				string PlatformSDKString = "";
 				if (bOutputSDKs)
 				{
diff --git a/Engine/Source/Programs/UnrealBuildTool/Modes/WriteDocumentationMode.cs b/Engine/Source/Programs/UnrealBuildTool/Modes/WriteDocumentationMode.cs
--- a/Engine/Source/Programs/UnrealBuildTool/Modes/WriteDocumentationMode.cs
+++ b/Engine/Source/Programs/UnrealBuildTool/Modes/WriteDocumentationMode.cs
@@ -35,7 +35,7 @@ namespace UnrealBuildTool
 		/// The HTML file to write to
 		/// </summary>
 		[CommandLine(Required = true)]
-		FileReference OutputFile = null;
+		FileReference OutputFile = null!;
 
 		/// <summary>
 		/// Entry point for this command
diff --git a/Engine/Source/Programs/UnrealBuildTool/Modes/WriteMetadataMode.cs b/Engine/Source/Programs/UnrealBuildTool/Modes/WriteMetadataMode.cs
--- a/Engine/Source/Programs/UnrealBuildTool/Modes/WriteMetadataMode.cs
+++ b/Engine/Source/Programs/UnrealBuildTool/Modes/WriteMetadataMode.cs
@@ -2,6 +2,7 @@
 
 using System;
 using System.Collections.Generic;
+using System.Diagnostics.CodeAnalysis;
 using System.IO;
 using System.Linq;
 using System.Runtime.Serialization.Formatters.Binary;
@@ -20,12 +21,12 @@ namespace UnrealBuildTool
 		/// <summary>
 		/// The project file
 		/// </summary>
-		public FileReference ProjectFile;
+		public FileReference? ProjectFile;
 
 		/// <summary>
 		/// Output location for the version file
 		/// </summary>
-		public FileReference VersionFile;
+		public FileReference? VersionFile;
 
 		/// <summary>
 		/// Output location for the target file
@@ -50,7 +51,7 @@ namespace UnrealBuildTool
 		/// <param name="ReceiptFile"></param>
 		/// <param name="Receipt"></param>
 		/// <param name="FileToManifest"></param>
-		public WriteMetadataTargetInfo(FileReference ProjectFile, FileReference VersionFile, FileReference ReceiptFile, TargetReceipt Receipt, Dictionary<FileReference, ModuleManifest> FileToManifest)//string EngineManifestName, string ProjectManifestName, Dictionary<string, FileReference> ModuleNameToLocation)
+		public WriteMetadataTargetInfo(FileReference? ProjectFile, FileReference? VersionFile, FileReference ReceiptFile, TargetReceipt Receipt, Dictionary<FileReference, ModuleManifest> FileToManifest)//string EngineManifestName, string ProjectManifestName, Dictionary<string, FileReference> ModuleNameToLocation)
 		{
 			this.ProjectFile = ProjectFile;
 			this.VersionFile = VersionFile;
@@ -112,7 +113,7 @@ namespace UnrealBuildTool
 			if(String.IsNullOrEmpty(Receipt.Version.BuildId))
 			{
 				// Check if there's an existing version file. If it exists, try to merge in any manifests that are valid (and reuse the existing build id)
-				BuildVersion PreviousVersion;
+				BuildVersion? PreviousVersion;
 				if(TargetInfo.VersionFile != null && BuildVersion.TryRead(TargetInfo.VersionFile, out PreviousVersion))
 				{
 					// Check if we can reuse the existing manifests. This prevents unnecessary builds when switching between projects.
@@ -142,7 +143,7 @@ namespace UnrealBuildTool
 				// Read all the manifests and merge them into the new ones, if they have the same build id
 				foreach(KeyValuePair<FileReference, ModuleManifest> Pair in TargetInfo.FileToManifest)
 				{
-					ModuleManifest SourceManifest;
+					ModuleManifest? SourceManifest;
 					if(TryReadManifest(Pair.Key, out SourceManifest) && SourceManifest.BuildId == Receipt.Version.BuildId)
 					{
 						MergeManifests(SourceManifest, Pair.Value);
@@ -224,12 +225,12 @@ namespace UnrealBuildTool
 		/// <param name="ManifestFiles">List of new manifest files</param>
 		/// <param name="RecycleFileToManifest">If successful, is populated with a map of filename to existing manifests that can be merged with the new manifests.</param>
 		/// <returns>True if the manifests can be recycled (and fills RecycleFileToManifest)</returns>
-		bool TryRecyclingManifests(string BuildId, IEnumerable<FileReference> ManifestFiles, Dictionary<FileReference, ModuleManifest> RecycleFileToManifest)
+		bool TryRecyclingManifests(string? BuildId, IEnumerable<FileReference> ManifestFiles, Dictionary<FileReference, ModuleManifest> RecycleFileToManifest)
 		{
 			bool bCanRecycleManifests = true;
 			foreach(FileReference ManifestFileName in ManifestFiles)
 			{
-				ModuleManifest Manifest;
+				ModuleManifest? Manifest;
 				if(ManifestFileName.IsUnderDirectory(UnrealBuildTool.EngineDirectory) && TryReadManifest(ManifestFileName, out Manifest))
 				{
 					if(Manifest.BuildId == BuildId)
@@ -252,7 +253,7 @@ namespace UnrealBuildTool
 		/// <param name="ManifestFileName">Path to the manifest</param>
 		/// <param name="Manifest">If successful, receives the manifest that was read</param>
 		/// <returns>True if the manifest was read correctly, false otherwise</returns>
-		public static bool TryReadManifest(FileReference ManifestFileName, out ModuleManifest Manifest)
+		public static bool TryReadManifest(FileReference ManifestFileName, [NotNullWhen(true)] out ModuleManifest? Manifest)
 		{
 			if(FileReference.Exists(ManifestFileName))
 			{
diff --git a/Engine/Source/Programs/UnrealBuildTool/Preprocessor/Identifier.cs b/Engine/Source/Programs/UnrealBuildTool/Preprocessor/Identifier.cs
--- a/Engine/Source/Programs/UnrealBuildTool/Preprocessor/Identifier.cs
+++ b/Engine/Source/Programs/UnrealBuildTool/Preprocessor/Identifier.cs
@@ -41,7 +41,7 @@ namespace UnrealBuildTool
 		/// <returns>New Identifier instance</returns>
 		public static Identifier FindOrAdd(string Name)
 		{
-			Identifier Result;
+			Identifier? Result;
 			if(!NameToIdentifier.TryGetValue(Name, out Result))
 			{
 				Identifier NewIdentifier = new Identifier(Name);
@@ -62,9 +62,9 @@ namespace UnrealBuildTool
 		/// </summary>
 		/// <param name="Other">Identifier to compare to</param>
 		/// <returns>Value indicating which identifier should sort first</returns>
-		public int CompareTo(Identifier Other)
+		public int CompareTo(Identifier? Other)
 		{
-			return Name.CompareTo(Other.Name);
+			return ReferenceEquals(Other, null)? 1 : Name.CompareTo(Other.Name);
 		}
 
 		/// <summary>
@@ -127,9 +127,9 @@ namespace UnrealBuildTool
 		/// </summary>
 		/// <param name="Writer">Writer to serialize data to</param>
 		/// <param name="Identifier">Identifier to write</param>
-		public static void WriteIdentifier(this BinaryArchiveWriter Writer, Identifier Identifier)
+		public static void WriteIdentifier(this BinaryArchiveWriter Writer, Identifier? Identifier)
 		{
-			Writer.WriteObjectReference<Identifier>(Identifier, () => Writer.WriteString(Identifier.ToString()));
+			Writer.WriteObjectReference<Identifier?>(Identifier, () => Writer.WriteString(Identifier!.ToString()));
 		}
 	}
 }
diff --git a/Engine/Source/Programs/UnrealBuildTool/Preprocessor/Preprocessor.cs b/Engine/Source/Programs/UnrealBuildTool/Preprocessor/Preprocessor.cs
--- a/Engine/Source/Programs/UnrealBuildTool/Preprocessor/Preprocessor.cs
+++ b/Engine/Source/Programs/UnrealBuildTool/Preprocessor/Preprocessor.cs
@@ -10,6 +10,7 @@ using System.IO;
 using System.Collections;
 using System.Collections.Concurrent;
 using Tools.DotNETCommon;
+using System.Diagnostics.CodeAnalysis;
 
 namespace UnrealBuildTool
 {
@@ -21,7 +22,7 @@ namespace UnrealBuildTool
 		/// <summary>
 		/// The context when the error was encountered
 		/// </summary>
-		PreprocessorContext Context;
+		PreprocessorContext? Context;
 
 		/// <summary>
 		/// Constructor
@@ -29,7 +30,7 @@ namespace UnrealBuildTool
 		/// <param name="Context">The current preprocesor context</param>
 		/// <param name="Format">Format string, to be passed to String.Format</param>
 		/// <param name="Args">Optional argument list for the format string</param>
-		public PreprocessorException(PreprocessorContext Context, string Format, params object[] Args)
+		public PreprocessorException(PreprocessorContext? Context, string Format, params object[] Args)
 			: base(String.Format(Format, Args))
 		{
 			this.Context = Context;
@@ -123,7 +124,7 @@ namespace UnrealBuildTool
 			}
 			if(Tokens[0].Type != TokenType.Identifier)
 			{
-				throw new PreprocessorException(null, "'{0}' is not a valid macro name", Tokens[0].ToString());
+				throw new PreprocessorException(null, "'{0}' is not a valid macro name", Tokens[0].ToString()!);
 			}
 
 			List<Token> ValueTokens = new List<Token>();
@@ -140,7 +141,7 @@ namespace UnrealBuildTool
 				ValueTokens.AddRange(Tokens.Skip(2));
 			}
 
-			PreprocessorMacro Macro = new PreprocessorMacro(Tokens[0].Identifier, null, Tokens);
+			PreprocessorMacro Macro = new PreprocessorMacro(Tokens[0].Identifier!, null, Tokens);
 			State.DefineMacro(Macro);
 		}
 
@@ -215,7 +216,7 @@ namespace UnrealBuildTool
 		/// <param name="Type">Specifies rules for how to resolve the include path (normal/system)</param>
 		/// <param name="File">If found, receives the resolved file</param>
 		/// <returns>True if the The resolved file</returns>
-		public bool TryResolveIncludePath(PreprocessorContext Context, string IncludePath, IncludePathType Type, out FileItem File)
+		public bool TryResolveIncludePath(PreprocessorContext Context, string IncludePath, IncludePathType Type, [NotNullWhen(true)] out FileItem? File)
 		{
 			// From MSDN (https://msdn.microsoft.com/en-us/library/36k2cdd4.aspx?f=255&MSPPError=-2147217396)
 			//
@@ -254,12 +255,12 @@ namespace UnrealBuildTool
 			// Try to match the include path against any of the included directories
 			if(Type == IncludePathType.Normal)
 			{
-				for(PreprocessorContext OuterContext = Context; OuterContext != null; OuterContext = OuterContext.Outer)
+				for(PreprocessorContext? OuterContext = Context; OuterContext != null; OuterContext = OuterContext.Outer)
 				{
-					PreprocessorFileContext OuterFileContext = OuterContext as PreprocessorFileContext;
+					PreprocessorFileContext? OuterFileContext = OuterContext as PreprocessorFileContext;
 					if(OuterFileContext != null)
 					{
-						FileItem ResolvedFile;
+						FileItem? ResolvedFile;
 						if(TryResolveRelativeIncludePath(OuterFileContext.Directory, Fragments, out ResolvedFile))
 						{
 							File = ResolvedFile;
@@ -272,7 +273,7 @@ namespace UnrealBuildTool
 			// Try to match the include path against any of the system directories
 			foreach(DirectoryItem BaseDirectory in IncludeDirectories)
 			{
-				FileItem ResolvedFile;
+				FileItem? ResolvedFile;
 				if(TryResolveRelativeIncludePath(BaseDirectory, Fragments, out ResolvedFile))
 				{
 					File = ResolvedFile;
@@ -292,9 +293,9 @@ namespace UnrealBuildTool
 		/// <param name="Fragments">Fragments of the relative path to follow</param>
 		/// <param name="File">The file that was found, if successful</param>
 		/// <returns>True if the The resolved file</returns>
-		public bool TryResolveRelativeIncludePath(DirectoryItem BaseDirectory, string[] Fragments, out FileItem File)
+		public bool TryResolveRelativeIncludePath(DirectoryItem BaseDirectory, string[] Fragments, [NotNullWhen(true)] out FileItem? File)
 		{
-			DirectoryItem Directory = BaseDirectory;
+			DirectoryItem? Directory = BaseDirectory;
 			for(int Idx = 0; Idx < Fragments.Length - 1; Idx++)
 			{
 				if(!Directory.TryGetDirectory(Fragments[Idx], out Directory))
@@ -380,7 +381,7 @@ namespace UnrealBuildTool
 		{
 			// Expand macros in the given tokens
 			List<Token> ExpandedTokens = new List<Token>();
-			ExpandMacros(Markup.Tokens, ExpandedTokens, false, Context);
+			ExpandMacros(Markup.Tokens!, ExpandedTokens, false, Context);
 						
 			// Convert the string to a single token
 			string IncludeToken = Token.Format(ExpandedTokens);
@@ -404,7 +405,7 @@ namespace UnrealBuildTool
 			string IncludePath = IncludeToken.Substring(1, IncludeToken.Length - 2);
 
 			// Resolve the included file
-			FileItem IncludedFile;
+			FileItem? IncludedFile;
 			if(!TryResolveIncludePath(Context, IncludePath, Type, out IncludedFile))
 			{
 				throw new PreprocessorException(Context, "Couldn't resolve include '{0}'", IncludePath);
@@ -435,9 +436,9 @@ namespace UnrealBuildTool
 						for(; Context.MarkupIdx < Fragment.MarkupMax; Context.MarkupIdx++)
 						{
 							SourceFileMarkup Markup = SourceFile.Markup[Context.MarkupIdx];
-							ParseMarkup(Markup.Type, Markup.Tokens, Context);
+							ParseMarkup(Markup.Type, Markup.Tokens!, Context);
 						}
-						Transform = State.EndCapture();
+						Transform = State.EndCapture()!;
 
 						// Add it to the fragment for future fragments
 						PreprocessorTransform[] NewTransforms = new PreprocessorTransform[Fragment.Transforms.Length + 1];
@@ -479,7 +480,7 @@ namespace UnrealBuildTool
 		/// <param name="Name">Name of the macro</param>
 		/// <param name="Parameters">Parameter list for the macro</param>
 		/// <param name="Tokens">List of tokens</param>
-		void AddMacro(PreprocessorContext Context, Identifier Name, List<Identifier> Parameters, List<Token> Tokens)
+		void AddMacro(PreprocessorContext Context, Identifier Name, List<Identifier>? Parameters, List<Token> Tokens)
 		{
 			if(Tokens.Count == 0)
 			{
@@ -578,11 +579,11 @@ namespace UnrealBuildTool
 				}
 
 				// Read the macro name
-				Identifier Name = Tokens[0].Identifier;
+				Identifier Name = Tokens[0].Identifier!;
 				int TokenIdx = 1;
 
 				// Read the macro parameter list, if there is one
-				List<Identifier> Parameters = null;
+				List<Identifier>? Parameters = null;
 				if (TokenIdx < Tokens.Count && !Tokens[TokenIdx].HasLeadingSpace && Tokens[TokenIdx].Type == TokenType.LeftParen)
 				{
 					Parameters = new List<Identifier>();
@@ -619,12 +620,12 @@ namespace UnrealBuildTool
 								{
 									throw new PreprocessorException(Context, "Invalid preprocessor token: {0}", NameToken);
 								}
-								if (Parameters.Contains(NameToken.Identifier))
+								if (Parameters.Contains(NameToken.Identifier!))
 								{
 									throw new PreprocessorException(Context, "'{0}' has already been used as an argument name", NameToken);
 								}
 							}
-							Parameters.Add(NameToken.Identifier);
+							Parameters.Add(NameToken.Identifier!);
 
 							// Read the separator
 							Token SeparatorToken = Tokens[TokenIdx];
@@ -666,7 +667,7 @@ namespace UnrealBuildTool
 				}
 
 				// Remove the macro from the list of definitions
-				State.UndefMacro(Tokens[0].Identifier);
+				State.UndefMacro(Tokens[0].Identifier!);
 			}
 		}
 
@@ -711,7 +712,7 @@ namespace UnrealBuildTool
 				}
 
 				// Check if the macro is defined
-				if(State.IsMacroDefined(Tokens[0].Identifier))
+				if(State.IsMacroDefined(Tokens[0].Identifier!))
 				{
 					Branch |= PreprocessorBranch.Active | PreprocessorBranch.Taken;
 				}
@@ -736,7 +737,7 @@ namespace UnrealBuildTool
 				}
 
 				// Check if the macro is defined
-				if(!State.IsMacroDefined(Tokens[0].Identifier))
+				if(!State.IsMacroDefined(Tokens[0].Identifier!))
 				{
 					Branch |= PreprocessorBranch.Active | PreprocessorBranch.Taken;
 				}
@@ -842,7 +843,7 @@ namespace UnrealBuildTool
 			{
 				if(Tokens.Count == 1 && Tokens[0].Identifier == Identifiers.Once)
 				{
-					SourceFile SourceFile = GetCurrentSourceFile(Context);
+					SourceFile SourceFile = GetCurrentSourceFile(Context)!;
 					PragmaOnceFiles.Add(SourceFile.File);
 					State.MarkPragmaOnce();
 				}
@@ -949,7 +950,7 @@ namespace UnrealBuildTool
 				}
 
 				// Insert a token for whether it's defined or not
-				OutputTokens.Add(new Token(TokenType.Number, TokenFlags.None, State.IsMacroDefined(NameToken.Identifier)? OneLiteral : ZeroLiteral));
+				OutputTokens.Add(new Token(TokenType.Number, TokenFlags.None, State.IsMacroDefined(NameToken.Identifier!)? OneLiteral : ZeroLiteral));
 				bMoveNext = InputEnumerator.MoveNext();
 			}
 			else
@@ -958,8 +959,8 @@ namespace UnrealBuildTool
 				while(OutputTokens[OutputTokens.Count - 1].Type == TokenType.Identifier && !OutputTokens[OutputTokens.Count - 1].Flags.HasFlag(TokenFlags.DisableExpansion))
 				{
 					// Try to get a macro for the current token
-					PreprocessorMacro Macro;
-					if (!State.TryGetMacro(OutputTokens[OutputTokens.Count - 1].Identifier, out Macro) || IgnoreMacros.Contains(Macro))
+					PreprocessorMacro? Macro;
+					if (!State.TryGetMacro(OutputTokens[OutputTokens.Count - 1].Identifier!, out Macro) || IgnoreMacros.Contains(Macro))
 					{
 						break;
 					}
@@ -1020,7 +1021,7 @@ namespace UnrealBuildTool
 		/// <returns>String representing the current context</returns>
 		string GetCurrentFileMacroValue(PreprocessorContext Context)
 		{
-			SourceFile SourceFile = GetCurrentSourceFile(Context);
+			SourceFile? SourceFile = GetCurrentSourceFile(Context);
 			if(SourceFile == null)
 			{
 				return "<unknown>";
@@ -1036,12 +1037,12 @@ namespace UnrealBuildTool
 		/// </summary>
 		/// <param name="Context">Context to scan to find the current file</param>
 		/// <returns>Current source file being parsed</returns>
-		SourceFile GetCurrentSourceFile(PreprocessorContext Context)
+		SourceFile? GetCurrentSourceFile(PreprocessorContext Context)
 		{
-			SourceFile SourceFile = null;
-			for(PreprocessorContext OuterContext = Context; OuterContext != null; OuterContext = OuterContext.Outer)
+			SourceFile? SourceFile = null;
+			for(PreprocessorContext? OuterContext = Context; OuterContext != null; OuterContext = OuterContext.Outer)
 			{
-				PreprocessorFileContext OuterFileContext = OuterContext as PreprocessorFileContext;
+				PreprocessorFileContext? OuterFileContext = OuterContext as PreprocessorFileContext;
 				if(OuterFileContext != null)
 				{
 					SourceFile = OuterFileContext.SourceFile;
@@ -1058,9 +1059,9 @@ namespace UnrealBuildTool
 		/// <returns>Line number in the first file encountered</returns>
 		int GetCurrentLine(PreprocessorContext Context)
 		{
-			for(PreprocessorContext OuterContext = Context; OuterContext != null; OuterContext = OuterContext.Outer)
+			for(PreprocessorContext? OuterContext = Context; OuterContext != null; OuterContext = OuterContext.Outer)
 			{
-				PreprocessorFileContext OuterFileContext = OuterContext as PreprocessorFileContext;
+				PreprocessorFileContext? OuterFileContext = OuterContext as PreprocessorFileContext;
 				if(OuterFileContext != null)
 				{
 					return OuterFileContext.SourceFile.Markup[OuterFileContext.MarkupIdx].LineNumber;
@@ -1147,7 +1148,7 @@ namespace UnrealBuildTool
 			{
 				for(int Idx = 1;;Idx++)
 				{
-					if (!Macro.HasVariableArgumentList || Arguments.Count < Macro.Parameters.Count)
+					if (!Macro.HasVariableArgumentList || Arguments.Count < Macro.Parameters!.Count)
 					{
 						Arguments.Add(new List<Token>());
 					}
@@ -1199,7 +1200,7 @@ namespace UnrealBuildTool
 						throw new PreprocessorException(Context, "Expected ',' between arguments");
 					}
 
-					if (Macro.HasVariableArgumentList && Arguments.Count == Macro.Parameters.Count && Idx < ArgumentListTokens.Count - 1)
+					if (Macro.HasVariableArgumentList && Arguments.Count == Macro.Parameters!.Count && Idx < ArgumentListTokens.Count - 1)
 					{
 						Arguments[Arguments.Count - 1].Add(ArgumentListTokens[Idx]);
 					}
@@ -1207,13 +1208,13 @@ namespace UnrealBuildTool
 			}
 
 			// Add an empty variable argument if one was not specified
-			if (Macro.HasVariableArgumentList && Arguments.Count == Macro.Parameters.Count - 1)
+			if (Macro.HasVariableArgumentList && Arguments.Count == Macro.Parameters!.Count - 1)
 			{
 				Arguments.Add(new List<Token> { new Token(TokenType.Placemarker, TokenFlags.None) });
 			}
 
 			// Validate the argument list
-			if (Arguments.Count != Macro.Parameters.Count)
+			if (Arguments.Count != Macro.Parameters!.Count)
 			{
 				throw new PreprocessorException(Context, "Incorrect number of arguments to macro");
 			}
@@ -1235,7 +1236,7 @@ namespace UnrealBuildTool
 				if(Token.Type == TokenType.Hash && Idx + 1 < Macro.Tokens.Count)
 				{
 					// Stringizing operator
-					int ParamIdx = Macro.FindParameterIndex(Macro.Tokens[++Idx].Identifier);
+					int ParamIdx = Macro.FindParameterIndex(Macro.Tokens[++Idx].Identifier!);
 					if (ParamIdx == -1)
 					{
 						throw new PreprocessorException(Context, "{0} is not an argument name", Macro.Tokens[Idx].Text);
@@ -1262,7 +1263,7 @@ namespace UnrealBuildTool
 				else if (Token.Type == TokenType.Identifier)
 				{
 					// Expand a parameter
-					int ParamIdx = Macro.FindParameterIndex(Token.Identifier);
+					int ParamIdx = Macro.FindParameterIndex(Token.Identifier!);
 					if(ParamIdx == -1)
 					{
 						ExpandedTokens.Add(Token);
diff --git a/Engine/Source/Programs/UnrealBuildTool/Preprocessor/PreprocessorContext.cs b/Engine/Source/Programs/UnrealBuildTool/Preprocessor/PreprocessorContext.cs
--- a/Engine/Source/Programs/UnrealBuildTool/Preprocessor/PreprocessorContext.cs
+++ b/Engine/Source/Programs/UnrealBuildTool/Preprocessor/PreprocessorContext.cs
@@ -16,13 +16,13 @@ namespace UnrealBuildTool
 		/// <summary>
 		/// The outer context
 		/// </summary>
-		public readonly PreprocessorContext Outer;
+		public readonly PreprocessorContext? Outer;
 
 		/// <summary>
 		/// Constructor
 		/// </summary>
 		/// <param name="Outer">The outer context</param>
-		public PreprocessorContext(PreprocessorContext Outer)
+		public PreprocessorContext(PreprocessorContext? Outer)
 		{
 			this.Outer = Outer;
 		}
diff --git a/Engine/Source/Programs/UnrealBuildTool/Preprocessor/PreprocessorExpression.cs b/Engine/Source/Programs/UnrealBuildTool/Preprocessor/PreprocessorExpression.cs
--- a/Engine/Source/Programs/UnrealBuildTool/Preprocessor/PreprocessorExpression.cs
+++ b/Engine/Source/Programs/UnrealBuildTool/Preprocessor/PreprocessorExpression.cs
@@ -229,7 +229,7 @@ namespace UnrealBuildTool
 					Idx++;
 					return ~EvaluateUnary(Context, Tokens, ref Idx);
 				case TokenType.Identifier:
-					Identifier Text = Tokens[Idx].Identifier;
+					Identifier Text = Tokens[Idx].Identifier!;
 					if(Text == Identifiers.Sizeof)
 					{
 						throw new NotImplementedException();
@@ -297,7 +297,7 @@ namespace UnrealBuildTool
 			}
 			else if(Tokens[Idx].Type == TokenType.Number)
 			{
-				return ParseNumericLiteral(Context, Tokens[Idx++].Literal);
+				return ParseNumericLiteral(Context, Tokens[Idx++].Literal!);
 			}
 			else
 			{
diff --git a/Engine/Source/Programs/UnrealBuildTool/Preprocessor/PreprocessorMacro.cs b/Engine/Source/Programs/UnrealBuildTool/Preprocessor/PreprocessorMacro.cs
--- a/Engine/Source/Programs/UnrealBuildTool/Preprocessor/PreprocessorMacro.cs
+++ b/Engine/Source/Programs/UnrealBuildTool/Preprocessor/PreprocessorMacro.cs
@@ -22,7 +22,7 @@ namespace UnrealBuildTool
 		/// <summary>
 		/// Parameter names for the macro. The '...' placeholder is represented by the __VA_ARGS__ string.
 		/// </summary>
-		public List<Identifier> Parameters;
+		public List<Identifier>? Parameters;
 
 		/// <summary>
 		/// Raw list of tokens for this macro
@@ -35,7 +35,7 @@ namespace UnrealBuildTool
 		/// <param name="Name">Name of the macro</param>
 		/// <param name="Parameters">Parameter list for the macro. Should be null for object macros. Ownership of this list is transferred.</param>
 		/// <param name="Tokens">Tokens for the macro. Ownership of this list is transferred.</param>
-		public PreprocessorMacro(Identifier Name, List<Identifier> Parameters, List<Token> Tokens)
+		public PreprocessorMacro(Identifier Name, List<Identifier>? Parameters, List<Token> Tokens)
 		{
 			this.Name = Name;
 			this.Parameters = Parameters;
@@ -71,7 +71,7 @@ namespace UnrealBuildTool
 		/// <returns>Index of the parameter, or -1 if it's not found.</returns>
 		public int FindParameterIndex(Identifier Parameter)
 		{
-			for (int Idx = 0; Idx < Parameters.Count; Idx++)
+			for (int Idx = 0; Idx < Parameters!.Count; Idx++)
 			{
 				if (Parameters[Idx] == Parameter)
 				{
@@ -137,7 +137,7 @@ namespace UnrealBuildTool
 		/// </summary>
 		public int MinRequiredParameters
 		{
-			get { return HasVariableArgumentList? (Parameters.Count - 1) : Parameters.Count; }
+			get { return HasVariableArgumentList? (Parameters!.Count - 1) : Parameters!.Count; }
 		}
 
 		/// <summary>
@@ -145,7 +145,7 @@ namespace UnrealBuildTool
 		/// </summary>
 		public bool HasVariableArgumentList
 		{
-			get { return Parameters.Count > 0 && Parameters[Parameters.Count - 1] == Identifiers.__VA_ARGS__; }
+			get { return Parameters!.Count > 0 && Parameters[Parameters.Count - 1] == Identifiers.__VA_ARGS__; }
 		}
 
 		/// <summary>
diff --git a/Engine/Source/Programs/UnrealBuildTool/Preprocessor/PreprocessorState.cs b/Engine/Source/Programs/UnrealBuildTool/Preprocessor/PreprocessorState.cs
--- a/Engine/Source/Programs/UnrealBuildTool/Preprocessor/PreprocessorState.cs
+++ b/Engine/Source/Programs/UnrealBuildTool/Preprocessor/PreprocessorState.cs
@@ -2,6 +2,7 @@
 
 using System;
 using System.Collections.Generic;
+using System.Diagnostics.CodeAnalysis;
 using System.Linq;
 using System.Text;
 using System.Threading.Tasks;
@@ -68,7 +69,7 @@ namespace UnrealBuildTool
 		/// <summary>
 		/// The current transform. Any queries or state modifications will be recorded in this.
 		/// </summary>
-		PreprocessorTransform Transform;
+		PreprocessorTransform? Transform;
 
 		/// <summary>
 		/// Enumerates the current branches. Do not use this for decision logic; any dependencies will not be tracked.
@@ -126,9 +127,9 @@ namespace UnrealBuildTool
 		/// Detach the current transform
 		/// </summary>
 		/// <returns>The transform object</returns>
-		public PreprocessorTransform EndCapture()
+		public PreprocessorTransform? EndCapture()
 		{
-			PreprocessorTransform PrevTransform = Transform;
+			PreprocessorTransform? PrevTransform = Transform;
 			Transform = null;
 			return PrevTransform;
 		}
@@ -166,7 +167,7 @@ namespace UnrealBuildTool
 		public bool IsMacroDefined(Identifier Name)
 		{
 			// Could account for the fact that we don't need the full definition later...
-			PreprocessorMacro Macro;
+			PreprocessorMacro? Macro;
 			return TryGetMacro(Name, out Macro);
 		}
 
@@ -190,7 +191,7 @@ namespace UnrealBuildTool
 		/// <param name="Name">Name of the macro</param>
 		/// <param name="Macro">Receives the macro definition, or null if it's not defined</param>
 		/// <returns>True if the macro is defined, otherwise false</returns>
-		public bool TryGetMacro(Identifier Name, out PreprocessorMacro Macro)
+		public bool TryGetMacro(Identifier Name, [NotNullWhen(true)] out PreprocessorMacro? Macro)
 		{
 			bool bResult = NameToMacro.TryGetValue(Name, out Macro);
 
@@ -199,7 +200,7 @@ namespace UnrealBuildTool
 				Transform.RequiredMacros[Name] = Macro;
 			}
 
-			return bResult;
+			return Macro != null;
 		}
 
 		/// <summary>
@@ -305,9 +306,9 @@ namespace UnrealBuildTool
 			}
 
 			// Check all the required macros match
-			foreach(KeyValuePair<Identifier, PreprocessorMacro> RequiredPair in Transform.RequiredMacros)
+			foreach(KeyValuePair<Identifier, PreprocessorMacro?> RequiredPair in Transform.RequiredMacros)
 			{
-				PreprocessorMacro Macro;
+				PreprocessorMacro? Macro;
 				if(NameToMacro.TryGetValue(RequiredPair.Key, out Macro))
 				{
 					if(RequiredPair.Value == null || !Macro.IsEquivalentTo(RequiredPair.Value))
@@ -344,7 +345,7 @@ namespace UnrealBuildTool
 			Branches.AddRange(Transform.NewBranches);
 
 			// Update the macro definitions
-			foreach(KeyValuePair<Identifier, PreprocessorMacro> NewMacro in Transform.NewMacros)
+			foreach(KeyValuePair<Identifier, PreprocessorMacro?> NewMacro in Transform.NewMacros)
 			{
 				if(NewMacro.Value == null)
 				{
diff --git a/Engine/Source/Programs/UnrealBuildTool/Preprocessor/PreprocessorTransform.cs b/Engine/Source/Programs/UnrealBuildTool/Preprocessor/PreprocessorTransform.cs
--- a/Engine/Source/Programs/UnrealBuildTool/Preprocessor/PreprocessorTransform.cs
+++ b/Engine/Source/Programs/UnrealBuildTool/Preprocessor/PreprocessorTransform.cs
@@ -27,7 +27,7 @@ namespace UnrealBuildTool
 		/// <summary>
 		/// Map of macro name to their required value
 		/// </summary>
-		public Dictionary<Identifier, PreprocessorMacro> RequiredMacros;
+		public Dictionary<Identifier, PreprocessorMacro?> RequiredMacros;
 
 		/// <summary>
 		/// List of new branches that will be pushed onto the stack
@@ -37,7 +37,7 @@ namespace UnrealBuildTool
 		/// <summary>
 		/// List of new macros that will be defined
 		/// </summary>
-		public Dictionary<Identifier, PreprocessorMacro> NewMacros;
+		public Dictionary<Identifier, PreprocessorMacro?> NewMacros;
 
 		/// <summary>
 		/// This fragment contains a pragma once directive
@@ -50,9 +50,9 @@ namespace UnrealBuildTool
 		public PreprocessorTransform()
 		{
 			RequiredBranches = new List<PreprocessorBranch>();
-			RequiredMacros = new Dictionary<Identifier, PreprocessorMacro>();
+			RequiredMacros = new Dictionary<Identifier, PreprocessorMacro?>();
 			NewBranches = new List<PreprocessorBranch>();
-			NewMacros = new Dictionary<Identifier, PreprocessorMacro>();
+			NewMacros = new Dictionary<Identifier, PreprocessorMacro?>();
 		}
 
 		/// <summary>
@@ -67,9 +67,9 @@ namespace UnrealBuildTool
 			}
 
 			RequiredBranches = Reader.ReadList(() => (PreprocessorBranch)Reader.ReadByte());
-			RequiredMacros = Reader.ReadDictionary(() => Reader.ReadIdentifier(), () => Reader.ReadObjectReference(() => new PreprocessorMacro(Reader)));
+			RequiredMacros = Reader.ReadDictionary(() => Reader.ReadIdentifier(), () => (PreprocessorMacro?)Reader.ReadObjectReference(() => new PreprocessorMacro(Reader)));
 			NewBranches = Reader.ReadList(() => (PreprocessorBranch)Reader.ReadByte());
-			NewMacros = Reader.ReadDictionary(() => Reader.ReadIdentifier(), () => Reader.ReadObjectReference(() => new PreprocessorMacro(Reader)));
+			NewMacros = Reader.ReadDictionary(() => Reader.ReadIdentifier(), () => (PreprocessorMacro?)Reader.ReadObjectReference(() => new PreprocessorMacro(Reader)));
 			bHasPragmaOnce = Reader.ReadBool();
 		}
 
@@ -86,9 +86,9 @@ namespace UnrealBuildTool
 			}
 
 			Writer.WriteList(RequiredBranches, x => Writer.WriteByte((byte)x));
-			Writer.WriteDictionary(RequiredMacros, k => Writer.WriteIdentifier(k), v => Writer.WriteObjectReference(v, () => v.Write(Writer)));
+			Writer.WriteDictionary(RequiredMacros, k => Writer.WriteIdentifier(k), v => Writer.WriteObjectReference(v, () => v!.Write(Writer)));
 			Writer.WriteList(NewBranches, x => Writer.WriteByte((byte)x));
-			Writer.WriteDictionary(NewMacros, k => Writer.WriteIdentifier(k), v => Writer.WriteObjectReference(v, () => v.Write(Writer)));
+			Writer.WriteDictionary(NewMacros, k => Writer.WriteIdentifier(k), v => Writer.WriteObjectReference(v, () => v!.Write(Writer)));
 			Writer.WriteBool(bHasPragmaOnce);
 		}
 	}
diff --git a/Engine/Source/Programs/UnrealBuildTool/Preprocessor/SourceFile.cs b/Engine/Source/Programs/UnrealBuildTool/Preprocessor/SourceFile.cs
--- a/Engine/Source/Programs/UnrealBuildTool/Preprocessor/SourceFile.cs
+++ b/Engine/Source/Programs/UnrealBuildTool/Preprocessor/SourceFile.cs
@@ -41,7 +41,7 @@ namespace UnrealBuildTool
 		/// <summary>
 		/// Macro which is defined as part of the header guard. If this macro is defined, we don't need to include the header more than once.
 		/// </summary>
-		public Identifier HeaderGuardMacro;
+		public Identifier? HeaderGuardMacro;
 
 		/// <summary>
 		/// Parsed fragments in this file. Only set for files which are being optimized.
@@ -61,18 +61,20 @@ namespace UnrealBuildTool
 		/// </summary>
 		/// <param name="File">Location of the file</param>
 		public SourceFile(FileItem File)
-			: this(ReadFileWithNullTerminator(File.Location))
+			: this(File, ReadFileWithNullTerminator(File.Location))
 		{
-			this.File = File;
 			this.LastWriteTimeUtc = File.LastWriteTimeUtc.Ticks;
 		}
 
 		/// <summary>
 		/// Construct a SourceFile from a raw byte array
 		/// </summary>
+		/// <param name="File">Location of the file</param>
 		/// <param name="Data">Contents of the file</param>
-		public SourceFile(byte[] Data)
+		public SourceFile(FileItem File, byte[] Data)
 		{
+			this.File = File;
+
 			// Compute the hash
 			Hash = ContentHash.MD5(Data);
 
@@ -201,14 +203,14 @@ namespace UnrealBuildTool
 				int MinIdx = 0;
 
 				// Get the define used for the header guard
-				Identifier PossibleHeaderGuardMacro;
-				if(Markup[MinIdx].Type == SourceFileMarkupType.Ifndef && Markup[MinIdx].Tokens.Count == 1 && Markup[MinIdx].Tokens[0].Type == TokenType.Identifier)
+				Identifier? PossibleHeaderGuardMacro;
+				if(Markup[MinIdx].Type == SourceFileMarkupType.Ifndef && Markup[MinIdx].Tokens!.Count == 1 && Markup[MinIdx].Tokens![0].Type == TokenType.Identifier)
 				{
-					PossibleHeaderGuardMacro = Markup[MinIdx].Tokens[0].Identifier;
+					PossibleHeaderGuardMacro = Markup[MinIdx].Tokens![0].Identifier!;
 				}
-				else if(Markup[MinIdx].Type == SourceFileMarkupType.If && Markup[MinIdx].Tokens.Count == 3 && Markup[MinIdx].Tokens[0].Type == TokenType.LogicalNot && Markup[MinIdx].Tokens[1].Identifier == Identifiers.Defined && Markup[MinIdx].Tokens[2].Type == TokenType.Identifier)
+				else if(Markup[MinIdx].Type == SourceFileMarkupType.If && Markup[MinIdx].Tokens!.Count == 3 && Markup[MinIdx].Tokens![0].Type == TokenType.LogicalNot && Markup[MinIdx].Tokens![1].Identifier == Identifiers.Defined && Markup[MinIdx].Tokens![2].Type == TokenType.Identifier)
 				{
-					PossibleHeaderGuardMacro = Markup[MinIdx].Tokens[2].Identifier;
+					PossibleHeaderGuardMacro = Markup[MinIdx].Tokens![2].Identifier!;
 				}
 				else
 				{
@@ -216,7 +218,7 @@ namespace UnrealBuildTool
 				}
 
 				// Check it starts with an #if or #ifdef
-				if (Markup[MinIdx + 1].Type == SourceFileMarkupType.Define && Markup[MinIdx + 1].Tokens.Count == 1 && Markup[MinIdx + 1].Tokens[0].Identifier == PossibleHeaderGuardMacro)
+				if (Markup[MinIdx + 1].Type == SourceFileMarkupType.Define && Markup[MinIdx + 1].Tokens!.Count == 1 && Markup[MinIdx + 1].Tokens![0].Identifier == PossibleHeaderGuardMacro)
 				{
 					// Find the point at which the include depth goes back to zero
 					int MaxIdx = MinIdx + 1;
diff --git a/Engine/Source/Programs/UnrealBuildTool/Preprocessor/SourceFileMarkup.cs b/Engine/Source/Programs/UnrealBuildTool/Preprocessor/SourceFileMarkup.cs
--- a/Engine/Source/Programs/UnrealBuildTool/Preprocessor/SourceFileMarkup.cs
+++ b/Engine/Source/Programs/UnrealBuildTool/Preprocessor/SourceFileMarkup.cs
@@ -110,7 +110,7 @@ namespace UnrealBuildTool
 		/// <summary>
 		/// The tokens parsed for this markup. Set for directives.
 		/// </summary>
-		public List<Token> Tokens;
+		public List<Token>? Tokens;
 
 		/// <summary>
 		/// Construct the annotation with the given range
@@ -118,7 +118,7 @@ namespace UnrealBuildTool
 		/// <param name="Type">The type of this directive</param>
 		/// <param name="LineNumber">The line number of this markup</param>
 		/// <param name="Tokens">List of tokens</param>
-		public SourceFileMarkup(SourceFileMarkupType Type, int LineNumber, List<Token> Tokens)
+		public SourceFileMarkup(SourceFileMarkupType Type, int LineNumber, List<Token>? Tokens)
 		{
 			this.Type = Type;
 			this.LineNumber = LineNumber;
@@ -236,7 +236,7 @@ namespace UnrealBuildTool
 						{
 							if(Reader.Current.Type == TokenType.Identifier)
 							{
-								Identifier Directive = Reader.Current.Identifier;
+								Identifier Directive = Reader.Current.Identifier!;
 								if(Directive == Identifiers.Include)
 								{
 									Type = SourceFileMarkupType.Include;
diff --git a/Engine/Source/Programs/UnrealBuildTool/Preprocessor/Token.cs b/Engine/Source/Programs/UnrealBuildTool/Preprocessor/Token.cs
--- a/Engine/Source/Programs/UnrealBuildTool/Preprocessor/Token.cs
+++ b/Engine/Source/Programs/UnrealBuildTool/Preprocessor/Token.cs
@@ -110,7 +110,7 @@ namespace UnrealBuildTool
 		/// <summary>
 		/// If this token is identifier, contains the identifier name
 		/// </summary>
-		public Identifier Identifier
+		public Identifier? Identifier
 		{
 			get;
 		}
@@ -118,7 +118,7 @@ namespace UnrealBuildTool
 		/// <summary>
 		/// If this token is a literal, contains the raw utf-8 representation of it.
 		/// </summary>
-		public byte[] Literal
+		public byte[]? Literal
 		{
 			get;
 		}
@@ -194,7 +194,7 @@ namespace UnrealBuildTool
 			Writer.WriteByte((byte)Flags);
 			if(Type == TokenType.Identifier)
 			{
-				Writer.WriteIdentifier(Identifier);
+				Writer.WriteIdentifier(Identifier!);
 			}
 			else if(IsLiteral(Type))
 			{
@@ -207,7 +207,7 @@ namespace UnrealBuildTool
 		/// </summary>
 		/// <param name="Other">The object to compare against</param>
 		/// <returns>True if the objects are equal, false otherwise</returns>
-		public override bool Equals(object Other)
+		public override bool Equals(object? Other)
 		{
 			return Equals(Other as Token);
 		}
@@ -217,9 +217,9 @@ namespace UnrealBuildTool
 		/// </summary>
 		/// <param name="Other">The object to compare against</param>
 		/// <returns>True if the tokens are equal, false otherwise</returns>
-		public bool Equals(Token Other)
+		public bool Equals(Token? Other)
 		{
-			if((object)Other == null)
+			if(ReferenceEquals(Other, null))
 			{
 				return false;
 			}
diff --git a/Engine/Source/Programs/UnrealBuildTool/Preprocessor/TokenReader.cs b/Engine/Source/Programs/UnrealBuildTool/Preprocessor/TokenReader.cs
--- a/Engine/Source/Programs/UnrealBuildTool/Preprocessor/TokenReader.cs
+++ b/Engine/Source/Programs/UnrealBuildTool/Preprocessor/TokenReader.cs
@@ -84,7 +84,7 @@ namespace UnrealBuildTool
 		/// <summary>
 		/// The current token
 		/// </summary>
-		Token CurrentToken;
+		Token? CurrentToken;
 
 		/// <summary>
 		/// Constructor
@@ -183,7 +183,7 @@ namespace UnrealBuildTool
 		/// </summary>
 		public Token Current
 		{
-			get { return CurrentToken; }
+			get { return CurrentToken!; }
 		}
 
 		/// <summary>
@@ -446,7 +446,7 @@ namespace UnrealBuildTool
 		public bool MoveNextIncludePath()
 		{
 			bool bResult = MoveNext();
-			if(bResult && CurrentToken.Type == TokenType.CompareLess)
+			if(bResult && Current.Type == TokenType.CompareLess)
 			{
 				int StartOffset = Offset - 1;
 				for(int EndOffset = Offset; Data[EndOffset] != 0 && Data[EndOffset] != '\n'; EndOffset++)
diff --git a/Engine/Source/Programs/UnrealBuildTool/Storage/IStorageProvider.cs b/Engine/Source/Programs/UnrealBuildTool/Storage/IStorageProvider.cs
--- a/Engine/Source/Programs/UnrealBuildTool/Storage/IStorageProvider.cs
+++ b/Engine/Source/Programs/UnrealBuildTool/Storage/IStorageProvider.cs
@@ -21,7 +21,7 @@ namespace UnrealBuildTool.Storage
 		/// <summary>
 		/// Accessor for the read stream
 		/// </summary>
-		Stream Stream { get; }
+		Stream? Stream { get; }
 
 		/// <summary>
 		/// Whether the reader contains any data
@@ -37,7 +37,7 @@ namespace UnrealBuildTool.Storage
 		/// <summary>
 		/// Accessor for the write stream
 		/// </summary>
-		Stream Stream { get; }
+		Stream? Stream { get; }
 
 		/// <summary>
 		/// Commits the written data to the cache
diff --git a/Engine/Source/Programs/UnrealBuildTool/Storage/Impl/FileSystemStorageProvider.cs b/Engine/Source/Programs/UnrealBuildTool/Storage/Impl/FileSystemStorageProvider.cs
--- a/Engine/Source/Programs/UnrealBuildTool/Storage/Impl/FileSystemStorageProvider.cs
+++ b/Engine/Source/Programs/UnrealBuildTool/Storage/Impl/FileSystemStorageProvider.cs
@@ -24,7 +24,7 @@ namespace UnrealBuildTool.Storage.Impl
 		/// </summary>
 		class StorageReader : IStorageReader
 		{
-			public Stream Stream
+			public Stream? Stream
 			{
 				get;
 				private set;
@@ -55,10 +55,10 @@ namespace UnrealBuildTool.Storage.Impl
 		/// </summary>
 		class StorageWriter : IStorageWriter
 		{
-			FileReference TempLocation;
+			FileReference? TempLocation;
 			FileReference FinalLocation;
 
-			public Stream Stream
+			public Stream? Stream
 			{
 				get;
 				private set;
@@ -80,7 +80,7 @@ namespace UnrealBuildTool.Storage.Impl
 					throw new InvalidOperationException("Item has already been committed");
 				}
 
-				Stream.Close();
+				Stream?.Close();
 
 				try
 				{
diff --git a/Engine/Source/Programs/UnrealBuildTool/System/Asn.cs b/Engine/Source/Programs/UnrealBuildTool/System/Asn.cs
--- a/Engine/Source/Programs/UnrealBuildTool/System/Asn.cs
+++ b/Engine/Source/Programs/UnrealBuildTool/System/Asn.cs
@@ -2,6 +2,7 @@
 
 using System;
 using System.Collections.Generic;
+using System.Diagnostics.CodeAnalysis;
 using System.IO;
 using System.Linq;
 using System.Text;
@@ -118,7 +119,7 @@ namespace UnrealBuildTool
 
 		public static string GetObjectIdentifierName(int[] Values)
 		{
-			string Description;
+			string? Description;
 			if(!TryGetObjectIdentifierName(Values, out Description))
 			{
 				Description = String.Format("Unknown ({0})", String.Join(", ", Values.Select(x => x.ToString())));
@@ -126,7 +127,7 @@ namespace UnrealBuildTool
 			return Description;
 		}
 
-		public static bool TryGetObjectIdentifierName(int[] Values, out string Name)
+		public static bool TryGetObjectIdentifierName(int[] Values, [NotNullWhen(true)] out string? Name)
 		{
 			foreach(KeyValuePair<int[], string> KnownObjectIdentifier in KnownObjectIdentifiers)
 			{
diff --git a/Engine/Source/Programs/UnrealBuildTool/System/BuildException.cs b/Engine/Source/Programs/UnrealBuildTool/System/BuildException.cs
--- a/Engine/Source/Programs/UnrealBuildTool/System/BuildException.cs
+++ b/Engine/Source/Programs/UnrealBuildTool/System/BuildException.cs
@@ -35,7 +35,7 @@ namespace UnrealBuildTool
 		/// </summary>
 		/// <param name="Format">Formatting string for the error message</param>
 		/// <param name="Arguments">Arguments for the formatting string</param>
-		public BuildException(string Format, params object[] Arguments) 
+		public BuildException(string Format, params object?[] Arguments) 
 			: base(String.Format(Format, Arguments))
 		{
 		}
@@ -46,7 +46,7 @@ namespace UnrealBuildTool
 		/// <param name="InnerException">The inner exception being wrapped</param>
 		/// <param name="Format">Format for the message string</param>
 		/// <param name="Arguments">Format arguments</param>
-		public BuildException(Exception InnerException, string Format, params object[] Arguments) 
+		public BuildException(Exception InnerException, string Format, params object?[] Arguments) 
 			: base(String.Format(Format, Arguments), InnerException)
 		{
 		}
diff --git a/Engine/Source/Programs/UnrealBuildTool/System/BuildHostPlatform.cs b/Engine/Source/Programs/UnrealBuildTool/System/BuildHostPlatform.cs
--- a/Engine/Source/Programs/UnrealBuildTool/System/BuildHostPlatform.cs
+++ b/Engine/Source/Programs/UnrealBuildTool/System/BuildHostPlatform.cs
@@ -32,7 +32,7 @@ namespace UnrealBuildTool
 	/// </summary>
 	public abstract class BuildHostPlatform
 	{
-		private static BuildHostPlatform CurrentPlatform;
+		private static BuildHostPlatform? CurrentPlatform;
 		private static bool bIsMac = File.Exists("/System/Library/CoreServices/SystemVersion.plist");
 
 		/// <summary>
@@ -76,6 +76,10 @@ namespace UnrealBuildTool
 					{
 						CurrentPlatform = new LinuxBuildHostPlatform();
 					}
+					else
+					{
+						throw new NotImplementedException();
+					}
 				}
 				return CurrentPlatform;
 			}
@@ -177,7 +181,7 @@ namespace UnrealBuildTool
 		/// </summary>
 		/// <param name="Name">Name of the process to get information for.</param>
 		/// <returns></returns>
-		public virtual ProcessInfo GetProcessByName(string Name)
+		public virtual ProcessInfo? GetProcessByName(string Name)
 		{
 			ProcessInfo[] AllProcess = GetProcesses();
 			foreach (ProcessInfo Info in AllProcess)
@@ -364,7 +368,7 @@ namespace UnrealBuildTool
 		}
 		private void ProcessVMMapOutput(Process Proc, HashSet<string> Modules)
 		{
-			for (string Line = Proc.StandardOutput.ReadLine(); Line != null; Line = Proc.StandardOutput.ReadLine())
+			for (string? Line = Proc.StandardOutput.ReadLine(); Line != null; Line = Proc.StandardOutput.ReadLine())
 			{
 				Line = Line.Trim();
 				if (Line.EndsWith(".dylib"))
diff --git a/Engine/Source/Programs/UnrealBuildTool/System/BuildVersion.cs b/Engine/Source/Programs/UnrealBuildTool/System/BuildVersion.cs
--- a/Engine/Source/Programs/UnrealBuildTool/System/BuildVersion.cs
+++ b/Engine/Source/Programs/UnrealBuildTool/System/BuildVersion.cs
@@ -2,6 +2,7 @@
 
 using System;
 using System.Collections.Generic;
+using System.Diagnostics.CodeAnalysis;
 using System.IO;
 using System.Linq;
 using System.Text;
@@ -54,17 +55,17 @@ namespace UnrealBuildTool
 		/// <summary>
 		/// Name of the current branch, with '/' characters escaped as '+'
 		/// </summary>
-		public string BranchName;
+		public string? BranchName;
 
 		/// <summary>
 		/// The current build id. This will be generated automatically whenever engine binaries change if not set in the default Engine/Build/Build.version.
 		/// </summary>
-		public string BuildId;
+		public string? BuildId;
 
 		/// <summary>
 		/// The build version string
 		/// </summary>
-		public string BuildVersionString;
+		public string? BuildVersionString;
 
 		/// <summary>
 		/// Returns the value which can be used as the compatible changelist. Requires that the regular changelist is also set, and defaults to the 
@@ -81,7 +82,7 @@ namespace UnrealBuildTool
 		/// <param name="FileName">Path to the version file</param>
 		/// <param name="Version">The version information</param>
 		/// <returns>True if the version was read successfully, false otherwise</returns>
-		public static bool TryRead(FileReference FileName, out BuildVersion Version)
+		public static bool TryRead(FileReference FileName, [NotNullWhen(true)] out BuildVersion? Version)
 		{
 			JsonObject Object;
 			if (!JsonObject.TryRead(FileName, out Object))
@@ -136,7 +137,7 @@ namespace UnrealBuildTool
 		/// <param name="Object">The object to read from</param>
 		/// <param name="Version">The resulting version field</param>
 		/// <returns>True if the build version could be read, false otherwise</returns>
-		public static bool TryParse(JsonObject Object, out BuildVersion Version)
+		public static bool TryParse(JsonObject Object, [NotNullWhen(true)] out BuildVersion? Version)
 		{
 			BuildVersion NewVersion = new BuildVersion();
 			if (!Object.TryGetIntegerField("MajorVersion", out NewVersion.MajorVersion) || !Object.TryGetIntegerField("MinorVersion", out NewVersion.MinorVersion) || !Object.TryGetIntegerField("PatchVersion", out NewVersion.PatchVersion))
@@ -229,7 +230,7 @@ namespace UnrealBuildTool
 		/// <summary>
 		/// Cached copy of the current build version
 		/// </summary>
-		private static ReadOnlyBuildVersion CurrentCached;
+		private static ReadOnlyBuildVersion? CurrentCached;
 
 		/// <summary>
 		/// Constructor
@@ -255,7 +256,7 @@ namespace UnrealBuildTool
 						throw new BuildException("Version file is missing ({0})", File);
 					}
 
-					BuildVersion Version;
+					BuildVersion? Version;
 					if(!BuildVersion.TryRead(File, out Version))
 					{
 						throw new BuildException("Unable to read version file ({0}). Check that this file is present and well-formed JSON.", File);
@@ -315,12 +316,12 @@ namespace UnrealBuildTool
 			get { return Inner.IsPromotedBuild; }
 		}
 
-		public string BranchName
+		public string? BranchName
 		{
 			get { return Inner.BranchName; }
 		}
 
-		public string BuildVersionString
+		public string? BuildVersionString
 		{
 			get { return Inner.BuildVersionString; }
 		}
diff --git a/Engine/Source/Programs/UnrealBuildTool/System/CommandLine.cs b/Engine/Source/Programs/UnrealBuildTool/System/CommandLine.cs
--- a/Engine/Source/Programs/UnrealBuildTool/System/CommandLine.cs
+++ b/Engine/Source/Programs/UnrealBuildTool/System/CommandLine.cs
@@ -2,6 +2,7 @@
 
 using System;
 using System.Collections.Generic;
+using System.Diagnostics.CodeAnalysis;
 using System.Linq;
 using System.Reflection;
 using System.Text;
@@ -34,6 +35,19 @@ namespace UnrealBuildTool
 			/// The attribute containing this argument's info
 			/// </summary>
 			public CommandLineAttribute Attribute;
+
+			/// <summary>
+			/// Constructor
+			/// </summary>
+			/// <param name="Prefix"></param>
+			/// <param name="FieldInfo"></param>
+			/// <param name="Attribute"></param>
+			public Parameter(string Prefix, FieldInfo FieldInfo, CommandLineAttribute Attribute)
+			{
+				this.Prefix = Prefix;
+				this.FieldInfo = FieldInfo;
+				this.Attribute = Attribute;
+			}
 		}
 
 		/// <summary>
@@ -48,7 +62,7 @@ namespace UnrealBuildTool
 
 		private static Type NonNullableType(Type type)
 		{
-			Type UnderlyingType = Nullable.GetUnderlyingType(type);
+			Type? UnderlyingType = Nullable.GetUnderlyingType(type);
 			if (UnderlyingType != null)
 			{
 				return UnderlyingType;
@@ -65,7 +79,7 @@ namespace UnrealBuildTool
 		{
 			// Build a mapping from name to field and attribute for this object
 			Dictionary<string, Parameter> PrefixToParameter = new Dictionary<string, Parameter>(StringComparer.InvariantCultureIgnoreCase);
-			for(Type TargetType = TargetObject.GetType(); TargetType != typeof(object); TargetType = TargetType.BaseType)
+			for(Type TargetType = TargetObject.GetType(); TargetType != typeof(object); TargetType = TargetType.BaseType!)
 			{
 				foreach(FieldInfo FieldInfo in TargetType.GetFields(BindingFlags.Instance | BindingFlags.GetField | BindingFlags.GetProperty | BindingFlags.Public | BindingFlags.NonPublic | BindingFlags.DeclaredOnly))
 				{
@@ -91,7 +105,7 @@ namespace UnrealBuildTool
 								Prefix = Prefix + "=";
 							}
 						}
-						PrefixToParameter.Add(Prefix, new Parameter { Prefix = Prefix, Attribute = Attribute, FieldInfo = FieldInfo });
+						PrefixToParameter.Add(Prefix, new Parameter(Prefix, FieldInfo, Attribute));
 					}
 				}
 			}
@@ -108,7 +122,7 @@ namespace UnrealBuildTool
 					string Prefix = (EqualsIdx == -1)? Argument : Argument.Substring(0, EqualsIdx + 1);
 
 					// Check if there's a matching argument registered
-					Parameter Parameter;
+					Parameter? Parameter;
 					if(PrefixToParameter.TryGetValue(Prefix, out Parameter))
 					{
 						int NextIdx = Idx + 1;
@@ -196,7 +210,7 @@ namespace UnrealBuildTool
 		static void AssignValue(Parameter Parameter, string Text, object TargetObject, Dictionary<FieldInfo, Parameter> AssignedFieldToParameter)
 		{
 			// Check if the field type implements ICollection<>. If so, we can take multiple values.
-			Type CollectionType = null;
+			Type? CollectionType = null;
 			foreach (Type InterfaceType in Parameter.FieldInfo.FieldType.GetInterfaces())
 			{
 				if (InterfaceType.IsGenericType && InterfaceType.GetGenericTypeDefinition() == typeof(ICollection<>))
@@ -210,7 +224,7 @@ namespace UnrealBuildTool
 			if (CollectionType == null)
 			{
 				// Try to parse the value
-				object Value;
+				object? Value;
 				if(!TryParseValue(Parameter.FieldInfo.FieldType, Text, out Value))
 				{
 					Log.WriteLine(LogEventType.Warning, "Invalid value for {0}... - ignoring {1}", Parameter.Prefix, Text);
@@ -218,11 +232,11 @@ namespace UnrealBuildTool
 				}
 
 				// Check if this field has already been assigned to. Output a warning if the previous value is in conflict with the new one.
-				Parameter PreviousParameter;
+				Parameter? PreviousParameter;
 				if(AssignedFieldToParameter.TryGetValue(Parameter.FieldInfo, out PreviousParameter))
 				{
-					object PreviousValue = Parameter.FieldInfo.GetValue(TargetObject);
-					if(!PreviousValue.Equals(Value))
+					object? PreviousValue = Parameter.FieldInfo.GetValue(TargetObject);
+					if(!Object.Equals(PreviousValue, Value))
 					{
 						if(PreviousParameter.Prefix == Parameter.Prefix)
 						{
@@ -256,7 +270,7 @@ namespace UnrealBuildTool
 				// Parse each of the argument values separately
 				foreach(string Item in ItemArray)
 				{
-					object Value;
+					object? Value;
 					if(TryParseValue(CollectionType.GenericTypeArguments[0], Item, out Value))
 					{
 						CollectionType.InvokeMember("Add", BindingFlags.InvokeMethod, null, Parameter.FieldInfo.GetValue(TargetObject), new object[] { Value });
@@ -276,7 +290,7 @@ namespace UnrealBuildTool
 		/// <param name="Text">The value text</param>
 		/// <param name="Value">On success, contains the parsed object</param>
 		/// <returns>True if the text could be parsed, false otherwise</returns>
-		static bool TryParseValue(Type FieldType, string Text, out object Value)
+		static bool TryParseValue(Type FieldType, string Text, [NotNullWhen(true)] out object? Value)
 		{
 			FieldType = NonNullableType(FieldType);
 			if(FieldType.IsEnum)
diff --git a/Engine/Source/Programs/UnrealBuildTool/System/ConfigCache.cs b/Engine/Source/Programs/UnrealBuildTool/System/ConfigCache.cs
--- a/Engine/Source/Programs/UnrealBuildTool/System/ConfigCache.cs
+++ b/Engine/Source/Programs/UnrealBuildTool/System/ConfigCache.cs
@@ -2,6 +2,7 @@
 
 using System;
 using System.Collections.Generic;
+using System.Diagnostics.CodeAnalysis;
 using System.Linq;
 using System.Reflection;
 using System.Text;
@@ -21,7 +22,7 @@ namespace UnrealBuildTool
 		/// </summary>
 		/// <param name="TargetObject">The object containing the field to be modified</param>
 		/// <param name="ValueObject">The value to add</param>
-		delegate void AddElementDelegate(object TargetObject, object ValueObject); 
+		delegate void AddElementDelegate(object TargetObject, object? ValueObject); 
 
 		/// <summary>
 		/// Caches information about a field with a [ConfigFile] attribute in a type
@@ -41,12 +42,23 @@ namespace UnrealBuildTool
 			/// <summary>
 			/// For fields implementing ICollection, specifies the element type
 			/// </summary>
-			public Type ElementType;
+			public Type? ElementType;
 
 			/// <summary>
 			/// For fields implementing ICollection, a callback to add an element type.
 			/// </summary>
-			public AddElementDelegate AddElement;
+			public AddElementDelegate? AddElement;
+
+			/// <summary>
+			/// Constructor
+			/// </summary>
+			/// <param name="FieldInfo"></param>
+			/// <param name="Attribute"></param>
+			public ConfigField(FieldInfo FieldInfo, ConfigFileAttribute Attribute)
+			{
+				this.FieldInfo = FieldInfo;
+				this.Attribute = Attribute;
+			}
 		}
 
 		/// <summary>
@@ -62,7 +74,7 @@ namespace UnrealBuildTool
 			/// <summary>
 			/// The project directory to read from
 			/// </summary>
-			public DirectoryReference ProjectDir;
+			public DirectoryReference? ProjectDir;
 
 			/// <summary>
 			/// Which platform-specific files to read
@@ -75,7 +87,7 @@ namespace UnrealBuildTool
 			/// <param name="Type">The hierarchy type</param>
 			/// <param name="ProjectDir">The project directory to read from</param>
 			/// <param name="Platform">Which platform-specific files to read</param>
-			public ConfigHierarchyKey(ConfigHierarchyType Type, DirectoryReference ProjectDir, UnrealTargetPlatform Platform)
+			public ConfigHierarchyKey(ConfigHierarchyType Type, DirectoryReference? ProjectDir, UnrealTargetPlatform Platform)
 			{
 				this.Type = Type;
 				this.ProjectDir = ProjectDir;
@@ -87,9 +99,9 @@ namespace UnrealBuildTool
 			/// </summary>
 			/// <param name="Other">The object to compare against</param>
 			/// <returns>True if the objects match, false otherwise</returns>
-			public override bool Equals(object Other)
+			public override bool Equals(object? Other)
 			{
-				ConfigHierarchyKey OtherKey = Other as ConfigHierarchyKey;
+				ConfigHierarchyKey? OtherKey = Other as ConfigHierarchyKey;
 				return (OtherKey != null && OtherKey.Type == Type && OtherKey.ProjectDir == ProjectDir && OtherKey.Platform == Platform);
 			}
 
@@ -124,7 +136,7 @@ namespace UnrealBuildTool
 		/// <param name="Location">Location of the file to read</param>
 		/// <param name="ConfigFile">On success, receives the parsed config file</param>
 		/// <returns>True if the file exists and was read, false otherwise</returns>
-		internal static bool TryReadFile(FileReference Location, out ConfigFile ConfigFile)
+		internal static bool TryReadFile(FileReference Location, [NotNullWhen(true)] out ConfigFile? ConfigFile)
 		{
 			lock (LocationToConfigFile)
 			{
@@ -152,13 +164,13 @@ namespace UnrealBuildTool
 		/// <param name="ProjectDir">The project directory to read the hierarchy for</param>
 		/// <param name="Platform">Which platform to read platform-specific config files for</param>
 		/// <returns>The requested config hierarchy</returns>
-		public static ConfigHierarchy ReadHierarchy(ConfigHierarchyType Type, DirectoryReference ProjectDir, UnrealTargetPlatform Platform)
+		public static ConfigHierarchy ReadHierarchy(ConfigHierarchyType Type, DirectoryReference? ProjectDir, UnrealTargetPlatform Platform)
 		{
 			// Get the key to use for the cache. It cannot be null, so we use the engine directory if a project directory is not given.
 			ConfigHierarchyKey Key = new ConfigHierarchyKey(Type, ProjectDir, Platform);
 
 			// Try to get the cached hierarchy with this key
-			ConfigHierarchy Hierarchy;
+			ConfigHierarchy? Hierarchy;
 			lock (HierarchyKeyToHierarchy)
 			{
 				if (!HierarchyKeyToHierarchy.TryGetValue(Key, out Hierarchy))
@@ -167,7 +179,7 @@ namespace UnrealBuildTool
 					List<ConfigFile> Files = new List<ConfigFile>();
 					foreach (FileReference IniFileName in ConfigHierarchy.EnumerateConfigFileLocations(Type, ProjectDir, Platform))
 					{
-						ConfigFile File;
+						ConfigFile? File;
 						if (TryReadFile(IniFileName, out File))
 						{
 							Files.Add(File);
@@ -201,7 +213,7 @@ namespace UnrealBuildTool
 		/// <returns>List of config fields for the given type</returns>
 		static List<ConfigField> FindConfigFieldsForType(Type TargetObjectType)
 		{
-			List<ConfigField> Fields;
+			List<ConfigField>? Fields;
 			lock(TypeToConfigFields)
 			{
 				if (!TypeToConfigFields.TryGetValue(TargetObjectType, out Fields))
@@ -217,17 +229,15 @@ namespace UnrealBuildTool
 						foreach (ConfigFileAttribute Attribute in Attributes)
 						{
 							// Copy the field 
-							ConfigField Setter = new ConfigField();
-							Setter.FieldInfo = FieldInfo;
-							Setter.Attribute = Attribute;
+							ConfigField Setter = new ConfigField(FieldInfo, Attribute);
 
 							// Check if the field type implements ICollection<>. If so, we can take multiple values.
 							foreach (Type InterfaceType in FieldInfo.FieldType.GetInterfaces())
 							{
 								if (InterfaceType.IsGenericType && InterfaceType.GetGenericTypeDefinition() == typeof(ICollection<>))
 								{
-									MethodInfo MethodInfo = InterfaceType.GetRuntimeMethod("Add", new Type[] { InterfaceType.GenericTypeArguments[0] });
-									Setter.AddElement = (Target, Value) => { MethodInfo.Invoke(Setter.FieldInfo.GetValue(Target), new object[] { Value }); };
+									MethodInfo MethodInfo = InterfaceType.GetRuntimeMethod("Add", new Type[] { InterfaceType.GenericTypeArguments[0] })!;
+									Setter.AddElement = (Target, Value) => { MethodInfo.Invoke(Setter.FieldInfo.GetValue(Target), new object?[] { Value }); };
 									Setter.ElementType = InterfaceType.GenericTypeArguments[0];
 									break;
 								}
@@ -261,7 +271,7 @@ namespace UnrealBuildTool
 		/// <param name="Platform">The platform being built</param>
 		/// <param name="TargetObject">Object to receive the settings</param>
 		/// <param name="Tracker">Tracks the set of config values that were retrieved. May be null.</param>
-		internal static void ReadSettings(DirectoryReference ProjectDir, UnrealTargetPlatform Platform, object TargetObject, ConfigValueTracker Tracker)
+		internal static void ReadSettings(DirectoryReference ProjectDir, UnrealTargetPlatform Platform, object TargetObject, ConfigValueTracker? Tracker)
 		{
 			List<ConfigField> Fields = FindConfigFieldsForType(TargetObject.GetType());
 			foreach(ConfigField Field in Fields)
@@ -273,7 +283,7 @@ namespace UnrealBuildTool
 				string KeyName = Field.Attribute.KeyName ?? Field.FieldInfo.Name;
 
 				// Get the value(s) associated with this key
-				IReadOnlyList<string> Values;
+				IReadOnlyList<string>? Values;
 				Hierarchy.TryGetValues(Field.Attribute.SectionName, KeyName, out Values);
 
 				// Parse the values from the config files and update the target object
@@ -281,7 +291,7 @@ namespace UnrealBuildTool
 				{
 					if(Values != null && Values.Count == 1)
 					{
-						object Value;
+						object? Value;
 						if(TryParseValue(Values[0], Field.FieldInfo.FieldType, out Value))
 						{
 							Field.FieldInfo.SetValue(TargetObject, Value);
@@ -294,8 +304,8 @@ namespace UnrealBuildTool
 					{
 						foreach(string Item in Values)
 						{
-							object Value;
-							if(TryParseValue(Item, Field.ElementType, out Value))
+							object? Value;
+							if(TryParseValue(Item, Field.ElementType!, out Value))
 							{
 								Field.AddElement(TargetObject, Value);
 							}
@@ -318,7 +328,7 @@ namespace UnrealBuildTool
 		/// <param name="FieldType">The type of field to parse</param>
 		/// <param name="Value">If successful, a value of type 'FieldType'</param>
 		/// <returns>True if the value could be parsed, false otherwise</returns>
-		public static bool TryParseValue(string Text, Type FieldType, out object Value)
+		public static bool TryParseValue(string Text, Type FieldType, out object? Value)
 		{
 			if(FieldType == typeof(string))
 			{
diff --git a/Engine/Source/Programs/UnrealBuildTool/System/ConfigFile.cs b/Engine/Source/Programs/UnrealBuildTool/System/ConfigFile.cs
--- a/Engine/Source/Programs/UnrealBuildTool/System/ConfigFile.cs
+++ b/Engine/Source/Programs/UnrealBuildTool/System/ConfigFile.cs
@@ -2,6 +2,7 @@
 
 using System;
 using System.Collections.Generic;
+using System.Diagnostics.CodeAnalysis;
 using System.IO;
 using System.Linq;
 using System.Text;
@@ -112,7 +113,7 @@ namespace UnrealBuildTool
 		/// <param name="Name">Name of the line you want to get</param>
 		/// <param name="OutLine">The result of the operation</param>
 		/// <returns>return true if the line is retrieved return false and null OutLine if Name isn't found in this section</returns>
-		public bool TryGetLine(string Name, out ConfigLine OutLine)
+		public bool TryGetLine(string Name, out ConfigLine? OutLine)
 		{
 			foreach ( ConfigLine Line in Lines)
 			{
@@ -154,12 +155,17 @@ namespace UnrealBuildTool
 		{
 			using (StreamReader Reader = new StreamReader(Location.FullName))
 			{
-				ConfigFileSection CurrentSection = null;
-				while(!Reader.EndOfStream)
+				ConfigFileSection? CurrentSection = null;
+				for(; ;)
 				{
+					string? Line = Reader.ReadLine();
+					if (Line == null)
+					{
+						break;
+					}
+
 					// Find the first non-whitespace character
-					string Line = Reader.ReadLine();
-					for(int StartIdx = 0; StartIdx < Line.Length; StartIdx++)
+					for (int StartIdx = 0; StartIdx < Line.Length; StartIdx++)
 					{
 						if (Line[StartIdx] != ' ' && Line[StartIdx] != '\t')
 						{
@@ -169,7 +175,7 @@ namespace UnrealBuildTool
 							{
 								if(Line[EndIdx - 1] == '\\')
 								{
-									string NextLine = Reader.ReadLine();
+									string? NextLine = Reader.ReadLine();
 									if(NextLine == null)
 									{
 										break;
@@ -246,7 +252,7 @@ namespace UnrealBuildTool
 				string SectionName = Setting.Remove(Setting.IndexOf(':')).Trim(new char[] { '[', ']' });
 				if (SectionName.Length > 0)
 				{
-					ConfigFileSection CurrentSection = null;
+					ConfigFileSection? CurrentSection = null;
 					if (!Sections.TryGetValue(SectionName, out CurrentSection))
 					{
 						CurrentSection = new ConfigFileSection(SectionName);
@@ -361,7 +367,7 @@ namespace UnrealBuildTool
 		/// <returns>The config section</returns>
 		public ConfigFileSection FindOrAddSection(string SectionName)
 		{
-			ConfigFileSection Section;
+			ConfigFileSection? Section;
 			if(!Sections.TryGetValue(SectionName, out Section))
 			{
 				Section = new ConfigFileSection(SectionName);
@@ -376,7 +382,7 @@ namespace UnrealBuildTool
 		/// <param name="SectionName">Name of the section to look for</param>
 		/// <param name="RawSection">On success, the config section that was found</param>
 		/// <returns>True if the section was found, false otherwise</returns>
-		public bool TryGetSection(string SectionName, out ConfigFileSection RawSection)
+		public bool TryGetSection(string SectionName, [NotNullWhen(true)] out ConfigFileSection? RawSection)
 		{
 			return Sections.TryGetValue(SectionName, out RawSection);
 		}
diff --git a/Engine/Source/Programs/UnrealBuildTool/System/ConfigFileAttribute.cs b/Engine/Source/Programs/UnrealBuildTool/System/ConfigFileAttribute.cs
--- a/Engine/Source/Programs/UnrealBuildTool/System/ConfigFileAttribute.cs
+++ b/Engine/Source/Programs/UnrealBuildTool/System/ConfigFileAttribute.cs
@@ -27,7 +27,7 @@ namespace UnrealBuildTool
 		/// <summary>
 		/// Key name to search for
 		/// </summary>
-		public string KeyName;
+		public string? KeyName;
 
 		/// <summary>
 		/// Constructor
@@ -35,7 +35,7 @@ namespace UnrealBuildTool
 		/// <param name="ConfigType">Type of the config hierarchy to read from</param>
 		/// <param name="SectionName">Section containing the setting</param>
 		/// <param name="KeyName">Key name to search for. Optional; uses the name of the field if not set.</param>
-		public ConfigFileAttribute(ConfigHierarchyType ConfigType, string SectionName, string KeyName = null)
+		public ConfigFileAttribute(ConfigHierarchyType ConfigType, string SectionName, string? KeyName = null)
 		{
 			this.ConfigType = ConfigType;
 			this.SectionName = SectionName;
diff --git a/Engine/Source/Programs/UnrealBuildTool/System/ConfigHierarchy.cs b/Engine/Source/Programs/UnrealBuildTool/System/ConfigHierarchy.cs
--- a/Engine/Source/Programs/UnrealBuildTool/System/ConfigHierarchy.cs
+++ b/Engine/Source/Programs/UnrealBuildTool/System/ConfigHierarchy.cs
@@ -7,6 +7,7 @@ using System.Text;
 using System.Threading.Tasks;
 using Tools.DotNETCommon;
 using System.IO;
+using System.Diagnostics.CodeAnalysis;
 
 namespace UnrealBuildTool
 {
@@ -83,7 +84,7 @@ namespace UnrealBuildTool
                     }
 
                     // Find or create the values for this key
-                    List<string> Values;
+                    List<string>? Values;
 
 					if(KeyToValue.TryGetValue(Line.Key, out Values))
 					{
@@ -130,9 +131,9 @@ namespace UnrealBuildTool
 		/// <param name="KeyName">The key name to search for</param>
 		/// <param name="Value">On success, receives the corresponding value</param>
 		/// <returns>True if the key was found, false otherwise</returns>
-		public bool TryGetValue(string KeyName, out string Value)
+		public bool TryGetValue(string KeyName, [NotNullWhen(true)] out string? Value)
 		{
-			List<string> ValuesList;
+			List<string>? ValuesList;
 			if(KeyToValue.TryGetValue(KeyName, out ValuesList) && ValuesList.Count > 0)
 			{
 				Value = ValuesList[0];
@@ -151,9 +152,9 @@ namespace UnrealBuildTool
 		/// <param name="KeyName">The key name to search for</param>
 		/// <param name="Values">On success, receives a list of the corresponding values</param>
 		/// <returns>True if the key was found, false otherwise</returns>
-		public bool TryGetValues(string KeyName, out IReadOnlyList<string> Values)
+		public bool TryGetValues(string KeyName, [NotNullWhen(true)] out IReadOnlyList<string>? Values)
 		{
-			List<string> ValuesList;
+			List<string>? ValuesList;
 			if(KeyToValue.TryGetValue(KeyName, out ValuesList))
 			{
 				Values = ValuesList;
@@ -226,7 +227,7 @@ namespace UnrealBuildTool
 		/// <returns>The merged config section</returns>
 		public ConfigHierarchySection FindSection(string SectionName)
 		{
-            ConfigHierarchySection Section;
+            ConfigHierarchySection? Section;
             try
             {
                 // Acquire a read lock and do a quick check for the config section
@@ -243,7 +244,7 @@ namespace UnrealBuildTool
                             List<ConfigFileSection> RawSections = new List<ConfigFileSection>();
                             foreach (ConfigFile File in Files)
                             {
-                                ConfigFileSection RawSection;
+                                ConfigFileSection? RawSection;
                                 if (File.TryGetSection(SectionName, out RawSection))
                                 {
                                     RawSections.Add(RawSection);
@@ -287,9 +288,9 @@ namespace UnrealBuildTool
 		/// <param name="KeyName">Key name</param>
 		/// <param name="Values">Value associated with the specified key. If the key has more than one value, only the first one is returned</param>
 		/// <returns>True if the key exists</returns>
-		public bool GetArray(string SectionName, string KeyName, out List<string> Values)
+		public bool GetArray(string SectionName, string KeyName, [NotNullWhen(true)] out List<string>? Values)
 		{
-			IReadOnlyList<string> ValuesEnumerable;
+			IReadOnlyList<string>? ValuesEnumerable;
 			if(TryGetValues(SectionName, KeyName, out ValuesEnumerable))
 			{
 				Values = ValuesEnumerable.ToList();
@@ -309,7 +310,7 @@ namespace UnrealBuildTool
 		/// <param name="KeyName">Key name</param>
 		/// <param name="Value">Value associated with the specified key. If the key has more than one value, only the first one is returned</param>
 		/// <returns>True if the key exists</returns>
-		public bool GetString(string SectionName, string KeyName, out string Value)
+		public bool GetString(string SectionName, string KeyName, [NotNullWhen(true)] out string? Value)
 		{
 			if(TryGetValue(SectionName, KeyName, out Value))
 			{
@@ -341,7 +342,7 @@ namespace UnrealBuildTool
 		/// <param name="KeyName">Key name</param>
 		/// <param name="Value">Value associated with the specified key. If the key has more than one value, only the first one is returned</param>
 		/// <returns>True if the key exists</returns>
-		public bool TryGetValue(string SectionName, string KeyName, out string Value)
+		public bool TryGetValue(string SectionName, string KeyName, [NotNullWhen(true)] out string? Value)
 		{
 			return FindSection(SectionName).TryGetValue(KeyName, out Value);
 		}
@@ -355,7 +356,7 @@ namespace UnrealBuildTool
 		/// <returns>True if the key exists</returns>
 		public bool TryGetValue(string SectionName, string KeyName, out bool Value)
 		{
-			string Text;
+			string? Text;
 			if(!TryGetValue(SectionName, KeyName, out Text))
 			{
 				Value = false;
@@ -373,7 +374,7 @@ namespace UnrealBuildTool
 		/// <returns>True if the key exists</returns>
 		public bool TryGetValue(string SectionName, string KeyName, out int Value)
 		{
-			string Text;
+			string? Text;
 			if(!TryGetValue(SectionName, KeyName, out Text))
 			{
 				Value = 0;
@@ -391,7 +392,7 @@ namespace UnrealBuildTool
 		/// <returns>True if the key exists</returns>
 		public bool TryGetValue(string SectionName, string KeyName, out Guid Value)
 		{
-			string Text;
+			string? Text;
 			if(!TryGetValue(SectionName, KeyName, out Text))
 			{
 				Value = Guid.Empty;
@@ -409,7 +410,7 @@ namespace UnrealBuildTool
 		/// <returns>True if the key exists</returns>
 		public bool TryGetValue(string SectionName, string KeyName, out float Value)
 		{
-			string Text;
+			string? Text;
 			if(!TryGetValue(SectionName, KeyName, out Text))
 			{
 				Value = 0;
@@ -427,7 +428,7 @@ namespace UnrealBuildTool
 		/// <returns>True if the key exists</returns>
 		public bool TryGetValue(string SectionName, string KeyName, out double Value)
 		{
-			string Text;
+			string? Text;
 			if(!TryGetValue(SectionName, KeyName, out Text))
 			{
 				Value = 0;
@@ -443,7 +444,7 @@ namespace UnrealBuildTool
 		/// <param name="KeyName">Key name</param>
 		/// <param name="Values">Copy of the list containing all values associated with the specified key</param>
 		/// <returns>True if the key exists</returns>
-		public bool TryGetValues(string SectionName, string KeyName, out IReadOnlyList<string> Values)
+		public bool TryGetValues(string SectionName, string KeyName, [NotNullWhen(true)] out IReadOnlyList<string>? Values)
 		{
 			return FindSection(SectionName).TryGetValues(KeyName, out Values);
 		}
@@ -558,7 +559,7 @@ namespace UnrealBuildTool
 		/// <param name="Line">Line of text to parse</param>
 		/// <param name="Properties">Receives key/value pairs for the config object</param>
 		/// <returns>True if an object was parsed, false otherwise</returns>
-		public static bool TryParse(string Line, out Dictionary<string, string> Properties)
+		public static bool TryParse(string Line, [NotNullWhen(true)] out Dictionary<string, string>? Properties)
 		{
 			// Convert the string to a zero-terminated array, to make parsing easier.
 			char[] Chars = new char[Line.Length + 1];
@@ -737,10 +738,10 @@ namespace UnrealBuildTool
 		class ConfigLayerExpansion
 		{
 			// a set of replacements from the source file to possible other files
-			public string Before1 = null;
-			public string After1 = null;
-			public string Before2 = null;
-			public string After2 = null;
+			public string? Before1 = null;
+			public string? After1 = null;
+			public string? Before2 = null;
+			public string? After2 = null;
 		};
 
 		static string [] ConfigLayers =
@@ -784,11 +785,11 @@ namespace UnrealBuildTool
 		};
 
 		// Match FPlatformProcess::UserDir()
-		private static string GetUserDir()
+		private static string? GetUserDir()
 		{
 			// Some user accounts (eg. SYSTEM on Windows) don't have a home directory. Ignore them if Environment.GetFolderPath() returns an empty string.
 			string PersonalFolder = Environment.GetFolderPath(Environment.SpecialFolder.Personal);
-			string PersonalConfigFolder = null;
+			string? PersonalConfigFolder = null;
 			if (!String.IsNullOrEmpty(PersonalFolder))
 			{
 				PersonalConfigFolder = PersonalFolder;
@@ -805,14 +806,20 @@ namespace UnrealBuildTool
 		private static string PerformBasicReplacements(string InString, string BaseIniName)
 		{
 			string OutString = InString.Replace("{TYPE}", BaseIniName);
-			OutString = OutString.Replace("{USERSETTINGS}", Utils.GetUserSettingDirectory().FullName);
+
+			DirectoryReference? UserSettingsDir = Utils.GetUserSettingDirectory();
+			if (UserSettingsDir != null)
+			{
+				OutString = OutString.Replace("{USERSETTINGS}", UserSettingsDir.FullName);
+			}
+
 			OutString = OutString.Replace("{USER}", GetUserDir());
 
 			return OutString;
 		}
 
 
-		private static string PerformExpansionReplacements(ConfigLayerExpansion Expansion, string InString)
+		private static string? PerformExpansionReplacements(ConfigLayerExpansion Expansion, string InString)
 		{
 			// if there's replacement to do, the output is just the output
 			if (Expansion.Before1 == null)
@@ -835,7 +842,7 @@ namespace UnrealBuildTool
 			return OutString;
 		}
 
-		private static string PerformFinalExpansions(string InString, string PlatformName, DirectoryReference ProjectDir)
+		private static string PerformFinalExpansions(string InString, string PlatformName, DirectoryReference? ProjectDir)
 		{
 			string PlatformExtensionEngineConfigDir = DirectoryReference.Combine(UnrealBuildTool.EngineDirectory, "Platforms", PlatformName).FullName;
 
@@ -872,9 +879,9 @@ namespace UnrealBuildTool
 		/// <summary>
 		/// Returns a list of INI filenames for the given project
 		/// </summary>
-		public static IEnumerable<FileReference> EnumerateConfigFileLocations(ConfigHierarchyType Type, DirectoryReference ProjectDir, UnrealTargetPlatform Platform)
+		public static IEnumerable<FileReference> EnumerateConfigFileLocations(ConfigHierarchyType Type, DirectoryReference? ProjectDir, UnrealTargetPlatform Platform)
 		{
-			string BaseIniName = Enum.GetName(typeof(ConfigHierarchyType), Type);
+			string BaseIniName = Enum.GetName(typeof(ConfigHierarchyType), Type) ?? String.Empty;
 			string PlatformName = GetIniPlatformName(Platform);
 
 			foreach (string Layer in ConfigLayers)
@@ -899,7 +906,7 @@ namespace UnrealBuildTool
 					foreach (ConfigLayerExpansion Expansion in ConfigLayerExpansions)
 					{
 						// expansion replacements
-						string ExpandedPath = PerformExpansionReplacements(Expansion, LayerPath);
+						string? ExpandedPath = PerformExpansionReplacements(Expansion, LayerPath);
 
 						// if nothing was replaced, then skip it, as it won't change anything
 						if (ExpandedPath == null)
@@ -910,7 +917,7 @@ namespace UnrealBuildTool
 						// now go up the ini parent chain
 						if (bHasPlatformTag)
 						{
-							DataDrivenPlatformInfo.ConfigDataDrivenPlatformInfo Info = DataDrivenPlatformInfo.GetDataDrivenInfoForPlatform(PlatformName);
+							DataDrivenPlatformInfo.ConfigDataDrivenPlatformInfo? Info = DataDrivenPlatformInfo.GetDataDrivenInfoForPlatform(PlatformName);
 							if (Info != null && Info.IniParentChain != null)
 							{
 								// the IniParentChain
@@ -948,9 +955,9 @@ namespace UnrealBuildTool
 		/// <summary>
 		/// Returns a list of INI filenames for the given project
 		/// </summary>
-		public static IEnumerable<FileReference> EnumerateGeneratedConfigFileLocations(ConfigHierarchyType Type, DirectoryReference ProjectDir, UnrealTargetPlatform Platform)
+		public static IEnumerable<FileReference> EnumerateGeneratedConfigFileLocations(ConfigHierarchyType Type, DirectoryReference? ProjectDir, UnrealTargetPlatform Platform)
 		{
-			string BaseIniName = Enum.GetName(typeof(ConfigHierarchyType), Type);
+			string BaseIniName = Enum.GetName(typeof(ConfigHierarchyType), Type)!;
 			string PlatformName = GetIniPlatformName(Platform);
 
 			// Get the generated config file too. EditorSettings overrides this from 
@@ -968,7 +975,7 @@ namespace UnrealBuildTool
 		/// Determines the path to the generated config directory (same as FPaths::GeneratedConfigDir())
 		/// </summary>
 		/// <returns></returns>
-		public static DirectoryReference GetGeneratedConfigDir(DirectoryReference ProjectDir)
+		public static DirectoryReference GetGeneratedConfigDir(DirectoryReference? ProjectDir)
 		{
 			if(ProjectDir == null)
 			{
diff --git a/Engine/Source/Programs/UnrealBuildTool/System/ConfigValueTracker.cs b/Engine/Source/Programs/UnrealBuildTool/System/ConfigValueTracker.cs
--- a/Engine/Source/Programs/UnrealBuildTool/System/ConfigValueTracker.cs
+++ b/Engine/Source/Programs/UnrealBuildTool/System/ConfigValueTracker.cs
@@ -89,7 +89,7 @@ namespace UnrealBuildTool
 		/// </summary>
 		/// <param name="Other">The object to compare to</param>
 		/// <returns>True if the keys are equal, false otherwise</returns>
-		public override bool Equals(object Other)
+		public override bool Equals(object? Other)
 		{
 			return (Other is ConfigDependencyKey) && Equals((ConfigDependencyKey)Other);
 		}
@@ -99,9 +99,9 @@ namespace UnrealBuildTool
 		/// </summary>
 		/// <param name="Other">The key to compare to</param>
 		/// <returns>True if the keys are equal, false otherwise</returns>
-		public bool Equals(ConfigDependencyKey Other)
+		public bool Equals(ConfigDependencyKey? Other)
 		{
-			return Type == Other.Type && ProjectDir == Other.ProjectDir && Platform == Other.Platform && SectionName == Other.SectionName && KeyName == Other.KeyName;
+			return !ReferenceEquals(Other, null) && Type == Other.Type && ProjectDir == Other.ProjectDir && Platform == Other.Platform && SectionName == Other.SectionName && KeyName == Other.KeyName;
 		}
 
 		/// <summary>
@@ -128,14 +128,14 @@ namespace UnrealBuildTool
 		/// <summary>
 		/// The dependencies list
 		/// </summary>
-		Dictionary<ConfigDependencyKey, IReadOnlyList<string>> Dependencies;
+		Dictionary<ConfigDependencyKey, IReadOnlyList<string>?> Dependencies;
 
 		/// <summary>
 		/// Constructor
 		/// </summary>
 		public ConfigValueTracker()
 		{
-			Dependencies = new Dictionary<ConfigDependencyKey, IReadOnlyList<string>>();
+			Dependencies = new Dictionary<ConfigDependencyKey, IReadOnlyList<string>?>();
 		}
 
 		/// <summary>
@@ -144,7 +144,7 @@ namespace UnrealBuildTool
 		/// <param name="Reader">Archive to read from</param>
 		public ConfigValueTracker(BinaryArchiveReader Reader)
 		{
-			Dependencies = Reader.ReadDictionary(() => new ConfigDependencyKey(Reader), () => (IReadOnlyList<string>)Reader.ReadList(() => Reader.ReadString()));
+			Dependencies = Reader.ReadDictionary(() => new ConfigDependencyKey(Reader), () => (IReadOnlyList<string>?)Reader.ReadList(() => Reader.ReadString()));
 		}
 
 		/// <summary>
@@ -165,7 +165,7 @@ namespace UnrealBuildTool
 		/// <param name="SectionName">Name of the config file section</param>
 		/// <param name="KeyName">Name of the config file key</param>
 		/// <param name="Values">Current values for this key</param>
-		public void Add(ConfigHierarchyType Type, DirectoryReference ProjectDir, UnrealTargetPlatform Platform, string SectionName, string KeyName, IReadOnlyList<string> Values)
+		public void Add(ConfigHierarchyType Type, DirectoryReference ProjectDir, UnrealTargetPlatform Platform, string SectionName, string KeyName, IReadOnlyList<string>? Values)
 		{
 			ConfigDependencyKey Key = new ConfigDependencyKey(Type, ProjectDir, Platform, SectionName, KeyName);
 			Dependencies[Key] = Values;
@@ -177,13 +177,13 @@ namespace UnrealBuildTool
 		/// <returns></returns>
 		public bool IsValid()
 		{
-			foreach(KeyValuePair<ConfigDependencyKey, IReadOnlyList<string>> Pair in Dependencies)
+			foreach(KeyValuePair<ConfigDependencyKey, IReadOnlyList<string>?> Pair in Dependencies)
 			{
 				// Read the appropriate hierarchy
 				ConfigHierarchy Hierarchy = ConfigCache.ReadHierarchy(Pair.Key.Type, Pair.Key.ProjectDir, Pair.Key.Platform);
 
 				// Get the value(s) associated with this key
-				IReadOnlyList<string> NewValues;
+				IReadOnlyList<string>? NewValues;
 				Hierarchy.TryGetValues(Pair.Key.SectionName, Pair.Key.KeyName, out NewValues);
 
 				// Check if they're different
diff --git a/Engine/Source/Programs/UnrealBuildTool/System/CppCompileEnvironment.cs b/Engine/Source/Programs/UnrealBuildTool/System/CppCompileEnvironment.cs
--- a/Engine/Source/Programs/UnrealBuildTool/System/CppCompileEnvironment.cs
+++ b/Engine/Source/Programs/UnrealBuildTool/System/CppCompileEnvironment.cs
@@ -80,7 +80,7 @@ namespace UnrealBuildTool
 		public List<FileItem> ObjectFiles = new List<FileItem>();
 		public List<FileItem> CompiledModuleInterfaces = new List<FileItem>();
 		public List<FileItem> GeneratedHeaderFiles = new List<FileItem>();
-		public FileItem PrecompiledHeaderFile = null;
+		public FileItem? PrecompiledHeaderFile = null;
 	}
 
 	/// <summary>
@@ -116,7 +116,7 @@ namespace UnrealBuildTool
 		/// <summary>
 		/// The name of the header file which is precompiled.
 		/// </summary>
-		public FileReference PrecompiledHeaderIncludeFilename = null;
+		public FileReference? PrecompiledHeaderIncludeFilename = null;
 
 		/// <summary>
 		/// Whether the compilation should create, use, or do nothing with the precompiled header.
@@ -299,12 +299,12 @@ namespace UnrealBuildTool
         /// <summary>
         /// Platform specific directory where PGO profiling data is stored.
         /// </summary>
-        public string PGODirectory;
+        public string? PGODirectory;
 
         /// <summary>
         /// Platform specific filename where PGO profiling data is saved.
         /// </summary>
-        public string PGOFilenamePrefix;
+        public string? PGOFilenamePrefix;
 
 		/// <summary>
 		/// Whether to log detailed timing info from the compiler
@@ -370,7 +370,7 @@ namespace UnrealBuildTool
 		/// <summary>
 		/// The file containing the precompiled header data.
 		/// </summary>
-		public FileItem PrecompiledHeaderFile = null;
+		public FileItem? PrecompiledHeaderFile = null;
 
 		/// <summary>
 		/// Whether or not UHT is being built
