UBT: Convert UnrealBuildTool to use nullable references.

diff --git a/Engine/Source/Programs/UnrealBuildTool/System/CppDependencyCache.cs b/Engine/Source/Programs/UnrealBuildTool/System/CppDependencyCache.cs
--- a/Engine/Source/Programs/UnrealBuildTool/System/CppDependencyCache.cs
+++ b/Engine/Source/Programs/UnrealBuildTool/System/CppDependencyCache.cs
@@ -4,6 +4,7 @@ using System;
 using System.Collections.Concurrent;
 using System.Collections.Generic;
 using System.Diagnostics;
+using System.Diagnostics.CodeAnalysis;
 using System.IO;
 using System.Linq;
 using System.Text;
@@ -23,11 +24,11 @@ namespace UnrealBuildTool
 		class DependencyInfo
 		{
 			public long LastWriteTimeUtc;
-			public string ProducedModule;
-			public List<string> ImportedModules;
+			public string? ProducedModule;
+			public List<string>? ImportedModules;
 			public List<FileItem> Files;
 
-			public DependencyInfo(long LastWriteTimeUtc, string ProducedModule, List<string> ImportedModules, List<FileItem> Files)
+			public DependencyInfo(long LastWriteTimeUtc, string? ProducedModule, List<string>? ImportedModules, List<FileItem> Files)
 			{
 				this.LastWriteTimeUtc = LastWriteTimeUtc;
 				this.ProducedModule = ProducedModule;
@@ -71,7 +72,7 @@ namespace UnrealBuildTool
 		/// <summary>
 		/// The parent cache.
 		/// </summary>
-		CppDependencyCache Parent;
+		CppDependencyCache? Parent;
 
 		/// <summary>
 		/// Map from file item to dependency info
@@ -94,7 +95,7 @@ namespace UnrealBuildTool
 		/// <param name="Location">File to store the cache</param>
 		/// <param name="BaseDir">Base directory for files that this cache should store data for</param>
 		/// <param name="Parent">The parent cache to use</param>
-		private CppDependencyCache(FileReference Location, DirectoryReference BaseDir, CppDependencyCache Parent)
+		private CppDependencyCache(FileReference Location, DirectoryReference BaseDir, CppDependencyCache? Parent)
 		{
 			this.Location = Location;
 			this.BaseDir = BaseDir;
@@ -112,9 +113,9 @@ namespace UnrealBuildTool
 		/// <param name="InputFile">The dependencies file</param>
 		/// <param name="OutModule">The produced module name</param>
 		/// <returns>True if a produced module was found</returns>
-		public bool TryGetProducedModule(FileItem InputFile, out string OutModule)
+		public bool TryGetProducedModule(FileItem InputFile, [NotNullWhen(true)] out string? OutModule)
 		{
-			DependencyInfo Info;
+			DependencyInfo? Info;
 			if (TryGetDependencyInfo(InputFile, out Info) && Info.ProducedModule != null)
 			{
 				OutModule = Info.ProducedModule;
@@ -133,13 +134,13 @@ namespace UnrealBuildTool
 		/// <param name="InputFile">The dependency file to query</param>
 		/// <param name="OutImportedModules">List of imported modules</param>
 		/// <returns>True if a list of imported modules was obtained</returns>
-		public bool TryGetImportedModules(FileItem InputFile, out List<string> OutImportedModules)
+		public bool TryGetImportedModules(FileItem InputFile, [NotNullWhen(true)] out List<string>? OutImportedModules)
 		{
-			DependencyInfo Info;
+			DependencyInfo? Info;
 			if (TryGetDependencyInfo(InputFile, out Info))
 			{
 				OutImportedModules = Info.ImportedModules;
-				return true;
+				return OutImportedModules != null;
 			}
 			else
 			{
@@ -154,9 +155,9 @@ namespace UnrealBuildTool
 		/// <param name="InputFile">File to be read</param>
 		/// <param name="OutDependencyItems">Receives a list of output items</param>
 		/// <returns>True if the input file exists and the dependencies were read</returns>
-		public bool TryGetDependencies(FileItem InputFile, out List<FileItem> OutDependencyItems)
+		public bool TryGetDependencies(FileItem InputFile, [NotNullWhen(true)] out List<FileItem>? OutDependencyItems)
 		{
-			DependencyInfo Info;
+			DependencyInfo? Info;
 			if (TryGetDependencyInfo(InputFile, out Info))
 			{
 				OutDependencyItems = Info.Files;
@@ -175,7 +176,7 @@ namespace UnrealBuildTool
 		/// <param name="InputFile">File to be read</param>
 		/// <param name="OutInfo">The dependency info</param>
 		/// <returns>True if the input file exists and the dependencies were read</returns>
-		private bool TryGetDependencyInfo(FileItem InputFile, out DependencyInfo OutInfo)
+		private bool TryGetDependencyInfo(FileItem InputFile, [NotNullWhen(true)] out DependencyInfo? OutInfo)
 		{
 			if (!InputFile.Exists)
 			{
@@ -201,7 +202,7 @@ namespace UnrealBuildTool
 		/// <param name="InputFile">File to be read</param>
 		/// <param name="OutInfo">The dependency info</param>
 		/// <returns>True if the input file exists and the dependencies were read</returns>
-		private bool TryGetDependencyInfoInternal(FileItem InputFile, out DependencyInfo OutInfo)
+		private bool TryGetDependencyInfoInternal(FileItem InputFile, [NotNullWhen(true)] out DependencyInfo? OutInfo)
 		{
 			if (Parent != null && !InputFile.Location.IsUnderDirectory(BaseDir))
 			{
@@ -209,7 +210,7 @@ namespace UnrealBuildTool
 			}
 			else
 			{
-				DependencyInfo Info;
+				DependencyInfo? Info;
 				if (!DependencyFileToInfo.TryGetValue(InputFile, out Info) || InputFile.LastWriteTimeUtc.Ticks > Info.LastWriteTimeUtc)
 				{
 					Info = ReadDependencyInfo(InputFile);
@@ -232,9 +233,9 @@ namespace UnrealBuildTool
 		/// <param name="TargetType">The target type</param>
 		/// <param name="Architecture">The target architecture</param>
 		/// <returns>Dependency cache hierarchy for the given project</returns>
-		public static CppDependencyCache CreateHierarchy(FileReference ProjectFile, string TargetName, UnrealTargetPlatform Platform, UnrealTargetConfiguration Configuration, TargetType TargetType, string Architecture)
+		public static CppDependencyCache CreateHierarchy(FileReference? ProjectFile, string TargetName, UnrealTargetPlatform Platform, UnrealTargetConfiguration Configuration, TargetType TargetType, string Architecture)
 		{
-			CppDependencyCache Cache = null;
+			CppDependencyCache Cache = null!;
 
 			if (ProjectFile == null || !UnrealBuildTool.IsEngineInstalled())
 			{
@@ -268,11 +269,11 @@ namespace UnrealBuildTool
 		/// <param name="BaseDir">Base directory for files that this cache should store data for</param>
 		/// <param name="Parent">The parent cache to use</param>
 		/// <returns>Reference to a dependency cache with the given settings</returns>
-		static CppDependencyCache FindOrAddCache(FileReference Location, DirectoryReference BaseDir, CppDependencyCache Parent)
+		static CppDependencyCache FindOrAddCache(FileReference Location, DirectoryReference BaseDir, CppDependencyCache? Parent)
 		{
 			lock (Caches)
 			{
-				CppDependencyCache Cache;
+				CppDependencyCache? Cache;
 				if (Caches.TryGetValue(Location, out Cache))
 				{
 					Debug.Assert(Cache.BaseDir == BaseDir);
diff --git a/Engine/Source/Programs/UnrealBuildTool/System/CustomBuildSteps.cs b/Engine/Source/Programs/UnrealBuildTool/System/CustomBuildSteps.cs
--- a/Engine/Source/Programs/UnrealBuildTool/System/CustomBuildSteps.cs
+++ b/Engine/Source/Programs/UnrealBuildTool/System/CustomBuildSteps.cs
@@ -2,6 +2,7 @@
 
 using System;
 using System.Collections.Generic;
+using System.Diagnostics.CodeAnalysis;
 using System.Linq;
 using System.Text;
 using System.Threading.Tasks;
@@ -38,7 +39,7 @@ namespace UnrealBuildTool
 		/// <param name="FieldName">Name of the field to read</param>
 		/// <param name="OutBuildSteps">Output variable to store the sorted dictionary that was read</param>
 		/// <returns>True if the field was read (and OutBuildSteps is set), false otherwise.</returns>
-		public static bool TryRead(JsonObject RawObject, string FieldName, out CustomBuildSteps OutBuildSteps)
+		public static bool TryRead(JsonObject RawObject, string FieldName, [NotNullWhen(true)] out CustomBuildSteps? OutBuildSteps)
 		{
 			JsonObject BuildStepsObject;
 			if(RawObject.TryGetObjectField(FieldName, out BuildStepsObject))
@@ -80,9 +81,9 @@ namespace UnrealBuildTool
 		/// <param name="HostPlatform">The host platform to look for</param>
 		/// <param name="OutCommands">Array of commands</param>
 		/// <returns>True if a list of commands was generated</returns>
-		public bool TryGetCommands(UnrealTargetPlatform HostPlatform, out string[] OutCommands)
+		public bool TryGetCommands(UnrealTargetPlatform HostPlatform, [NotNullWhen(true)] out string[]? OutCommands)
 		{
-			string[] Commands;
+			string[]? Commands;
 			if(HostPlatformToCommands.TryGetValue(HostPlatform, out Commands) && Commands.Length > 0)
 			{
 				OutCommands = Commands;
diff --git a/Engine/Source/Programs/UnrealBuildTool/System/DataDrivenPlatformInfo.cs b/Engine/Source/Programs/UnrealBuildTool/System/DataDrivenPlatformInfo.cs
--- a/Engine/Source/Programs/UnrealBuildTool/System/DataDrivenPlatformInfo.cs
+++ b/Engine/Source/Programs/UnrealBuildTool/System/DataDrivenPlatformInfo.cs
@@ -32,15 +32,15 @@ namespace UnrealBuildTool
 			/// <summary>
 			/// Additional restricted folders for this platform.
 			/// </summary>
-			public string[] AdditionalRestrictedFolders = null;
+			public string[]? AdditionalRestrictedFolders = null;
 
 			/// <summary>
 			/// Entire ini parent chain, ending with this platform
 			/// </summary>
-			public string[] IniParentChain = null;
+			public string[]? IniParentChain = null;
 		};
 
-		static Dictionary<string, ConfigDataDrivenPlatformInfo> PlatformInfos = null;
+		static Dictionary<string, ConfigDataDrivenPlatformInfo>? PlatformInfos = null;
 
 		/// <summary>
 		/// Return all data driven infos found
@@ -67,12 +67,12 @@ namespace UnrealBuildTool
 						if (FileRef.IsUnderDirectory(DirectoryReference.Combine(UnrealBuildTool.EngineDirectory, "Config")))
 						{
 							// Foo/Engine/Config/<Platform>/DataDrivenPlatformInfo.ini
-							IniPlatformName = Path.GetFileName(Path.GetDirectoryName(FilePath));
+							IniPlatformName = Path.GetFileName(Path.GetDirectoryName(FilePath))!;
 						}
 						else
 						{
 							// Foo/Engine/Platforms/<Platform>/Config/DataDrivenPlatformInfo.ini
-							IniPlatformName = Path.GetFileName(Path.GetDirectoryName(Path.GetDirectoryName(FilePath)));
+							IniPlatformName = Path.GetFileName(Path.GetDirectoryName(Path.GetDirectoryName(FilePath)))!;
 						}
 
 						// load the DataDrivenPlatformInfo from the path (with Add support in a file that doesn't use +'s in the arrays for C++ usage)
@@ -81,20 +81,20 @@ namespace UnrealBuildTool
 
 
 						// we must have the key section 
-						ConfigFileSection Section = null;
+						ConfigFileSection? Section = null;
 						if (Config.TryGetSection("DataDrivenPlatformInfo", out Section))
 						{
 							ConfigHierarchySection ParsedSection = new ConfigHierarchySection(new List<ConfigFileSection>() { Section });
 
 							// get string values
-							string IniParent;
+							string? IniParent;
 							if (ParsedSection.TryGetValue("IniParent", out IniParent))
 							{
 								IniParents[IniPlatformName] = IniParent;
 							}
 
 							// slightly nasty bool parsing for bool values
-							string Temp;
+							string? Temp;
 							if (ParsedSection.TryGetValue("bIsConfidential", out Temp))
 							{
 								ConfigHierarchy.TryParse(Temp, out NewInfo.bIsConfidential);
@@ -108,7 +108,7 @@ namespace UnrealBuildTool
 							}
 
 							// get a list of additional restricted folders
-							IReadOnlyList<string> AdditionalRestrictedFolders;
+							IReadOnlyList<string>? AdditionalRestrictedFolders;
 							if(ParsedSection.TryGetValues("AdditionalRestrictedFolders", out AdditionalRestrictedFolders) && AdditionalRestrictedFolders.Count > 0)
 							{
 								NewInfo.AdditionalRestrictedFolders = AdditionalRestrictedFolders.Select(x => x.Trim()).Where(x => x.Length > 0).ToArray();
@@ -123,7 +123,7 @@ namespace UnrealBuildTool
 				// now that all are read in, calculate the ini parent chain, starting with parent-most
 				foreach (KeyValuePair<string, ConfigDataDrivenPlatformInfo> Pair in PlatformInfos)
 				{
-					string CurrentPlatform;
+					string? CurrentPlatform;
 
 					// walk up the chain and build up the ini chain
 					List<string> Chain = new List<string>();
@@ -156,11 +156,11 @@ namespace UnrealBuildTool
 		/// </summary>
 		/// <param name="PlatformName"></param>
 		/// <returns></returns>
-		public static ConfigDataDrivenPlatformInfo GetDataDrivenInfoForPlatform(string PlatformName)
+		public static ConfigDataDrivenPlatformInfo? GetDataDrivenInfoForPlatform(string PlatformName)
 		{
 
 			// lookup the platform name (which is not guaranteed to be there)
-			ConfigDataDrivenPlatformInfo Info;
+			ConfigDataDrivenPlatformInfo? Info;
 			GetAllPlatformInfos().TryGetValue(PlatformName, out Info);
 
 			// return what we found of null if nothing
diff --git a/Engine/Source/Programs/UnrealBuildTool/System/DirectoryItem.cs b/Engine/Source/Programs/UnrealBuildTool/System/DirectoryItem.cs
--- a/Engine/Source/Programs/UnrealBuildTool/System/DirectoryItem.cs
+++ b/Engine/Source/Programs/UnrealBuildTool/System/DirectoryItem.cs
@@ -3,6 +3,7 @@
 using System;
 using System.Collections.Concurrent;
 using System.Collections.Generic;
+using System.Diagnostics.CodeAnalysis;
 using System.IO;
 using System.Linq;
 using System.Text;
@@ -29,12 +30,12 @@ namespace UnrealBuildTool
 		/// <summary>
 		/// Cached map of name to subdirectory item
 		/// </summary>
-		Dictionary<string, DirectoryItem> Directories;
+		Dictionary<string, DirectoryItem>? Directories;
 
 		/// <summary>
 		/// Cached map of name to file
 		/// </summary>
-		Dictionary<string, FileItem> Files;
+		Dictionary<string, FileItem>? Files;
 
 		/// <summary>
 		/// Global map of location to item
@@ -87,7 +88,7 @@ namespace UnrealBuildTool
 		/// <summary>
 		/// Gets the parent directory item
 		/// </summary>
-		public DirectoryItem GetParentDirectoryItem()
+		public DirectoryItem? GetParentDirectoryItem()
 		{
 			if(Info.Parent == null)
 			{
@@ -127,7 +128,7 @@ namespace UnrealBuildTool
 		/// <returns>The directory item for this location</returns>
 		public static DirectoryItem GetItemByDirectoryReference(DirectoryReference Location)
 		{
-			DirectoryItem Result;
+			DirectoryItem? Result;
 			if(!LocationToItem.TryGetValue(Location, out Result))
 			{
 				DirectoryItem NewItem = new DirectoryItem(Location, new DirectoryInfo(Location.FullName));
@@ -152,7 +153,7 @@ namespace UnrealBuildTool
 		{
 			DirectoryReference Location = new DirectoryReference(Info);
 
-			DirectoryItem Result;
+			DirectoryItem? Result;
 			if(!LocationToItem.TryGetValue(Location, out Result))
 			{
 				DirectoryItem NewItem = new DirectoryItem(Location, Info);
@@ -175,7 +176,7 @@ namespace UnrealBuildTool
 		{
 			Info = new DirectoryInfo(Info.FullName);
 
-			Dictionary<string, DirectoryItem> PrevDirectories = Directories;
+			Dictionary<string, DirectoryItem>? PrevDirectories = Directories;
 			if(PrevDirectories != null)
 			{
 				foreach(DirectoryItem SubDirectory in PrevDirectories.Values)
@@ -185,7 +186,7 @@ namespace UnrealBuildTool
 				Directories = null;
 			}
 
-			Dictionary<string, FileItem> PrevFiles = Files;
+			Dictionary<string, FileItem>? PrevFiles = Files;
 			if(PrevFiles != null)
 			{
 				foreach(FileItem File in PrevFiles.Values)
@@ -247,7 +248,7 @@ namespace UnrealBuildTool
 		public IEnumerable<DirectoryItem> EnumerateDirectories()
 		{
 			CacheDirectories();
-			return Directories.Values;
+			return Directories!.Values;
 		}
 
 		/// <summary>
@@ -256,7 +257,7 @@ namespace UnrealBuildTool
 		/// <param name="Name">Name of the directory</param>
 		/// <param name="OutDirectory">If successful receives the matching directory item with this name</param>
 		/// <returns>True if the file exists, false otherwise</returns>
-		public bool TryGetDirectory(string Name, out DirectoryItem OutDirectory)
+		public bool TryGetDirectory(string Name, [NotNullWhen(true)] out DirectoryItem? OutDirectory)
 		{
 			if(Name.Length > 0 && Name[0] == '.')
 			{
@@ -273,7 +274,7 @@ namespace UnrealBuildTool
 			}
 
 			CacheDirectories();
-			return Directories.TryGetValue(Name, out OutDirectory);
+			return Directories!.TryGetValue(Name, out OutDirectory);
 		}
 
 		/// <summary>
@@ -304,7 +305,7 @@ namespace UnrealBuildTool
 		public IEnumerable<FileItem> EnumerateFiles()
 		{
 			CacheFiles();
-			return Files.Values;
+			return Files!.Values;
 		}
 
 		/// <summary>
@@ -314,10 +315,10 @@ namespace UnrealBuildTool
 		/// <param name="Name">Name of the file</param>
 		/// <param name="OutFile">If successful receives the matching file item with this name</param>
 		/// <returns>True if the file exists, false otherwise</returns>
-		public bool TryGetFile(string Name, out FileItem OutFile)
+		public bool TryGetFile(string Name, [NotNullWhen(true)] out FileItem? OutFile)
 		{
 			CacheFiles();
-			return Files.TryGetValue(Name, out OutFile);
+			return Files!.TryGetValue(Name, out OutFile);
 		}
 
 		/// <summary>
diff --git a/Engine/Source/Programs/UnrealBuildTool/System/DynamicCompilation.cs b/Engine/Source/Programs/UnrealBuildTool/System/DynamicCompilation.cs
--- a/Engine/Source/Programs/UnrealBuildTool/System/DynamicCompilation.cs
+++ b/Engine/Source/Programs/UnrealBuildTool/System/DynamicCompilation.cs
@@ -157,7 +157,7 @@ namespace UnrealBuildTool
 			}
 		}
 
-		private static Assembly CompileAssembly(FileReference OutputAssemblyPath, HashSet<FileReference> SourceFileNames, List<string> ReferencedAssembies, List<string> PreprocessorDefines = null, bool TreatWarningsAsErrors = false)
+		private static Assembly? CompileAssembly(FileReference OutputAssemblyPath, HashSet<FileReference> SourceFileNames, List<string>? ReferencedAssembies, List<string>? PreprocessorDefines = null, bool TreatWarningsAsErrors = false)
 		{
 			CSharpParseOptions ParseOptions = new CSharpParseOptions(
 				languageVersion:LanguageVersion.Latest, 
@@ -283,7 +283,7 @@ namespace UnrealBuildTool
 		/// <param name="DoNotCompile"></param>
 		/// <param name="TreatWarningsAsErrors"></param>
 		/// <returns>The assembly that was loaded</returns>
-		public static Assembly CompileAndLoadAssembly(FileReference OutputAssemblyPath, HashSet<FileReference> SourceFileNames, List<string> ReferencedAssembies = null, List<string> PreprocessorDefines = null, bool DoNotCompile = false, bool TreatWarningsAsErrors = false)
+		public static Assembly? CompileAndLoadAssembly(FileReference OutputAssemblyPath, HashSet<FileReference> SourceFileNames, List<string>? ReferencedAssembies = null, List<string>? PreprocessorDefines = null, bool DoNotCompile = false, bool TreatWarningsAsErrors = false)
 		{
 			// Check to see if the resulting assembly is compiled and up to date
 			FileReference AssemblyManifestFilePath = FileReference.Combine(OutputAssemblyPath.Directory, Path.GetFileNameWithoutExtension(OutputAssemblyPath.FullName) + "Manifest.json");
@@ -295,7 +295,7 @@ namespace UnrealBuildTool
 			}
 
 			// Load the assembly to ensure it is correct
-			Assembly CompiledAssembly = null;
+			Assembly? CompiledAssembly = null;
 			if (!bNeedsCompilation)
 			{
 				try
diff --git a/Engine/Source/Programs/UnrealBuildTool/System/ExternalExecution.cs b/Engine/Source/Programs/UnrealBuildTool/System/ExternalExecution.cs
--- a/Engine/Source/Programs/UnrealBuildTool/System/ExternalExecution.cs
+++ b/Engine/Source/Programs/UnrealBuildTool/System/ExternalExecution.cs
@@ -231,7 +231,7 @@ namespace UnrealBuildTool
 		/// <summary>
 		/// Base (i.e. extensionless) path+filename of the .gen files
 		/// </summary>
-		public string GeneratedCPPFilenameBase;
+		public string? GeneratedCPPFilenameBase;
 
 		/// <summary>
 		/// Version of code generated by UHT
@@ -336,7 +336,7 @@ namespace UnrealBuildTool
 			public List<string> ClassesHeaders;
 			public List<string> PublicHeaders;
 			public List<string> PrivateHeaders;
-			public string GeneratedCPPFilenameBase;
+			public string? GeneratedCPPFilenameBase;
 			public bool SaveExportedHeaders;
 			public EGeneratedCodeVersion UHTGeneratedCodeVersion;
 
@@ -347,7 +347,7 @@ namespace UnrealBuildTool
 				OverrideModuleType = Info.OverrideModuleType;
 				BaseDirectory = Info.ModuleDirectories[0].FullName;
 				IncludeBase = Info.ModuleDirectories[0].ParentDirectory.FullName;
-				OutputDirectory = Path.GetDirectoryName(Info.GeneratedCPPFilenameBase);
+				OutputDirectory = Path.GetDirectoryName(Info.GeneratedCPPFilenameBase)!;
 				ClassesHeaders = Info.PublicUObjectClassesHeaders.Select((Header) => Header.AbsolutePath).ToList();
 				PublicHeaders = Info.PublicUObjectHeaders.Select((Header) => Header.AbsolutePath).ToList();
 				PrivateHeaders = Info.PrivateUObjectHeaders.Select((Header) => Header.AbsolutePath).ToList();
@@ -477,7 +477,7 @@ namespace UnrealBuildTool
 		/// <returns>True if ToModule is a dependency of FromModule, false otherwise</returns>
 		static bool IsDependency(UEBuildModuleCPP FromModule, UEBuildModuleCPP ToModule, Dictionary<UEBuildModule, HashSet<UEBuildModule>> Cache)
 		{
-			HashSet<UEBuildModule> Dependencies;
+			HashSet<UEBuildModule>? Dependencies;
 			if(!Cache.TryGetValue(FromModule, out Dependencies))
 			{
 				Dependencies = new HashSet<UEBuildModule>();
@@ -493,7 +493,7 @@ namespace UnrealBuildTool
 		/// <param name="RulesObject">The rules object</param>
 		/// <param name="ProjectDescriptor">Descriptor for the project being built</param>
 		/// <returns>The module type</returns>
-		static UHTModuleType GetModuleType(ModuleRules RulesObject, ProjectDescriptor ProjectDescriptor)
+		static UHTModuleType GetModuleType(ModuleRules RulesObject, ProjectDescriptor? ProjectDescriptor)
 		{
 			ModuleRulesContext Context = RulesObject.Context;
 			if (Context.bClassifyAsGameModuleForUHT)
@@ -570,7 +570,7 @@ namespace UnrealBuildTool
 			}
 		}
 
-		public static void SetupUObjectModules(IEnumerable<UEBuildModuleCPP> ModulesToGenerateHeadersFor, UnrealTargetPlatform Platform, ProjectDescriptor ProjectDescriptor, List<UHTModuleInfo> UObjectModules, List<UHTModuleHeaderInfo> UObjectModuleHeaders, EGeneratedCodeVersion GeneratedCodeVersion, bool bIsAssemblingBuild, SourceFileMetadataCache MetadataCache)
+		public static void SetupUObjectModules(IEnumerable<UEBuildModuleCPP> ModulesToGenerateHeadersFor, UnrealTargetPlatform Platform, ProjectDescriptor? ProjectDescriptor, List<UHTModuleInfo> UObjectModules, List<UHTModuleHeaderInfo> UObjectModuleHeaders, EGeneratedCodeVersion GeneratedCodeVersion, bool bIsAssemblingBuild, SourceFileMetadataCache MetadataCache)
 		{
 			// Find the type of each module
 			Dictionary<UEBuildModuleCPP, UHTModuleType> ModuleToType = new Dictionary<UEBuildModuleCPP, UHTModuleType>();
@@ -592,7 +592,7 @@ namespace UnrealBuildTool
 				{
 					UEBuildModuleCPP Module = ModulesSortedByType[Idx];
 
-					UHTModuleInfo Info = new UHTModuleInfo(Module.Name, Module.RulesFile, Module.ModuleDirectories, ModuleToType[Module], DirectoryItem.GetItemByDirectoryReference(Module.GeneratedCodeDirectory), GeneratedCodeVersion, Module.Rules.bUsePrecompiled, Module.Rules.OverridePackageType);
+					UHTModuleInfo Info = new UHTModuleInfo(Module.Name, Module.RulesFile, Module.ModuleDirectories, ModuleToType[Module], DirectoryItem.GetItemByDirectoryReference(Module.GeneratedCodeDirectory!), GeneratedCodeVersion, Module.Rules.bUsePrecompiled, Module.Rules.OverridePackageType);
 					ModuleInfoArray[Idx] = Info;
 
 					Queue.Enqueue(() => SetupUObjectModule(Info, ExcludedFolders, MetadataCache, Queue));
@@ -611,7 +611,7 @@ namespace UnrealBuildTool
 
 					// If we've got this far and there are no source files then it's likely we're installed and ignoring
 					// engine files, so we don't need a .gen.cpp either
-					Info.GeneratedCPPFilenameBase = Path.Combine(Module.GeneratedCodeDirectory.FullName, Info.ModuleName) + ".gen";
+					Info.GeneratedCPPFilenameBase = Path.Combine(Module.GeneratedCodeDirectory!.FullName, Info.ModuleName) + ".gen";
 					if (!Module.Rules.bUsePrecompiled)
 					{
 						Module.GeneratedCppDirectories = new List<string>();
@@ -703,7 +703,7 @@ namespace UnrealBuildTool
 		/// Gets the path to the receipt for UHT
 		/// </summary>
 		/// <returns>Path to the UHT receipt</returns>
-		public static FileReference GetHeaderToolReceiptFile(FileReference ProjectFile, UnrealTargetConfiguration Configuration, bool bHasProjectScriptPlugin)
+		public static FileReference GetHeaderToolReceiptFile(FileReference? ProjectFile, UnrealTargetConfiguration Configuration, bool bHasProjectScriptPlugin)
 		{
 			if(bHasProjectScriptPlugin && ProjectFile != null)
 			{
@@ -721,6 +721,10 @@ namespace UnrealBuildTool
 		static FileReference GetHeaderToolPath(FileReference ReceiptFile)
 		{
 			TargetReceipt Receipt = TargetReceipt.Read(ReceiptFile);
+			if (Receipt.Launch == null)
+			{
+				throw new BuildException("'Launch' property not set in UHT receipt.");
+			}
 			return Receipt.Launch;
 		}
 
@@ -743,7 +747,7 @@ namespace UnrealBuildTool
 				// Don't check timestamps for individual binaries if we're using the installed version of UHT. It will always be up to date.
 				if (!UnrealBuildTool.IsFileInstalled(ReceiptFile.Location))
 				{
-					TargetReceipt Receipt;
+					TargetReceipt? Receipt;
 					if (!TargetReceipt.TryRead(ReceiptPath, out Receipt))
 					{
 						Timestamp = DateTime.MaxValue;
@@ -811,7 +815,7 @@ namespace UnrealBuildTool
 				{
 					if (Module.ModuleName.Equals("CoreUObject", StringComparison.InvariantCultureIgnoreCase))
 					{
-						CoreGeneratedTimestampUtc = GetCoreGeneratedTimestampUtc(Module.ModuleName, Path.GetDirectoryName(Module.GeneratedCPPFilenameBase));
+						CoreGeneratedTimestampUtc = GetCoreGeneratedTimestampUtc(Module.ModuleName, Path.GetDirectoryName(Module.GeneratedCPPFilenameBase)!);
 						break;
 					}
 				}
@@ -831,7 +835,7 @@ namespace UnrealBuildTool
 				}
 
 				// Make sure we have an existing folder for generated code.  If not, then we definitely need to generate code!
-				string GeneratedCodeDirectory = Path.GetDirectoryName(Module.GeneratedCPPFilenameBase);
+				string GeneratedCodeDirectory = Path.GetDirectoryName(Module.GeneratedCPPFilenameBase)!;
 				FileSystemInfo TestDirectory = (FileSystemInfo)new DirectoryInfo(GeneratedCodeDirectory);
 				if (!TestDirectory.Exists)
 				{
@@ -980,7 +984,7 @@ namespace UnrealBuildTool
 			{
 				if(!Module.bIsReadOnly)
 				{
-					string GeneratedCodeDirectory = Path.GetDirectoryName(Module.GeneratedCPPFilenameBase);
+					string GeneratedCodeDirectory = Path.GetDirectoryName(Module.GeneratedCPPFilenameBase)!;
 					DirectoryInfo GeneratedCodeDirectoryInfo = new DirectoryInfo(GeneratedCodeDirectory);
 
 					try
@@ -1054,7 +1058,7 @@ namespace UnrealBuildTool
 		/// Builds and runs the header tool and touches the header directories.
 		/// Performs any early outs if headers need no changes, given the UObject modules, tool path, game name, and configuration
 		/// </summary>
-		public static void ExecuteHeaderToolIfNecessary(BuildConfiguration BuildConfiguration, FileReference ProjectFile, string TargetName, TargetType TargetType, bool bHasProjectScriptPlugin, List<UHTModuleInfo> UObjectModules, FileReference ModuleInfoFileName, bool bIsGatheringBuild, bool bIsAssemblingBuild, ISourceFileWorkingSet WorkingSet)
+		public static void ExecuteHeaderToolIfNecessary(BuildConfiguration BuildConfiguration, FileReference? ProjectFile, string TargetName, TargetType TargetType, bool bHasProjectScriptPlugin, List<UHTModuleInfo> UObjectModules, FileReference ModuleInfoFileName, bool bIsGatheringBuild, bool bIsAssemblingBuild, ISourceFileWorkingSet WorkingSet)
 		{
 			if (ProgressWriter.bWriteMarkup)
 			{
@@ -1153,7 +1157,7 @@ namespace UnrealBuildTool
 						string Architecture = UEBuildPlatform.GetBuildPlatform(Platform).GetDefaultArchitecture(null);
 
 						// Add UHT plugins to UBT command line as external plugins
-						FileReference ScriptProjectFile = null;
+						FileReference? ScriptProjectFile = null;
 						if(bHasProjectScriptPlugin && ProjectFile != null)
 						{
 							ScriptProjectFile = ProjectFile;
diff --git a/Engine/Source/Programs/UnrealBuildTool/System/FileItem.cs b/Engine/Source/Programs/UnrealBuildTool/System/FileItem.cs
--- a/Engine/Source/Programs/UnrealBuildTool/System/FileItem.cs
+++ b/Engine/Source/Programs/UnrealBuildTool/System/FileItem.cs
@@ -21,7 +21,7 @@ namespace UnrealBuildTool
 		/// <summary>
 		/// The directory containing this file
 		/// </summary>
-		DirectoryItem CachedDirectory;
+		DirectoryItem? CachedDirectory;
 
 		/// <summary>
 		/// Location of this file
@@ -168,7 +168,7 @@ namespace UnrealBuildTool
 		{
 			FileReference Location = new FileReference(Info);
 
-			FileItem Result;
+			FileItem? Result;
 			if (!UniqueSourceFileMap.TryGetValue(Location, out Result))
 			{
 				FileItem NewFileItem = new FileItem(Location, Info);
@@ -191,7 +191,7 @@ namespace UnrealBuildTool
 		/// <returns>The FileItem that represents the given a full file path.</returns>
 		public static FileItem GetItemByFileReference(FileReference Location)
 		{
-			FileItem Result;
+			FileItem? Result;
 			if (!UniqueSourceFileMap.TryGetValue(Location, out Result))
 			{
 				FileItem NewFileItem = new FileItem(Location, Location.ToFileInfo());
diff --git a/Engine/Source/Programs/UnrealBuildTool/System/InstalledPlatformInfo.cs b/Engine/Source/Programs/UnrealBuildTool/System/InstalledPlatformInfo.cs
--- a/Engine/Source/Programs/UnrealBuildTool/System/InstalledPlatformInfo.cs
+++ b/Engine/Source/Programs/UnrealBuildTool/System/InstalledPlatformInfo.cs
@@ -120,12 +120,12 @@ namespace UnrealBuildTool
 			}
 		}
 
-		private static List<InstalledPlatformConfiguration> InstalledPlatformConfigurations;
+		private static List<InstalledPlatformConfiguration>? InstalledPlatformConfigurations;
 
 		static InstalledPlatformInfo()
 		{
-			List<string> InstalledPlatforms;
-			ConfigHierarchy Ini = ConfigCache.ReadHierarchy(ConfigHierarchyType.Engine, (DirectoryReference)null, BuildHostPlatform.Current.Platform);
+			List<string>? InstalledPlatforms;
+			ConfigHierarchy Ini = ConfigCache.ReadHierarchy(ConfigHierarchyType.Engine, (DirectoryReference?)null, BuildHostPlatform.Current.Platform);
 
 			bool bHasInstalledPlatformInfo;
 			if(Ini.TryGetValue("InstalledPlatforms", "HasInstalledPlatformInfo", out bHasInstalledPlatformInfo) && bHasInstalledPlatformInfo)
@@ -227,7 +227,7 @@ namespace UnrealBuildTool
 
 			if (bCanCreateEntry)
 			{
-				InstalledPlatformConfigurations.Add(new InstalledPlatformConfiguration(Configuration, UnrealTargetPlatform.Parse(PlatformName), PlatformType, Architecture, RequiredFile, ProjectType, bCanBeDisplayed));
+				InstalledPlatformConfigurations!.Add(new InstalledPlatformConfiguration(Configuration, UnrealTargetPlatform.Parse(PlatformName), PlatformType, Architecture, RequiredFile, ProjectType, bCanBeDisplayed));
 			}
 		}
 
diff --git a/Engine/Source/Programs/UnrealBuildTool/System/JunkDeleter.cs b/Engine/Source/Programs/UnrealBuildTool/System/JunkDeleter.cs
--- a/Engine/Source/Programs/UnrealBuildTool/System/JunkDeleter.cs
+++ b/Engine/Source/Programs/UnrealBuildTool/System/JunkDeleter.cs
@@ -34,7 +34,7 @@ namespace UnrealBuildTool
 				using (StreamReader reader = new StreamReader(ManifestPath))
 				{
 					string CurrentToRootDir = ".." + Path.DirectorySeparatorChar + "..";
-					string LineRead;
+					string? LineRead;
 					while ((LineRead = reader.ReadLine()) != null)
 					{
 						string JunkEntry = LineRead.Trim();
@@ -104,7 +104,7 @@ namespace UnrealBuildTool
 					if (FileName.Contains('*'))
 					{
 						// Wildcard search and delete
-						string DirectoryToLookIn = Path.GetDirectoryName(Junk);
+						string DirectoryToLookIn = Path.GetDirectoryName(Junk)!;
 						if (Directory.Exists(DirectoryToLookIn))
 						{
 							// Delete all files within the specified folder
diff --git a/Engine/Source/Programs/UnrealBuildTool/System/LinkEnvironment.cs b/Engine/Source/Programs/UnrealBuildTool/System/LinkEnvironment.cs
--- a/Engine/Source/Programs/UnrealBuildTool/System/LinkEnvironment.cs
+++ b/Engine/Source/Programs/UnrealBuildTool/System/LinkEnvironment.cs
@@ -32,22 +32,22 @@ namespace UnrealBuildTool
 		/// <summary>
 		/// On Mac, indicates the path to the target's application bundle
 		/// </summary>
-		public DirectoryReference BundleDirectory;
+		public DirectoryReference? BundleDirectory;
 
 		/// <summary>
 		/// The directory to put the non-executable files in (PDBs, import library, etc)
 		/// </summary>
-		public DirectoryReference OutputDirectory;
+		public DirectoryReference? OutputDirectory;
 
 		/// <summary>
 		/// Intermediate file directory
 		/// </summary>
-		public DirectoryReference IntermediateDirectory;
+		public DirectoryReference? IntermediateDirectory;
 
 		/// <summary>
 		/// The directory to shadow source files in for syncing to remote compile servers
 		/// </summary>
-		public DirectoryReference LocalShadowDirectory = null;
+		public DirectoryReference? LocalShadowDirectory = null;
 
 		/// <summary>
 		/// The file path for the executable file that is output by the linker.
@@ -229,12 +229,12 @@ namespace UnrealBuildTool
         /// <summary>
         /// Platform specific directory where PGO profiling data is stored.
         /// </summary>
-        public string PGODirectory;
+        public string? PGODirectory;
 
         /// <summary>
         /// Platform specific filename where PGO profiling data is saved.
         /// </summary>
-        public string PGOFilenamePrefix;
+        public string? PGOFilenamePrefix;
 
         /// <summary>
         /// Whether to request the linker create a map file as part of the build
@@ -264,12 +264,12 @@ namespace UnrealBuildTool
 		/// <summary>
 		/// Bundle version for Mac apps
 		/// </summary>
-		public string BundleVersion;
+		public string? BundleVersion;
 
 		/// <summary>
 		/// When building a dynamic library on Apple platforms, specifies the installed name for other binaries that link against it.
 		/// </summary>
-		public string InstallName;
+		public string? InstallName;
 
 		/// <summary>
 		/// A list of the object files to be linked.
@@ -295,7 +295,7 @@ namespace UnrealBuildTool
 		/// Provides a Module Definition File (.def) to the linker to describe various attributes of a DLL.
 		/// Necessary when exporting functions by ordinal values instead of by name.
 		/// </summary>
-		public string ModuleDefinitionFile;
+		public string? ModuleDefinitionFile;
 
 		/// <summary>
 		/// All the additional properties from the modules linked into this binary
diff --git a/Engine/Source/Programs/UnrealBuildTool/System/LocalizationDescriptor.cs b/Engine/Source/Programs/UnrealBuildTool/System/LocalizationDescriptor.cs
--- a/Engine/Source/Programs/UnrealBuildTool/System/LocalizationDescriptor.cs
+++ b/Engine/Source/Programs/UnrealBuildTool/System/LocalizationDescriptor.cs
@@ -100,7 +100,7 @@ namespace UnrealBuildTool
 		/// <param name="Writer">The Json writer to output to</param>
 		/// <param name="Name">Name of the array</param>
 		/// <param name="Targets">Array of targets</param>
-		public static void WriteArray(JsonWriter Writer, string Name, LocalizationTargetDescriptor[] Targets)
+		public static void WriteArray(JsonWriter Writer, string Name, LocalizationTargetDescriptor[]? Targets)
 		{
 			if (Targets != null && Targets.Length > 0)
 			{
diff --git a/Engine/Source/Programs/UnrealBuildTool/System/ModuleDescriptor.cs b/Engine/Source/Programs/UnrealBuildTool/System/ModuleDescriptor.cs
--- a/Engine/Source/Programs/UnrealBuildTool/System/ModuleDescriptor.cs
+++ b/Engine/Source/Programs/UnrealBuildTool/System/ModuleDescriptor.cs
@@ -170,47 +170,47 @@ namespace UnrealBuildTool
 		/// <summary>
 		/// List of allowed platforms
 		/// </summary>
-		public List<UnrealTargetPlatform> WhitelistPlatforms;
+		public List<UnrealTargetPlatform>? WhitelistPlatforms;
 
 		/// <summary>
 		/// List of disallowed platforms
 		/// </summary>
-		public List<UnrealTargetPlatform> BlacklistPlatforms;
+		public List<UnrealTargetPlatform>? BlacklistPlatforms;
 
 		/// <summary>
 		/// List of allowed targets
 		/// </summary>
-		public TargetType[] WhitelistTargets;
+		public TargetType[]? WhitelistTargets;
 
 		/// <summary>
 		/// List of disallowed targets
 		/// </summary>
-		public TargetType[] BlacklistTargets;
+		public TargetType[]? BlacklistTargets;
 
 		/// <summary>
 		/// List of allowed target configurations
 		/// </summary>
-		public UnrealTargetConfiguration[] WhitelistTargetConfigurations;
+		public UnrealTargetConfiguration[]? WhitelistTargetConfigurations;
 
 		/// <summary>
 		/// List of disallowed target configurations
 		/// </summary>
-		public UnrealTargetConfiguration[] BlacklistTargetConfigurations;
+		public UnrealTargetConfiguration[]? BlacklistTargetConfigurations;
 
 		/// <summary>
 		/// List of allowed programs
 		/// </summary>
-		public string[] WhitelistPrograms;
+		public string[]? WhitelistPrograms;
 
 		/// <summary>
 		/// List of disallowed programs
 		/// </summary>
-		public string[] BlacklistPrograms;
+		public string[]? BlacklistPrograms;
 
 		/// <summary>
 		/// List of additional dependencies for building this module.
 		/// </summary>
-		public string[] AdditionalDependencies;
+		public string[]? AdditionalDependencies;
 
 		/// <summary>
 		/// Constructor
@@ -423,7 +423,7 @@ namespace UnrealBuildTool
 		/// <param name="Writer">The Json writer to output to</param>
 		/// <param name="Name">Name of the array</param>
 		/// <param name="Modules">Array of modules</param>
-		public static void WriteArray(JsonWriter Writer, string Name, ModuleDescriptor[] Modules)
+		public static void WriteArray(JsonWriter Writer, string Name, ModuleDescriptor[]? Modules)
 		{
 			if (Modules != null && Modules.Length > 0)
 			{
diff --git a/Engine/Source/Programs/UnrealBuildTool/System/ModuleManifest.cs b/Engine/Source/Programs/UnrealBuildTool/System/ModuleManifest.cs
--- a/Engine/Source/Programs/UnrealBuildTool/System/ModuleManifest.cs
+++ b/Engine/Source/Programs/UnrealBuildTool/System/ModuleManifest.cs
@@ -1,7 +1,8 @@
-﻿// Copyright Epic Games, Inc. All Rights Reserved.
+// Copyright Epic Games, Inc. All Rights Reserved.
 
 using System;
 using System.Collections.Generic;
+using System.Diagnostics.CodeAnalysis;
 using System.IO;
 using System.Linq;
 using System.Text;
@@ -95,7 +96,7 @@ namespace UnrealBuildTool
 		/// <param name="FileName">The filename that was read</param>
 		/// <param name="Result">If successful, the receipt that was read. Null otherwise.</param>
 		/// <returns>True if the file was read succesfully.</returns>
-		public static bool TryRead(FileReference FileName, out ModuleManifest Result)
+		public static bool TryRead(FileReference FileName, [NotNullWhen(true)] out ModuleManifest? Result)
 		{
 			if (!FileReference.Exists(FileName))
 			{
diff --git a/Engine/Source/Programs/UnrealBuildTool/System/NativeProjects.cs b/Engine/Source/Programs/UnrealBuildTool/System/NativeProjects.cs
--- a/Engine/Source/Programs/UnrealBuildTool/System/NativeProjects.cs
+++ b/Engine/Source/Programs/UnrealBuildTool/System/NativeProjects.cs
@@ -2,6 +2,7 @@
 
 using System;
 using System.Collections.Generic;
+using System.Diagnostics.CodeAnalysis;
 using System.IO;
 using System.Linq;
 using System.Text;
@@ -23,17 +24,17 @@ namespace UnrealBuildTool
 		/// <summary>
 		/// The native project base directories
 		/// </summary>
-		static HashSet<DirectoryReference> CachedBaseDirectories;
+		static HashSet<DirectoryReference>? CachedBaseDirectories;
 
 		/// <summary>
 		/// Cached list of project files within all the base directories
 		/// </summary>
-		static HashSet<FileReference> CachedProjectFiles;
+		static HashSet<FileReference>? CachedProjectFiles;
 
 		/// <summary>
 		/// Cached map of target names to the project file they belong to
 		/// </summary>
-		static Dictionary<string, FileReference> CachedTargetNameToProjectFile;
+		static Dictionary<string, FileReference>? CachedTargetNameToProjectFile;
 
 		/// <summary>
 		/// Clear our cached properties. Generally only needed if your script has modified local files...
@@ -128,7 +129,7 @@ namespace UnrealBuildTool
 		/// <param name="InTargetName">Name of the target of interest</param>
 		/// <param name="OutProjectFileName">The project filename</param>
 		/// <returns>True if the target was found</returns>
-		public static bool TryGetProjectForTarget(string InTargetName, out FileReference OutProjectFileName)
+		public static bool TryGetProjectForTarget(string InTargetName, [NotNullWhen(true)] out FileReference? OutProjectFileName)
 		{
 			if(CachedTargetNameToProjectFile == null)
 			{
@@ -199,7 +200,7 @@ namespace UnrealBuildTool
 		public static bool IsNativeProject(FileReference ProjectFile)
 		{
 			EnumerateProjectFiles();
-			return CachedProjectFiles.Contains(ProjectFile);
+			return CachedProjectFiles!.Contains(ProjectFile);
 		}
 	}
 }
diff --git a/Engine/Source/Programs/UnrealBuildTool/System/PlatformExports.cs b/Engine/Source/Programs/UnrealBuildTool/System/PlatformExports.cs
--- a/Engine/Source/Programs/UnrealBuildTool/System/PlatformExports.cs
+++ b/Engine/Source/Programs/UnrealBuildTool/System/PlatformExports.cs
@@ -52,7 +52,7 @@ namespace UnrealBuildTool
 		/// <returns>True if the project uses the default build configuration</returns>
 		public static bool HasDefaultBuildConfig(FileReference ProjectFile, UnrealTargetPlatform Platform)
 		{
-			UEBuildPlatform.TryGetBuildPlatform(Platform, out UEBuildPlatform BuildPlat);
+			UEBuildPlatform.TryGetBuildPlatform(Platform, out UEBuildPlatform? BuildPlat);
 			return (BuildPlat == null)? true : BuildPlat.HasDefaultBuildConfig(Platform, ProjectFile.Directory);
 		}
 
@@ -64,7 +64,7 @@ namespace UnrealBuildTool
 		/// <returns>True if the project requires a build for the platform</returns>
 		public static bool RequiresBuild(FileReference ProjectFile, UnrealTargetPlatform Platform)
 		{
-			UEBuildPlatform.TryGetBuildPlatform(Platform, out UEBuildPlatform BuildPlat);
+			UEBuildPlatform.TryGetBuildPlatform(Platform, out UEBuildPlatform? BuildPlat);
 			return (BuildPlat == null) ? false : BuildPlat.RequiresBuild(Platform, ProjectFile.Directory);
 		}
 
diff --git a/Engine/Source/Programs/UnrealBuildTool/System/PluginDescriptor.cs b/Engine/Source/Programs/UnrealBuildTool/System/PluginDescriptor.cs
--- a/Engine/Source/Programs/UnrealBuildTool/System/PluginDescriptor.cs
+++ b/Engine/Source/Programs/UnrealBuildTool/System/PluginDescriptor.cs
@@ -67,52 +67,52 @@ namespace UnrealBuildTool
 		/// Name of the version for this plugin.  This is the front-facing part of the version number.  It doesn't need to match
 		/// the version number numerically, but should be updated when the version number is increased accordingly.
 		/// </summary>
-		public string VersionName;
+		public string? VersionName;
 
 		/// <summary>
 		/// Friendly name of the plugin
 		/// </summary>
-		public string FriendlyName;
+		public string? FriendlyName;
 
 		/// <summary>
 		/// Description of the plugin
 		/// </summary>
-		public string Description;
+		public string? Description;
 
 		/// <summary>
 		/// The name of the category this plugin
 		/// </summary>
-		public string Category;
+		public string? Category;
 
 		/// <summary>
 		/// The company or individual who created this plugin.  This is an optional field that may be displayed in the user interface.
 		/// </summary>
-		public string CreatedBy;
+		public string? CreatedBy;
 
 		/// <summary>
 		/// Hyperlink URL string for the company or individual who created this plugin.  This is optional.
 		/// </summary>
-		public string CreatedByURL;
+		public string? CreatedByURL;
 
 		/// <summary>
 		/// Documentation URL string.
 		/// </summary>
-		public string DocsURL;
+		public string? DocsURL;
 
 		/// <summary>
 		/// Marketplace URL for this plugin. This URL will be embedded into projects that enable this plugin, so we can redirect to the marketplace if a user doesn't have it installed.
 		/// </summary>
-		public string MarketplaceURL;
+		public string? MarketplaceURL;
 
 		/// <summary>
 		/// Support URL/email for this plugin.
 		/// </summary>
-		public string SupportURL;
+		public string? SupportURL;
 
 		/// <summary>
 		/// Sets the version of the engine that this plugin is compatible with.
 		/// </summary>
-		public string EngineVersion;
+		public string? EngineVersion;
 
 		/// <summary>4
 		/// If true, this plugin from a platform extension extending another plugin */
@@ -122,22 +122,22 @@ namespace UnrealBuildTool
 		/// <summary>
 		/// List of platforms supported by this plugin. This list will be copied to any plugin reference from a project file, to allow filtering entire plugins from staged builds.
 		/// </summary>
-		public List<UnrealTargetPlatform> SupportedTargetPlatforms;
+		public List<UnrealTargetPlatform>? SupportedTargetPlatforms;
 
 		/// <summary>
 		/// List of programs supported by this plugin.
 		/// </summary>
-		public string[] SupportedPrograms;
+		public string[]? SupportedPrograms;
 
 		/// <summary>
 		/// List of all modules associated with this plugin
 		/// </summary>
-		public List<ModuleDescriptor> Modules;
+		public List<ModuleDescriptor>? Modules;
 
 		/// <summary>
 		/// List of all localization targets associated with this plugin
 		/// </summary>
-		public LocalizationTargetDescriptor[] LocalizationTargets;
+		public LocalizationTargetDescriptor[]? LocalizationTargets;
 
 		/// <summary>
 		/// Whether this plugin should be enabled by default for all projects
@@ -177,17 +177,17 @@ namespace UnrealBuildTool
 		/// <summary>
 		/// Set of pre-build steps to execute, keyed by host platform name.
 		/// </summary>
-		public CustomBuildSteps PreBuildSteps;
+		public CustomBuildSteps? PreBuildSteps;
 
 		/// <summary>
 		/// Set of post-build steps to execute, keyed by host platform name.
 		/// </summary>
-		public CustomBuildSteps PostBuildSteps;
+		public CustomBuildSteps? PostBuildSteps;
 
 		/// <summary>
 		/// Additional plugins that this plugin depends on
 		/// </summary>
-		public List<PluginReferenceDescriptor> Plugins;
+		public List<PluginReferenceDescriptor>? Plugins;
 
 		/// <summary>
 		/// Private constructor. This object should not be created directly; read it from disk using FromFile() instead.
diff --git a/Engine/Source/Programs/UnrealBuildTool/System/PluginReferenceDescriptor.cs b/Engine/Source/Programs/UnrealBuildTool/System/PluginReferenceDescriptor.cs
--- a/Engine/Source/Programs/UnrealBuildTool/System/PluginReferenceDescriptor.cs
+++ b/Engine/Source/Programs/UnrealBuildTool/System/PluginReferenceDescriptor.cs
@@ -34,47 +34,47 @@ namespace UnrealBuildTool
 		/// <summary>
 		/// Description of the plugin for users that do not have it installed.
 		/// </summary>
-		public string Description;
+		public string? Description;
 
 		/// <summary>
 		/// URL for this plugin on the marketplace, if the user doesn't have it installed.
 		/// </summary>
-		public string MarketplaceURL;
+		public string? MarketplaceURL;
 
 		/// <summary>
 		/// If enabled, list of platforms for which the plugin should be enabled (or all platforms if blank).
 		/// </summary>
-		public List<UnrealTargetPlatform> WhitelistPlatforms;
+		public List<UnrealTargetPlatform>? WhitelistPlatforms;
 
 		/// <summary>
 		/// If enabled, list of platforms for which the plugin should be disabled.
 		/// </summary>
-		public List<UnrealTargetPlatform> BlacklistPlatforms;
+		public List<UnrealTargetPlatform>? BlacklistPlatforms;
 
 		/// <summary>
 		/// If enabled, list of target configurations for which the plugin should be enabled (or all target configurations if blank).
 		/// </summary>
-		public UnrealTargetConfiguration[] WhitelistTargetConfigurations;
+		public UnrealTargetConfiguration[]? WhitelistTargetConfigurations;
 
 		/// <summary>
 		/// If enabled, list of target configurations for which the plugin should be disabled.
 		/// </summary>
-		public UnrealTargetConfiguration[] BlacklistTargetConfigurations;
+		public UnrealTargetConfiguration[]? BlacklistTargetConfigurations;
 
 		/// <summary>
 		/// If enabled, list of targets for which the plugin should be enabled (or all targets if blank).
 		/// </summary>
-		public TargetType[] WhitelistTargets;
+		public TargetType[]? WhitelistTargets;
 
 		/// <summary>
 		/// If enabled, list of targets for which the plugin should be disabled.
 		/// </summary>
-		public TargetType[] BlacklistTargets;
+		public TargetType[]? BlacklistTargets;
 
 		/// <summary>
 		/// The list of supported platforms for this plugin. This field is copied from the plugin descriptor, and supplements the user's whitelisted and blacklisted platforms.
 		/// </summary>
-		public List<UnrealTargetPlatform> SupportedTargetPlatforms;
+		public List<UnrealTargetPlatform>? SupportedTargetPlatforms;
 
 		/// <summary>
 		/// Constructor
@@ -82,7 +82,7 @@ namespace UnrealBuildTool
 		/// <param name="InName">Name of the plugin</param>
 		/// <param name="InMarketplaceURL">The marketplace URL for plugins which are not installed</param>
 		/// <param name="bInEnabled">Whether the plugin is enabled</param>
-		public PluginReferenceDescriptor(string InName, string InMarketplaceURL, bool bInEnabled)
+		public PluginReferenceDescriptor(string InName, string? InMarketplaceURL, bool bInEnabled)
 		{
 			Name = InName;
 			MarketplaceURL = InMarketplaceURL;
@@ -147,7 +147,7 @@ namespace UnrealBuildTool
 		/// <param name="Writer">The Json writer to output to</param>
 		/// <param name="Name">Name of the array</param>
 		/// <param name="Plugins">Array of plugins</param>
-		public static void WriteArray(JsonWriter Writer, string Name, PluginReferenceDescriptor[] Plugins)
+		public static void WriteArray(JsonWriter Writer, string Name, PluginReferenceDescriptor[]? Plugins)
 		{
 			if (Plugins != null && Plugins.Length > 0)
 			{
@@ -167,9 +167,9 @@ namespace UnrealBuildTool
 		/// <returns>New PluginReferenceDescriptor object</returns>
 		public static PluginReferenceDescriptor FromJsonObject(JsonObject RawObject)
 		{
-			string[] WhitelistPlatformNames = null;
-			string[] BlacklistPlatformNames = null;
-			string[] SupportedTargetPlatformNames = null;
+			string[]? WhitelistPlatformNames = null;
+			string[]? BlacklistPlatformNames = null;
+			string[]? SupportedTargetPlatformNames = null;
 
 			PluginReferenceDescriptor Descriptor = new PluginReferenceDescriptor(RawObject.GetStringField("Name"), null, RawObject.GetBoolField("Enabled"));
 			RawObject.TryGetBoolField("Optional", out Descriptor.bOptional);
diff --git a/Engine/Source/Programs/UnrealBuildTool/System/Plugins.cs b/Engine/Source/Programs/UnrealBuildTool/System/Plugins.cs
--- a/Engine/Source/Programs/UnrealBuildTool/System/Plugins.cs
+++ b/Engine/Source/Programs/UnrealBuildTool/System/Plugins.cs
@@ -179,7 +179,7 @@ namespace UnrealBuildTool
 			Dictionary<string, PluginInfo> NameToPluginInfo = new Dictionary<string, PluginInfo>(StringComparer.InvariantCultureIgnoreCase);
 			foreach(PluginInfo Plugin in Plugins)
 			{
-				PluginInfo ExistingPluginInfo;
+				PluginInfo? ExistingPluginInfo;
 				if(!NameToPluginInfo.TryGetValue(Plugin.Name, out ExistingPluginInfo))
 				{
 					NameToPluginInfo.Add(Plugin.Name, Plugin);
@@ -233,7 +233,7 @@ namespace UnrealBuildTool
 		/// </summary>
 		/// <param name="ProjectFile">Path to the project file</param>
 		/// <returns>List of project files</returns>
-		public static IEnumerable<FileReference> EnumeratePlugins(FileReference ProjectFile)
+		public static IEnumerable<FileReference> EnumeratePlugins(FileReference? ProjectFile)
 		{
 			List<DirectoryReference> BaseDirs = new List<DirectoryReference>();
 			BaseDirs.AddRange(UnrealBuildTool.GetExtensionDirs(UnrealBuildTool.EngineDirectory, "Plugins"));
@@ -318,7 +318,7 @@ namespace UnrealBuildTool
 		private static void TryMergeWithParent(PluginInfo Child, FileReference Filename)
 		{
 			// find the parent
-			PluginInfo Parent = null;
+			PluginInfo? Parent = null;
 
 			string[] Tokens = Filename.GetFileNameWithoutAnyExtensions().Split("_".ToCharArray());
 			if (Tokens.Length == 2)
@@ -486,7 +486,7 @@ namespace UnrealBuildTool
 					continue;
 				}
 
-				List<PluginInfo> Plugins;
+				List<PluginInfo>? Plugins;
 				if (!PluginInfoCache.TryGetValue(Dir, out Plugins))
 				{
 					Plugins = new List<PluginInfo>();
@@ -526,7 +526,7 @@ namespace UnrealBuildTool
 		/// <param name="ParentDirectory">Parent directory to look in. Plugins will be found in any *subfolders* of this directory.</param>
 		public static IEnumerable<FileReference> EnumeratePlugins(DirectoryReference ParentDirectory)
 		{
-			List<FileReference> FileNames;
+			List<FileReference>? FileNames;
 			if (!PluginFileCache.TryGetValue(ParentDirectory, out FileNames))
 			{
 				FileNames = new List<FileReference>();
diff --git a/Engine/Source/Programs/UnrealBuildTool/System/PrecompiledManifest.cs b/Engine/Source/Programs/UnrealBuildTool/System/PrecompiledManifest.cs
--- a/Engine/Source/Programs/UnrealBuildTool/System/PrecompiledManifest.cs
+++ b/Engine/Source/Programs/UnrealBuildTool/System/PrecompiledManifest.cs
@@ -1,6 +1,7 @@
 // Copyright Epic Games, Inc. All Rights Reserved.
 using System;
 using System.Collections.Generic;
+using System.Diagnostics.CodeAnalysis;
 using System.IO;
 using System.Linq;
 using System.Text;
@@ -52,7 +53,7 @@ namespace UnrealBuildTool
 		/// <param name="Location">Filename to read from</param>
 		/// <param name="Manifest">If successful, the manifest that was read</param>
 		/// <returns>True if successful</returns>
-		public static bool TryRead(FileReference Location, out PrecompiledManifest Manifest)
+		public static bool TryRead(FileReference Location, [NotNullWhen(true)] out PrecompiledManifest? Manifest)
 		{
 			if (!FileReference.Exists(Location))
 			{
diff --git a/Engine/Source/Programs/UnrealBuildTool/System/ProgressWriter.cs b/Engine/Source/Programs/UnrealBuildTool/System/ProgressWriter.cs
--- a/Engine/Source/Programs/UnrealBuildTool/System/ProgressWriter.cs
+++ b/Engine/Source/Programs/UnrealBuildTool/System/ProgressWriter.cs
@@ -27,12 +27,12 @@ namespace UnrealBuildTool
 		/// <summary>
 		/// The inner scope object
 		/// </summary>
-		LogStatusScope Status;
+		LogStatusScope? Status;
 
 		/// <summary>
 		/// The current progress message
 		/// </summary>
-		string CurrentProgressString;
+		string? CurrentProgressString;
 
 		/// <summary>
 		/// Constructor
diff --git a/Engine/Source/Programs/UnrealBuildTool/System/ProjectDescriptor.cs b/Engine/Source/Programs/UnrealBuildTool/System/ProjectDescriptor.cs
--- a/Engine/Source/Programs/UnrealBuildTool/System/ProjectDescriptor.cs
+++ b/Engine/Source/Programs/UnrealBuildTool/System/ProjectDescriptor.cs
@@ -58,27 +58,27 @@ namespace UnrealBuildTool
 		/// <summary>
 		/// The engine to open this project with.
 		/// </summary>
-		public string EngineAssociation;
+		public string? EngineAssociation;
 
 		/// <summary>
 		/// Category to show under the project browser
 		/// </summary>
-		public string Category;
+		public string? Category;
 
 		/// <summary>
 		/// Description to show in the project browser
 		/// </summary>
-		public string Description;
+		public string? Description;
 
 		/// <summary>
 		/// List of all modules associated with this project
 		/// </summary>
-		public ModuleDescriptor[] Modules;
+		public ModuleDescriptor[]? Modules;
 
 		/// <summary>
 		/// List of plugins for this project (may be enabled/disabled)
 		/// </summary>
-		public PluginReferenceDescriptor[] Plugins;
+		public PluginReferenceDescriptor[]? Plugins;
 
 		/// <summary>
 		/// Array of additional root directories
@@ -93,7 +93,7 @@ namespace UnrealBuildTool
 		/// <summary>
 		/// Array of platforms that this project is targeting
 		/// </summary>
-		public string[] TargetPlatforms;
+		public string[]? TargetPlatforms;
 
 		/// <summary>
 		/// A hash that is used to determine if the project was forked from a sample
@@ -103,12 +103,12 @@ namespace UnrealBuildTool
 		/// <summary>
 		/// Steps to execute before building targets in this project
 		/// </summary>
-		public CustomBuildSteps PreBuildSteps;
+		public CustomBuildSteps? PreBuildSteps;
 
 		/// <summary>
 		/// Steps to execute before building targets in this project
 		/// </summary>
-		public CustomBuildSteps PostBuildSteps;
+		public CustomBuildSteps? PostBuildSteps;
 
 		/// <summary>
 		/// Indicates if this project is an Enterprise project
diff --git a/Engine/Source/Programs/UnrealBuildTool/System/RestrictedFolders.cs b/Engine/Source/Programs/UnrealBuildTool/System/RestrictedFolders.cs
--- a/Engine/Source/Programs/UnrealBuildTool/System/RestrictedFolders.cs
+++ b/Engine/Source/Programs/UnrealBuildTool/System/RestrictedFolders.cs
@@ -24,17 +24,17 @@ namespace UnrealBuildTool
 		/// <summary>
 		/// Array of all restricted folder names
 		/// </summary>
-		private static string[] Names;
+		private static string[]? Names;
 
 		/// <summary>
 		/// Array of all restricted folders
 		/// </summary>
-		private static RestrictedFolder[] Values;
+		private static RestrictedFolder[]? Values;
 
 		/// <summary>
 		/// Set of permitted references for each restricted folder. Determined via data-driven platform info.
 		/// </summary>
-		private static Dictionary<RestrictedFolder, RestrictedFolder[]> PermittedReferences;
+		private static Dictionary<RestrictedFolder, RestrictedFolder[]>? PermittedReferences;
 
 		/// <summary>
 		/// Constructor
@@ -82,7 +82,7 @@ namespace UnrealBuildTool
 		/// </summary>
 		/// <param name="Other">The restricted folder to compare against</param>
 		/// <returns>True if the restricted folder is equal to the other instance</returns>
-		public override bool Equals(object Other)
+		public override bool Equals(object? Other)
 		{
 			return Other is RestrictedFolder && Id == ((RestrictedFolder)Other).Id;
 		}
@@ -102,8 +102,8 @@ namespace UnrealBuildTool
 		/// <returns>Collection of restricted folders</returns>
 		public IEnumerable<RestrictedFolder> GetPermittedReferences()
 		{
-			RestrictedFolder[] References;
-			if (PermittedReferences.TryGetValue(this, out References))
+			RestrictedFolder[]? References;
+			if (PermittedReferences != null && PermittedReferences.TryGetValue(this, out References))
 			{
 				foreach (RestrictedFolder Reference in References)
 				{
diff --git a/Engine/Source/Programs/UnrealBuildTool/System/RulesAssembly.cs b/Engine/Source/Programs/UnrealBuildTool/System/RulesAssembly.cs
--- a/Engine/Source/Programs/UnrealBuildTool/System/RulesAssembly.cs
+++ b/Engine/Source/Programs/UnrealBuildTool/System/RulesAssembly.cs
@@ -2,6 +2,7 @@
 
 using System;
 using System.Collections.Generic;
+using System.Diagnostics.CodeAnalysis;
 using System.IO;
 using System.Linq;
 using System.Reflection;
@@ -25,7 +26,7 @@ namespace UnrealBuildTool
 		/// <summary>
 		/// The compiled assembly
 		/// </summary>
-		private Assembly CompiledAssembly;
+		private Assembly? CompiledAssembly;
 
 		/// <summary>
 		/// The base directories for this assembly
@@ -71,7 +72,7 @@ namespace UnrealBuildTool
 		/// <summary>
 		/// The parent rules assembly that this assembly inherits. Game assemblies inherit the engine assembly, and the engine assembly inherits nothing.
 		/// </summary>
-		private RulesAssembly Parent;
+		private RulesAssembly? Parent;
 
 		/// <summary>
 		/// Constructor. Compiles a rules assembly from the given source files.
@@ -87,7 +88,7 @@ namespace UnrealBuildTool
 		/// <param name="bReadOnly">Whether the modules and targets in this assembly are installed, and should be created with the bUsePrecompiled flag set</param> 
 		/// <param name="bSkipCompile">Whether to skip compiling this assembly</param>
 		/// <param name="Parent">The parent rules assembly</param>
-		internal RulesAssembly(RulesScope Scope, List<DirectoryReference> BaseDirs, IReadOnlyList<PluginInfo> Plugins, Dictionary<FileReference, ModuleRulesContext> ModuleFileToContext, List<FileReference> TargetFiles, FileReference AssemblyFileName, bool bContainsEngineModules, BuildSettingsVersion? DefaultBuildSettings, bool bReadOnly, bool bSkipCompile, RulesAssembly Parent)
+		internal RulesAssembly(RulesScope Scope, List<DirectoryReference> BaseDirs, IReadOnlyList<PluginInfo> Plugins, Dictionary<FileReference, ModuleRulesContext> ModuleFileToContext, List<FileReference> TargetFiles, FileReference AssemblyFileName, bool bContainsEngineModules, BuildSettingsVersion? DefaultBuildSettings, bool bReadOnly, bool bSkipCompile, RulesAssembly? Parent)
 		{
 			this.Scope = Scope;
 			this.BaseDirs = BaseDirs;
@@ -138,10 +139,10 @@ namespace UnrealBuildTool
 				{
 					foreach (MethodInfo Method in CompiledType.GetMethods(BindingFlags.Instance | BindingFlags.Public | BindingFlags.DeclaredOnly))
 					{
-						ObsoleteOverrideAttribute Attribute = Method.GetCustomAttribute<ObsoleteOverrideAttribute>(true);
+						ObsoleteOverrideAttribute? Attribute = Method.GetCustomAttribute<ObsoleteOverrideAttribute>(true);
 						if (Attribute != null)
 						{
-							FileReference Location;
+							FileReference? Location;
 							if (!TryGetFileNameFromType(CompiledType, out Location))
 							{
 								Location = new FileReference(CompiledAssembly.Location);
@@ -151,10 +152,10 @@ namespace UnrealBuildTool
 					}
 					if(CompiledType.BaseType == typeof(ModuleRules))
 					{
-						ConstructorInfo Constructor = CompiledType.GetConstructor(new Type[] { typeof(TargetInfo) });
+						ConstructorInfo? Constructor = CompiledType.GetConstructor(new Type[] { typeof(TargetInfo) });
 						if(Constructor != null)
 						{
-							FileReference Location;
+							FileReference? Location;
 							if (!TryGetFileNameFromType(CompiledType, out Location))
 							{
 								Location = new FileReference(CompiledAssembly.Location);
@@ -198,7 +199,7 @@ namespace UnrealBuildTool
 			PreprocessorDefines.Add("WITH_FORWARDED_TARGET_RULES_CTOR");
 
 			// Define macros for the UE4 version, starting with 4.17
-			BuildVersion Version;
+			BuildVersion? Version;
 			if (BuildVersion.TryRead(BuildVersion.GetDefaultFileName(), out Version))
 			{
 				for(int MinorVersion = 17; MinorVersion <= 30; MinorVersion++)
@@ -249,7 +250,7 @@ namespace UnrealBuildTool
 		/// <param name="ExistingType"></param>
 		/// <param name="File"></param>
 		/// <returns>True if the type was found, false otherwise</returns>
-		public bool TryGetFileNameFromType(Type ExistingType, out FileReference File)
+		public bool TryGetFileNameFromType(Type ExistingType, [NotNullWhen(true)] out FileReference? File)
 		{
 			if (ExistingType.Assembly == CompiledAssembly)
 			{
@@ -286,9 +287,9 @@ namespace UnrealBuildTool
 		/// </summary>
 		/// <param name="ModuleName">The name of the module</param>
 		/// <returns>The filename containing rules for this module, or an empty string if not found</returns>
-		public FileReference GetModuleFileName(string ModuleName)
+		public FileReference? GetModuleFileName(string ModuleName)
 		{
-			FileReference ModuleFile;
+			FileReference? ModuleFile;
 			if (ModuleNameToModuleFile.TryGetValue(ModuleName, out ModuleFile))
 			{
 				return ModuleFile;
@@ -304,7 +305,7 @@ namespace UnrealBuildTool
 		/// </summary>
 		/// <param name="ModuleName">The name of the module</param>
 		/// <returns>The rules type for this module, or null if not found</returns>
-		public Type GetModuleRulesType(string ModuleName)
+		public Type? GetModuleRulesType(string ModuleName)
 		{
 			if (ModuleNameToModuleFile.ContainsKey(ModuleName))
 			{
@@ -321,15 +322,15 @@ namespace UnrealBuildTool
 		/// </summary>
 		/// <param name="ModuleName">The name of the module</param>
 		/// <returns>The rules type for this module, or null if not found</returns>
-		private Type GetModuleRulesTypeInternal(string ModuleName)
+		private Type? GetModuleRulesTypeInternal(string ModuleName)
 		{
 			// The build module must define a type named 'Rules' that derives from our 'ModuleRules' type.  
-			Type RulesObjectType = CompiledAssembly.GetType(ModuleName);
+			Type? RulesObjectType = CompiledAssembly?.GetType(ModuleName);
 			if (RulesObjectType == null)
 			{
 				// Temporary hack to avoid System namespace collisions
 				// @todo projectfiles: Make rules assemblies require namespaces.
-				RulesObjectType = CompiledAssembly.GetType("UnrealBuildTool.Rules." + ModuleName);
+				RulesObjectType = CompiledAssembly?.GetType("UnrealBuildTool.Rules." + ModuleName);
 			}
 			return RulesObjectType;
 		}
@@ -339,9 +340,9 @@ namespace UnrealBuildTool
 		/// </summary>
 		/// <param name="TargetName">The name of the target</param>
 		/// <returns>The filename containing rules for this target, or an empty string if not found</returns>
-		public FileReference GetTargetFileName(string TargetName)
+		public FileReference? GetTargetFileName(string TargetName)
 		{
-			FileReference TargetFile;
+			FileReference? TargetFile;
 			if (TargetNameToTargetFile.TryGetValue(TargetName, out TargetFile))
 			{
 				return TargetFile;
@@ -365,7 +366,7 @@ namespace UnrealBuildTool
 			string ModuleTypeName = ModuleName;
 
 			// Make sure the base module file is known to us
-			FileReference ModuleFileName;
+			FileReference? ModuleFileName;
 			if (!ModuleNameToModuleFile.TryGetValue(ModuleTypeName, out ModuleFileName))
 			{
 				if (Parent == null)
@@ -379,16 +380,16 @@ namespace UnrealBuildTool
 			}
 
 			// get the standard Rules object class from the assembly
-			Type BaseRulesObjectType = GetModuleRulesTypeInternal(ModuleTypeName);
+			Type BaseRulesObjectType = GetModuleRulesTypeInternal(ModuleTypeName)!;
 
 			// look around for platform/group modules that we will use instead of the basic module
-			Type PlatformRulesObjectType = GetModuleRulesTypeInternal(ModuleTypeName + "_" + Target.Platform.ToString());
+			Type? PlatformRulesObjectType = GetModuleRulesTypeInternal(ModuleTypeName + "_" + Target.Platform.ToString());
 			if (PlatformRulesObjectType == null)
 			{
 				foreach (UnrealPlatformGroup Group in UEBuildPlatform.GetPlatformGroups(Target.Platform))
 				{
 					// look to see if the group has an override
-					Type GroupRulesObjectType = GetModuleRulesTypeInternal(ModuleName + "_" + Group.ToString());
+					Type? GroupRulesObjectType = GetModuleRulesTypeInternal(ModuleName + "_" + Group.ToString());
 
 					// we expect only one platform group to be found in the extensions
 					if (GroupRulesObjectType != null && PlatformRulesObjectType != null)
@@ -407,10 +408,10 @@ namespace UnrealBuildTool
 
 
 			// Figure out the best rules object to use
-			Type RulesObjectType = PlatformRulesObjectType != null ? PlatformRulesObjectType : BaseRulesObjectType;
+			Type? RulesObjectType = PlatformRulesObjectType != null ? PlatformRulesObjectType : BaseRulesObjectType;
 			if (RulesObjectType == null)
 			{
-				throw new BuildException("Expecting to find a type to be declared in a module rules named '{0}' in {1}.  This type must derive from the 'ModuleRules' type defined by Unreal Build Tool.", ModuleTypeName, CompiledAssembly.FullName);
+				throw new BuildException("Expecting to find a type to be declared in a module rules named '{0}' in {1}.  This type must derive from the 'ModuleRules' type defined by Unreal Build Tool.", ModuleTypeName, CompiledAssembly?.FullName);
 			}
 
 			// Create an instance of the module's rules object
@@ -442,7 +443,7 @@ namespace UnrealBuildTool
 					RulesObject.SubclassRules = new List<string>();
 					while (SubType != null && SubType != BaseRulesObjectType)
 					{
-						FileReference SubTypeFileName;
+						FileReference? SubTypeFileName;
 						if (TryGetFileNameFromType(SubType, out SubTypeFileName))
 						{
 							RulesObject.DirectoriesForModuleSubClasses.Add(SubType, SubTypeFileName.Directory);
@@ -457,7 +458,7 @@ namespace UnrealBuildTool
 				}
 
 				// Call the constructor
-				ConstructorInfo Constructor = RulesObjectType.GetConstructor(new Type[] { typeof(ReadOnlyTargetRules) });
+				ConstructorInfo? Constructor = RulesObjectType.GetConstructor(new Type[] { typeof(ReadOnlyTargetRules) });
 				if(Constructor == null)
 				{
 					throw new BuildException("No valid constructor found for {0}.", ModuleName);
@@ -482,7 +483,7 @@ namespace UnrealBuildTool
 		protected TargetRules CreateTargetRulesInstance(string TypeName, TargetInfo TargetInfo)
 		{
 			// The build module must define a type named '<TargetName>Target' that derives from our 'TargetRules' type.  
-			Type BaseRulesType = CompiledAssembly.GetType(TypeName);
+			Type? BaseRulesType = CompiledAssembly?.GetType(TypeName);
 			if (BaseRulesType == null)
 			{
 				throw new BuildException("Expecting to find a type to be declared in a target rules named '{0}'.  This type must derive from the 'TargetRules' type defined by Unreal Build Tool.", TypeName);
@@ -490,14 +491,14 @@ namespace UnrealBuildTool
 
 			// Look for platform/group rules that we will use instead of the base rules
 			string PlatformRulesName = TargetInfo.Name + "_" + TargetInfo.Platform.ToString();
-			Type PlatformRulesType = CompiledAssembly.GetType(TypeName + "_" + TargetInfo.Platform.ToString());
+			Type? PlatformRulesType = CompiledAssembly?.GetType(TypeName + "_" + TargetInfo.Platform.ToString());
 			if (PlatformRulesType == null)
 			{
 				foreach (UnrealPlatformGroup Group in UEBuildPlatform.GetPlatformGroups(TargetInfo.Platform))
 				{
 					// look to see if the group has an override
 					string GroupRulesName = TargetInfo.Name + "_" + Group.ToString();
-					Type GroupRulesObjectType = CompiledAssembly.GetType(TypeName + "_" + Group.ToString());
+					Type? GroupRulesObjectType = CompiledAssembly?.GetType(TypeName + "_" + Group.ToString());
 
 					// we expect only one platform group to be found in the extensions
 					if (GroupRulesObjectType != null && PlatformRulesType != null)
@@ -530,7 +531,7 @@ namespace UnrealBuildTool
 			Rules.File = TargetNameToTargetFile[TargetInfo.Name];
 
 			// Find the constructor
-			ConstructorInfo Constructor = RulesType.GetConstructor(new Type[] { typeof(TargetInfo) });
+			ConstructorInfo? Constructor = RulesType.GetConstructor(new Type[] { typeof(TargetInfo) });
 			if(Constructor == null)
 			{
 				throw new BuildException("No constructor found on {0} which takes an argument of type TargetInfo.", RulesType.Name);
@@ -543,7 +544,7 @@ namespace UnrealBuildTool
 			}
 			catch (Exception Ex)
 			{
-				throw new BuildException(Ex, "Unable to instantiate instance of '{0}' object type from compiled assembly '{1}'.  Unreal Build Tool creates an instance of your module's 'Rules' object in order to find out about your module's requirements.  The CLR exception details may provide more information:  {2}", TypeName, Path.GetFileNameWithoutExtension(CompiledAssembly.Location), Ex.ToString());
+				throw new BuildException(Ex, "Unable to instantiate instance of '{0}' object type from compiled assembly '{1}'.  Unreal Build Tool creates an instance of your module's 'Rules' object in order to find out about your module's requirements.  The CLR exception details may provide more information:  {2}", TypeName, Path.GetFileNameWithoutExtension(CompiledAssembly?.Location), Ex.ToString());
 			}
 
 			// Set the default overriddes for the configured target type
@@ -626,7 +627,7 @@ namespace UnrealBuildTool
 		/// <param name="ProjectFile">Path to the project file for this target</param>
 		/// <param name="Arguments">Command line arguments for this target</param>
 		/// <returns>The build target rules for the specified target</returns>
-		public TargetRules CreateTargetRules(string TargetName, UnrealTargetPlatform Platform, UnrealTargetConfiguration Configuration, string Architecture, FileReference ProjectFile, CommandLineArguments Arguments)
+		public TargetRules CreateTargetRules(string TargetName, UnrealTargetPlatform Platform, UnrealTargetConfiguration Configuration, string Architecture, FileReference? ProjectFile, CommandLineArguments? Arguments)
 		{
 			bool bFoundTargetName = TargetNameToTargetFile.ContainsKey(TargetName);
 			if (bFoundTargetName == false)
@@ -637,10 +638,10 @@ namespace UnrealBuildTool
 					string ExceptionMessage = "Couldn't find target rules file for target '";
 					ExceptionMessage += TargetName;
 					ExceptionMessage += "' in rules assembly '";
-					ExceptionMessage += CompiledAssembly.FullName;
+					ExceptionMessage += CompiledAssembly?.FullName;
 					ExceptionMessage += "'." + Environment.NewLine;
 
-					ExceptionMessage += "Location: " + CompiledAssembly.Location + Environment.NewLine;
+					ExceptionMessage += "Location: " + CompiledAssembly?.Location + Environment.NewLine;
 
 					ExceptionMessage += "Target rules found:" + Environment.NewLine;
 					foreach (KeyValuePair<string, FileReference> entry in TargetNameToTargetFile)
@@ -741,9 +742,9 @@ namespace UnrealBuildTool
 		/// <param name="ModuleFile">The module to search for</param>
 		/// <param name="Plugin">The matching plugin info, or null.</param>
 		/// <returns>True if the module belongs to a plugin</returns>
-		public bool TryGetPluginForModule(FileReference ModuleFile, out PluginInfo Plugin)
+		public bool TryGetPluginForModule(FileReference ModuleFile, [NotNullWhen(true)] out PluginInfo? Plugin)
 		{
-			ModuleRulesContext Context;
+			ModuleRulesContext? Context;
 			if (ModuleFileToContext.TryGetValue(ModuleFile, out Context))
 			{
 				Plugin = Context.Plugin;
diff --git a/Engine/Source/Programs/UnrealBuildTool/System/RulesCompiler.cs b/Engine/Source/Programs/UnrealBuildTool/System/RulesCompiler.cs
--- a/Engine/Source/Programs/UnrealBuildTool/System/RulesCompiler.cs
+++ b/Engine/Source/Programs/UnrealBuildTool/System/RulesCompiler.cs
@@ -174,7 +174,7 @@ namespace UnrealBuildTool
 		/// <param name="Directories">The directories to cache</param>
 		private static void PrefetchRulesFiles(IEnumerable<DirectoryReference> Directories)
 		{
-			ThreadPoolWorkQueue Queue = null;
+			ThreadPoolWorkQueue? Queue = null;
 			try
 			{
 				foreach(DirectoryReference Directory in Directories)
@@ -213,7 +213,7 @@ namespace UnrealBuildTool
 		private static IReadOnlyList<FileReference> FindAllRulesFiles(DirectoryReference Directory, RulesFileType Type)
 		{
 			// Check to see if we've already cached source files for this folder
-			RulesFileCache Cache;
+			RulesFileCache? Cache;
 			if (!RootFolderToRulesFileCache.TryGetValue(Directory, out Cache))
 			{
 				Cache = new RulesFileCache();
@@ -330,7 +330,7 @@ namespace UnrealBuildTool
 		/// <summary>
 		/// The cached rules assembly for engine modules and targets.
 		/// </summary>
-		private static RulesAssembly EngineRulesAssembly;
+		private static RulesAssembly? EngineRulesAssembly;
 
 		/// <summary>
 		/// Map of assembly names we've already compiled and loaded to their Assembly and list of game folders.  This is used to prevent
@@ -371,7 +371,7 @@ namespace UnrealBuildTool
 		/// <param name="bSkipCompile">Whether to skip compilation for this assembly</param>
 		/// <param name="Parent">The parent rules assembly</param>
 		/// <returns>New rules assembly</returns>
-		private static RulesAssembly CreateEngineRulesAssemblyInternal(RulesScope Scope, List<DirectoryReference> RootDirectories, string AssemblyPrefix, IReadOnlyList<PluginInfo> Plugins, bool bReadOnly, bool bSkipCompile, RulesAssembly Parent)
+		private static RulesAssembly CreateEngineRulesAssemblyInternal(RulesScope Scope, List<DirectoryReference> RootDirectories, string AssemblyPrefix, IReadOnlyList<PluginInfo> Plugins, bool bReadOnly, bool bSkipCompile, RulesAssembly? Parent)
 		{
 			// Scope hierarchy
 			RulesScope PluginsScope = new RulesScope(Scope.Name + " Plugins", Scope);
@@ -447,7 +447,7 @@ namespace UnrealBuildTool
 		public static RulesAssembly CreateProjectRulesAssembly(FileReference ProjectFileName, bool bUsePrecompiled, bool bSkipCompile)
 		{
 			// Check if there's an existing assembly for this project
-			RulesAssembly ProjectRulesAssembly;
+			RulesAssembly? ProjectRulesAssembly;
 			if (!LoadedAssemblyMap.TryGetValue(ProjectFileName, out ProjectRulesAssembly))
 			{
 				ProjectDescriptor Project = ProjectDescriptor.FromFile(ProjectFileName);
@@ -538,7 +538,7 @@ namespace UnrealBuildTool
         public static RulesAssembly CreatePluginRulesAssembly(FileReference PluginFileName, bool bSkipCompile, RulesAssembly Parent, bool bContainsEngineModules)
 		{
 			// Check if there's an existing assembly for this project
-			RulesAssembly PluginRulesAssembly;
+			RulesAssembly? PluginRulesAssembly;
 			if (!LoadedAssemblyMap.TryGetValue(PluginFileName, out PluginRulesAssembly))
 			{
 				// Find all the rules source files
@@ -547,7 +547,7 @@ namespace UnrealBuildTool
 
 				// Create a list of plugins for this assembly. If it already exists in the parent assembly, just create an empty assembly.
 				List<PluginInfo> ForeignPlugins = new List<PluginInfo>();
-				if (Parent == null || !Parent.EnumeratePlugins().Any(x => x.File == PluginFileName))
+				if (!Parent.EnumeratePlugins().Any(x => x.File == PluginFileName))
 				{
 					ForeignPlugins.Add(new PluginInfo(PluginFileName, PluginType.External));
 				}
@@ -577,7 +577,7 @@ namespace UnrealBuildTool
 		/// <param name="bUsePrecompiled">Whether to use a precompiled engine build</param>
 		/// <param name="ForeignPlugin">Foreign plugin to be compiled</param>
 		/// <returns>The compiled rules assembly</returns>
-		public static RulesAssembly CreateTargetRulesAssembly(FileReference ProjectFile, string TargetName, bool bSkipRulesCompile, bool bUsePrecompiled, FileReference ForeignPlugin)
+		public static RulesAssembly CreateTargetRulesAssembly(FileReference? ProjectFile, string TargetName, bool bSkipRulesCompile, bool bUsePrecompiled, FileReference? ForeignPlugin)
 		{
 			RulesAssembly RulesAssembly;
 			if (ProjectFile != null)
@@ -628,9 +628,9 @@ namespace UnrealBuildTool
 		/// </summary>
 		/// <param name="ExistingType">The type to search for.</param>
 		/// <returns>The filename that declared the given type, or null</returns>
-		public static string GetFileNameFromType(Type ExistingType)
+		public static string? GetFileNameFromType(Type ExistingType)
 		{
-			FileReference FileName;
+			FileReference? FileName;
 			if (EngineRulesAssembly != null && EngineRulesAssembly.TryGetFileNameFromType(ExistingType, out FileName))
 			{
 				return FileName.FullName;
diff --git a/Engine/Source/Programs/UnrealBuildTool/System/RulesDocumentation.cs b/Engine/Source/Programs/UnrealBuildTool/System/RulesDocumentation.cs
--- a/Engine/Source/Programs/UnrealBuildTool/System/RulesDocumentation.cs
+++ b/Engine/Source/Programs/UnrealBuildTool/System/RulesDocumentation.cs
@@ -108,7 +108,7 @@ namespace UnrealBuildTool
 
 		static void WriteFieldUDN(XmlDocument InputDocumentation, FieldInfo Field, TextWriter Writer)
 		{
-			XmlNode Node = InputDocumentation.SelectSingleNode(String.Format("//member[@name='F:{0}.{1}']/summary", Field.DeclaringType.FullName, Field.Name));
+			XmlNode Node = InputDocumentation.SelectSingleNode(String.Format("//member[@name='F:{0}.{1}']/summary", Field.DeclaringType!.FullName, Field.Name));
 			if (Node != null)
 			{
 				// Reflow the comments into paragraphs, assuming that each paragraph will be separated by a blank line
@@ -185,7 +185,7 @@ namespace UnrealBuildTool
 
 		static void WriteFieldHTML(XmlDocument InputDocumentation, FieldInfo Field, TextWriter Writer)
 		{
-			XmlNode Node = InputDocumentation.SelectSingleNode(String.Format("//member[@name='F:{0}.{1}']/summary", Field.DeclaringType.FullName, Field.Name));
+			XmlNode Node = InputDocumentation.SelectSingleNode(String.Format("//member[@name='F:{0}.{1}']/summary", Field.DeclaringType!.FullName, Field.Name));
 			if(Node != null)
 			{
 				// Reflow the comments into paragraphs, assuming that each paragraph will be separated by a blank line
diff --git a/Engine/Source/Programs/UnrealBuildTool/System/RulesScope.cs b/Engine/Source/Programs/UnrealBuildTool/System/RulesScope.cs
--- a/Engine/Source/Programs/UnrealBuildTool/System/RulesScope.cs
+++ b/Engine/Source/Programs/UnrealBuildTool/System/RulesScope.cs
@@ -21,14 +21,14 @@ namespace UnrealBuildTool
 		/// <summary>
 		/// The parent scope
 		/// </summary>
-		public RulesScope Parent;
+		public RulesScope? Parent;
 
 		/// <summary>
 		/// Constructor
 		/// </summary>
 		/// <param name="Name">Name of this scope</param>
 		/// <param name="Parent">The parent scope</param>
-		public RulesScope(string Name, RulesScope Parent)
+		public RulesScope(string Name, RulesScope? Parent)
 		{
 			this.Name = Name;
 			this.Parent = Parent;
@@ -41,7 +41,7 @@ namespace UnrealBuildTool
 		/// <returns>True if this scope contains the other scope</returns>
 		public bool Contains(RulesScope Other)
 		{
-			for(RulesScope Scope = this; Scope != null; Scope = Scope.Parent)
+			for(RulesScope? Scope = this; Scope != null; Scope = Scope.Parent)
 			{
 				if(Scope == Other)
 				{
diff --git a/Engine/Source/Programs/UnrealBuildTool/System/SingleInstanceMutex.cs b/Engine/Source/Programs/UnrealBuildTool/System/SingleInstanceMutex.cs
--- a/Engine/Source/Programs/UnrealBuildTool/System/SingleInstanceMutex.cs
+++ b/Engine/Source/Programs/UnrealBuildTool/System/SingleInstanceMutex.cs
@@ -19,7 +19,7 @@ namespace UnrealBuildTool
 		/// <summary>
 		/// The global mutex instance
 		/// </summary>
-		Mutex GlobalMutex;
+		Mutex? GlobalMutex;
 
 		/// <summary>
 		/// Constructor. Attempts to acquire the global mutex
diff --git a/Engine/Source/Programs/UnrealBuildTool/System/SourceFileMetadataCache.cs b/Engine/Source/Programs/UnrealBuildTool/System/SourceFileMetadataCache.cs
--- a/Engine/Source/Programs/UnrealBuildTool/System/SourceFileMetadataCache.cs
+++ b/Engine/Source/Programs/UnrealBuildTool/System/SourceFileMetadataCache.cs
@@ -11,6 +11,8 @@ using System.Text.RegularExpressions;
 using System.Threading.Tasks;
 using Tools.DotNETCommon;
 
+#nullable enable
+
 namespace UnrealBuildTool
 {
 	/// <summary>
@@ -31,7 +33,7 @@ namespace UnrealBuildTool
 			/// <summary>
 			/// Contents of the include directive
 			/// </summary>
-			public string IncludeText;
+			public string? IncludeText;
 		}
 
 		/// <summary>
@@ -68,7 +70,7 @@ namespace UnrealBuildTool
 		/// <summary>
 		/// The parent cache.
 		/// </summary>
-		SourceFileMetadataCache Parent;
+		SourceFileMetadataCache? Parent;
 
 		/// <summary>
 		/// Map from file item to source file info
@@ -116,7 +118,7 @@ namespace UnrealBuildTool
 		/// <param name="Location">File to store the cache</param>
 		/// <param name="BaseDir">Base directory for files that this cache should store data for</param>
 		/// <param name="Parent">The parent cache to use</param>
-		private SourceFileMetadataCache(FileReference Location, DirectoryReference BaseDir, SourceFileMetadataCache Parent)
+		private SourceFileMetadataCache(FileReference Location, DirectoryReference BaseDir, SourceFileMetadataCache? Parent)
 		{
 			this.Location = Location;
 			this.BaseDirectory = BaseDir;
@@ -136,7 +138,7 @@ namespace UnrealBuildTool
 		/// </summary>
 		/// <param name="SourceFile">The source file to parse</param>
 		/// <returns>Text from the first include directive. Null if the file did not contain any include directives.</returns>
-		public string GetFirstInclude(FileItem SourceFile)
+		public string? GetFirstInclude(FileItem SourceFile)
 		{
 			if(Parent != null && !SourceFile.Location.IsUnderDirectory(BaseDirectory))
 			{
@@ -144,7 +146,7 @@ namespace UnrealBuildTool
 			}
 			else
 			{
-				IncludeInfo IncludeInfo;
+				IncludeInfo? IncludeInfo;
 				if(!FileToIncludeInfo.TryGetValue(SourceFile, out IncludeInfo) || SourceFile.LastWriteTimeUtc.Ticks > IncludeInfo.LastWriteTimeUtc)
 				{
 					IncludeInfo = new IncludeInfo();
@@ -170,7 +172,7 @@ namespace UnrealBuildTool
 			}
 			else
 			{
-				SourceFile Result;
+				SourceFile? Result;
 				if (!FileToSourceFile.TryGetValue(File, out Result) || File.LastWriteTimeUtc.Ticks > Result.LastWriteTimeUtc)
 				{
 					SourceFile NewSourceFile = new SourceFile(File);
@@ -214,7 +216,7 @@ namespace UnrealBuildTool
 			}
 			else
 			{
-				ReflectionInfo ReflectionInfo;
+				ReflectionInfo? ReflectionInfo;
 				if(!FileToReflectionInfo.TryGetValue(SourceFile, out ReflectionInfo) || SourceFile.LastWriteTimeUtc.Ticks > ReflectionInfo.LastWriteTimeUtc)
 				{
 					ReflectionInfo = new ReflectionInfo();
@@ -232,14 +234,14 @@ namespace UnrealBuildTool
 		/// </summary>
 		/// <param name="SourceFile">The source file to parse</param>
 		/// <returns>The first include directive</returns>
-		static string ParseFirstInclude(FileReference SourceFile)
+		static string? ParseFirstInclude(FileReference SourceFile)
 		{
 			bool bMatchImport = SourceFile.HasExtension(".m") || SourceFile.HasExtension(".mm");
 			using(StreamReader Reader = new StreamReader(SourceFile.FullName, true))
 			{
 				for(;;)
 				{
-					string Line = Reader.ReadLine();
+					string? Line = Reader.ReadLine();
 					if(Line == null)
 					{
 						return null;
@@ -268,9 +270,9 @@ namespace UnrealBuildTool
 		/// </summary>
 		/// <param name="ProjectFile">Project file for the target being built</param>
 		/// <returns>Dependency cache hierarchy for the given project</returns>
-		public static SourceFileMetadataCache CreateHierarchy(FileReference ProjectFile)
+		public static SourceFileMetadataCache CreateHierarchy(FileReference? ProjectFile)
 		{
-			SourceFileMetadataCache Cache = null;
+			SourceFileMetadataCache? Cache = null;
 
 			if(ProjectFile == null || !UnrealBuildTool.IsEngineInstalled())
 			{
@@ -284,7 +286,7 @@ namespace UnrealBuildTool
 				Cache = FindOrAddCache(ProjectCacheLocation, ProjectFile.Directory, Cache);
 			}
 
-			return Cache;
+			return Cache!;
 		}
 
 		/// <summary>
@@ -311,11 +313,11 @@ namespace UnrealBuildTool
 		/// <param name="BaseDirectory">Base directory for files that this cache should store data for</param>
 		/// <param name="Parent">The parent cache to use</param>
 		/// <returns>Reference to a dependency cache with the given settings</returns>
-		static SourceFileMetadataCache FindOrAddCache(FileReference Location, DirectoryReference BaseDirectory, SourceFileMetadataCache Parent)
+		static SourceFileMetadataCache FindOrAddCache(FileReference Location, DirectoryReference BaseDirectory, SourceFileMetadataCache? Parent)
 		{
 			lock(Caches)
 			{
-				SourceFileMetadataCache Cache;
+				SourceFileMetadataCache? Cache;
 				if(Caches.TryGetValue(Location, out Cache))
 				{
 					Debug.Assert(Cache.BaseDirectory == BaseDirectory);
diff --git a/Engine/Source/Programs/UnrealBuildTool/System/SourceFileSearch.cs b/Engine/Source/Programs/UnrealBuildTool/System/SourceFileSearch.cs
--- a/Engine/Source/Programs/UnrealBuildTool/System/SourceFileSearch.cs
+++ b/Engine/Source/Programs/UnrealBuildTool/System/SourceFileSearch.cs
@@ -38,7 +38,7 @@ namespace UnrealBuildTool
 		};
 
 		/// Finds mouse source files
-		public static List<FileReference> FindModuleSourceFiles(FileReference ModuleRulesFile, bool SearchSubdirectories = true, HashSet<DirectoryReference> SearchedDirectories = null)
+		public static List<FileReference> FindModuleSourceFiles(FileReference ModuleRulesFile, bool SearchSubdirectories = true, HashSet<DirectoryReference>? SearchedDirectories = null)
 		{
 			// The module's "base directory" is simply the directory where its xxx.Build.cs file is stored.  We'll always
 			// harvest source files for this module in this base directory directory and all of its sub-directories.
@@ -52,7 +52,7 @@ namespace UnrealBuildTool
 		/// <param name="SubdirectoryNamesToExclude">Directory base names to ignore when searching subdirectories.  Can be null.</param>
 		/// <param name="SearchSubdirectories">True to include subdirectories, otherwise we only search the list of base directories</param>
 		/// <param name="SearchedDirectories">If non-null, is updated with a list of the directories searched</param>
-		public static List<FileReference> FindFiles(DirectoryReference DirectoryToSearch, List<string> SubdirectoryNamesToExclude = null, bool SearchSubdirectories = true, HashSet<DirectoryReference> SearchedDirectories = null)
+		public static List<FileReference> FindFiles(DirectoryReference DirectoryToSearch, List<string>? SubdirectoryNamesToExclude = null, bool SearchSubdirectories = true, HashSet<DirectoryReference>? SearchedDirectories = null)
 		{
 			// Build a list of directory names that we don't want to search under. We always ignore intermediate directories.
 			string[] ExcludedDirectorySuffixes;
@@ -78,7 +78,7 @@ namespace UnrealBuildTool
 			return FoundFiles;
 		}
 
-		static void FindFilesInternal(DirectoryReference Directory, string[] ExcludedDirectorySuffixes, List<FileReference> FoundFiles, HashSet<DirectoryReference> SearchedDirectories)
+		static void FindFilesInternal(DirectoryReference Directory, string[] ExcludedDirectorySuffixes, List<FileReference> FoundFiles, HashSet<DirectoryReference>? SearchedDirectories)
 		{
 			if(SearchedDirectories != null)
 			{
@@ -94,7 +94,7 @@ namespace UnrealBuildTool
 			}
 		}
 
-		static void FindFilesInternalRecursive(DirectoryReference Directory, string[] ExcludedDirectorySuffixes, List<FileReference> FoundFiles, HashSet<DirectoryReference> SearchedDirectories)
+		static void FindFilesInternalRecursive(DirectoryReference Directory, string[] ExcludedDirectorySuffixes, List<FileReference> FoundFiles, HashSet<DirectoryReference>? SearchedDirectories)
 		{
 			FindFilesInternal(Directory, ExcludedDirectorySuffixes, FoundFiles, SearchedDirectories);
 
diff --git a/Engine/Source/Programs/UnrealBuildTool/System/SourceFileWorkingSet.cs b/Engine/Source/Programs/UnrealBuildTool/System/SourceFileWorkingSet.cs
--- a/Engine/Source/Programs/UnrealBuildTool/System/SourceFileWorkingSet.cs
+++ b/Engine/Source/Programs/UnrealBuildTool/System/SourceFileWorkingSet.cs
@@ -3,6 +3,7 @@
 using System;
 using System.Collections.Generic;
 using System.Diagnostics;
+using System.Diagnostics.CodeAnalysis;
 using System.IO;
 using System.Linq;
 using System.Text;
@@ -90,11 +91,11 @@ namespace UnrealBuildTool
 	class GitSourceFileWorkingSet : ISourceFileWorkingSet
 	{
 		DirectoryReference RootDir;
-		Process BackgroundProcess;
+		Process? BackgroundProcess;
 		HashSet<FileReference> Files;
 		List<DirectoryReference> Directories;
 		List<string> ErrorOutput;
-		GitSourceFileWorkingSet Inner;
+		GitSourceFileWorkingSet? Inner;
 
 		/// <summary>
 		/// Constructor
@@ -102,7 +103,7 @@ namespace UnrealBuildTool
 		/// <param name="GitPath">Path to the Git executable</param>
 		/// <param name="RootDir">Root directory to run queries from (typically the directory containing the .git folder, to ensure all subfolders can be searched)</param>
 		/// <param name="Inner">An inner working set. This allows supporting multiple Git repositories (one containing the engine, another containing the project, for example)</param>
-		public GitSourceFileWorkingSet(string GitPath, DirectoryReference RootDir, GitSourceFileWorkingSet Inner)
+		public GitSourceFileWorkingSet(string GitPath, DirectoryReference RootDir, GitSourceFileWorkingSet? Inner)
 		{
 			this.RootDir = RootDir;
 			this.Files = new HashSet<FileReference>();
@@ -300,7 +301,7 @@ namespace UnrealBuildTool
 		/// Sets the path to use for the repository. Interpreted relative to the Unreal Engine root directory (the folder above the Engine folder) -- if relative.
 		/// </summary>
 		[XmlConfigFile]
-		public static string RepositoryPath = null;
+		public static string? RepositoryPath = null;
 
 		/// <summary>
 		/// Sets the path to use for the Git executable. Defaults to "git" (assuming it is in the PATH).
@@ -322,7 +323,7 @@ namespace UnrealBuildTool
 			}
 			else if (Provider == ProviderType.Git)
 			{
-				GitSourceFileWorkingSet WorkingSet;
+				GitSourceFileWorkingSet? WorkingSet;
 				if (!String.IsNullOrEmpty(RepositoryPath))
 				{
 					WorkingSet = new GitSourceFileWorkingSet(GitPath, DirectoryReference.Combine(RootDir, RepositoryPath), null);
@@ -339,7 +340,7 @@ namespace UnrealBuildTool
 			}
 			else
 			{
-				GitSourceFileWorkingSet WorkingSet;
+				GitSourceFileWorkingSet? WorkingSet;
 				if(TryCreateGitWorkingSet(RootDir, ProjectDirs, out WorkingSet))
 				{
 					return WorkingSet;
@@ -351,9 +352,9 @@ namespace UnrealBuildTool
 			}
 		}
 
-		static bool TryCreateGitWorkingSet(DirectoryReference RootDir, IEnumerable<DirectoryReference> ProjectDirs, out GitSourceFileWorkingSet OutWorkingSet)
+		static bool TryCreateGitWorkingSet(DirectoryReference RootDir, IEnumerable<DirectoryReference> ProjectDirs, [NotNullWhen(true)] out GitSourceFileWorkingSet? OutWorkingSet)
 		{
-			GitSourceFileWorkingSet WorkingSet  = null;
+			GitSourceFileWorkingSet? WorkingSet  = null;
 
 			// Create the working set for the engine directory
 			if (DirectoryReference.Exists(DirectoryReference.Combine(RootDir, ".git")))
@@ -379,7 +380,7 @@ namespace UnrealBuildTool
 
 			// Set the output value
 			OutWorkingSet = WorkingSet;
-			return WorkingSet != null;
+			return OutWorkingSet != null;
 		}
 	}
 }
diff --git a/Engine/Source/Programs/UnrealBuildTool/System/TargetMakefile.cs b/Engine/Source/Programs/UnrealBuildTool/System/TargetMakefile.cs
--- a/Engine/Source/Programs/UnrealBuildTool/System/TargetMakefile.cs
+++ b/Engine/Source/Programs/UnrealBuildTool/System/TargetMakefile.cs
@@ -81,17 +81,17 @@ namespace UnrealBuildTool
 		/// <summary>
 		/// The array of command-line arguments. The makefile will be invalidated whenever these change.
 		/// </summary>
-		public string[] AdditionalArguments;
+		public string[]? AdditionalArguments;
 
 		/// <summary>
 		/// Scripts which should be run before building anything
 		/// </summary>
-		public FileReference[] PreBuildScripts;
+		public FileReference[] PreBuildScripts = Array.Empty<FileReference>();
 
 		/// <summary>
 		/// Targets which should be build before building anything
 		/// </summary>
-		public TargetInfo[] PreBuildTargets;
+		public TargetInfo[] PreBuildTargets = Array.Empty<TargetInfo>();
 
 		/// <summary>
 		/// Every action in the action graph
@@ -303,7 +303,7 @@ namespace UnrealBuildTool
 		/// <param name="Arguments">Command line arguments for this target</param>
 		/// <param name="ReasonNotLoaded">If the function returns null, this string will contain the reason why</param>
 		/// <returns>The loaded makefile, or null if it failed for some reason.  On failure, the 'ReasonNotLoaded' variable will contain information about why</returns>
-		public static TargetMakefile Load(FileReference MakefilePath, FileReference ProjectFile, UnrealTargetPlatform Platform, string[] Arguments, out string ReasonNotLoaded)
+		public static TargetMakefile? Load(FileReference MakefilePath, FileReference? ProjectFile, UnrealTargetPlatform Platform, string[] Arguments, out string? ReasonNotLoaded)
 		{
 			FileInfo MakefileInfo;
 			using(Timeline.ScopeEvent("Checking dependent timestamps"))
@@ -464,7 +464,7 @@ namespace UnrealBuildTool
 		/// <param name="WorkingSet">The current working set of source files</param>
 		/// <param name="ReasonNotLoaded">If the makefile is not valid, is set to a message describing why</param>
 		/// <returns>True if the makefile is valid, false otherwise</returns>
-		public static bool IsValidForSourceFiles(TargetMakefile Makefile, FileReference ProjectFile, UnrealTargetPlatform Platform, ISourceFileWorkingSet WorkingSet, out string ReasonNotLoaded)
+		public static bool IsValidForSourceFiles(TargetMakefile Makefile, FileReference? ProjectFile, UnrealTargetPlatform Platform, ISourceFileWorkingSet WorkingSet, out string? ReasonNotLoaded)
 		{
 			using(Timeline.ScopeEvent("TargetMakefile.IsValidForSourceFiles()"))
 			{
@@ -673,7 +673,7 @@ namespace UnrealBuildTool
 		/// <param name="Platform">The platform that the target is being built for</param>
 		/// <param name="Configuration">The configuration being built</param>
 		/// <returns>Path to the makefile</returns>
-		public static FileReference GetLocation(FileReference ProjectFile, string TargetName, UnrealTargetPlatform Platform, UnrealTargetConfiguration Configuration)
+		public static FileReference GetLocation(FileReference? ProjectFile, string TargetName, UnrealTargetPlatform Platform, UnrealTargetConfiguration Configuration)
 		{
 			DirectoryReference BaseDirectory = DirectoryReference.FromFile(ProjectFile) ?? UnrealBuildTool.EngineDirectory;
 			return FileReference.Combine(BaseDirectory, "Intermediate", "Build", Platform.ToString(), TargetName, Configuration.ToString(), "Makefile.bin");
