UBT: Convert UnrealBuildTool to use nullable references.

diff --git a/Engine/Source/Programs/UnrealBuildTool/System/TargetReceipt.cs b/Engine/Source/Programs/UnrealBuildTool/System/TargetReceipt.cs
--- a/Engine/Source/Programs/UnrealBuildTool/System/TargetReceipt.cs
+++ b/Engine/Source/Programs/UnrealBuildTool/System/TargetReceipt.cs
@@ -9,6 +9,7 @@ using System.Xml.Serialization;
 using System.ComponentModel;
 using System.Runtime.Serialization;
 using Tools.DotNETCommon;
+using System.Diagnostics.CodeAnalysis;
 
 namespace UnrealBuildTool
 {
@@ -74,6 +75,7 @@ namespace UnrealBuildTool
 		/// </summary>
 		private BuildProduct()
 		{
+			Path = null!;
 		}
 
 		/// <summary>
@@ -154,6 +156,7 @@ namespace UnrealBuildTool
 		/// </summary>
 		private RuntimeDependency()
 		{
+			Path = null!;
 		}
 
 		/// <summary>
@@ -258,12 +261,12 @@ namespace UnrealBuildTool
 		/// <summary>
 		/// Path to the project file for this target
 		/// </summary>
-		public FileReference ProjectFile;
+		public FileReference? ProjectFile;
 
 		/// <summary>
 		/// The project directory
 		/// </summary>
-		public DirectoryReference ProjectDir;
+		public DirectoryReference? ProjectDir;
 
 		/// <summary>
 		/// The name of this target
@@ -298,7 +301,7 @@ namespace UnrealBuildTool
 		/// <summary>
 		/// The exectuable to launch for this target
 		/// </summary>
-		public FileReference Launch;
+		public FileReference? Launch;
 
 		/// <summary>
 		/// The build products which are part of this target
@@ -320,13 +323,6 @@ namespace UnrealBuildTool
 		/// </summary>
 		public List<ReceiptProperty> AdditionalProperties = new List<ReceiptProperty>();
 
-		/// <summary>
-		/// Default constructor
-		/// </summary>
-		public TargetReceipt()
-		{
-		}
-
 		/// <summary>
 		/// Constructor
 		/// </summary>
@@ -337,7 +333,7 @@ namespace UnrealBuildTool
 		/// <param name="InConfiguration">Configuration of the target being compiled</param>
 		/// <param name="InVersion">Version information for the target</param>
 		/// <param name="InArchitecture">Architecture information for the target</param>
-		public TargetReceipt(FileReference InProjectFile, string InTargetName, TargetType InTargetType, UnrealTargetPlatform InPlatform, UnrealTargetConfiguration InConfiguration, BuildVersion InVersion, string InArchitecture)
+		public TargetReceipt(FileReference? InProjectFile, string InTargetName, TargetType InTargetType, UnrealTargetPlatform InPlatform, UnrealTargetConfiguration InConfiguration, BuildVersion InVersion, string InArchitecture)
 		{
 			ProjectFile = InProjectFile;
 			ProjectDir = DirectoryReference.FromFile(InProjectFile);
@@ -349,31 +345,6 @@ namespace UnrealBuildTool
 			Architecture = InArchitecture;
 		}
 
-		/// <summary>
-		/// Copy constructor
-		/// </summary>
-		/// <param name="Other">Receipt to copy from</param>
-		public TargetReceipt(TargetReceipt Other)
-		{
-			Launch = Other.Launch;
-			foreach (BuildProduct OtherBuildProduct in Other.BuildProducts)
-			{
-				BuildProducts.Add(new BuildProduct(OtherBuildProduct));
-			}
-			foreach (RuntimeDependency OtherRuntimeDependency in Other.RuntimeDependencies)
-			{
-				RuntimeDependencies.Add(new RuntimeDependency(OtherRuntimeDependency));
-			}
-			foreach (KeyValuePair<string, bool> Pair in Other.PluginNameToEnabledState)
-			{
-				if (!PluginNameToEnabledState.ContainsKey(Pair.Key))
-				{
-					PluginNameToEnabledState.Add(Pair.Key, Pair.Value);
-				}
-			}
-			AdditionalProperties.AddRange(Other.AdditionalProperties);
-		}
-
 		/// <summary>
 		/// Adds a build product to the receipt. Does not check whether it already exists.
 		/// </summary>
@@ -420,7 +391,7 @@ namespace UnrealBuildTool
 		/// <param name="EngineDir">Value of the $(EngineDir) variable.</param>
 		/// <param name="ProjectDir">Value of the $(ProjectDir) variable.</param>
 		/// <returns>Converted path for the file.</returns>
-		static string InsertPathVariables(FileReference File, DirectoryReference EngineDir, DirectoryReference ProjectDir)
+		static string InsertPathVariables(FileReference File, DirectoryReference EngineDir, DirectoryReference? ProjectDir)
 		{
 			if (File.IsUnderDirectory(EngineDir))
 			{
@@ -443,7 +414,7 @@ namespace UnrealBuildTool
 		/// <param name="EngineDir">Value of the $(EngineDir) variable.</param>
 		/// <param name="ProjectDir">Value of the $(ProjectDir) variable.</param>
 		/// <returns>Converted path for the file.</returns>
-		static FileReference ExpandPathVariables(string Path, DirectoryReference EngineDir, DirectoryReference ProjectDir)
+		static FileReference ExpandPathVariables(string Path, DirectoryReference EngineDir, DirectoryReference? ProjectDir)
 		{
 			const string EnginePrefix = "$(EngineDir)";
 			if(Path.StartsWith(EnginePrefix, StringComparison.InvariantCultureIgnoreCase))
@@ -549,14 +520,14 @@ namespace UnrealBuildTool
 			UnrealTargetConfiguration Configuration = RawObject.GetEnumField<UnrealTargetConfiguration>("Configuration");
 
 			// Try to read the build version
-			BuildVersion Version;
+			BuildVersion? Version;
 			if (!BuildVersion.TryParse(RawObject.GetObjectField("Version"), out Version))
 			{
 				throw new JsonParseException("Invalid 'Version' field");
 			}
 
 			// Read the project path
-			FileReference ProjectFile;
+			FileReference? ProjectFile;
 
 			string RelativeProjectFile;
 			if(RawObject.TryGetStringField("Project", out RelativeProjectFile))
@@ -580,7 +551,7 @@ namespace UnrealBuildTool
 			TargetReceipt Receipt = new TargetReceipt(ProjectFile, TargetName, TargetType, Platform, Configuration, Version, Architecture);
 
 			// Get the project directory
-			DirectoryReference ProjectDir = Receipt.ProjectDir;
+			DirectoryReference? ProjectDir = Receipt.ProjectDir;
 
 			// Read the launch executable
 			string Launch;
@@ -682,7 +653,7 @@ namespace UnrealBuildTool
 		/// <param name="Location">Filename to read from</param>
 		/// <param name="Receipt">If successful, the receipt that was read</param>
 		/// <returns>True if successful</returns>
-		public static bool TryRead(FileReference Location, out TargetReceipt Receipt)
+		public static bool TryRead(FileReference Location, [NotNullWhen(true)] out TargetReceipt? Receipt)
 		{
 			return TryRead(Location, UnrealBuildTool.EngineDirectory, out Receipt);
 		}
@@ -694,7 +665,7 @@ namespace UnrealBuildTool
 		/// <param name="EngineDir">Engine directory for expanded paths</param>
 		/// <param name="Receipt">If successful, the receipt that was read</param>
 		/// <returns>True if successful</returns>
-		public static bool TryRead(FileReference Location, DirectoryReference EngineDir, out TargetReceipt Receipt)
+		public static bool TryRead(FileReference Location, DirectoryReference EngineDir, [NotNullWhen(true)] out TargetReceipt? Receipt)
 		{
 			if (!FileReference.Exists(Location))
 			{
diff --git a/Engine/Source/Programs/UnrealBuildTool/System/Timeline.cs b/Engine/Source/Programs/UnrealBuildTool/System/Timeline.cs
--- a/Engine/Source/Programs/UnrealBuildTool/System/Timeline.cs
+++ b/Engine/Source/Programs/UnrealBuildTool/System/Timeline.cs
@@ -160,7 +160,7 @@ namespace UnrealBuildTool
 				for (; OuterEvents.Count > 1; OuterEvents.RemoveAt(OuterEvents.Count - 1))
 				{
 					Event OuterEvent = OuterEvents.Last();
-					if (Event.StartTime < OuterEvent.FinishTime.Value)
+					if (Event.StartTime < OuterEvent.FinishTime!.Value)
 					{
 						break;
 					}
@@ -190,7 +190,7 @@ namespace UnrealBuildTool
 			// Remove everything from the stack
 			for(; OuterEvents.Count > 0; OuterEvents.RemoveAt(OuterEvents.Count - 1))
 			{
-				UpdateLastEventTime(ref LastTime, OuterEvents.Last().FinishTime.Value, MaxUnknownTime, OuterEvents, Verbosity);
+				UpdateLastEventTime(ref LastTime, OuterEvents.Last().FinishTime!.Value, MaxUnknownTime, OuterEvents, Verbosity);
 			}
 
 			// Print the finish time
diff --git a/Engine/Source/Programs/UnrealBuildTool/System/TimingData.cs b/Engine/Source/Programs/UnrealBuildTool/System/TimingData.cs
--- a/Engine/Source/Programs/UnrealBuildTool/System/TimingData.cs
+++ b/Engine/Source/Programs/UnrealBuildTool/System/TimingData.cs
@@ -55,8 +55,10 @@ namespace UnrealBuildTool
 		/// <summary>
 		/// Constructs a new instance of <see cref="TimingData"/>.
 		/// </summary>
-		public TimingData()
+		public TimingData(string Name, TimingDataType Type)
 		{
+			this.Name = Name;
+			this.Type = Type;
 		}
 
 		/// <summary>
@@ -105,7 +107,7 @@ namespace UnrealBuildTool
 		/// <summary>
 		/// If this event is a sub-event, the parent is the event it is the sub-event to.
 		/// </summary>
-		public TimingData Parent { get; set; }
+		public TimingData? Parent { get; set; }
 
 		/// <summary>
 		/// Any sub-events of this event.
@@ -119,7 +121,7 @@ namespace UnrealBuildTool
 		/// <param name="ChildData">The sub-event to add to this event.</param>
 		public void AddChild(TimingData ChildData)
 		{
-			TimingData MatchingData;
+			TimingData? MatchingData;
 			if (Children.TryGetValue(ChildData.Name, out MatchingData))
 			{
 				MatchingData.ExclusiveDuration += ChildData.ExclusiveDuration;
@@ -140,10 +142,8 @@ namespace UnrealBuildTool
 		/// <returns>A deep clone of this event.</returns>
 		public TimingData Clone()
 		{
-			TimingData ClonedTimingData = new TimingData()
+			TimingData ClonedTimingData = new TimingData(Name, Type)
 			{
-				Name = Name,
-				Type = Type,
 				ExclusiveDuration = ExclusiveDuration,
 			};
 
diff --git a/Engine/Source/Programs/UnrealBuildTool/System/Unity.cs b/Engine/Source/Programs/UnrealBuildTool/System/Unity.cs
--- a/Engine/Source/Programs/UnrealBuildTool/System/Unity.cs
+++ b/Engine/Source/Programs/UnrealBuildTool/System/Unity.cs
@@ -134,8 +134,8 @@ namespace UnrealBuildTool
 				List<FileCollection> Result = UnityFiles;
 
 				// Null everything to ensure that failure will occur if you accidentally reuse this object.
-				CurrentCollection = null;
-				UnityFiles = null;
+				CurrentCollection = null!;
+				UnityFiles = null!;
 
 				return Result;
 			}
diff --git a/Engine/Source/Programs/UnrealBuildTool/System/Utils.cs b/Engine/Source/Programs/UnrealBuildTool/System/Utils.cs
--- a/Engine/Source/Programs/UnrealBuildTool/System/Utils.cs
+++ b/Engine/Source/Programs/UnrealBuildTool/System/Utils.cs
@@ -106,7 +106,7 @@ namespace UnrealBuildTool
 		/// <param name="AdditionalVariables">Lookup of variable names to values</param>
 		/// <param name="bUseAdditionalVariablesOnly">If true, then Environment.GetEnvironmentVariable will not be used if the var is not found in AdditionalVariables</param>
 		/// <returns>String with all variables replaced</returns>
-		public static string ExpandVariables(string InputString, Dictionary<string, string> AdditionalVariables = null, bool bUseAdditionalVariablesOnly = false)
+		public static string ExpandVariables(string InputString, Dictionary<string, string>? AdditionalVariables = null, bool bUseAdditionalVariablesOnly = false)
 		{
 			string Result = InputString;
 			for (int Idx = Result.IndexOf("$("); Idx != -1; Idx = Result.IndexOf("$(", Idx))
@@ -122,7 +122,7 @@ namespace UnrealBuildTool
 				string Name = Result.Substring(Idx + 2, EndIdx - (Idx + 2));
 
 				// Find the value for it, either from the dictionary or the environment block
-				string Value = null;
+				string? Value = null;
 				if (AdditionalVariables == null || !AdditionalVariables.TryGetValue(Name, out Value))
 				{
 					if (bUseAdditionalVariablesOnly == false)
@@ -199,7 +199,7 @@ namespace UnrealBuildTool
 		/// <returns>the value of the environment variable if found and the default value if missing</returns>
 		public static bool GetEnvironmentVariable(string VarName, bool bDefault)
 		{
-			string Value = Environment.GetEnvironmentVariable(VarName);
+			string? Value = Environment.GetEnvironmentVariable(VarName);
 			if (Value != null)
 			{
 				// Convert the string to its boolean value
@@ -216,7 +216,7 @@ namespace UnrealBuildTool
 		/// <returns>the value of the environment variable if found and the default value if missing</returns>
 		public static string GetStringEnvironmentVariable(string VarName, string Default)
 		{
-			string Value = Environment.GetEnvironmentVariable(VarName);
+			string? Value = Environment.GetEnvironmentVariable(VarName);
 			if (Value != null)
 			{
 				return Value;
@@ -232,7 +232,7 @@ namespace UnrealBuildTool
 		/// <returns>the value of the environment variable if found and the default value if missing</returns>
 		public static double GetEnvironmentVariable(string VarName, double Default)
 		{
-			string Value = Environment.GetEnvironmentVariable(VarName);
+			string? Value = Environment.GetEnvironmentVariable(VarName);
 			if (Value != null)
 			{
 				return Convert.ToDouble(Value);
@@ -248,7 +248,7 @@ namespace UnrealBuildTool
 		/// <returns>the value of the environment variable if found and the default value if missing</returns>
 		public static string GetEnvironmentVariable(string VarName, string Default)
 		{
-			string Value = Environment.GetEnvironmentVariable(VarName);
+			string? Value = Environment.GetEnvironmentVariable(VarName);
 			if (Value != null)
 			{
 				return Value;
@@ -431,7 +431,7 @@ namespace UnrealBuildTool
 		/// <param name="Args">Arguments to Command</param>
 		/// <param name="ExitCode">The return code from the process after it exits</param>
 		/// <param name="LogOutput">Whether to also log standard output and standard error</param>
-		public static string RunLocalProcessAndReturnStdOut(string Command, string Args, out int ExitCode, bool LogOutput = false)
+		public static string RunLocalProcessAndReturnStdOut(string Command, string? Args, out int ExitCode, bool LogOutput = false)
 		{
 			// Process Arguments follow windows conventions in .NET Core
 			// Which means single quotes ' are not considered quotes.
@@ -592,7 +592,7 @@ namespace UnrealBuildTool
 		/// <returns>File path with consistent separators</returns>
 		public static string CleanDirectorySeparators(string FilePath, char UseDirectorySeparatorChar = '\0')
 		{
-			StringBuilder CleanPath = null;
+			StringBuilder? CleanPath = null;
 			if (UseDirectorySeparatorChar == '\0')
 			{
 				UseDirectorySeparatorChar = Environment.OSVersion.Platform == PlatformID.Unix ? '/' : '\\';
@@ -837,7 +837,7 @@ namespace UnrealBuildTool
 		static public T ReadClass<T>(string FileName) where T : new()
 		{
 			T Instance = new T();
-			StreamReader XmlStream = null;
+			StreamReader? XmlStream = null;
 			try
 			{
 				// Get the XML data stream to read from
@@ -879,7 +879,7 @@ namespace UnrealBuildTool
 		static public bool WriteClass<T>(T Data, string FileName, string DefaultNameSpace)
 		{
 			bool bSuccess = true;
-			StreamWriter XmlStream = null;
+			StreamWriter? XmlStream = null;
 			try
 			{
 				FileInfo Info = new FileInfo(FileName);
@@ -976,7 +976,7 @@ namespace UnrealBuildTool
 		{
 			if (!String.IsNullOrEmpty(Path.GetExtension(Filename)))
 			{
-				return Path.Combine(Path.GetDirectoryName(Filename), Path.GetFileNameWithoutExtension(Filename));
+				return Path.Combine(Path.GetDirectoryName(Filename)!, Path.GetFileNameWithoutExtension(Filename));
 			}
 			else
 			{
@@ -1017,7 +1017,7 @@ namespace UnrealBuildTool
 		/// NOTE: This function may return null. Some accounts (eg. the SYSTEM account on Windows) do not have a personal folder, and Jenkins
 		/// runs using this account by default.
 		/// </summary>
-		public static DirectoryReference GetUserSettingDirectory()
+		public static DirectoryReference? GetUserSettingDirectory()
 		{
 			if (BuildHostPlatform.Current.Platform == UnrealTargetPlatform.Mac)
 			{
diff --git a/Engine/Source/Programs/UnrealBuildTool/System/VersionNumber.cs b/Engine/Source/Programs/UnrealBuildTool/System/VersionNumber.cs
--- a/Engine/Source/Programs/UnrealBuildTool/System/VersionNumber.cs
+++ b/Engine/Source/Programs/UnrealBuildTool/System/VersionNumber.cs
@@ -2,6 +2,7 @@
 
 using System;
 using System.Collections.Generic;
+using System.Diagnostics.CodeAnalysis;
 using System.Linq;
 using System.Text;
 using System.Threading.Tasks;
@@ -52,10 +53,10 @@ namespace UnrealBuildTool
 		/// </summary>
 		/// <param name="Obj">Object to compare against</param>
 		/// <returns>True if the objects are equal, false otherwise.</returns>
-		public override bool Equals(object Obj)
+		public override bool Equals(object? Obj)
 		{
-			VersionNumber Version = Obj as VersionNumber;
-			return Version != null && this == Version;
+			VersionNumber? Version = Obj as VersionNumber;
+			return !ReferenceEquals(Version, null) && this == Version;
 		}
 
 		/// <summary>
@@ -150,9 +151,9 @@ namespace UnrealBuildTool
 		/// </summary>
 		/// <param name="Other">Other version number to compare to</param>
 		/// <returns>A negative value if this version is before Other, a positive value if this version is after Other, and zero otherwise.</returns>
-		public int CompareTo(VersionNumber Other)
+		public int CompareTo(VersionNumber? Other)
 		{
-			return Compare(this, Other);
+			return ReferenceEquals(Other, null)? 1 : Compare(this, Other);
 		}
 
 		/// <summary>
@@ -211,7 +212,7 @@ namespace UnrealBuildTool
 		/// <param name="Text">The string to parse</param>
 		/// <param name="OutNumber">Variable to receive the parsed version number</param>
 		/// <returns>A version number object</returns>
-		public static bool TryParse(string Text, out VersionNumber OutNumber)
+		public static bool TryParse(string Text, [NotNullWhen(true)] out VersionNumber? OutNumber)
 		{
 			List<int> Components = new List<int>();
 			foreach(string TextElement in Text.Split(Delimiters))
diff --git a/Engine/Source/Programs/UnrealBuildTool/System/XmlConfig.cs b/Engine/Source/Programs/UnrealBuildTool/System/XmlConfig.cs
--- a/Engine/Source/Programs/UnrealBuildTool/System/XmlConfig.cs
+++ b/Engine/Source/Programs/UnrealBuildTool/System/XmlConfig.cs
@@ -2,6 +2,7 @@
 
 using System;
 using System.Collections.Generic;
+using System.Diagnostics.CodeAnalysis;
 using System.IO;
 using System.Linq;
 using System.Reflection;
@@ -33,28 +34,39 @@ namespace UnrealBuildTool
 			/// Which folder to display the config file under in the generated project files
 			/// </summary>
 			public string FolderName;
+
+			/// <summary>
+			/// Constructor
+			/// </summary>
+			/// <param name="Location"></param>
+			/// <param name="FolderName"></param>
+			public InputFile(FileReference Location, string FolderName)
+			{
+				this.Location = Location;
+				this.FolderName = FolderName;
+			}
 		}
 
 		/// <summary>
 		/// The cache file that is being used
 		/// </summary>
-		public static FileReference CacheFile;
+		public static FileReference? CacheFile;
 
 		/// <summary>
 		/// Parsed config values
 		/// </summary>
-		static XmlConfigData Values;
+		static XmlConfigData? Values;
 
 		/// <summary>
 		/// Cached serializer for the XML schema
 		/// </summary>
-		static XmlSerializer CachedSchemaSerializer;
+		static XmlSerializer? CachedSchemaSerializer;
 
 		/// <summary>
 		/// Initialize the config system with the given types
 		/// </summary>
 		/// <param name="OverrideCacheFile">Force use of the cached XML config without checking if it's valid (useful for remote builds)</param>
-		public static void ReadConfigFiles(FileReference OverrideCacheFile)
+		public static void ReadConfigFiles(FileReference? OverrideCacheFile)
 		{
 			// Find all the configurable types
 			List<Type> ConfigTypes = FindConfigurableTypes();
@@ -77,7 +89,7 @@ namespace UnrealBuildTool
 				CacheFile = FileReference.Combine(UnrealBuildTool.EngineDirectory, "Intermediate", "Build", "XmlConfigCache.bin");
 				if(UnrealBuildTool.IsEngineInstalled())
 				{
-					DirectoryReference UserSettingsDir = Utils.GetUserSettingDirectory();
+					DirectoryReference? UserSettingsDir = Utils.GetUserSettingDirectory();
 					if(UserSettingsDir != null)
 					{
 						CacheFile = FileReference.Combine(UserSettingsDir, "UnrealEngine", String.Format("XmlConfigCache-{0}.bin", UnrealBuildTool.RootDirectory.FullName.Replace(":", "").Replace(Path.DirectorySeparatorChar, '+')));
@@ -91,7 +103,7 @@ namespace UnrealBuildTool
 				FileReference SchemaFile = GetSchemaLocation();
 
 				// Try to read the existing cache from disk
-				XmlConfigData CachedValues;
+				XmlConfigData? CachedValues;
 				if(IsCacheUpToDate(CacheFile, InputFiles) && FileReference.Exists(SchemaFile))
 				{
 					if(XmlConfigData.TryRead(CacheFile, ConfigTypes, out CachedValues) && Enumerable.SequenceEqual(InputFiles, CachedValues.InputFiles))
@@ -166,7 +178,7 @@ namespace UnrealBuildTool
 			}
 			catch (ReflectionTypeLoadException Ex)
 			{
-				Console.WriteLine("TypeLoadException: {0}", string.Join("\n", Ex.LoaderExceptions.Select(x => x.Message)));
+				Console.WriteLine("TypeLoadException: {0}", string.Join("\n", Ex.LoaderExceptions.Select(x => x?.Message)));
 				throw;
 			}
 			return ConfigTypes;
@@ -216,7 +228,7 @@ namespace UnrealBuildTool
  				FileReference NotForLicenseesConfigLocation = FileReference.Combine(UnrealBuildTool.EngineDirectory, "Restricted", "NotForLicensees", "Programs", "UnrealBuildTool", "BuildConfiguration.xml");
  				if(FileReference.Exists(NotForLicenseesConfigLocation))
  				{
- 					InputFiles.Add(new InputFile { Location = NotForLicenseesConfigLocation, FolderName = "NotForLicensees" });
+ 					InputFiles.Add(new InputFile(NotForLicenseesConfigLocation, "NotForLicensees"));
  				}
 
 				// Check for the user config file under /Engine/Saved/UnrealBuildTool
@@ -225,7 +237,7 @@ namespace UnrealBuildTool
 				{
 					CreateDefaultConfigFile(UserConfigLocation);
 				}
-				InputFiles.Add(new InputFile { Location = UserConfigLocation, FolderName = "User" });
+				InputFiles.Add(new InputFile(UserConfigLocation, "User"));
 			}
 
 			// Check for the global config file under AppData/Unreal Engine/UnrealBuildTool
@@ -237,7 +249,7 @@ namespace UnrealBuildTool
 				{
 					CreateDefaultConfigFile(AppDataConfigLocation);
 				}
-				InputFiles.Add(new InputFile { Location = AppDataConfigLocation, FolderName = "Global (AppData)" });
+				InputFiles.Add(new InputFile(AppDataConfigLocation, "Global (AppData)"));
 			}
 
 			// Check for the global config file under My Documents/Unreal Engine/UnrealBuildTool
@@ -247,7 +259,7 @@ namespace UnrealBuildTool
 				FileReference PersonalConfigLocation = FileReference.Combine(new DirectoryReference(PersonalFolder), "Unreal Engine", "UnrealBuildTool", "BuildConfiguration.xml");
 				if(FileReference.Exists(PersonalConfigLocation))
 				{
-					InputFiles.Add(new InputFile { Location = PersonalConfigLocation, FolderName = "Global (Documents)" });
+					InputFiles.Add(new InputFile(PersonalConfigLocation, "Global (Documents)"));
 				}
 			}
 
@@ -275,10 +287,10 @@ namespace UnrealBuildTool
 		/// <param name="TargetObject">The object instance to be configured</param>
 		public static void ApplyTo(object TargetObject)
 		{
-			for(Type TargetType = TargetObject.GetType(); TargetType != null; TargetType = TargetType.BaseType)
+			for(Type? TargetType = TargetObject.GetType(); TargetType != null; TargetType = TargetType.BaseType)
 			{
-				KeyValuePair<FieldInfo, object>[] FieldValues;
-				if(Values.TypeToValues.TryGetValue(TargetType, out FieldValues))
+				KeyValuePair<FieldInfo, object>[]? FieldValues;
+				if(Values!.TypeToValues.TryGetValue(TargetType, out FieldValues))
 				{
 					foreach(KeyValuePair<FieldInfo, object> FieldValuePair in FieldValues)
 					{
@@ -317,11 +329,11 @@ namespace UnrealBuildTool
 		/// <param name="Name">Name of the field to receive</param>
 		/// <param name="Value">On success, receives the value of the field</param>
 		/// <returns>True if the value was read, false otherwise</returns>
-		public static bool TryGetValue(Type TargetType, string Name, out object Value)
+		public static bool TryGetValue(Type TargetType, string Name, [NotNullWhen(true)] out object? Value)
 		{
 			// Find all the config values for this type
-			KeyValuePair<FieldInfo, object>[] FieldValues;
-			if(!Values.TypeToValues.TryGetValue(TargetType, out FieldValues))
+			KeyValuePair<FieldInfo, object>[]? FieldValues;
+			if(!Values!.TypeToValues.TryGetValue(TargetType, out FieldValues))
 			{
 				Value = null;
 				return false;
@@ -358,7 +370,7 @@ namespace UnrealBuildTool
 					{
 						string CategoryName = Attribute.Category ?? ConfigType.Name;
 
-						Dictionary<string, FieldInfo> NameToField;
+						Dictionary<string, FieldInfo>? NameToField;
 						if(!CategoryToFields.TryGetValue(CategoryName, out NameToField))
 						{
 							NameToField = new Dictionary<string, FieldInfo>();
@@ -544,7 +556,7 @@ namespace UnrealBuildTool
 		static bool TryReadFile(FileReference Location, Dictionary<string, Dictionary<string, FieldInfo>> CategoryToFields, Dictionary<Type, Dictionary<FieldInfo, object>> TypeToValues, XmlSchema Schema)
 		{
 			// Read the XML file, and validate it against the schema
-			XmlConfigFile ConfigFile;
+			XmlConfigFile? ConfigFile;
 			if(!XmlConfigFile.TryRead(Location, Schema, out ConfigFile))
 			{
 				return false;
@@ -553,12 +565,12 @@ namespace UnrealBuildTool
 			// Parse the document
 			foreach(XmlElement CategoryElement in ConfigFile.DocumentElement.ChildNodes.OfType<XmlElement>())
 			{
-				Dictionary<string, FieldInfo> NameToField;
+				Dictionary<string, FieldInfo>? NameToField;
 				if(CategoryToFields.TryGetValue(CategoryElement.Name, out NameToField))
 				{
 					foreach(XmlElement KeyElement in CategoryElement.ChildNodes.OfType<XmlElement>())
 					{
-						FieldInfo Field;
+						FieldInfo? Field;
 						if(NameToField.TryGetValue(KeyElement.Name, out Field))
 						{
 							// Parse the corresponding value
@@ -573,11 +585,11 @@ namespace UnrealBuildTool
 							}
 
 							// Add it to the set of values for the type containing this field
-							Dictionary<FieldInfo, object> FieldToValue;
-							if(!TypeToValues.TryGetValue(Field.DeclaringType, out FieldToValue))
+							Dictionary<FieldInfo, object>? FieldToValue;
+							if(!TypeToValues.TryGetValue(Field.DeclaringType!, out FieldToValue))
 							{
 								FieldToValue = new Dictionary<FieldInfo, object>();
-								TypeToValues.Add(Field.DeclaringType, FieldToValue);
+								TypeToValues.Add(Field.DeclaringType!, FieldToValue);
 							}
 							FieldToValue[Field] = Value;
 						}
@@ -739,9 +751,9 @@ namespace UnrealBuildTool
 		/// <param name="Field">The field to search for</param>
 		/// <param name="Lines">Receives the description for the requested field</param>
 		/// <returns>True if a comment was found for the field</returns>
-		private static bool TryGetXmlComment(XmlDocument Documentation, FieldInfo Field, out List<string> Lines)
+		private static bool TryGetXmlComment(XmlDocument Documentation, FieldInfo Field, [NotNullWhen(true)] out List<string>? Lines)
 		{
-			XmlNode Node = Documentation.SelectSingleNode(String.Format("//member[@name='F:{0}.{1}']/summary", Field.DeclaringType.FullName, Field.Name));
+			XmlNode Node = Documentation.SelectSingleNode(String.Format("//member[@name='F:{0}.{1}']/summary", Field.DeclaringType!.FullName, Field.Name));
 			if (Node == null)
 			{
 				Lines = null;
@@ -790,7 +802,7 @@ namespace UnrealBuildTool
 					foreach (KeyValuePair<string, FieldInfo> FieldPair in Fields)
 					{
 						// Get the XML comment for this field
-						List<string> Lines;
+						List<string>? Lines;
 						if(!TryGetXmlComment(InputDocumentation, FieldPair.Value, out Lines) || Lines.Count == 0)
 						{
 							Log.TraceWarning("Missing XML comment for {0}", FieldPair.Value.Name);
@@ -839,7 +851,7 @@ namespace UnrealBuildTool
 					foreach (KeyValuePair<string, FieldInfo> FieldPair in Fields)
 					{
 						// Get the XML comment for this field
-						List<string> Lines;
+						List<string>? Lines;
 						if (!TryGetXmlComment(InputDocumentation, FieldPair.Value, out Lines) || Lines.Count == 0)
 						{
 							Log.TraceWarning("Missing XML comment for {0}", FieldPair.Value.Name);
diff --git a/Engine/Source/Programs/UnrealBuildTool/System/XmlConfigData.cs b/Engine/Source/Programs/UnrealBuildTool/System/XmlConfigData.cs
--- a/Engine/Source/Programs/UnrealBuildTool/System/XmlConfigData.cs
+++ b/Engine/Source/Programs/UnrealBuildTool/System/XmlConfigData.cs
@@ -2,6 +2,7 @@
 
 using System;
 using System.Collections.Generic;
+using System.Diagnostics.CodeAnalysis;
 using System.IO;
 using System.Linq;
 using System.Reflection;
@@ -50,7 +51,7 @@ namespace UnrealBuildTool
 		/// <param name="Types">Array of valid types. Used to resolve serialized type names to concrete types.</param>
 		/// <param name="Data">On success, receives the parsed data</param>
 		/// <returns>True if the data was read and is valid</returns>
-		public static bool TryRead(FileReference Location, IEnumerable<Type> Types, out XmlConfigData Data)
+		public static bool TryRead(FileReference Location, IEnumerable<Type> Types, [NotNullWhen(true)] out XmlConfigData? Data)
 		{
 			// Check the file exists first
 			if(!FileReference.Exists(Location))
@@ -95,7 +96,7 @@ namespace UnrealBuildTool
 						string FieldName = Reader.ReadString();
 
 						// Find the matching field on the output type
-						FieldInfo Field = Type.GetField(FieldName, BindingFlags.Public | BindingFlags.NonPublic | BindingFlags.Instance | BindingFlags.Static);
+						FieldInfo? Field = Type.GetField(FieldName, BindingFlags.Public | BindingFlags.NonPublic | BindingFlags.Instance | BindingFlags.Static);
 						if(Field == null || Field.GetCustomAttribute<XmlConfigFileAttribute>() == null)
 						{
 							Data = null;
diff --git a/Engine/Source/Programs/UnrealBuildTool/System/XmlConfigFile.cs b/Engine/Source/Programs/UnrealBuildTool/System/XmlConfigFile.cs
--- a/Engine/Source/Programs/UnrealBuildTool/System/XmlConfigFile.cs
+++ b/Engine/Source/Programs/UnrealBuildTool/System/XmlConfigFile.cs
@@ -2,6 +2,7 @@
 
 using System;
 using System.Collections.Generic;
+using System.Diagnostics.CodeAnalysis;
 using System.Linq;
 using System.Text;
 using System.Threading.Tasks;
@@ -35,7 +36,7 @@ namespace UnrealBuildTool
 		/// <summary>
 		/// Interface to the LineInfo on the active XmlReader
 		/// </summary>
-		IXmlLineInfo LineInfo;
+		IXmlLineInfo LineInfo = null!;
 
 		/// <summary>
 		/// Set to true if the reader encounters an error
@@ -65,7 +66,7 @@ namespace UnrealBuildTool
 		/// <param name="Schema">The schema to validate against</param>
 		/// <param name="OutConfigFile">If successful, the document that was read</param>
 		/// <returns>True if the document could be read, false otherwise</returns>
-		public static bool TryRead(FileReference File, XmlSchema Schema, out XmlConfigFile OutConfigFile)
+		public static bool TryRead(FileReference File, XmlSchema Schema, [NotNullWhen(true)] out XmlConfigFile? OutConfigFile)
 		{
 			XmlConfigFile ConfigFile = new XmlConfigFile(File);
 
diff --git a/Engine/Source/Programs/UnrealBuildTool/System/XmlConfigFileAttribute.cs b/Engine/Source/Programs/UnrealBuildTool/System/XmlConfigFileAttribute.cs
--- a/Engine/Source/Programs/UnrealBuildTool/System/XmlConfigFileAttribute.cs
+++ b/Engine/Source/Programs/UnrealBuildTool/System/XmlConfigFileAttribute.cs
@@ -1,4 +1,4 @@
-﻿// Copyright Epic Games, Inc. All Rights Reserved.
+// Copyright Epic Games, Inc. All Rights Reserved.
 
 using System;
 using System.Collections.Generic;
@@ -17,11 +17,11 @@ namespace UnrealBuildTool
 		/// <summary>
 		/// The category for this config value. Optional; defaults to the declaring type name.
 		/// </summary>
-		public string Category = null;
+		public string? Category = null;
 
 		/// <summary>
 		/// Name of the key to read. Optional; defaults to the field name.
 		/// </summary>
-		public string Name = null;
+		public string? Name = null;
 	}
 }
diff --git a/Engine/Source/Programs/UnrealBuildTool/UnrealBuildTool.csproj b/Engine/Source/Programs/UnrealBuildTool/UnrealBuildTool.csproj
--- a/Engine/Source/Programs/UnrealBuildTool/UnrealBuildTool.csproj
+++ b/Engine/Source/Programs/UnrealBuildTool/UnrealBuildTool.csproj
@@ -15,7 +15,8 @@
     <IsWindows Condition="'$([System.Runtime.InteropServices.RuntimeInformation]::IsOSPlatform($([System.Runtime.InteropServices.OSPlatform]::Windows)))' == 'true'">true</IsWindows> 
     <IsOSX Condition="'$([System.Runtime.InteropServices.RuntimeInformation]::IsOSPlatform($([System.Runtime.InteropServices.OSPlatform]::OSX)))' == 'true'">true</IsOSX> 
     <IsLinux Condition="'$([System.Runtime.InteropServices.RuntimeInformation]::IsOSPlatform($([System.Runtime.InteropServices.OSPlatform]::Linux)))' == 'true'">true</IsLinux> 
-    <DebugType>pdbonly</DebugType>
+    <DebugType>pdbonly</DebugType> 
+    <Nullable>enable</Nullable>
   </PropertyGroup>
   <Choose>
     <When Condition="'$(IsOSX)'=='true'">
