UBT: Fix up a few more nullable references.

#rnx

diff --git a/Engine/Source/Programs/UnrealBuildTool/Configuration/ModuleRules.cs b/Engine/Source/Programs/UnrealBuildTool/Configuration/ModuleRules.cs
--- a/Engine/Source/Programs/UnrealBuildTool/Configuration/ModuleRules.cs
+++ b/Engine/Source/Programs/UnrealBuildTool/Configuration/ModuleRules.cs
@@ -1152,6 +1152,7 @@ namespace UnrealBuildTool
 			}
 		}
 
+#nullable disable
 		/// <summary>
 		/// Constructor. For backwards compatibility while the parameterless constructor is being phased out, initialization which would happen here is done by 
 		/// RulesAssembly.CreateModulRules instead.
@@ -1160,12 +1161,8 @@ namespace UnrealBuildTool
 		public ModuleRules(ReadOnlyTargetRules Target)
 		{
 			this.Target = Target;
-
-			Name = Name!;
-			File = File!;
-			Directory = Directory!;
-			Context = Context!;
 		}
+#nullable restore
 
 		/// <summary>
 		/// Add the given Engine ThirdParty modules as static private dependencies
diff --git a/Engine/Source/Programs/UnrealBuildTool/Configuration/TargetRules.cs b/Engine/Source/Programs/UnrealBuildTool/Configuration/TargetRules.cs
--- a/Engine/Source/Programs/UnrealBuildTool/Configuration/TargetRules.cs
+++ b/Engine/Source/Programs/UnrealBuildTool/Configuration/TargetRules.cs
@@ -1544,6 +1544,7 @@ namespace UnrealBuildTool
 		[ConfigSubObject]
 		public HoloLensTargetRules HoloLensPlatform;
 
+#nullable disable
 		/// <summary>
 		/// Constructor.
 		/// </summary>
@@ -1559,9 +1560,6 @@ namespace UnrealBuildTool
 			this.WindowsPlatform = new WindowsTargetRules(this);
 			this.HoloLensPlatform = new HoloLensTargetRules(Target);
 
-			// Suppress nullable warnings (initialized before constructor is run)
-			File = File!;
-
 			// Read settings from config files
 			foreach (object ConfigurableObject in GetConfigurableObjects())
 			{
@@ -1595,7 +1593,7 @@ namespace UnrealBuildTool
 			// Get the directory to use for crypto settings. We can build engine targets (eg. UHT) with 
 			// a project file, but we can't use that to determine crypto settings without triggering
 			// constant rebuilds of UHT.
-			DirectoryReference? CryptoSettingsDir = DirectoryReference.FromFile(ProjectFile);
+			DirectoryReference CryptoSettingsDir = DirectoryReference.FromFile(ProjectFile);
 			if (CryptoSettingsDir != null && File != null && !File.IsUnderDirectory(CryptoSettingsDir))
 			{
 				CryptoSettingsDir = null;
@@ -1621,6 +1619,7 @@ namespace UnrealBuildTool
 				ProjectDefinitions.Add("IMPLEMENT_SIGNING_KEY_REGISTRATION()=");
 			}
 		}
+#nullable restore
 
 		/// <summary>
 		/// Formats an array of bytes as a sequence of values
diff --git a/Engine/Source/Programs/UnrealBuildTool/Configuration/UEBuildPlatform.cs b/Engine/Source/Programs/UnrealBuildTool/Configuration/UEBuildPlatform.cs
--- a/Engine/Source/Programs/UnrealBuildTool/Configuration/UEBuildPlatform.cs
+++ b/Engine/Source/Programs/UnrealBuildTool/Configuration/UEBuildPlatform.cs
@@ -633,7 +633,7 @@ namespace UnrealBuildTool
 			return Configurations;
 		}
 
-		protected static bool DoProjectSettingsMatchDefault(UnrealTargetPlatform Platform, DirectoryReference ProjectDirectoryName, string Section, string[] BoolKeys, string[]? IntKeys, string[]? StringKeys)
+		protected static bool DoProjectSettingsMatchDefault(UnrealTargetPlatform Platform, DirectoryReference ProjectDirectoryName, string Section, string[]? BoolKeys, string[]? IntKeys, string[]? StringKeys)
 		{
 			ConfigHierarchy ProjIni = ConfigCache.ReadHierarchy(ConfigHierarchyType.Engine, ProjectDirectoryName, Platform);
 			ConfigHierarchy DefaultIni = ConfigCache.ReadHierarchy(ConfigHierarchyType.Engine, (DirectoryReference?)null, Platform);
diff --git a/Engine/Source/Programs/UnrealBuildTool/Executors/LocalExecutor.cs b/Engine/Source/Programs/UnrealBuildTool/Executors/LocalExecutor.cs
--- a/Engine/Source/Programs/UnrealBuildTool/Executors/LocalExecutor.cs
+++ b/Engine/Source/Programs/UnrealBuildTool/Executors/LocalExecutor.cs
@@ -11,8 +11,6 @@ using System.Threading;
 using System.Linq;
 using Tools.DotNETCommon;
 
-#nullable disable
-
 namespace UnrealBuildTool
 {
 	class ActionThread
@@ -70,7 +68,7 @@ namespace UnrealBuildTool
 		/// </summary>
 		private void ThreadFunc()
 		{
-			string Args = Action.CommandArguments;
+			string? Args = Action.CommandArguments;
 			// TODO: Process Arguments follow windows conventions in .NET Core
 			// Which means single quotes ' are not considered quotes.
 			// see [REDACTED_PROFILE_URL]/runtime/issues/29857
@@ -106,7 +104,7 @@ namespace UnrealBuildTool
 			}
 
 			// Try to launch the action's process, and produce a friendly error message if it fails.
-			Process ActionProcess = null;
+			Process? ActionProcess = null;
 			try
 			{
 				try
@@ -295,7 +293,7 @@ namespace UnrealBuildTool
 			// The number of actions to execute in parallel is trying to keep the CPU busy enough in presence of I/O stalls.
 			Log.TraceInformation("Performing {0} actions ({1} in parallel)", Actions.Count, NumParallelProcesses);
 
-			Dictionary<LinkedAction, ActionThread> ActionThreadDictionary = new Dictionary<LinkedAction, ActionThread>();
+			Dictionary<LinkedAction, ActionThread?> ActionThreadDictionary = new Dictionary<LinkedAction, ActionThread?>();
 			int JobNumber = 1;
 			using (ProgressWriter ProgressWriter = new ProgressWriter("Compiling C++ source code...", false))
 			{
@@ -307,7 +305,7 @@ namespace UnrealBuildTool
 					int NumExecutingActions = 0;
 					foreach (LinkedAction Action in Actions)
 					{
-						ActionThread ActionThread = null;
+						ActionThread? ActionThread = null;
 						bool bFoundActionProcess = ActionThreadDictionary.TryGetValue(Action, out ActionThread);
 						if (bFoundActionProcess == false)
 						{
@@ -341,7 +339,7 @@ namespace UnrealBuildTool
 					// prerequisites.
 					foreach (LinkedAction Action in Actions)
 					{
-						ActionThread ActionProcess = null;
+						ActionThread? ActionProcess = null;
 						bool bFoundActionProcess = ActionThreadDictionary.TryGetValue(Action, out ActionProcess);
 						if (bFoundActionProcess == false)
 						{
@@ -354,7 +352,7 @@ namespace UnrealBuildTool
 								{
 									if (Actions.Contains(PrerequisiteAction))
 									{
-										ActionThread PrerequisiteProcess = null;
+										ActionThread? PrerequisiteProcess = null;
 										bool bFoundPrerequisiteProcess = ActionThreadDictionary.TryGetValue(PrerequisiteAction, out PrerequisiteProcess);
 										if (bFoundPrerequisiteProcess == true)
 										{
@@ -413,10 +411,10 @@ namespace UnrealBuildTool
 
 			// Check whether any of the tasks failed and log action stats if wanted.
 			bool bSuccess = true;
-			foreach (KeyValuePair<LinkedAction, ActionThread> ActionProcess in ActionThreadDictionary)
+			foreach (KeyValuePair<LinkedAction, ActionThread?> ActionProcess in ActionThreadDictionary)
 			{
 				LinkedAction Action = ActionProcess.Key;
-				ActionThread ActionThread = ActionProcess.Value;
+				ActionThread? ActionThread = ActionProcess.Value;
 
 				// Check for pending actions, preemptive failure
 				if (ActionThread == null)
diff --git a/Engine/Source/Programs/UnrealBuildTool/Executors/ParallelExecutor.cs b/Engine/Source/Programs/UnrealBuildTool/Executors/ParallelExecutor.cs
--- a/Engine/Source/Programs/UnrealBuildTool/Executors/ParallelExecutor.cs
+++ b/Engine/Source/Programs/UnrealBuildTool/Executors/ParallelExecutor.cs
@@ -9,8 +9,6 @@ using System.Threading;
 using System.Threading.Tasks;
 using Tools.DotNETCommon;
 
-#nullable disable
-
 namespace UnrealBuildTool
 {
 	/// <summary>
@@ -32,6 +30,11 @@ namespace UnrealBuildTool
 
 			public List<string> LogLines = new List<string>();
 			public int ExitCode = -1;
+
+			public BuildAction(LinkedAction Inner)
+			{
+				this.Inner = Inner;
+			}
 		}
 
 		/// <summary>
@@ -106,9 +109,8 @@ namespace UnrealBuildTool
 			List<BuildAction> Actions = new List<BuildAction>();
 			for(int Idx = 0; Idx < InputActions.Count; Idx++)
 			{
-				BuildAction Action = new BuildAction();
+				BuildAction Action = new BuildAction(InputActions[Idx]);
 				Action.SortIndex = Idx;
-				Action.Inner = InputActions[Idx];
 
 				if (!Action.Inner.StatusDescription.EndsWith(".ispc"))
 				{
@@ -124,7 +126,7 @@ namespace UnrealBuildTool
 			{
 				foreach(LinkedAction PrerequisiteAction in Action.Inner.PrerequisiteActions)
 				{
-					BuildAction Dependency;
+					BuildAction? Dependency;
 					if(LinkedActionToBuildAction.TryGetValue(PrerequisiteAction, out Dependency))
 					{
 						Action.Dependencies.Add(Dependency);
diff --git a/Engine/Source/Programs/UnrealBuildTool/Modes/BuildMode.cs b/Engine/Source/Programs/UnrealBuildTool/Modes/BuildMode.cs
--- a/Engine/Source/Programs/UnrealBuildTool/Modes/BuildMode.cs
+++ b/Engine/Source/Programs/UnrealBuildTool/Modes/BuildMode.cs
@@ -352,7 +352,7 @@ namespace UnrealBuildTool
 				if (ModuleDependencyActions.Count > 0)
 				{
 					Dictionary<LinkedAction, bool> PreprocessActionToOutdatedFlag = new Dictionary<LinkedAction, bool>();
-					ActionGraph.GatherAllOutdatedActions(ModuleDependencyActions, History, PreprocessActionToOutdatedFlag, null, BuildConfiguration.bIgnoreOutdatedImportLibraries);
+					ActionGraph.GatherAllOutdatedActions(ModuleDependencyActions, History, PreprocessActionToOutdatedFlag, null!, BuildConfiguration.bIgnoreOutdatedImportLibraries);
 
 					List<LinkedAction> PreprocessActions = PreprocessActionToOutdatedFlag.Where(x => x.Value).Select(x => x.Key).ToList();
 					if (PreprocessActions.Count > 0)
@@ -391,7 +391,7 @@ namespace UnrealBuildTool
 							if (CppModulesAction != null && CppModulesAction.CompiledModuleInterfaceFile != null)
 							{
 								string? ProducedModule;
-								if (CppDependencies.TryGetProducedModule(PrerequisiteAction.DependencyListFile, out ProducedModule))
+								if (CppDependencies.TryGetProducedModule(PrerequisiteAction.DependencyListFile!, out ProducedModule))
 								{
 									ModuleOutputs[ProducedModule] = CppModulesAction.CompiledModuleInterfaceFile;
 								}
@@ -405,7 +405,7 @@ namespace UnrealBuildTool
 							if (PrerequisiteAction.ActionType == ActionType.CompileModuleInterface)
 							{
 								List<string>? ImportedModules;
-								if(CppDependencies.TryGetImportedModules(PrerequisiteAction.DependencyListFile, out ImportedModules))
+								if(CppDependencies.TryGetImportedModules(PrerequisiteAction.DependencyListFile!, out ImportedModules))
 								{
 									foreach (string ImportedModule in ImportedModules)
 									{
diff --git a/Engine/Source/Programs/UnrealBuildTool/Modes/CleanMode.cs b/Engine/Source/Programs/UnrealBuildTool/Modes/CleanMode.cs
--- a/Engine/Source/Programs/UnrealBuildTool/Modes/CleanMode.cs
+++ b/Engine/Source/Programs/UnrealBuildTool/Modes/CleanMode.cs
@@ -9,8 +9,6 @@ using System.Text;
 using System.Threading.Tasks;
 using Tools.DotNETCommon;
 
-#nullable disable
-
 namespace UnrealBuildTool
 {
 	/// <summary>
@@ -59,7 +57,7 @@ namespace UnrealBuildTool
 				const string UnrealHeaderToolTarget = "UnrealHeaderTool";
 
 				// Get a list of project files to clean UHT for
-				List<FileReference> ProjectFiles = new List<FileReference>();
+				List<FileReference?> ProjectFiles = new List<FileReference?>();
 				foreach (TargetDescriptor TargetDesc in TargetDescriptors)
 				{
 					if (TargetDesc.Name != UnrealHeaderToolTarget && !RemoteMac.HandlesTargetPlatform(TargetDesc.Platform))
@@ -80,7 +78,7 @@ namespace UnrealBuildTool
 				{
 					UnrealTargetConfiguration Configuration = BuildConfiguration.bForceDebugUnrealHeaderTool ? UnrealTargetConfiguration.Debug : UnrealTargetConfiguration.Development;
 					string Architecture = UEBuildPlatform.GetBuildPlatform(BuildHostPlatform.Current.Platform).GetDefaultArchitecture(null);
-					foreach (FileReference ProjectFile in ProjectFiles)
+					foreach (FileReference? ProjectFile in ProjectFiles)
 					{
 						TargetDescriptors.Add(new TargetDescriptor(ProjectFile, UnrealHeaderToolTarget, BuildHostPlatform.Current.Platform, Configuration, Architecture, null));
 					}
diff --git a/Engine/Source/Programs/UnrealBuildTool/Platform/Windows/MicrosoftPlatformSDK.cs b/Engine/Source/Programs/UnrealBuildTool/Platform/Windows/MicrosoftPlatformSDK.cs
--- a/Engine/Source/Programs/UnrealBuildTool/Platform/Windows/MicrosoftPlatformSDK.cs
+++ b/Engine/Source/Programs/UnrealBuildTool/Platform/Windows/MicrosoftPlatformSDK.cs
@@ -7,8 +7,6 @@ using System.Linq;
 using System.Text.RegularExpressions;
 using Microsoft.Win32;
 
-#nullable disable
-
 namespace UnrealBuildTool
 {
 	internal class MicrosoftPlatformSDK : UEBuildPlatformSDK
@@ -25,13 +23,13 @@ namespace UnrealBuildTool
 			MaxVersion = "10.9.99999";
 		}
 
-		public override void GetValidSoftwareVersionRange(out string MinVersion, out string MaxVersion)
+		public override void GetValidSoftwareVersionRange(out string? MinVersion, out string? MaxVersion)
 		{
 			// @todo we may want to split this for Win64 vs other things
 			MinVersion = MaxVersion = null;
 		}
 
-		public override string GetInstalledSDKVersion()
+		public override string? GetInstalledSDKVersion()
 		{
 			if (Utils.IsRunningOnMono)
 			{
diff --git a/Engine/Source/Programs/UnrealBuildTool/Platform/Windows/UEBuildWindows.cs b/Engine/Source/Programs/UnrealBuildTool/Platform/Windows/UEBuildWindows.cs
--- a/Engine/Source/Programs/UnrealBuildTool/Platform/Windows/UEBuildWindows.cs
+++ b/Engine/Source/Programs/UnrealBuildTool/Platform/Windows/UEBuildWindows.cs
@@ -10,8 +10,7 @@ using System.Linq;
 using Tools.DotNETCommon;
 using Microsoft.VisualStudio.Setup.Configuration;
 using System.Runtime.InteropServices;
-
-#nullable disable
+using System.Diagnostics.CodeAnalysis;
 
 namespace UnrealBuildTool
 {
@@ -145,14 +144,14 @@ namespace UnrealBuildTool
 		/// </summary>
 		[XmlConfigFile(Category = "WindowsPlatform")]
 		[CommandLine("-CompilerVersion")]
-		public string CompilerVersion = null;
+		public string? CompilerVersion = null;
 
 		/// <summary>
 		/// The specific Windows SDK version to use. This may be a specific version number (for example, "8.1", "10.0" or "10.0.10150.0"), or the string "Latest", to select the newest available version.
 		/// By default, and if it is available, we use the Windows SDK version indicated by WindowsPlatform.DefaultWindowsSdkVersion (otherwise, we use the latest version).
 		/// </summary>
 		[XmlConfigFile(Category = "WindowsPlatform")]
-		public string WindowsSdkVersion = null;
+		public string? WindowsSdkVersion = null;
 
 		/// <summary>
 		/// Value for the WINVER macro, defining the minimum supported Windows version.
@@ -181,19 +180,19 @@ namespace UnrealBuildTool
 		/// The name of the company (author, provider) that created the project.
 		/// </summary>
 		[ConfigFile(ConfigHierarchyType.Game, "/Script/EngineSettings.GeneralProjectSettings", "CompanyName")]
-		public string CompanyName;
+		public string? CompanyName;
 
 		/// <summary>
 		/// The project's copyright and/or trademark notices.
 		/// </summary>
 		[ConfigFile(ConfigHierarchyType.Game, "/Script/EngineSettings.GeneralProjectSettings", "CopyrightNotice")]
-		public string CopyrightNotice;
+		public string? CopyrightNotice;
 
 		/// <summary>
 		/// The product name.
 		/// </summary>
 		[ConfigFile(ConfigHierarchyType.Game, "/Script/EngineSettings.GeneralProjectSettings", "ProjectName")]
-		public string ProductName;
+		public string? ProductName;
 
 		/// <summary>
 		/// The static analyzer to use.
@@ -207,19 +206,19 @@ namespace UnrealBuildTool
 		/// </summary>
 		[XmlConfigFile]
 		[CommandLine("-ObjSrcMap")]
-		public string ObjSrcMapFile = null;
+		public string? ObjSrcMapFile = null;
 
 		/// <summary>
 		/// Provides a Module Definition File (.def) to the linker to describe various attributes of a DLL.
 		/// Necessary when exporting functions by ordinal values instead of by name.
 		/// </summary>
-		public string ModuleDefinitionFile;
+		public string? ModuleDefinitionFile;
 
 		/// <summary>
 		/// Specifies the path to a manifest file for the linker to embed. Defaults to the manifest in Engine/Build/Windows/Resources. Can be assigned to null
 		/// if the target wants to specify its own manifest.
 		/// </summary>
-		public string ManifestFile;
+		public string? ManifestFile;
 
 		/// <summary>
 		/// Enables strict standard conformance mode (/permissive-) in VS2017+.
@@ -329,12 +328,12 @@ namespace UnrealBuildTool
 		/// <summary>
 		/// The Visual C++ environment to use for this target. Only initialized after all the target settings are finalized, in ValidateTarget().
 		/// </summary>
-		internal VCEnvironment Environment;
+		internal VCEnvironment? Environment;
 
 		/// <summary>
 		/// Directory containing the toolchain
 		/// </summary>
-		public string ToolChainDir
+		public string? ToolChainDir
 		{
 			get { return (Environment == null)? null : Environment.ToolChainDir.FullName; }
 		}
@@ -342,7 +341,7 @@ namespace UnrealBuildTool
 		/// <summary>
 		/// The version number of the toolchain
 		/// </summary>
-		public string ToolChainVersion
+		public string? ToolChainVersion
 		{
 			get { return (Environment == null)? null : Environment.ToolChainVersion.ToString(); }
 		}
@@ -350,7 +349,7 @@ namespace UnrealBuildTool
 		/// <summary>
 		/// Root directory containing the Windows Sdk
 		/// </summary>
-		public string WindowsSdkDir
+		public string? WindowsSdkDir
 		{
 			get { return (Environment == null)? null : Environment.WindowsSdkDir.FullName; }
 		}
@@ -360,7 +359,7 @@ namespace UnrealBuildTool
 		/// </summary>
 		public string DiaSdkDir
 		{
-			get { return WindowsPlatform.FindDiaSdkDirs(Environment.Compiler).Select(x => x.FullName).FirstOrDefault(); }
+			get { return WindowsPlatform.FindDiaSdkDirs(Environment!.Compiler).Select(x => x.FullName).FirstOrDefault(); }
 		}
 
 		/// <summary>
@@ -372,7 +371,7 @@ namespace UnrealBuildTool
 			{
 				try
 				{
-					return WindowsPlatform.FindVisualStudioInstallations(Environment.Compiler).Select(x => x.BaseDir.FullName).FirstOrDefault();
+					return WindowsPlatform.FindVisualStudioInstallations(Environment!.Compiler).Select(x => x.BaseDir.FullName).FirstOrDefault();
 				}
 				catch(Exception) // Find function will throw if there is no visual studio installed! This can happen w/ clang builds
 				{
@@ -451,12 +450,12 @@ namespace UnrealBuildTool
 			get { return Inner.Architecture; }
 		}
 
-		public string CompilerVersion
+		public string? CompilerVersion
 		{
 			get { return Inner.CompilerVersion; }
 		}
 
-		public string WindowsSdkVersion
+		public string? WindowsSdkVersion
 		{
 			get { return Inner.WindowsSdkVersion; }
 		}
@@ -481,17 +480,17 @@ namespace UnrealBuildTool
 			get { return Inner.bEnableRayTracing; }
 		}
 
-		public string CompanyName
+		public string? CompanyName
 		{
 			get { return Inner.CompanyName; }
 		}
 
-		public string CopyrightNotice
+		public string? CopyrightNotice
 		{
 			get { return Inner.CopyrightNotice; }
 		}
 
-		public string ProductName
+		public string? ProductName
 		{
 			get { return Inner.ProductName; }
 		}
@@ -501,17 +500,17 @@ namespace UnrealBuildTool
 			get { return Inner.StaticAnalyzer; }
 		}
 
-		public string ObjSrcMapFile
+		public string? ObjSrcMapFile
 		{
 			get { return Inner.ObjSrcMapFile; }
 		}
 
-		public string ModuleDefinitionFile
+		public string? ModuleDefinitionFile
 		{
 			get { return Inner.ModuleDefinitionFile; }
 		}
 
-		public string ManifestFile
+		public string? ManifestFile
 		{
 			get { return Inner.ManifestFile; }
 		}
@@ -596,22 +595,22 @@ namespace UnrealBuildTool
 			get; private set;
 		}
 
-		internal VCEnvironment Environment
+		internal VCEnvironment? Environment
 		{
 			get { return Inner.Environment; }
 		}
 
-		public string ToolChainDir
+		public string? ToolChainDir
 		{
 			get { return Inner.ToolChainDir; }
 		}
 
-		public string ToolChainVersion
+		public string? ToolChainVersion
 		{
 			get { return Inner.ToolChainVersion; }
 		}
 
-		public string WindowsSdkDir
+		public string? WindowsSdkDir
 		{
 			get { return Inner.WindowsSdkDir; }
 		}
@@ -740,12 +739,12 @@ namespace UnrealBuildTool
 		/// <summary>
 		/// Cache of Windows SDK installation directories
 		/// </summary>
-		private static IReadOnlyDictionary<VersionNumber, DirectoryReference> CachedWindowsSdkDirs;
+		private static IReadOnlyDictionary<VersionNumber, DirectoryReference>? CachedWindowsSdkDirs;
 
 		/// <summary>
 		/// Cache of Universal CRT installation directories
 		/// </summary>
-		private static IReadOnlyDictionary<VersionNumber, DirectoryReference> CachedUniversalCrtDirs;
+		private static IReadOnlyDictionary<VersionNumber, DirectoryReference>? CachedUniversalCrtDirs;
 
 		/// <summary>
 		/// True if we should use the Clang linker (LLD) when we are compiling with Clang, otherwise we use the MSVC linker
@@ -916,7 +915,7 @@ namespace UnrealBuildTool
 		/// Gets the default compiler which should be used, if it's not set explicitly by the target, command line, or config file.
 		/// </summary>
 		/// <returns>The default compiler version</returns>
-		internal static WindowsCompiler GetDefaultCompiler(FileReference ProjectFile)
+		internal static WindowsCompiler GetDefaultCompiler(FileReference? ProjectFile)
 		{
 			// If there's no specific compiler set, try to pick the matching compiler for the selected IDE
 			if(ProjectFileGeneratorSettings.Format != null)
@@ -935,7 +934,7 @@ namespace UnrealBuildTool
 			}
 
 			// Also check the default format for the Visual Studio project generator
-			object ProjectFormatObject;
+			object? ProjectFormatObject;
 			if (XmlConfig.TryGetValue(typeof(VCProjectFileGenerator), "Version", out ProjectFormatObject))
 			{
 				VCProjectFileFormat ProjectFormat = (VCProjectFileFormat)ProjectFormatObject;
@@ -974,7 +973,7 @@ namespace UnrealBuildTool
 			}
 
 			// If we do have a Visual Studio installation, but we're missing just the C++ parts, warn about that.
-			DirectoryReference VSInstallDir;
+			DirectoryReference? VSInstallDir;
 			if (TryGetVSInstallDir(WindowsCompiler.VisualStudio2019, out VSInstallDir))
 			{
 				Log.TraceWarning("Visual Studio 2019 is installed, but is missing the C++ toolchain. Please verify that the \"VC++ 2019 toolset\" component is selected in the Visual Studio 2019 installation options.");
@@ -1019,7 +1018,7 @@ namespace UnrealBuildTool
 		/// <param name="Compiler">Version of the toolchain to look for.</param>
 		/// <param name="InstallDir">On success, the directory that Visual Studio is installed to.</param>
 		/// <returns>True if the directory was found, false otherwise.</returns>
-		public static bool TryGetVSInstallDir(WindowsCompiler Compiler, out DirectoryReference InstallDir)
+		public static bool TryGetVSInstallDir(WindowsCompiler Compiler, [NotNullWhen(true)] out DirectoryReference? InstallDir)
 		{
 			if (BuildHostPlatform.Current.Platform != UnrealTargetPlatform.Win64 && BuildHostPlatform.Current.Platform != UnrealTargetPlatform.Win32)
 			{
@@ -1048,7 +1047,7 @@ namespace UnrealBuildTool
 		/// <returns>List of directories containing Visual Studio installations</returns>
 		public static List<VisualStudioInstallation> FindVisualStudioInstallations(WindowsCompiler Compiler)
 		{
-			List<VisualStudioInstallation> Installations;
+			List<VisualStudioInstallation>? Installations;
 			if(!CachedVisualStudioInstallations.TryGetValue(Compiler, out Installations))
 			{
 				Installations = new List<VisualStudioInstallation>();
@@ -1057,7 +1056,7 @@ namespace UnrealBuildTool
 				    if(Compiler == WindowsCompiler.VisualStudio2015_DEPRECATED)
 				    {
 					    // VS2015 just installs one toolchain; use that.
-					    DirectoryReference InstallDir;
+					    DirectoryReference? InstallDir;
 					    if(TryReadInstallDirRegistryKey32("Microsoft\\VisualStudio\\SxS\\VS7", "14.0", out InstallDir))
 					    {
 						    Installations.Add(new VisualStudioInstallation(InstallDir));
@@ -1086,7 +1085,7 @@ namespace UnrealBuildTool
 								{
 									string VersionString = Instance.GetInstallationVersion();
 
-									VersionNumber Version;
+									VersionNumber? Version;
 									if (VersionNumber.TryParse(VersionString, out Version))
 									{
 										VersionNumber Version2019 = new VersionNumber(16);
@@ -1100,7 +1099,7 @@ namespace UnrealBuildTool
 										}
 									}
 
-									ISetupInstanceCatalog Catalog = Instance as ISetupInstanceCatalog;
+									ISetupInstanceCatalog? Catalog = Instance as ISetupInstanceCatalog;
 
 									VisualStudioInstallation Installation = new VisualStudioInstallation(new DirectoryReference(Instance.GetInstallationPath()));
 									if (Catalog != null && Catalog.IsPrerelease())
@@ -1140,7 +1139,7 @@ namespace UnrealBuildTool
 		/// <returns>Map of version number to directories</returns>
 		static Dictionary<VersionNumber, ToolChainInstallation> FindToolChainInstallations(WindowsCompiler Compiler)
 		{
-			Dictionary<VersionNumber, ToolChainInstallation> ToolChainInstallations;
+			Dictionary<VersionNumber, ToolChainInstallation>? ToolChainInstallations;
 			if(!CachedToolChainInstallations.TryGetValue(Compiler, out ToolChainInstallations))
 			{
 				ToolChainInstallations = new Dictionary<VersionNumber, ToolChainInstallation>();
@@ -1170,7 +1169,7 @@ namespace UnrealBuildTool
 							{
 								foreach(DirectoryReference ToolChainDir in DirectoryReference.EnumerateDirectories(ClangBaseDir))
 								{
-									VersionNumber Version;
+									VersionNumber? Version;
 									if(VersionNumber.TryParse(ToolChainDir.GetDirectoryName(), out Version) && IsValidToolChainDirClang(ToolChainDir))
 									{
 										ToolChainInstallations[Version] = new ToolChainInstallation(ToolChainDir, false);
@@ -1180,7 +1179,7 @@ namespace UnrealBuildTool
 						}
 
 						// Check for LLVM_PATH environment variable.
-						string LLVMPath = Environment.GetEnvironmentVariable("LLVM_PATH");
+						string? LLVMPath = Environment.GetEnvironmentVariable("LLVM_PATH");
 						if (!String.IsNullOrEmpty(LLVMPath))
 						{
 							DirectoryReference LLVMPathDir = new DirectoryReference(LLVMPath);
@@ -1267,7 +1266,7 @@ namespace UnrealBuildTool
 			{
 				foreach (DirectoryReference ToolChainDir in DirectoryReference.EnumerateDirectories(BaseDir))
 				{
-					VersionNumber Version;
+					VersionNumber? Version;
 					if (IsValidToolChainDir2017or2019(ToolChainDir, out Version))
 					{
 						
@@ -1278,7 +1277,7 @@ namespace UnrealBuildTool
 							ToolChainInstallations[Version] = new ToolChainInstallation(ToolChainDir, bPreview);
 						}
 
-						VersionNumber VersionFolder;
+						VersionNumber? VersionFolder;
 						if ( (VersionNumber.TryParse(ToolChainDir.GetDirectoryName(), out VersionFolder)) && (VersionFolder != Version) && (!ToolChainInstallations.ContainsKey(VersionFolder)))
 						{
 							//Add VS installation based on the version number in the installation path
@@ -1317,7 +1316,7 @@ namespace UnrealBuildTool
 		/// <param name="ToolChainDir">The toolchain directory</param>
 		/// <param name="Version">The version number for the toolchain</param>
 		/// <returns>True if the path is a valid version</returns>
-		static bool IsValidToolChainDir2017or2019(DirectoryReference ToolChainDir, out VersionNumber Version)
+		static bool IsValidToolChainDir2017or2019(DirectoryReference ToolChainDir, [NotNullWhen(true)] out VersionNumber? Version)
 		{
 			FileReference CompilerExe = FileReference.Combine(ToolChainDir, "bin", "Hostx86", "x64", "cl.exe");
 			if(!FileReference.Exists(CompilerExe))
@@ -1377,7 +1376,7 @@ namespace UnrealBuildTool
 		/// <returns>True if the toolchain is compatible with UE4</returns>
 		static bool IsCompatibleVisualCppToolChain(VersionNumber Version)
 		{
-			string Message;
+			string? Message;
 			return IsCompatibleVisualCppToolChain(Version, out Message);
 		}
 
@@ -1387,7 +1386,7 @@ namespace UnrealBuildTool
 		/// <param name="Version">The version number to check</param>
 		/// <param name="Message">Receives a message describing why the toolchain is not compatible</param>
 		/// <returns>True if the toolchain is compatible with UE4</returns>
-		static bool IsCompatibleVisualCppToolChain(VersionNumber Version, out string Message)
+		static bool IsCompatibleVisualCppToolChain(VersionNumber Version, [NotNullWhen(false)] out string? Message)
 		{
 			if (Version >= new VersionNumber(14, 23, 0) && Version < new VersionNumber(14, 23, 28107))
 			{
@@ -1409,13 +1408,13 @@ namespace UnrealBuildTool
 		/// <param name="OutToolChainVersion">Receives the chosen toolchain version</param>
 		/// <param name="OutToolChainDir">Receives the directory containing the toolchain</param>
 		/// <returns>True if the toolchain directory was found correctly</returns>
-		public static bool TryGetToolChainDir(WindowsCompiler Compiler, string CompilerVersion, out VersionNumber OutToolChainVersion, out DirectoryReference OutToolChainDir)
+		public static bool TryGetToolChainDir(WindowsCompiler Compiler, string? CompilerVersion, [NotNullWhen(true)] out VersionNumber? OutToolChainVersion, [NotNullWhen(true)] out DirectoryReference? OutToolChainDir)
 		{
 			// Find all the installed toolchains
 			Dictionary<VersionNumber, ToolChainInstallation> ToolChainInstallations = FindToolChainInstallations(Compiler);
 
 			// Figure out the actual version number that we want
-			VersionNumber ToolChainVersion = null;
+			VersionNumber? ToolChainVersion = null;
 			if(CompilerVersion != null)
 			{
 				if(String.Compare(CompilerVersion, "Latest", StringComparison.InvariantCultureIgnoreCase) == 0 && ToolChainInstallations.Count > 0)
@@ -1453,7 +1452,7 @@ namespace UnrealBuildTool
 
 				foreach (VersionNumber PreferredToolchain in PotentialToolchains)
 				{
-					ToolChainInstallation DefaultToolChain;
+					ToolChainInstallation? DefaultToolChain;
 					if (ToolChainInstallations.TryGetValue(PreferredToolchain, out DefaultToolChain) && (Compiler == WindowsCompiler.Clang || Has64BitToolChain(DefaultToolChain.BaseDir)))
 					{
 						ToolChainVersion = PreferredToolchain;
@@ -1469,14 +1468,14 @@ namespace UnrealBuildTool
 			}
 
 			// Check it's valid
-			string Message;
+			string? Message;
 			if (ToolChainVersion != null && !IsCompatibleVisualCppToolChain(ToolChainVersion, out Message))
 			{
 				throw new BuildException(Message);
 			}
 
 			// Get the actual directory for this version
-			ToolChainInstallation Installation;
+			ToolChainInstallation? Installation;
 			if (ToolChainVersion != null && ToolChainInstallations.TryGetValue(ToolChainVersion, out Installation))
 			{
 				OutToolChainVersion = ToolChainVersion;
@@ -1505,7 +1504,7 @@ namespace UnrealBuildTool
 		/// <param name="ValueName">Value to be read.</param>
 		/// <param name="InstallDir">On success, the directory corresponding to the value read.</param>
 		/// <returns>True if the key was read, false otherwise.</returns>
-		static bool TryReadInstallDirRegistryKey32(string KeySuffix, string ValueName, out DirectoryReference InstallDir)
+		static bool TryReadInstallDirRegistryKey32(string KeySuffix, string ValueName, [NotNullWhen(true)] out DirectoryReference? InstallDir)
 		{
 			foreach (KeyValuePair<RegistryKey, string> InstallRoot in InstallDirRoots)
 			{
@@ -1555,9 +1554,9 @@ namespace UnrealBuildTool
 		/// <param name="ValueName">Value within the key to read</param>
 		/// <param name="Value">The directory read from the registry key</param>
 		/// <returns>True if the key was read, false if it was missing or empty</returns>
-		static bool TryReadDirRegistryKey(string KeyName, string ValueName, out DirectoryReference Value)
+		static bool TryReadDirRegistryKey(string KeyName, string ValueName, [NotNullWhen(true)] out DirectoryReference? Value)
 		{
-			string StringValue = Registry.GetValue(KeyName, ValueName, null) as string;
+			string? StringValue = Registry.GetValue(KeyName, ValueName, null) as string;
 			if (String.IsNullOrEmpty(StringValue))
 			{
 				Value = null;
@@ -1575,7 +1574,7 @@ namespace UnrealBuildTool
 		/// </summary>
 		/// <param name="OutLocation">On success, receives the path to the MSBuild executable.</param>
 		/// <returns>True on success.</returns>
-		public static bool TryGetMsBuildPath(out FileReference OutLocation)
+		public static bool TryGetMsBuildPath([NotNullWhen(true)] out FileReference? OutLocation)
 		{
 			// Get the Visual Studio 2019 install directory
 			List<DirectoryReference> InstallDirs2019 = WindowsPlatform.FindVisualStudioInstallations(WindowsCompiler.VisualStudio2019).ConvertAll(x => x.BaseDir);
@@ -1602,7 +1601,7 @@ namespace UnrealBuildTool
 			}
 
 			// Try to get the MSBuild 14.0 path directly (see https://msdn.microsoft.com/en-us/library/hh162058(v=vs.120).aspx)
-			FileReference ToolPath = FileReference.Combine(DirectoryReference.GetSpecialFolder(Environment.SpecialFolder.ProgramFilesX86), "MSBuild", "14.0", "bin", "MSBuild.exe");
+			FileReference? ToolPath = FileReference.Combine(DirectoryReference.GetSpecialFolder(Environment.SpecialFolder.ProgramFilesX86), "MSBuild", "14.0", "bin", "MSBuild.exe");
 			if(FileReference.Exists(ToolPath))
 			{
 				OutLocation = ToolPath;
@@ -1636,7 +1635,7 @@ namespace UnrealBuildTool
 		/// <returns>Path to MSBuild</returns>
 		public static FileReference GetMsBuildToolPath()
 		{
-			FileReference Location;
+			FileReference? Location;
 			if(!TryGetMsBuildPath(out Location))
 			{
 				throw new BuildException("Unable to find installation of MSBuild.");
@@ -1671,7 +1670,7 @@ namespace UnrealBuildTool
 		/// This mirrors the logic in GetMSBuildPath.bat.
 		/// </summary>
 		/// <returns></returns>
-		static bool TryReadMsBuildInstallPath(string KeyRelativePath, string KeyName, string MsBuildRelativePath, out FileReference OutMsBuildPath)
+		static bool TryReadMsBuildInstallPath(string KeyRelativePath, string KeyName, string MsBuildRelativePath, [NotNullWhen(true)] out FileReference? OutMsBuildPath)
 		{
 			string[] KeyBasePaths =
 			{
@@ -1683,7 +1682,7 @@ namespace UnrealBuildTool
 
 			foreach (string KeyBasePath in KeyBasePaths)
 			{
-				string Value = Registry.GetValue(KeyBasePath + KeyRelativePath, KeyName, null) as string;
+				string? Value = Registry.GetValue(KeyBasePath + KeyRelativePath, KeyName, null) as string;
 				if (Value != null)
 				{
 					FileReference MsBuildPath = FileReference.Combine(new DirectoryReference(Value), MsBuildRelativePath);
@@ -1706,7 +1705,7 @@ namespace UnrealBuildTool
 		/// <returns>Map of version number to directories</returns>
 		public static List<DirectoryReference> FindDiaSdkDirs(WindowsCompiler Compiler)
 		{
-			List<DirectoryReference> DiaSdkDirs;
+			List<DirectoryReference>? DiaSdkDirs;
 			if(!CachedDiaSdkDirs.TryGetValue(Compiler, out DiaSdkDirs))
 			{
 				DiaSdkDirs = new List<DirectoryReference>();
@@ -1753,7 +1752,7 @@ namespace UnrealBuildTool
 			Dictionary<VersionNumber, DirectoryReference> UniversalCrtDirs = new Dictionary<VersionNumber, DirectoryReference>();
 
 			// Enumerate the Windows 8.1 SDK, if present
-			DirectoryReference InstallDir_8_1;
+			DirectoryReference? InstallDir_8_1;
 			if(TryReadInstallDirRegistryKey32("Microsoft\\Microsoft SDKs\\Windows\\v8.1", "InstallationFolder", out InstallDir_8_1))
 			{
 				if(FileReference.Exists(FileReference.Combine(InstallDir_8_1, "Include", "um", "windows.h")))
@@ -1776,7 +1775,7 @@ namespace UnrealBuildTool
 				{
 					foreach(DirectoryReference IncludeDir in DirectoryReference.EnumerateDirectories(IncludeRootDir))
 					{
-						VersionNumber IncludeVersion;
+						VersionNumber? IncludeVersion;
 						if(VersionNumber.TryParse(IncludeDir.GetDirectoryName(), out IncludeVersion))
 						{
 							if(FileReference.Exists(FileReference.Combine(IncludeDir, "um", "windows.h")))
@@ -1809,7 +1808,7 @@ namespace UnrealBuildTool
 			{
 				UpdateCachedWindowsSdks();
 			}
-			return CachedWindowsSdkDirs;
+			return CachedWindowsSdkDirs!;
 		}
 
 		/// <summary>
@@ -1822,7 +1821,7 @@ namespace UnrealBuildTool
 			{
 				UpdateCachedWindowsSdks();
 			}
-			return CachedUniversalCrtDirs;
+			return CachedUniversalCrtDirs!;
 		}
 
 		/// <summary>
@@ -1831,7 +1830,7 @@ namespace UnrealBuildTool
 		/// <param name="RootDirs">Receives all the Windows 10 sdk root directories</param>
 		private static void EnumerateSdkRootDirs(List<DirectoryReference> RootDirs)
 		{
-			DirectoryReference RootDir;
+			DirectoryReference? RootDir;
 			if(TryReadInstallDirRegistryKey32("Microsoft\\Windows Kits\\Installed Roots", "KitsRoot10", out RootDir))
 			{
 				Log.TraceLog("Found Windows 10 SDK root at {0} (1)", RootDir);
@@ -1862,16 +1861,16 @@ namespace UnrealBuildTool
 		/// <param name="OutSdkVersion">Receives the version number of the selected Windows SDK</param>
 		/// <param name="OutSdkDir">Receives the root directory for the selected SDK</param>
 		/// <returns>True if the toolchain directory was found correctly</returns>
-		public static bool TryGetWindowsSdkDir(string DesiredVersion, out VersionNumber OutSdkVersion, out DirectoryReference OutSdkDir)
+		public static bool TryGetWindowsSdkDir(string? DesiredVersion, [NotNullWhen(true)] out VersionNumber? OutSdkVersion, [NotNullWhen(true)] out DirectoryReference? OutSdkDir)
 		{
 			// Get a map of Windows SDK versions to their root directories
 			/*IReadOnlyDictionary<VersionNumber, DirectoryReference> WindowsSdkDirs =*/ FindWindowsSdkDirs();
 
 			// Figure out which version number to look for
-			VersionNumber WindowsSdkVersion = null;
+			VersionNumber? WindowsSdkVersion = null;
 			if(!string.IsNullOrEmpty(DesiredVersion))
 			{
-				if(String.Compare(DesiredVersion, "Latest", StringComparison.InvariantCultureIgnoreCase) == 0 && CachedWindowsSdkDirs.Count > 0)
+				if(String.Compare(DesiredVersion, "Latest", StringComparison.InvariantCultureIgnoreCase) == 0 && CachedWindowsSdkDirs!.Count > 0)
 				{
 					WindowsSdkVersion = CachedWindowsSdkDirs.OrderBy(x => x.Key).Last().Key;
 				}
@@ -1882,16 +1881,16 @@ namespace UnrealBuildTool
 			}
 			else
 			{
-				WindowsSdkVersion = PreferredWindowsSdkVersions.FirstOrDefault(x => CachedWindowsSdkDirs.ContainsKey(x));
-				if(WindowsSdkVersion == null && CachedWindowsSdkDirs.Count > 0)
+				WindowsSdkVersion = PreferredWindowsSdkVersions.FirstOrDefault(x => CachedWindowsSdkDirs!.ContainsKey(x));
+				if(WindowsSdkVersion == null && CachedWindowsSdkDirs!.Count > 0)
 				{
 					WindowsSdkVersion = CachedWindowsSdkDirs.OrderBy(x => x.Key).Last().Key;
 				}
 			}
 
 			// Get the actual directory for this version
-			DirectoryReference SdkDir;
-			if(WindowsSdkVersion != null && CachedWindowsSdkDirs.TryGetValue(WindowsSdkVersion, out SdkDir))
+			DirectoryReference? SdkDir;
+			if(WindowsSdkVersion != null && CachedWindowsSdkDirs!.TryGetValue(WindowsSdkVersion, out SdkDir))
 			{
 				OutSdkDir = SdkDir;
 				OutSdkVersion = WindowsSdkVersion;
@@ -1910,7 +1909,7 @@ namespace UnrealBuildTool
 		/// </summary>
 		/// <param name="OutInstallDir">Receives the installation directory on success</param>
 		/// <returns>True if the directory was found, false otherwise</returns>
-		public static bool TryGetNetFxSdkInstallDir(out DirectoryReference OutInstallDir)
+		public static bool TryGetNetFxSdkInstallDir([NotNullWhen(true)] out DirectoryReference? OutInstallDir)
 		{
 			DirectoryReference HostAutoSdkDir;
 			if(UEBuildPlatformSDK.TryGetHostPlatformAutoSDKDir(out HostAutoSdkDir))
@@ -1941,11 +1940,11 @@ namespace UnrealBuildTool
 			}
 
 			// If we didn't find one of our preferred versions for NetFXSDK, use the max version present on the system
-			Version MaxVersion = null;
-			string MaxVersionString = null;
+			Version? MaxVersion = null;
+			string? MaxVersionString = null;
 			foreach (string ExistingVersionString in ReadInstallDirSubKeys32(NetFxSDKKeyName))
 			{
-				Version ExistingVersion;
+				Version? ExistingVersion;
 				if (!Version.TryParse(ExistingVersionString, out ExistingVersion))
 				{
 					continue;
@@ -2370,7 +2369,7 @@ namespace UnrealBuildTool
 		/// </summary>
 		/// <param name="ProjectFile">The project file being built</param>
 		/// <param name="Metadata">String builder to contain build metadata</param>
-		public override void GetExternalBuildMetadata(FileReference ProjectFile, StringBuilder Metadata)
+		public override void GetExternalBuildMetadata(FileReference? ProjectFile, StringBuilder Metadata)
 		{
 			base.GetExternalBuildMetadata(ProjectFile, Metadata);
 
diff --git a/Engine/Source/Programs/UnrealBuildTool/Platform/Windows/VCEnvironment.cs b/Engine/Source/Programs/UnrealBuildTool/Platform/Windows/VCEnvironment.cs
--- a/Engine/Source/Programs/UnrealBuildTool/Platform/Windows/VCEnvironment.cs
+++ b/Engine/Source/Programs/UnrealBuildTool/Platform/Windows/VCEnvironment.cs
@@ -10,8 +10,6 @@ using Microsoft.Win32;
 using System.Text;
 using Tools.DotNETCommon;
 
-#nullable disable
-
 namespace UnrealBuildTool
 {
 	/// <summary>
@@ -437,7 +435,7 @@ namespace UnrealBuildTool
 			}
 
 			// Add the NETFXSDK include path. We need this for SwarmInterface.
-			DirectoryReference NetFxSdkDir;
+			DirectoryReference? NetFxSdkDir;
 			if(WindowsPlatform.TryGetNetFxSdkInstallDir(out NetFxSdkDir))
 			{
 				IncludePaths.Add(DirectoryReference.Combine(NetFxSdkDir, "include", "um"));
@@ -483,11 +481,11 @@ namespace UnrealBuildTool
 		/// <param name="WindowsSdkVersion">Version of the Windows SDK to use</param>
 		/// <param name="SuppliedSdkDirectoryForVersion">If specified, this is the SDK directory to use, otherwise, attempt to look up via registry. If specified, the WindowsSdkVersion is used directly</param>
 		/// <returns>New environment object with paths for the given settings</returns>
-		public static VCEnvironment Create(WindowsCompiler Compiler, UnrealTargetPlatform Platform, WindowsArchitecture Architecture, string CompilerVersion, string WindowsSdkVersion, string SuppliedSdkDirectoryForVersion)
+		public static VCEnvironment Create(WindowsCompiler Compiler, UnrealTargetPlatform Platform, WindowsArchitecture Architecture, string? CompilerVersion, string? WindowsSdkVersion, string? SuppliedSdkDirectoryForVersion)
 		{
 			// Get the compiler version info
-			VersionNumber SelectedCompilerVersion;
-			DirectoryReference SelectedCompilerDir;
+			VersionNumber? SelectedCompilerVersion;
+			DirectoryReference? SelectedCompilerDir;
 			if(!WindowsPlatform.TryGetToolChainDir(Compiler, CompilerVersion, out SelectedCompilerVersion, out SelectedCompilerDir))
 			{
 				throw new BuildException("{0}{1} must be installed in order to build this target.", WindowsPlatform.GetCompilerName(Compiler), String.IsNullOrEmpty(CompilerVersion)? "" : String.Format(" ({0})", CompilerVersion));
@@ -495,8 +493,8 @@ namespace UnrealBuildTool
 
 			// Get the toolchain info
 			WindowsCompiler ToolChain;
-			VersionNumber SelectedToolChainVersion;
-			DirectoryReference SelectedToolChainDir;
+			VersionNumber? SelectedToolChainVersion;
+			DirectoryReference? SelectedToolChainDir;
 			if(Compiler == WindowsCompiler.Clang || Compiler == WindowsCompiler.Intel)
 			{
 				if (WindowsPlatform.TryGetToolChainDir(WindowsCompiler.VisualStudio2019, null, out SelectedToolChainVersion, out SelectedToolChainDir))
@@ -520,12 +518,12 @@ namespace UnrealBuildTool
 			}
 
 			// Get the actual Windows SDK directory
-			VersionNumber SelectedWindowsSdkVersion;
-			DirectoryReference SelectedWindowsSdkDir;
+			VersionNumber? SelectedWindowsSdkVersion;
+			DirectoryReference? SelectedWindowsSdkDir;
 			if (SuppliedSdkDirectoryForVersion != null)
 			{
 				SelectedWindowsSdkDir = new DirectoryReference(SuppliedSdkDirectoryForVersion);
-				SelectedWindowsSdkVersion = VersionNumber.Parse(WindowsSdkVersion);
+				SelectedWindowsSdkVersion = VersionNumber.Parse(WindowsSdkVersion!);
 
 				if (!DirectoryReference.Exists(SelectedWindowsSdkDir))
 				{
diff --git a/Engine/Source/Programs/UnrealBuildTool/Platform/Windows/VCPreludeCache.cs b/Engine/Source/Programs/UnrealBuildTool/Platform/Windows/VCPreludeCache.cs
--- a/Engine/Source/Programs/UnrealBuildTool/Platform/Windows/VCPreludeCache.cs
+++ b/Engine/Source/Programs/UnrealBuildTool/Platform/Windows/VCPreludeCache.cs
@@ -7,8 +7,6 @@ using System.IO;
 using System.Linq;
 using Tools.DotNETCommon;
 
-#nullable disable
-
 namespace UnrealBuildTool
 {
 	/// <summary>
@@ -52,7 +50,7 @@ namespace UnrealBuildTool
 		/// <summary>
 		/// Path to the file used to generate prelude headers
 		/// </summary>
-		static FileReference PreludeGeneratorFile;
+		static FileReference? PreludeGeneratorFile;
 
 		/// <summary>
 		/// Cached mapping from compiler filename to its version string
@@ -68,7 +66,7 @@ namespace UnrealBuildTool
 		public static FileReference GetPreludeHeader(FileReference Compiler, List<string> Arguments)
 		{
 			// Cache the compiler version
-			string CompilerVersion;
+			string? CompilerVersion;
 			lock (CompilerExeToVersion)
 			{
 				if (!CompilerExeToVersion.TryGetValue(Compiler, out CompilerVersion))
diff --git a/Engine/Source/Programs/UnrealBuildTool/Platform/Windows/WindowsExports.cs b/Engine/Source/Programs/UnrealBuildTool/Platform/Windows/WindowsExports.cs
--- a/Engine/Source/Programs/UnrealBuildTool/Platform/Windows/WindowsExports.cs
+++ b/Engine/Source/Programs/UnrealBuildTool/Platform/Windows/WindowsExports.cs
@@ -2,6 +2,7 @@
 
 using System;
 using System.Collections.Generic;
+using System.Diagnostics.CodeAnalysis;
 using System.Linq;
 using System.Text;
 using System.Threading.Tasks;
@@ -20,7 +21,7 @@ namespace UnrealBuildTool
 		/// <param name="Compiler">The compiler version</param>
 		/// <param name="InstallDir">Receives the install directory on success</param>
 		/// <returns>True if successful</returns>
-		public static bool TryGetVSInstallDir(WindowsCompiler Compiler, out DirectoryReference InstallDir)
+		public static bool TryGetVSInstallDir(WindowsCompiler Compiler, [NotNullWhen(true)] out DirectoryReference? InstallDir)
 		{
 			return WindowsPlatform.TryGetVSInstallDir(Compiler, out InstallDir);
 		}
@@ -51,9 +52,9 @@ namespace UnrealBuildTool
 		/// <param name="OutSdkVersion">Version of SDK</param>
 		/// <param name="OutSdkDir">Path to SDK root folder</param>
 		/// <returns>String with the name</returns>
-		public static bool TryGetWindowsSdkDir(string DesiredVersion, out Version OutSdkVersion, out DirectoryReference OutSdkDir)
+		public static bool TryGetWindowsSdkDir(string DesiredVersion, [NotNullWhen(true)] out Version? OutSdkVersion, [NotNullWhen(true)] out DirectoryReference? OutSdkDir)
 		{
-			VersionNumber vn;
+			VersionNumber? vn;
 			if(WindowsPlatform.TryGetWindowsSdkDir(DesiredVersion, out vn, out OutSdkDir))
 			{
 				OutSdkVersion = new Version(vn.ToString());
@@ -72,8 +73,8 @@ namespace UnrealBuildTool
 			List<KeyValuePair<string, DirectoryReference>> WindowsSdkDirs = new List<KeyValuePair<string, DirectoryReference>>();
 
 			// Add the default directory first
-			VersionNumber Version;
-			DirectoryReference DefaultWindowsSdkDir;
+			VersionNumber? Version;
+			DirectoryReference? DefaultWindowsSdkDir;
 			if (WindowsPlatform.TryGetWindowsSdkDir(null, out Version, out DefaultWindowsSdkDir))
 			{
 				WindowsSdkDirs.Add(new KeyValuePair<string, DirectoryReference>(Version.ToString(), DefaultWindowsSdkDir));
diff --git a/Engine/Source/Programs/UnrealBuildTool/System/Action.cs b/Engine/Source/Programs/UnrealBuildTool/System/Action.cs
--- a/Engine/Source/Programs/UnrealBuildTool/System/Action.cs
+++ b/Engine/Source/Programs/UnrealBuildTool/System/Action.cs
@@ -9,8 +9,6 @@ using System.Text;
 using System.Threading.Tasks;
 using Tools.DotNETCommon;
 
-#nullable disable
-
 namespace UnrealBuildTool
 {
 	/// <summary>
@@ -64,7 +62,7 @@ namespace UnrealBuildTool
 		/// <summary>
 		/// For C++ source files, specifies a dependency list file used to check changes to header files
 		/// </summary>
-		FileItem DependencyListFile { get; }
+		FileItem? DependencyListFile { get; }
 
 		/// <summary>
 		/// Directory from which to execute the program to create produced items
@@ -150,27 +148,27 @@ namespace UnrealBuildTool
 		/// <summary>
 		/// For C++ source files, specifies a dependency list file used to check changes to header files
 		/// </summary>
-		public FileItem DependencyListFile { get; set; }
+		public FileItem? DependencyListFile { get; set; }
 
 		/// <summary>
 		/// Directory from which to execute the program to create produced items
 		/// </summary>
-		public DirectoryReference WorkingDirectory { get; set; } = null;
+		public DirectoryReference WorkingDirectory { get; set; } = null!;
 
 		/// <summary>
 		/// The command to run to create produced items
 		/// </summary>
-		public FileReference CommandPath { get; set; } = null;
+		public FileReference CommandPath { get; set; } = null!;
 
 		/// <summary>
 		/// Command-line parameters to pass to the program
 		/// </summary>
-		public string CommandArguments { get; set; } = null;
+		public string CommandArguments { get; set; } = null!;
 
 		/// <summary>
 		/// Optional friendly description of the type of command being performed, for example "Compile" or "Link".  Displayed by some executors.
 		/// </summary>
-		public string CommandDescription { get; set; } = null;
+		public string CommandDescription { get; set; } = null!;
 
 		/// <summary>
 		/// Human-readable description of this action that may be displayed as status while invoking the action.  This is often the name of the file being compiled, or an executable file name being linked.  Displayed by some executors.
@@ -364,7 +362,7 @@ namespace UnrealBuildTool
 				Action.DeleteItems.AddRange(DeleteItems.Select(x => FileItem.GetItemByPath(x)));
 			}
 
-			string DependencyListFile;
+			string? DependencyListFile;
 			if (Object.TryGetStringField("DependencyListFile", out DependencyListFile))
 			{
 				Action.DependencyListFile = FileItem.GetItemByPath(DependencyListFile);
@@ -473,7 +471,7 @@ namespace UnrealBuildTool
 		/// <summary>
 		/// Set of other actions that this action depends on. This set is built when the action graph is linked.
 		/// </summary>
-		public HashSet<LinkedAction> PrerequisiteActions;
+		public HashSet<LinkedAction> PrerequisiteActions = null!;
 
 		/// <summary>
 		/// Total number of actions depending on this one.
@@ -491,7 +489,7 @@ namespace UnrealBuildTool
 		public IEnumerable<FileItem> PrerequisiteItems => Inner.PrerequisiteItems;
 		public IEnumerable<FileItem> ProducedItems => Inner.ProducedItems;
 		public IEnumerable<FileItem> DeleteItems => Inner.DeleteItems;
-		public FileItem DependencyListFile => Inner.DependencyListFile;
+		public FileItem? DependencyListFile => Inner.DependencyListFile;
 		public DirectoryReference WorkingDirectory => Inner.WorkingDirectory;
 		public FileReference CommandPath => Inner.CommandPath;
 		public string CommandArguments => Inner.CommandArguments;
diff --git a/Engine/Source/Programs/UnrealBuildTool/System/ActionGraph.cs b/Engine/Source/Programs/UnrealBuildTool/System/ActionGraph.cs
--- a/Engine/Source/Programs/UnrealBuildTool/System/ActionGraph.cs
+++ b/Engine/Source/Programs/UnrealBuildTool/System/ActionGraph.cs
@@ -10,8 +10,6 @@ using Tools.DotNETCommon;
 using System.Text;
 using System.Threading.Tasks;
 
-#nullable disable
-
 namespace UnrealBuildTool
 {
 	static class ActionGraph
@@ -41,7 +39,7 @@ namespace UnrealBuildTool
 				Action.PrerequisiteActions = new HashSet<LinkedAction>();
 				foreach(FileItem PrerequisiteItem in Action.PrerequisiteItems)
 				{
-					LinkedAction PrerequisiteAction;
+					LinkedAction? PrerequisiteAction;
 					if(ItemToProducingAction.TryGetValue(PrerequisiteItem, out PrerequisiteAction))
 					{
 						Action.PrerequisiteActions.Add(PrerequisiteAction);
@@ -66,7 +64,7 @@ namespace UnrealBuildTool
 			{
 				foreach(FileItem ProducedItem in Action.ProducedItems)
 				{
-					IExternalAction ExistingAction;
+					IExternalAction? ExistingAction;
 					if(ItemToProducingAction.TryGetValue(ProducedItem, out ExistingAction))
 					{
 						bResult &= CheckForConflicts(ExistingAction, Action);
@@ -336,7 +334,7 @@ namespace UnrealBuildTool
 						bool bActionOnlyDependsOnNonCyclicalActions = true;
 						foreach (FileItem PrerequisiteItem in Action.PrerequisiteItems)
 						{
-							LinkedAction ProducingAction;
+							LinkedAction? ProducingAction;
 							if (ItemToProducingAction.TryGetValue(PrerequisiteItem, out ProducingAction))
 							{
 								if (!ActionIsNonCyclical.ContainsKey(ProducingAction))
@@ -497,7 +495,7 @@ namespace UnrealBuildTool
 			}
 
 			// Determine the last time the action was run based on the write times of its produced files.
-			string LatestUpdatedProducedItemName = null;
+			string? LatestUpdatedProducedItemName = null;
 			DateTimeOffset LastExecutionTimeUtc = DateTimeOffset.MaxValue;
 			foreach (FileItem ProducedItem in RootAction.ProducedItems)
 			{
@@ -590,7 +588,7 @@ namespace UnrealBuildTool
 			// Check the dependency list
 			if(!bIsOutdated && RootAction.DependencyListFile != null)
 			{
-				List<FileItem> DependencyFiles;
+				List<FileItem>? DependencyFiles;
 				if(!CppDependencies.TryGetDependencies(RootAction.DependencyListFile, out DependencyFiles))
 				{
 					Log.TraceLog("{0}: Missing dependency list file \"{1}\"", RootAction.StatusDescription, RootAction.DependencyListFile);
@@ -683,7 +681,7 @@ namespace UnrealBuildTool
 						Dependencies.Add(Action.DependencyListFile);
 					}
 				}
-				Parallel.ForEach(Dependencies, File => { List<FileItem> Temp; CppDependencies.TryGetDependencies(File, out Temp); });
+				Parallel.ForEach(Dependencies, File => { List<FileItem>? Temp; CppDependencies.TryGetDependencies(File, out Temp); });
 			}
 
 			using(Timeline.ScopeEvent("Cache outdated actions"))
@@ -769,11 +767,12 @@ namespace UnrealBuildTool
 				Writer.WriteObjectStart();
 
 				Writer.WriteObjectStart("Environment");
-				foreach (System.Collections.DictionaryEntry Pair in Environment.GetEnvironmentVariables())
+				foreach (object? Object in Environment.GetEnvironmentVariables())
 				{
-					if (!UnrealBuildTool.InitialEnvironment.Contains(Pair.Key) || (string)(UnrealBuildTool.InitialEnvironment[Pair.Key]) != (string)(Pair.Value))
+					System.Collections.DictionaryEntry Pair = (System.Collections.DictionaryEntry)Object!;
+					if (!UnrealBuildTool.InitialEnvironment!.Contains(Pair.Key) || (string)(UnrealBuildTool.InitialEnvironment![Pair.Key]!) != (string)(Pair.Value!))
 					{
-						Writer.WriteValue((string)Pair.Key, (string)Pair.Value);
+						Writer.WriteValue((string)Pair!.Key, (string)Pair.Value!);
 					}
 				}
 				Writer.WriteObjectEnd();
diff --git a/Engine/Source/Programs/UnrealBuildTool/System/ActionSerializer.cs b/Engine/Source/Programs/UnrealBuildTool/System/ActionSerializer.cs
--- a/Engine/Source/Programs/UnrealBuildTool/System/ActionSerializer.cs
+++ b/Engine/Source/Programs/UnrealBuildTool/System/ActionSerializer.cs
@@ -8,8 +8,6 @@ using System.Text;
 using System.Threading.Tasks;
 using Tools.DotNETCommon;
 
-#nullable disable
-
 namespace UnrealBuildTool
 {
 	/// <summary>
@@ -96,7 +94,7 @@ namespace UnrealBuildTool
 			{
 				if (Type.IsClass && !Type.IsAbstract && typeof(IActionSerializer).IsAssignableFrom(Type))
 				{
-					IActionSerializer Serializer = (IActionSerializer)Activator.CreateInstance(Type);
+					IActionSerializer Serializer = (IActionSerializer)Activator.CreateInstance(Type)!;
 					TypeToSerializer[Serializer.Type] = Serializer;
 					NameToSerializer[Type.Name] = Serializer;
 				}
@@ -123,7 +121,7 @@ namespace UnrealBuildTool
 		{
 			string Name = Reader.ReadString();
 
-			IActionSerializer Serializer;
+			IActionSerializer? Serializer;
 			if (!NameToSerializer.TryGetValue(Name, out Serializer))
 			{
 				throw new BuildException("Unable to find action type '{0}'", Name);
@@ -141,7 +139,7 @@ namespace UnrealBuildTool
 		{
 			Type Type = Action.GetType();
 
-			IActionSerializer Serializer;
+			IActionSerializer? Serializer;
 			if (!TypeToSerializer.TryGetValue(Type, out Serializer))
 			{
 				throw new BuildException("Unable to find serializer for action type '{0}'", Type.Name);
diff --git a/Engine/Source/Programs/UnrealBuildTool/System/FileItem.cs b/Engine/Source/Programs/UnrealBuildTool/System/FileItem.cs
--- a/Engine/Source/Programs/UnrealBuildTool/System/FileItem.cs
+++ b/Engine/Source/Programs/UnrealBuildTool/System/FileItem.cs
@@ -303,9 +303,9 @@ namespace UnrealBuildTool
 		/// </summary>
 		/// <param name="Writer">Writer to serialize data to</param>
 		/// <param name="FileItem">File item to write</param>
-		public static void WriteFileItem(this BinaryArchiveWriter Writer, FileItem FileItem)
+		public static void WriteFileItem(this BinaryArchiveWriter Writer, FileItem? FileItem)
 		{
-			Writer.WriteObjectReference<FileItem>(FileItem, () => Writer.WriteFileReference(FileItem.Location));
+			Writer.WriteObjectReference<FileItem>(FileItem!, () => Writer.WriteFileReference(FileItem!.Location));
 		}
 
 		/// <summary>
diff --git a/Engine/Source/Programs/UnrealBuildTool/System/SourceFileMetadataCache.cs b/Engine/Source/Programs/UnrealBuildTool/System/SourceFileMetadataCache.cs
--- a/Engine/Source/Programs/UnrealBuildTool/System/SourceFileMetadataCache.cs
+++ b/Engine/Source/Programs/UnrealBuildTool/System/SourceFileMetadataCache.cs
@@ -294,7 +294,7 @@ namespace UnrealBuildTool
 		/// </summary>
 		/// <param name="ProjectFile">Project file for the target being built</param>
 		/// <returns>Dependency cache hierarchy for the given project</returns>
-		public static IEnumerable<FileReference> GetFilesToClean(FileReference ProjectFile)
+		public static IEnumerable<FileReference> GetFilesToClean(FileReference? ProjectFile)
 		{
 			if(ProjectFile == null || !UnrealBuildTool.IsEngineInstalled())
 			{
diff --git a/Engine/Source/Programs/UnrealBuildTool/System/Utils.cs b/Engine/Source/Programs/UnrealBuildTool/System/Utils.cs
--- a/Engine/Source/Programs/UnrealBuildTool/System/Utils.cs
+++ b/Engine/Source/Programs/UnrealBuildTool/System/Utils.cs
@@ -333,7 +333,7 @@ namespace UnrealBuildTool
 				{
 					if (!Directory.Exists(LogDir))
 					{
-						string IniPath = UnrealBuildTool.GetRemoteIniPath();
+						string? IniPath = UnrealBuildTool.GetRemoteIniPath();
 						if(string.IsNullOrEmpty(IniPath))
 						{
 							break;
diff --git a/Engine/Source/Programs/UnrealBuildTool/System/VersionNumber.cs b/Engine/Source/Programs/UnrealBuildTool/System/VersionNumber.cs
--- a/Engine/Source/Programs/UnrealBuildTool/System/VersionNumber.cs
+++ b/Engine/Source/Programs/UnrealBuildTool/System/VersionNumber.cs
@@ -79,7 +79,7 @@ namespace UnrealBuildTool
 		/// <param name="Lhs">The first version number</param>
 		/// <param name="Rhs">The second version number</param>
 		/// <returns>True if the versions are equal.</returns>
-		public static bool operator==(VersionNumber Lhs, VersionNumber Rhs)
+		public static bool operator==(VersionNumber? Lhs, VersionNumber? Rhs)
 		{
 			if(Object.ReferenceEquals(Lhs, null))
 			{
@@ -97,7 +97,7 @@ namespace UnrealBuildTool
 		/// <param name="Lhs">The first version number</param>
 		/// <param name="Rhs">The second version number</param>
 		/// <returns>True if the versions are not equal.</returns>
-		public static bool operator!=(VersionNumber Lhs, VersionNumber Rhs)
+		public static bool operator!=(VersionNumber? Lhs, VersionNumber? Rhs)
 		{
 			return !(Lhs == Rhs);
 		}
diff --git a/Engine/Source/Programs/UnrealBuildTool/UnrealBuildTool.cs b/Engine/Source/Programs/UnrealBuildTool/UnrealBuildTool.cs
--- a/Engine/Source/Programs/UnrealBuildTool/UnrealBuildTool.cs
+++ b/Engine/Source/Programs/UnrealBuildTool/UnrealBuildTool.cs
@@ -10,8 +10,6 @@ using System.Text;
 using System.Threading;
 using Tools.DotNETCommon;
 
-#nullable disable
-
 namespace UnrealBuildTool
 {
 	static class UnrealBuildTool
@@ -25,7 +23,7 @@ namespace UnrealBuildTool
 		/// <summary>
 		/// The environment at boot time.
 		/// </summary>
-		static public System.Collections.IDictionary InitialEnvironment;
+		static public System.Collections.IDictionary? InitialEnvironment;
 
 		/// <summary>
 		/// Whether we're running with engine installed
@@ -40,12 +38,12 @@ namespace UnrealBuildTool
 		/// <summary>
 		/// If we are running with an installed project, specifies the path to it
 		/// </summary>
-		static FileReference InstalledProjectFile;
+		static FileReference? InstalledProjectFile;
 
 		/// <summary>
 		/// Directory for saved application settings (typically Engine/Programs)
 		/// </summary>
-		static DirectoryReference CachedEngineProgramSavedDirectory;
+		static DirectoryReference? CachedEngineProgramSavedDirectory;
 
 		/// <summary>
 		/// The path to UBT
@@ -137,7 +135,7 @@ namespace UnrealBuildTool
 		/// <returns>List of extension directories, including the given base directory</returns>
 		public static List<DirectoryReference> GetExtensionDirs(DirectoryReference BaseDir, bool bIncludePlatformDirectories=true, bool bIncludeRestrictedDirectories=true, bool bIncludeBaseDirectory=true)
 		{
-			Tuple<List<DirectoryReference>, List<DirectoryReference>> CachedDirs;
+			Tuple<List<DirectoryReference>, List<DirectoryReference>>? CachedDirs;
 			if (!CachedExtensionDirectories.TryGetValue(BaseDir, out CachedDirs))
 			{
 				CachedDirs = Tuple.Create(new List<DirectoryReference>(), new List<DirectoryReference>());
@@ -209,7 +207,7 @@ namespace UnrealBuildTool
 		/// <summary>
 		/// The Remote Ini directory.  This should always be valid when compiling using a remote server.
 		/// </summary>
-		static string RemoteIniPath = null;
+		static string? RemoteIniPath = null;
 
 		/// <summary>
 		/// Returns true if UnrealBuildTool is running using installed Engine components
@@ -251,7 +249,7 @@ namespace UnrealBuildTool
 		/// Gets the installed project file
 		/// </summary>
 		/// <returns>Location of the installed project file</returns>
-		static public FileReference GetInstalledProjectFile()
+		static public FileReference? GetInstalledProjectFile()
 		{
 			if(IsProjectInstalled())
 			{
@@ -274,7 +272,7 @@ namespace UnrealBuildTool
 			{
 				return true;
 			}
-			if(IsProjectInstalled() && File.IsUnderDirectory(InstalledProjectFile.Directory))
+			if(IsProjectInstalled() && File.IsUnderDirectory(InstalledProjectFile!.Directory))
 			{
 				return true;
 			}
@@ -294,7 +292,7 @@ namespace UnrealBuildTool
 		/// The Unreal remote tool ini directory.  This should be valid if compiling using a remote server
 		/// </summary>
 		/// <returns>The directory path</returns>
-		static public string GetRemoteIniPath()
+		static public string? GetRemoteIniPath()
 		{
 			return RemoteIniPath;
 		}
@@ -320,7 +318,7 @@ namespace UnrealBuildTool
 			/// Specifies the path to a log file to write. Note that the default mode (eg. building, generating project files) will create a log file by default if this not specified.
 			/// </summary>
 			[CommandLine(Prefix = "-Log")]
-			public FileReference LogFileName = null;
+			public FileReference? LogFileName = null;
 
 			/// <summary>
 			/// Whether to include timestamps in the log
@@ -376,7 +374,7 @@ namespace UnrealBuildTool
 			[CommandLine("-VSMac", Value="GenerateProjectFiles")]
 			[CommandLine("-CLion", Value="GenerateProjectFiles")]
 			[CommandLine("-Rider", Value="GenerateProjectFiles")]
-			public string Mode = null;
+			public string? Mode = null;
 
 			/// <summary>
 			/// Initialize the options with the given commnad line arguments
@@ -399,7 +397,7 @@ namespace UnrealBuildTool
 		/// <returns>Zero on success, non-zero on error</returns>
 		private static int Main(string[] ArgumentsArray)
 		{
-			SingleInstanceMutex Mutex = null;
+			SingleInstanceMutex? Mutex = null;
 			try
 			{
 				// Start capturing performance info
@@ -440,7 +438,7 @@ namespace UnrealBuildTool
 				Encoding.RegisterProvider(CodePagesEncodingProvider.Instance);
 
 				// Get the type of the mode to execute, using a fast-path for the build mode.
-				Type ModeType = typeof(BuildMode);
+				Type? ModeType = typeof(BuildMode);
 				if(Options.Mode != null)
 				{
 					// Find all the valid modes
@@ -449,7 +447,7 @@ namespace UnrealBuildTool
 					{
 						if(Type.IsClass && !Type.IsAbstract && Type.IsSubclassOf(typeof(ToolMode)))
 						{
-							ToolModeAttribute Attribute = Type.GetCustomAttribute<ToolModeAttribute>();
+							ToolModeAttribute? Attribute = Type.GetCustomAttribute<ToolModeAttribute>();
 							if(Attribute == null)
 							{
 								throw new BuildException("Class '{0}' should have a ToolModeAttribute", Type.Name);
@@ -467,7 +465,7 @@ namespace UnrealBuildTool
 				}
 
 				// Get the options for which systems have to be initialized for this mode
-				ToolModeOptions ModeOptions = ModeType.GetCustomAttribute<ToolModeAttribute>().Options;
+				ToolModeOptions ModeOptions = ModeType.GetCustomAttribute<ToolModeAttribute>()!.Options;
 
 				// if we don't care about the trace listener, toss it now
 				if ((ModeOptions & ToolModeOptions.UseStartupTraceListener) == 0)
@@ -489,7 +487,7 @@ namespace UnrealBuildTool
 				{
 					using(Timeline.ScopeEvent("XmlConfig.ReadConfigFiles()"))
 					{
-						string XmlConfigMutexName = SingleInstanceMutex.GetUniqueMutexForPath("UnrealBuildTool_Mutex_XmlConfig", Assembly.GetExecutingAssembly().CodeBase);
+						string XmlConfigMutexName = SingleInstanceMutex.GetUniqueMutexForPath("UnrealBuildTool_Mutex_XmlConfig", Assembly.GetExecutingAssembly().CodeBase!);
 						using(SingleInstanceMutex XmlConfigMutex = new SingleInstanceMutex(XmlConfigMutexName, true))
 						{
 							FileReference XmlConfigCache = Arguments.GetFileReferenceOrDefault("-XmlConfigCache=", null);
@@ -503,7 +501,7 @@ namespace UnrealBuildTool
 				{
 					using(Timeline.ScopeEvent("SingleInstanceMutex.Acquire()"))
 					{
-						string MutexName = SingleInstanceMutex.GetUniqueMutexForPath("UnrealBuildTool_Mutex", Assembly.GetExecutingAssembly().CodeBase);
+						string MutexName = SingleInstanceMutex.GetUniqueMutexForPath("UnrealBuildTool_Mutex", Assembly.GetExecutingAssembly().CodeBase!);
 						Mutex = new SingleInstanceMutex(MutexName, Options.bWaitMutex);
 					}
 				}
@@ -532,7 +530,7 @@ namespace UnrealBuildTool
 				}
 
 				// Create the appropriate handler
-				ToolMode Mode = (ToolMode)Activator.CreateInstance(ModeType);
+				ToolMode Mode = (ToolMode)Activator.CreateInstance(ModeType)!;
 
 				// Execute the mode
 				int Result = Mode.Execute(Arguments);
