UBT: Remove references to fastJSON.

#rnx

diff --git a/Engine/Source/Programs/UnrealBuildTool/System/EncryptionAndSigning.cs b/Engine/Source/Programs/UnrealBuildTool/System/EncryptionAndSigning.cs
--- a/Engine/Source/Programs/UnrealBuildTool/System/EncryptionAndSigning.cs
+++ b/Engine/Source/Programs/UnrealBuildTool/System/EncryptionAndSigning.cs
@@ -5,6 +5,7 @@ using System.Collections.Generic;
 using System.IO;
 using System.Linq;
 using System.Text;
+using System.Text.Json;
 using System.Text.RegularExpressions;
 using Tools.DotNETCommon;
 
@@ -85,15 +86,15 @@ namespace UnrealBuildTool
 			/// <summary>
 			/// Optional name for this encryption key
 			/// </summary>
-			public string Name;
+			public string Name { get; set; }
 			/// <summary>
 			/// Optional guid for this encryption key
 			/// </summary>
-			public string Guid;
+			public string Guid { get; set; }
 			/// <summary>
 			/// AES key
 			/// </summary>
-			public byte[] Key;
+			public byte[] Key { get; set; }
 			/// <summary>
 			/// Determine if this key is valid
 			/// </summary>
@@ -111,58 +112,58 @@ namespace UnrealBuildTool
 			/// <summary>
 			/// AES encyption key
 			/// </summary>
-			public EncryptionKey EncryptionKey = null;
+			public EncryptionKey EncryptionKey { get; set; } = null;
 
 			/// <summary>
 			/// RSA public/private key
 			/// </summary>
-			public SigningKeyPair SigningKey = null;
+			public SigningKeyPair SigningKey { get; set; } = null;
 
 			/// <summary>
 			/// Enable pak signature checking
 			/// </summary>
-			public bool bEnablePakSigning = false;
+			public bool bEnablePakSigning { get; set; } = false;
 
 			/// <summary>
 			/// Encrypt the index of the pak file. Stops the pak file being easily accessible by unrealpak
 			/// </summary>
-			public bool bEnablePakIndexEncryption = false;
+			public bool bEnablePakIndexEncryption { get; set; } = false;
 
 			/// <summary>
 			/// Encrypt all ini files in the pak. Good for game data obsfucation
 			/// </summary>
-			public bool bEnablePakIniEncryption = false;
+			public bool bEnablePakIniEncryption { get; set; } = false;
 
 			/// <summary>
 			/// Encrypt the uasset files in the pak file. After cooking, uasset files only contain package metadata / nametables / export and import tables. Provides good string data obsfucation without
 			/// the downsides of full package encryption, with the security drawbacks of still having some data stored unencrypted 
 			/// </summary>
-			public bool bEnablePakUAssetEncryption = false;
+			public bool bEnablePakUAssetEncryption { get; set; } = false;
 
 			/// <summary>
 			/// Encrypt all assets data files (including exp and ubulk) in the pak file. Likely to be slow, and to cause high data entropy (bad for delta patching)
 			/// </summary>
-			public bool bEnablePakFullAssetEncryption = false;
+			public bool bEnablePakFullAssetEncryption { get; set; } = false;
 
 			/// <summary>
 			/// Some platforms have their own data crypto systems, so allow the config settings to totally disable our own crypto
 			/// </summary>
-			public bool bDataCryptoRequired = false;
+			public bool bDataCryptoRequired { get; set; } = false;
 
 			/// <summary>
 			/// Config setting to enable pak signing
 			/// </summary>
-			public bool PakEncryptionRequired = true;
+			public bool PakEncryptionRequired { get; set; } = true;
 
 			/// <summary>
 			/// Config setting to enable pak encryption
 			/// </summary>
-			public bool PakSigningRequired = true;
+			public bool PakSigningRequired { get; set; } = true;
 			
 			/// <summary>
 			/// A set of named encryption keys that can be used to encrypt different sets of data with a different key that is delivered dynamically (i.e. not embedded within the game executable)
 			/// </summary>
-			public EncryptionKey[] SecondaryEncryptionKeys;
+			public EncryptionKey[] SecondaryEncryptionKeys { get; set; }
 
 			/// <summary>
 			/// 
@@ -186,7 +187,7 @@ namespace UnrealBuildTool
 			public void Save(FileReference InFile)
 			{
 				DirectoryReference.CreateDirectory(InFile.Directory);
-				FileReference.WriteAllText(InFile, fastJSON.JSON.Instance.ToJSON(this, new fastJSON.JSONParameters {}));
+				FileReference.WriteAllText(InFile, JsonSerializer.Serialize(this));
 			}
 		}
 
@@ -231,14 +232,21 @@ namespace UnrealBuildTool
 			CryptoSettings Settings = new CryptoSettings();
 			
 			ConfigHierarchy Ini = ConfigCache.ReadHierarchy(ConfigHierarchyType.Engine, InProjectDirectory, InTargetPlatform);
-			Ini.GetBool("PlatformCrypto", "PlatformRequiresDataCrypto", out Settings.bDataCryptoRequired);
-			Ini.GetBool("PlatformCrypto", "PakSigningRequired", out Settings.PakSigningRequired);
-			Ini.GetBool("PlatformCrypto", "PakEncryptionRequired", out Settings.PakEncryptionRequired);
+			Ini.GetBool("PlatformCrypto", "PlatformRequiresDataCrypto", out bool bDataCryptoRequired);
+			Settings.bDataCryptoRequired = bDataCryptoRequired;
+
+			Ini.GetBool("PlatformCrypto", "PakSigningRequired", out bool PakSigningRequired);
+			Settings.PakSigningRequired = PakSigningRequired;
+
+			Ini.GetBool("PlatformCrypto", "PakEncryptionRequired", out bool PakEncryptionRequired);
+			Settings.PakEncryptionRequired = PakEncryptionRequired;
 
 			{
 				// Start by parsing the legacy encryption.ini settings
 				Ini = ConfigCache.ReadHierarchy(ConfigHierarchyType.Encryption, InProjectDirectory, InTargetPlatform);
-				Ini.GetBool("Core.Encryption", "SignPak", out Settings.bEnablePakSigning);
+
+				Ini.GetBool("Core.Encryption", "SignPak", out bool bEnablePakSigning);
+				Settings.bEnablePakSigning = bEnablePakSigning;
 
 				string[] SigningKeyStrings = new string[3];
 				Ini.GetString("Core.Encryption", "rsa.privateexp", out SigningKeyStrings[0]);
@@ -266,7 +274,9 @@ namespace UnrealBuildTool
 					}
 				}
 
-				Ini.GetBool("Core.Encryption", "EncryptPak", out Settings.bEnablePakIndexEncryption);
+				Ini.GetBool("Core.Encryption", "EncryptPak", out bool bEnablePakIndexEncryption);
+				Settings.bEnablePakIndexEncryption = bEnablePakIndexEncryption;
+
 				Settings.bEnablePakFullAssetEncryption = false;
 				Settings.bEnablePakUAssetEncryption = false;
 				Settings.bEnablePakIniEncryption = Settings.bEnablePakIndexEncryption;
@@ -297,11 +307,20 @@ namespace UnrealBuildTool
 			// If we have new format crypto keys, read them in over the top of the legacy settings
 			if (CryptoSection != null && CryptoSection.KeyNames.Count() > 0)
 			{
-				Ini.GetBool(SectionName, "bEnablePakSigning", out Settings.bEnablePakSigning);
-				Ini.GetBool(SectionName, "bEncryptPakIniFiles", out Settings.bEnablePakIniEncryption);
-				Ini.GetBool(SectionName, "bEncryptPakIndex", out Settings.bEnablePakIndexEncryption);
-				Ini.GetBool(SectionName, "bEncryptUAssetFiles", out Settings.bEnablePakUAssetEncryption);
-				Ini.GetBool(SectionName, "bEncryptAllAssetFiles", out Settings.bEnablePakFullAssetEncryption);
+				Ini.GetBool(SectionName, "bEnablePakSigning", out bool bEnablePakSigning);
+				Settings.bEnablePakSigning = bEnablePakSigning;
+
+				Ini.GetBool(SectionName, "bEncryptPakIniFiles", out bool bEnablePakIniEncryption);
+				Settings.bEnablePakIniEncryption = bEnablePakIniEncryption;
+
+				Ini.GetBool(SectionName, "bEncryptPakIndex", out bool bEnablePakIndexEncryption);
+				Settings.bEnablePakIndexEncryption = bEnablePakIndexEncryption;
+
+				Ini.GetBool(SectionName, "bEncryptUAssetFiles", out bool bEnablePakUAssetEncryption);
+				Settings.bEnablePakUAssetEncryption = bEnablePakUAssetEncryption;
+
+				Ini.GetBool(SectionName, "bEncryptAllAssetFiles", out bool bEnablePakFullAssetEncryption);
+				Settings.bEnablePakFullAssetEncryption = bEnablePakFullAssetEncryption;
 
 				// Parse encryption key
 				string EncryptionKeyString;
diff --git a/Engine/Source/Programs/UnrealBuildTool/System/ExternalExecution.cs b/Engine/Source/Programs/UnrealBuildTool/System/ExternalExecution.cs
--- a/Engine/Source/Programs/UnrealBuildTool/System/ExternalExecution.cs
+++ b/Engine/Source/Programs/UnrealBuildTool/System/ExternalExecution.cs
@@ -6,11 +6,13 @@ using System.Linq;
 using System.Text;
 using System.Diagnostics;
 using System.IO;
+using System.Text.Json;
 using System.Text.RegularExpressions;
 using System.Reflection;
 using System.Runtime.InteropServices;
 using System.Runtime.Serialization;
 using Tools.DotNETCommon;
+using System.Text.Json.Serialization;
 
 namespace UnrealBuildTool
 {
@@ -327,18 +329,19 @@ namespace UnrealBuildTool
 	{
 		public class Module
 		{
-			public string Name;
-			public string ModuleType;
-			public string OverrideModuleType;
-			public string BaseDirectory;
-			public string IncludeBase;     // The include path which all UHT-generated includes should be relative to
-			public string OutputDirectory;
-			public List<string> ClassesHeaders;
-			public List<string> PublicHeaders;
-			public List<string> PrivateHeaders;
-			public string? GeneratedCPPFilenameBase;
-			public bool SaveExportedHeaders;
-			public EGeneratedCodeVersion UHTGeneratedCodeVersion;
+			public string Name { get; set; }
+			public string ModuleType { get; set; }
+			public string OverrideModuleType { get; set; }
+			public string BaseDirectory { get; set; }
+			public string IncludeBase { get; set; }     // The include path which all UHT-generated includes should be relative to
+			public string OutputDirectory { get; set; }
+			public List<string> ClassesHeaders { get; set; }
+			public List<string> PublicHeaders { get; set; }
+			public List<string> PrivateHeaders { get; set; }
+			public string? GeneratedCPPFilenameBase { get; set; }
+			public bool SaveExportedHeaders { get; set; }
+			[JsonConverter(typeof(JsonStringEnumConverter))]
+			public EGeneratedCodeVersion UHTGeneratedCodeVersion { get; set; }
 
 			public Module(UHTModuleInfo Info)
 			{
@@ -362,11 +365,11 @@ namespace UnrealBuildTool
 			}
 		}
 
-		public bool IsGameTarget;     // True if the current target is a game target
-		public string RootLocalPath;    // The engine path on the local machine
-		public string TargetName;       // Name of the target currently being compiled
-		public string ExternalDependenciesFile; // File to contain additional dependencies that the generated code depends on
-		public List<Module> Modules;
+		public bool IsGameTarget { get; set; }     // True if the current target is a game target
+		public string RootLocalPath { get; set; }    // The engine path on the local machine
+		public string TargetName { get; set; }       // Name of the target currently being compiled
+		public string ExternalDependenciesFile { get; set; } // File to contain additional dependencies that the generated code depends on
+		public List<Module> Modules { get; set; }
 
 		public UHTManifest(string InTargetName, TargetType InTargetType, string InRootLocalPath, string InExternalDependenciesFile, List<Module> InModules)
 		{
@@ -1186,7 +1189,7 @@ namespace UnrealBuildTool
 
 					// Disable extensions when serializing to remove the $type fields
 					Directory.CreateDirectory(ModuleInfoFileName.Directory.FullName);
-					System.IO.File.WriteAllText(ModuleInfoFileName.FullName, fastJSON.JSON.Instance.ToJSON(Manifest, new fastJSON.JSONParameters { UseExtensions = false }));
+					System.IO.File.WriteAllText(ModuleInfoFileName.FullName, JsonSerializer.Serialize(Manifest));
 
 					string CmdLine = (ProjectFile != null) ? "\"" + ProjectFile.FullName + "\"" : TargetName;
 					CmdLine += " \"" + ModuleInfoFileName + "\" -LogCmds=\"loginit warning, logexit warning, logdatabase error\" -Unattended -WarningsAsErrors";
