UBT: Move Windows resx file writer out of shared utility library. Only used in UBT.

#rnx

diff --git a/Engine/Source/Programs/DotNETCommon/DotNETUtilities/ResXResource/UEResXReader.cs b/Engine/Source/Programs/DotNETCommon/DotNETUtilities/ResXResource/UEResXReader.cs
deleted file mode 100644
diff --git a/Engine/Source/Programs/DotNETCommon/DotNETUtilities/ResXResource/UEResXWriter.cs b/Engine/Source/Programs/DotNETCommon/DotNETUtilities/ResXResource/UEResXWriter.cs
deleted file mode 100644
diff --git a/Engine/Source/Programs/UnrealBuildTool/Platform/Windows/ResXResource/UEResXReader.cs b/Engine/Source/Programs/UnrealBuildTool/Platform/Windows/ResXResource/UEResXReader.cs
new file mode 100644
--- /dev/null
+++ b/Engine/Source/Programs/UnrealBuildTool/Platform/Windows/ResXResource/UEResXReader.cs
@@ -0,0 +1,17 @@
+﻿// Copyright Epic Games, Inc. All Rights Reserved.
+
+using System;
+using System.Collections.Generic;
+using System.Linq;
+using System.Text;
+using System.Threading.Tasks;
+
+namespace UnrealBuildTool
+{
+	/// <summary>
+	/// Intentionally unimplemented, but declared so that it's type can be included in generated resx files
+	/// </summary>
+	public abstract class UEResXReader
+	{
+	}
+}
diff --git a/Engine/Source/Programs/UnrealBuildTool/Platform/Windows/ResXResource/UEResXWriter.cs b/Engine/Source/Programs/UnrealBuildTool/Platform/Windows/ResXResource/UEResXWriter.cs
new file mode 100644
--- /dev/null
+++ b/Engine/Source/Programs/UnrealBuildTool/Platform/Windows/ResXResource/UEResXWriter.cs
@@ -0,0 +1,70 @@
+﻿// Copyright Epic Games, Inc. All Rights Reserved.
+
+using System;
+using System.Collections.Generic;
+using System.Linq;
+using System.Text;
+using System.Threading.Tasks;
+using System.Xml;
+
+#nullable disable
+
+namespace UnrealBuildTool
+{
+	class UEResXWriter
+	{
+		public UEResXWriter(string InPath)
+		{
+
+			Filename = InPath;
+			Document = new XmlDocument();
+			Document.AppendChild(Document.CreateXmlDeclaration("1.0", "utf-8", null));
+			RootElement = Document.CreateElement("root");
+			Document.AppendChild(RootElement);
+
+			RootElement.AppendChild(CreateEntry("resheader", "resmimetype", "text/microsoft-resx", null));
+			RootElement.AppendChild(CreateEntry("resheader", "version", "2.0", null));
+			RootElement.AppendChild(CreateEntry("resheader", "reader", typeof(UEResXReader).AssemblyQualifiedName, null));
+			RootElement.AppendChild(CreateEntry("resheader", "writer", typeof(UEResXWriter).AssemblyQualifiedName, null));
+		}
+
+		public void Close()
+		{
+			Document.Save(Filename);
+		}
+
+		public void AddResource(string InName, string InValue)
+		{
+			RootElement.AppendChild(CreateEntry("data", InName, InValue, "preserve"));
+		}
+
+		private XmlNode CreateEntry(string InRootName, string InName, string InValue, string InSpace)
+		{
+			XmlElement Value = Document.CreateElement("value");
+			Value.InnerText = InValue;
+
+			XmlElement Data = Document.CreateElement("data");
+
+			{
+				XmlAttribute Attr = Document.CreateAttribute("name");
+				Attr.Value = InName;
+				Data.Attributes.Append(Attr);
+			}
+
+			if (!string.IsNullOrEmpty(InSpace))
+			{
+				XmlAttribute Attr = Document.CreateAttribute("xml:space");
+				Attr.Value = "preserve";
+				Data.Attributes.Append(Attr);
+			}
+
+			Data.AppendChild(Value);
+
+			return Data;
+		}
+
+		private string Filename;
+		private XmlDocument Document;
+		private XmlElement RootElement;
+	}
+}
