Missing files.

#rnx

diff --git a/Engine/Source/Programs/UnrealBuildTool/ProjectFiles/ProjectFileGenerator.cs b/Engine/Source/Programs/UnrealBuildTool/ProjectFiles/ProjectFileGenerator.cs
--- a/Engine/Source/Programs/UnrealBuildTool/ProjectFiles/ProjectFileGenerator.cs
+++ b/Engine/Source/Programs/UnrealBuildTool/ProjectFiles/ProjectFileGenerator.cs
@@ -434,7 +434,7 @@ namespace UnrealBuildTool
 		{
 			if (SupportsDotnetCoreProjects())
 			{
-				DirectoryInfo DotNetDir = new DirectoryInfo(DirectoryReference.Combine(UnrealBuildTool.EngineSourceDirectory, "Programs", "DotNet").FullName);
+				DirectoryInfo DotNetDir = new DirectoryInfo(DirectoryReference.Combine(UnrealBuildTool.EngineSourceDirectory, "Programs", "Shared").FullName);
 				if (DotNetDir.Exists)
 				{
 					List<FileInfo> ProjectFiles = new List<FileInfo>();
@@ -444,7 +444,7 @@ namespace UnrealBuildTool
 					}
 					if (ProjectFiles.Count > 0)
 					{
-						MasterProjectFolder Folder = ProgramsFolder.AddSubFolder("DotNet");
+						MasterProjectFolder Folder = ProgramsFolder.AddSubFolder("Shared");
 						foreach (FileInfo ProjectFile in ProjectFiles)
 						{
 							VCSharpProjectFile Project = new VCSharpProjectFile(new FileReference(ProjectFile));
