UBT: Fix null reference when running IncludeTool, due to module without a binary.

#rnx

diff --git a/Engine/Source/Programs/UnrealBuildTool/Configuration/UEBuildModule.cs b/Engine/Source/Programs/UnrealBuildTool/Configuration/UEBuildModule.cs
--- a/Engine/Source/Programs/UnrealBuildTool/Configuration/UEBuildModule.cs
+++ b/Engine/Source/Programs/UnrealBuildTool/Configuration/UEBuildModule.cs
@@ -987,7 +987,7 @@ namespace UnrealBuildTool
 		/// <param name="BinaryOutputDir">The output directory for the binary containing this module</param>
 		/// <param name="TargetOutputDir">The output directory for the target executable</param>
 		/// <param name="Writer">Writer for this binary's data</param>
-		public virtual void ExportJson(DirectoryReference BinaryOutputDir, DirectoryReference TargetOutputDir, JsonWriter Writer)
+		public virtual void ExportJson(DirectoryReference? BinaryOutputDir, DirectoryReference? TargetOutputDir, JsonWriter Writer)
 		{
 			Writer.WriteValue("Name", Name);
 			Writer.WriteValue("Directory", ModuleDirectory.FullName);
diff --git a/Engine/Source/Programs/UnrealBuildTool/Configuration/UEBuildTarget.cs b/Engine/Source/Programs/UnrealBuildTool/Configuration/UEBuildTarget.cs
--- a/Engine/Source/Programs/UnrealBuildTool/Configuration/UEBuildTarget.cs
+++ b/Engine/Source/Programs/UnrealBuildTool/Configuration/UEBuildTarget.cs
@@ -2214,7 +2214,7 @@ namespace UnrealBuildTool
 				foreach (UEBuildModule Module in Modules.Values)
 				{
 					Writer.WriteObjectStart(Module.Name);
-					Module.ExportJson(Module.Binary!.OutputDir, GetExecutableDir(), Writer);
+					Module.ExportJson(Module.Binary?.OutputDir, GetExecutableDir(), Writer);
 					Writer.WriteObjectEnd();
 				}
 				Writer.WriteObjectEnd();
