UBT: Fix PDB files not being output for rules assemblies.

#rnx

diff --git a/Engine/Source/Programs/UnrealBuildTool/System/DynamicCompilation.cs b/Engine/Source/Programs/UnrealBuildTool/System/DynamicCompilation.cs
--- a/Engine/Source/Programs/UnrealBuildTool/System/DynamicCompilation.cs
+++ b/Engine/Source/Programs/UnrealBuildTool/System/DynamicCompilation.cs
@@ -169,7 +169,7 @@ namespace UnrealBuildTool
 
 			foreach (FileReference SourceFileName in SourceFileNames)
 			{
-				SourceText Source = SourceText.From(File.ReadAllText(SourceFileName.FullName));
+				SourceText Source = SourceText.From(File.ReadAllText(SourceFileName.FullName), System.Text.Encoding.UTF8);
 				SyntaxTree Tree = CSharpSyntaxTree.ParseText(Source, ParseOptions, SourceFileName.FullName);
 
 				IEnumerable<Diagnostic> Diagnostics = Tree.GetDiagnostics();
@@ -253,18 +253,22 @@ namespace UnrealBuildTool
 
 			using (FileStream AssemblyStream = FileReference.Open(OutputAssemblyPath, FileMode.Create))
 			{
-				EmitOptions EmitOptions = new EmitOptions(
-					includePrivateMembers:true
-				);
+				using (FileStream PdbStream = FileReference.Open(OutputAssemblyPath.ChangeExtension(".pdb"), FileMode.Create))
+				{
+					EmitOptions EmitOptions = new EmitOptions(
+						includePrivateMembers: true
+					);
 
-				EmitResult Result = Compilation.Emit(
-					peStream:AssemblyStream,
-					options:EmitOptions);
+					EmitResult Result = Compilation.Emit(
+						peStream: AssemblyStream,
+						pdbStream: PdbStream,
+						options: EmitOptions);
 
-				if (!Result.Success)
-				{
-					LogDiagnostics(Result.Diagnostics);
-					return null;
+					if (!Result.Success)
+					{
+						LogDiagnostics(Result.Diagnostics);
+						return null;
+					}
 				}
 			}
 
