Do not generate PDB files when compiling rules assemblies on Mac.

#rnx

diff --git a/Engine/Source/Programs/UnrealBuildTool/System/DynamicCompilation.cs b/Engine/Source/Programs/UnrealBuildTool/System/DynamicCompilation.cs
--- a/Engine/Source/Programs/UnrealBuildTool/System/DynamicCompilation.cs
+++ b/Engine/Source/Programs/UnrealBuildTool/System/DynamicCompilation.cs
@@ -16,6 +16,7 @@ using Microsoft.CodeAnalysis.CSharp.Syntax;
 using Microsoft.CodeAnalysis.Emit;
 using System.Reflection.Metadata;
 using Microsoft.CodeAnalysis.Text;
+using System.Runtime.InteropServices;
 
 namespace UnrealBuildTool
 {
@@ -253,7 +254,7 @@ namespace UnrealBuildTool
 
 			using (FileStream AssemblyStream = FileReference.Open(OutputAssemblyPath, FileMode.Create))
 			{
-				using (FileStream PdbStream = FileReference.Open(OutputAssemblyPath.ChangeExtension(".pdb"), FileMode.Create))
+				using (FileStream? PdbStream = RuntimeInformation.IsOSPlatform(OSPlatform.Windows)? FileReference.Open(OutputAssemblyPath.ChangeExtension(".pdb"), FileMode.Create) : null)
 				{
 					EmitOptions EmitOptions = new EmitOptions(
 						includePrivateMembers: true
