Fix for pak signing keys not being saved out to the crypto.json metadata file

diff --git a/Engine/Source/Programs/UnrealBuildTool/System/EncryptionAndSigning.cs b/Engine/Source/Programs/UnrealBuildTool/System/EncryptionAndSigning.cs
--- a/Engine/Source/Programs/UnrealBuildTool/System/EncryptionAndSigning.cs
+++ b/Engine/Source/Programs/UnrealBuildTool/System/EncryptionAndSigning.cs
@@ -26,12 +26,12 @@ namespace UnrealBuildTool
 			/// <summary>
 			/// Exponent
 			/// </summary>
-			public byte[] Exponent;
+			public byte[] Exponent { get; set; }
 
 			/// <summary>
 			/// Modulus
 			/// </summary>
-			public byte[] Modulus;
+			public byte[] Modulus { get; set; }
 
 			/// <summary>
 			/// Determine if this key is valid
@@ -50,12 +50,12 @@ namespace UnrealBuildTool
 			/// <summary>
 			/// Public key
 			/// </summary>
-			public SigningKey PublicKey = new SigningKey();
+			public SigningKey PublicKey { get; set; } = new SigningKey();
 
 			/// <summary>
 			/// Private key
 			/// </summary>
-			public SigningKey PrivateKey = new SigningKey();
+			public SigningKey PrivateKey { get; set; } = new SigningKey();
 
 			/// <summary>
 			/// Determine if this key is valid
