Fix incorrect argument for UE5 project in Xcode project files.

diff --git a/Engine/Source/Programs/UnrealBuildTool/ProjectFiles/Xcode/XcodeProject.cs b/Engine/Source/Programs/UnrealBuildTool/ProjectFiles/Xcode/XcodeProject.cs
--- a/Engine/Source/Programs/UnrealBuildTool/ProjectFiles/Xcode/XcodeProject.cs
+++ b/Engine/Source/Programs/UnrealBuildTool/ProjectFiles/Xcode/XcodeProject.cs
@@ -1644,7 +1644,7 @@ namespace UnrealBuildTool
 			Content.Append("      </BuildableProductRunnable>" + ProjectFileGenerator.NewLine);
 			if (string.IsNullOrEmpty(OldCommandLineArguments))
 			{
-				if (bHasEditorConfiguration && TargetName != "UE4")
+				if (bHasEditorConfiguration && TargetName != "UE5")
 				{
 					Content.Append("      <CommandLineArguments>" + ProjectFileGenerator.NewLine);
 					if (IsForeignProject)
