- Improved Windows SDK detection - moved functionality from UEBuildWindows.cs to MicrosoftPlatformSDK.cs so Turnkey "installed sdk" uses the same function that UBT will use to choose the SDK from a list of preferred SDKs, with fallbacks to other locations.
   - Windows allows for any SDK to be used without updating Env vars, etc, so we don't really need a strict single "InstalledVersion" - we dynamically choose a version every run of UBT, and now Turnkey will match this logic
- Added device detection to Turnkey for Windows

diff --git a/Engine/Source/Programs/AutomationTool/Win/WinPlatform.Automation.cs b/Engine/Source/Programs/AutomationTool/Win/WinPlatform.Automation.cs
--- a/Engine/Source/Programs/AutomationTool/Win/WinPlatform.Automation.cs
+++ b/Engine/Source/Programs/AutomationTool/Win/WinPlatform.Automation.cs
@@ -16,6 +16,26 @@ public abstract class BaseWinPlatform : Platform
 		: base(P)
 	{
 	}
+
+	public override DeviceInfo[] GetDevices()
+	{
+		List<DeviceInfo> Devices = new List<DeviceInfo>();
+
+		if (HostPlatform.Current.HostEditorPlatform == UnrealTargetPlatform.Win64)
+		{
+			DeviceInfo LocalMachine = new DeviceInfo();
+			LocalMachine.Name = LocalMachine.Id = Environment.MachineName;
+			LocalMachine.Type = "Computer";
+			LocalMachine.SoftwareVersion = Environment.OSVersion.Version.ToString();
+
+			Devices.Add(LocalMachine);
+		}
+
+		return Devices.ToArray();
+	}
+
+
+
 	protected override string GetPlatformExeExtension()
 	{
 		return ".exe";
diff --git a/Engine/Source/Programs/UnrealBuildTool/Platform/Windows/MicrosoftPlatformSDK.cs b/Engine/Source/Programs/UnrealBuildTool/Platform/Windows/MicrosoftPlatformSDK.cs
--- a/Engine/Source/Programs/UnrealBuildTool/Platform/Windows/MicrosoftPlatformSDK.cs
+++ b/Engine/Source/Programs/UnrealBuildTool/Platform/Windows/MicrosoftPlatformSDK.cs
@@ -6,27 +6,38 @@ using EpicGames.Core;
 using System.Linq;
 using System.Text.RegularExpressions;
 using Microsoft.Win32;
+using System.Diagnostics.CodeAnalysis;
 
 namespace UnrealBuildTool
 {
 	internal class MicrosoftPlatformSDK : UEBuildPlatformSDK
 	{
+		/// <summary>
+		/// The default Windows SDK version to be used, if installed.
+		/// </summary>
+		static readonly VersionNumber[] PreferredWindowsSdkVersions = new VersionNumber[]
+		{
+			VersionNumber.Parse("10.0.18362.0"),
+			VersionNumber.Parse("10.0.16299.0")
+		};
+
 		public override string GetMainVersion()
 		{
-			// the current Windows SDK version we expect
-			return "10.0.18632";
+			// preferred/main version is the top of the Preferred list - 
+			return PreferredWindowsSdkVersions[0].ToString();
 		}
 
 		public override void GetValidVersionRange(out string MinVersion, out string MaxVersion)
 		{
-			MinVersion = "10.0.00000";
-			MaxVersion = "10.9.99999";
+			MinVersion = "10.0.00000.0";
+			MaxVersion = "10.9.99999.0";
 		}
 
 		public override void GetValidSoftwareVersionRange(out string? MinVersion, out string? MaxVersion)
 		{
-			// @todo we may want to split this for Win64 vs other things
-			MinVersion = MaxVersion = null;
+			// for now, allow any version of Windows 10
+			MinVersion = "10.0.0.0";
+			MaxVersion = null;
 		}
 
 		public override string? GetInstalledSDKVersion()
@@ -36,30 +47,14 @@ namespace UnrealBuildTool
 				return null;
 			}
 
-			// @todo turnkey: do we support pre-10?
-
-			// look for an installed SDK
-			string Version = "v10.0";
-			object Result = Registry.GetValue(@"HKEY_LOCAL_MACHINE\SOFTWARE\Wow6432Node\Microsoft\Microsoft SDKs\Windows\" + Version, "InstallationFolder", null);
-			if (Result != null)
+			// use the PreferredWindowsSdkVersions array to find the SDK version that UBT will use to build with - 
+			VersionNumber? WindowsVersion;
+			if (TryGetWindowsSdkDir(null, out WindowsVersion, out _))
 			{
-				return Registry.GetValue(@"HKEY_LOCAL_MACHINE\SOFTWARE\Wow6432Node\Microsoft\Microsoft SDKs\Windows\" + Version, "ProductVersion", null) as string;
+				return WindowsVersion.ToString();
 			}
 
-			// look in AutoSDK location (note that it doesn't have a setup.bat, at least right now, so the AutoSDK system isn't used
-			// @todo turnkey: make use of AutoSDKs, and maybe get rid of non-Latest style SDK selection in UEBUildWindows? (and move all that stuff here)
-			DirectoryReference HostAutoSdkDir;
-			if (TryGetHostPlatformAutoSDKDir(out HostAutoSdkDir))
-			{
-				DirectoryReference RootDirAutoSdk = DirectoryReference.Combine(HostAutoSdkDir, "Win64", "Windows Kits", "10", "include");
-				if (DirectoryReference.Exists(RootDirAutoSdk))
-				{
-					// sort the directories under 10, and use the highest one (this logic mirrors UEBuildWindows logic)
-					return DirectoryReference.EnumerateDirectories(RootDirAutoSdk).OrderBy(x => x.GetDirectoryName()).Last().GetDirectoryName();
-				}
-			}
-
-
+			// if that failed, we aren't able to build, so give up
 			return null;
 		}
 
@@ -70,7 +65,7 @@ namespace UnrealBuildTool
 			Match Result = Regex.Match(StringValue, @"^(\d+).(\d+).(\d+)");
 			if (Result.Success)
 			{
-				// 8 bits for major, 8 for minor, 16 for patch
+				// 8 bits for major, 8 for minor, 16 for build - we ignore patch (ie the 1234 in 10.0.14031.1234)
 				OutValue |= UInt64.Parse(Result.Groups[1].Value) << 24;
 				OutValue |= UInt64.Parse(Result.Groups[2].Value) << 16;
 				OutValue |= UInt64.Parse(Result.Groups[3].Value) << 0;
@@ -80,5 +75,327 @@ namespace UnrealBuildTool
 
 			return false;
 		}
+
+
+
+		#region Windows Specific SDK discovery
+
+		/// <summary>
+		/// Cache of Windows SDK installation directories
+		/// </summary>
+		private static IReadOnlyDictionary<VersionNumber, DirectoryReference>? CachedWindowsSdkDirs;
+
+		/// <summary>
+		/// Cache of Universal CRT installation directories
+		/// </summary>
+		private static IReadOnlyDictionary<VersionNumber, DirectoryReference>? CachedUniversalCrtDirs;
+
+
+
+
+		/// <summary>
+		/// Updates the CachedWindowsSdkDirs and CachedUniversalCrtDirs variables
+		/// </summary>
+		private static void UpdateCachedWindowsSdks()
+		{
+			Dictionary<VersionNumber, DirectoryReference> WindowsSdkDirs = new Dictionary<VersionNumber, DirectoryReference>();
+			Dictionary<VersionNumber, DirectoryReference> UniversalCrtDirs = new Dictionary<VersionNumber, DirectoryReference>();
+
+			// Enumerate the Windows 8.1 SDK, if present
+			DirectoryReference? InstallDir_8_1;
+			if (TryReadInstallDirRegistryKey32("Microsoft\\Microsoft SDKs\\Windows\\v8.1", "InstallationFolder", out InstallDir_8_1))
+			{
+				if (FileReference.Exists(FileReference.Combine(InstallDir_8_1, "Include", "um", "windows.h")))
+				{
+					Log.TraceLog("Found Windows 8.1 SDK at {0}", InstallDir_8_1);
+					VersionNumber Version_8_1 = new VersionNumber(8, 1);
+					WindowsSdkDirs[Version_8_1] = InstallDir_8_1;
+				}
+			}
+
+			// Find all the root directories for Windows 10 SDKs
+			List<DirectoryReference> InstallDirs_10 = new List<DirectoryReference>();
+			EnumerateSdkRootDirs(InstallDirs_10);
+
+			// Enumerate all the Windows 10 SDKs
+			foreach (DirectoryReference InstallDir_10 in InstallDirs_10.Distinct())
+			{
+				DirectoryReference IncludeRootDir = DirectoryReference.Combine(InstallDir_10, "Include");
+				if (DirectoryReference.Exists(IncludeRootDir))
+				{
+					foreach (DirectoryReference IncludeDir in DirectoryReference.EnumerateDirectories(IncludeRootDir))
+					{
+						VersionNumber? IncludeVersion;
+						if (VersionNumber.TryParse(IncludeDir.GetDirectoryName(), out IncludeVersion))
+						{
+							if (FileReference.Exists(FileReference.Combine(IncludeDir, "um", "windows.h")))
+							{
+								Log.TraceLog("Found Windows 10 SDK version {0} at {1}", IncludeVersion, InstallDir_10);
+								WindowsSdkDirs[IncludeVersion] = InstallDir_10;
+							}
+							if (FileReference.Exists(FileReference.Combine(IncludeDir, "ucrt", "corecrt.h")))
+							{
+								Log.TraceLog("Found Universal CRT version {0} at {1}", IncludeVersion, InstallDir_10);
+								UniversalCrtDirs[IncludeVersion] = InstallDir_10;
+							}
+						}
+					}
+				}
+			}
+
+			CachedWindowsSdkDirs = WindowsSdkDirs;
+			CachedUniversalCrtDirs = UniversalCrtDirs;
+		}
+
+		/// <summary>
+		/// Finds all the installed Windows SDK versions
+		/// </summary>
+		/// <returns>Map of version number to Windows SDK directories</returns>
+		public static IReadOnlyDictionary<VersionNumber, DirectoryReference> FindWindowsSdkDirs()
+		{
+			// Update the cache of install directories, if it's not set
+			if (CachedWindowsSdkDirs == null)
+			{
+				UpdateCachedWindowsSdks();
+			}
+			return CachedWindowsSdkDirs!;
+		}
+
+		/// <summary>
+		/// Finds all the installed Universal CRT versions
+		/// </summary>
+		/// <returns>Map of version number to universal CRT directories</returns>
+		public static IReadOnlyDictionary<VersionNumber, DirectoryReference> FindUniversalCrtDirs()
+		{
+			if (CachedUniversalCrtDirs == null)
+			{
+				UpdateCachedWindowsSdks();
+			}
+			return CachedUniversalCrtDirs!;
+		}
+
+		/// <summary>
+		/// Determines the directory containing the Windows SDK toolchain
+		/// </summary>
+		/// <param name="DesiredVersion">The desired Windows SDK version. This may be "Latest", a specific version number, or null. If null, the function will look for DefaultWindowsSdkVersion. Failing that, it will return the latest version.</param>
+		/// <param name="OutSdkVersion">Receives the version number of the selected Windows SDK</param>
+		/// <param name="OutSdkDir">Receives the root directory for the selected SDK</param>
+		/// <returns>True if the toolchain directory was found correctly</returns>
+		public static bool TryGetWindowsSdkDir(string? DesiredVersion, [NotNullWhen(true)] out VersionNumber? OutSdkVersion, [NotNullWhen(true)] out DirectoryReference? OutSdkDir)
+		{
+			// Get a map of Windows SDK versions to their root directories
+			/*IReadOnlyDictionary<VersionNumber, DirectoryReference> WindowsSdkDirs =*/
+			FindWindowsSdkDirs();
+
+			// Figure out which version number to look for
+			VersionNumber? WindowsSdkVersion = null;
+			if (!string.IsNullOrEmpty(DesiredVersion))
+			{
+				if (String.Compare(DesiredVersion, "Latest", StringComparison.InvariantCultureIgnoreCase) == 0 && CachedWindowsSdkDirs!.Count > 0)
+				{
+					WindowsSdkVersion = CachedWindowsSdkDirs.OrderBy(x => x.Key).Last().Key;
+				}
+				else if (!VersionNumber.TryParse(DesiredVersion, out WindowsSdkVersion))
+				{
+					throw new BuildException("Unable to find requested Windows SDK; '{0}' is an invalid version", DesiredVersion);
+				}
+			}
+			else
+			{
+				WindowsSdkVersion = PreferredWindowsSdkVersions.FirstOrDefault(x => CachedWindowsSdkDirs!.ContainsKey(x));
+				if (WindowsSdkVersion == null && CachedWindowsSdkDirs!.Count > 0)
+				{
+					WindowsSdkVersion = CachedWindowsSdkDirs.OrderBy(x => x.Key).Last().Key;
+				}
+			}
+
+			// Get the actual directory for this version
+			DirectoryReference? SdkDir;
+			if (WindowsSdkVersion != null && CachedWindowsSdkDirs!.TryGetValue(WindowsSdkVersion, out SdkDir))
+			{
+				OutSdkDir = SdkDir;
+				OutSdkVersion = WindowsSdkVersion;
+				return true;
+			}
+			else
+			{
+				OutSdkDir = null;
+				OutSdkVersion = null;
+				return false;
+			}
+		}
+
+		/// <summary>
+		/// Gets the installation directory for the NETFXSDK
+		/// </summary>
+		/// <param name="OutInstallDir">Receives the installation directory on success</param>
+		/// <returns>True if the directory was found, false otherwise</returns>
+		public static bool TryGetNetFxSdkInstallDir([NotNullWhen(true)] out DirectoryReference? OutInstallDir)
+		{
+			DirectoryReference HostAutoSdkDir;
+			if (UEBuildPlatformSDK.TryGetHostPlatformAutoSDKDir(out HostAutoSdkDir))
+			{
+				DirectoryReference NetFxDir_4_6 = DirectoryReference.Combine(HostAutoSdkDir, "Win64", "Windows Kits", "NETFXSDK", "4.6");
+				if (FileReference.Exists(FileReference.Combine(NetFxDir_4_6, "Include", "um", "mscoree.h")))
+				{
+					OutInstallDir = NetFxDir_4_6;
+					return true;
+				}
+
+				DirectoryReference NetFxDir_4_6_1 = DirectoryReference.Combine(HostAutoSdkDir, "Win64", "Windows Kits", "NETFXSDK", "4.6.1");
+				if (FileReference.Exists(FileReference.Combine(NetFxDir_4_6_1, "Include", "um", "mscoree.h")))
+				{
+					OutInstallDir = NetFxDir_4_6_1;
+					return true;
+				}
+			}
+
+			string NetFxSDKKeyName = "Microsoft\\Microsoft SDKs\\NETFXSDK";
+			string[] PreferredVersions = new string[] { "4.6.2", "4.6.1", "4.6" };
+			foreach (string PreferredVersion in PreferredVersions)
+			{
+				if (TryReadInstallDirRegistryKey32(NetFxSDKKeyName + "\\" + PreferredVersion, "KitsInstallationFolder", out OutInstallDir))
+				{
+					return true;
+				}
+			}
+
+			// If we didn't find one of our preferred versions for NetFXSDK, use the max version present on the system
+			Version? MaxVersion = null;
+			string? MaxVersionString = null;
+			foreach (string ExistingVersionString in ReadInstallDirSubKeys32(NetFxSDKKeyName))
+			{
+				Version? ExistingVersion;
+				if (!Version.TryParse(ExistingVersionString, out ExistingVersion))
+				{
+					continue;
+				}
+				if (MaxVersion == null || ExistingVersion.CompareTo(MaxVersion) > 0)
+				{
+					MaxVersion = ExistingVersion;
+					MaxVersionString = ExistingVersionString;
+				}
+			}
+
+			if (MaxVersionString != null)
+			{
+				return TryReadInstallDirRegistryKey32(NetFxSDKKeyName + "\\" + MaxVersionString, "KitsInstallationFolder", out OutInstallDir);
+			}
+
+			OutInstallDir = null;
+			return false;
+		}
+
+
+		static readonly KeyValuePair<RegistryKey, string>[] InstallDirRoots = {
+			new KeyValuePair<RegistryKey, string>(Registry.CurrentUser, "SOFTWARE\\"),
+			new KeyValuePair<RegistryKey, string>(Registry.LocalMachine, "SOFTWARE\\"),
+			new KeyValuePair<RegistryKey, string>(Registry.CurrentUser, "SOFTWARE\\Wow6432Node\\"),
+			new KeyValuePair<RegistryKey, string>(Registry.LocalMachine, "SOFTWARE\\Wow6432Node\\")
+		};
+
+		/// <summary>
+		/// Reads an install directory for a 32-bit program from a registry key. This checks for per-user and machine wide settings, and under the Wow64 virtual keys (HKCU\SOFTWARE, HKLM\SOFTWARE, HKCU\SOFTWARE\Wow6432Node, HKLM\SOFTWARE\Wow6432Node).
+		/// </summary>
+		/// <param name="KeySuffix">Path to the key to read, under one of the roots listed above.</param>
+		/// <param name="ValueName">Value to be read.</param>
+		/// <param name="InstallDir">On success, the directory corresponding to the value read.</param>
+		/// <returns>True if the key was read, false otherwise.</returns>
+		public static bool TryReadInstallDirRegistryKey32(string KeySuffix, string ValueName, [NotNullWhen(true)] out DirectoryReference? InstallDir)
+		{
+			foreach (KeyValuePair<RegistryKey, string> InstallRoot in InstallDirRoots)
+			{
+				using (RegistryKey Key = InstallRoot.Key.OpenSubKey(InstallRoot.Value + KeySuffix))
+				{
+					if (Key != null && TryReadDirRegistryKey(Key.Name, ValueName, out InstallDir))
+					{
+						return true;
+					}
+				}
+			}
+			InstallDir = null;
+			return false;
+		}
+
+		/// <summary>
+		/// For each root location relevant to install dirs, look for the given key and add its subkeys to the set of subkeys to return.
+		/// This checks for per-user and machine wide settings, and under the Wow64 virtual keys (HKCU\SOFTWARE, HKLM\SOFTWARE, HKCU\SOFTWARE\Wow6432Node, HKLM\SOFTWARE\Wow6432Node).
+		/// </summary>
+		/// <param name="KeyName">The subkey to look for under each root location</param>
+		/// <returns>A list of unique subkeys found under any of the existing subkeys</returns>
+		static string[] ReadInstallDirSubKeys32(string KeyName)
+		{
+			HashSet<string> AllSubKeys = new HashSet<string>(StringComparer.Ordinal);
+			foreach (KeyValuePair<RegistryKey, string> Root in InstallDirRoots)
+			{
+				using (RegistryKey Key = Root.Key.OpenSubKey(Root.Value + KeyName))
+				{
+					if (Key == null)
+					{
+						continue;
+					}
+
+					foreach (string SubKey in Key.GetSubKeyNames())
+					{
+						AllSubKeys.Add(SubKey);
+					}
+				}
+			}
+			return AllSubKeys.ToArray();
+		}
+
+		/// <summary>
+		/// Attempts to reads a directory name stored in a registry key
+		/// </summary>
+		/// <param name="KeyName">Key to read from</param>
+		/// <param name="ValueName">Value within the key to read</param>
+		/// <param name="Value">The directory read from the registry key</param>
+		/// <returns>True if the key was read, false if it was missing or empty</returns>
+		static bool TryReadDirRegistryKey(string KeyName, string ValueName, [NotNullWhen(true)] out DirectoryReference? Value)
+		{
+			string? StringValue = Registry.GetValue(KeyName, ValueName, null) as string;
+			if (String.IsNullOrEmpty(StringValue))
+			{
+				Value = null;
+				return false;
+			}
+			else
+			{
+				Value = new DirectoryReference(StringValue);
+				return true;
+			}
+		}
+
+		/// <summary>
+		/// Enumerates all the Windows 10 SDK root directories
+		/// </summary>
+		/// <param name="RootDirs">Receives all the Windows 10 sdk root directories</param>
+		public static void EnumerateSdkRootDirs(List<DirectoryReference> RootDirs)
+		{
+			DirectoryReference? RootDir;
+			if (TryReadInstallDirRegistryKey32("Microsoft\\Windows Kits\\Installed Roots", "KitsRoot10", out RootDir))
+			{
+				Log.TraceLog("Found Windows 10 SDK root at {0} (1)", RootDir);
+				RootDirs.Add(RootDir);
+			}
+			if (TryReadInstallDirRegistryKey32("Microsoft\\Microsoft SDKs\\Windows\\v10.0", "InstallationFolder", out RootDir))
+			{
+				Log.TraceLog("Found Windows 10 SDK root at {0} (2)", RootDir);
+				RootDirs.Add(RootDir);
+			}
+
+			DirectoryReference HostAutoSdkDir;
+			if (UEBuildPlatformSDK.TryGetHostPlatformAutoSDKDir(out HostAutoSdkDir))
+			{
+				DirectoryReference RootDirAutoSdk = DirectoryReference.Combine(HostAutoSdkDir, "Win64", "Windows Kits", "10");
+				if (DirectoryReference.Exists(RootDirAutoSdk))
+				{
+					Log.TraceLog("Found Windows 10 AutoSDK root at {0}", RootDirAutoSdk);
+					RootDirs.Add(RootDirAutoSdk);
+				}
+			}
+		}
+
+		#endregion
 	}
 }
diff --git a/Engine/Source/Programs/UnrealBuildTool/Platform/Windows/UEBuildWindows.cs b/Engine/Source/Programs/UnrealBuildTool/Platform/Windows/UEBuildWindows.cs
--- a/Engine/Source/Programs/UnrealBuildTool/Platform/Windows/UEBuildWindows.cs
+++ b/Engine/Source/Programs/UnrealBuildTool/Platform/Windows/UEBuildWindows.cs
@@ -712,15 +712,6 @@ namespace UnrealBuildTool
 //			VersionNumber.Parse("14.16.27023"), // fallback to VS2017 15.9 toolchain, microsoft updates these in places so for local installs only this version number is present
 		};
 
-		/// <summary>
-		/// The default Windows SDK version to be used, if installed.
-		/// </summary>
-		static readonly VersionNumber[] PreferredWindowsSdkVersions = new VersionNumber[]
-		{
-			VersionNumber.Parse("10.0.18362.0"),
-			VersionNumber.Parse("10.0.16299.0")
-		};
-
 		/// <summary>
 		/// Cache of Visual Studio installation directories
 		/// </summary>
@@ -736,16 +727,6 @@ namespace UnrealBuildTool
 		/// </summary>
 		private static Dictionary<WindowsCompiler, List<DirectoryReference>> CachedDiaSdkDirs = new Dictionary<WindowsCompiler, List<DirectoryReference>>();
 
-		/// <summary>
-		/// Cache of Windows SDK installation directories
-		/// </summary>
-		private static IReadOnlyDictionary<VersionNumber, DirectoryReference>? CachedWindowsSdkDirs;
-
-		/// <summary>
-		/// Cache of Universal CRT installation directories
-		/// </summary>
-		private static IReadOnlyDictionary<VersionNumber, DirectoryReference>? CachedUniversalCrtDirs;
-
 		/// <summary>
 		/// True if we should use the Clang linker (LLD) when we are compiling with Clang, otherwise we use the MSVC linker
 		/// </summary>
@@ -1057,7 +1038,7 @@ namespace UnrealBuildTool
 				    {
 					    // VS2015 just installs one toolchain; use that.
 					    DirectoryReference? InstallDir;
-					    if(TryReadInstallDirRegistryKey32("Microsoft\\VisualStudio\\SxS\\VS7", "14.0", out InstallDir))
+					    if(MicrosoftPlatformSDK.TryReadInstallDirRegistryKey32("Microsoft\\VisualStudio\\SxS\\VS7", "14.0", out InstallDir))
 					    {
 						    Installations.Add(new VisualStudioInstallation(InstallDir));
 					    }
@@ -1490,85 +1471,6 @@ namespace UnrealBuildTool
 			}
 		}
 
-		static readonly KeyValuePair<RegistryKey, string>[] InstallDirRoots = {
-			new KeyValuePair<RegistryKey, string>(Registry.CurrentUser, "SOFTWARE\\"),
-			new KeyValuePair<RegistryKey, string>(Registry.LocalMachine, "SOFTWARE\\"),
-			new KeyValuePair<RegistryKey, string>(Registry.CurrentUser, "SOFTWARE\\Wow6432Node\\"),
-			new KeyValuePair<RegistryKey, string>(Registry.LocalMachine, "SOFTWARE\\Wow6432Node\\")
-		};
-
-		/// <summary>
-		/// Reads an install directory for a 32-bit program from a registry key. This checks for per-user and machine wide settings, and under the Wow64 virtual keys (HKCU\SOFTWARE, HKLM\SOFTWARE, HKCU\SOFTWARE\Wow6432Node, HKLM\SOFTWARE\Wow6432Node).
-		/// </summary>
-		/// <param name="KeySuffix">Path to the key to read, under one of the roots listed above.</param>
-		/// <param name="ValueName">Value to be read.</param>
-		/// <param name="InstallDir">On success, the directory corresponding to the value read.</param>
-		/// <returns>True if the key was read, false otherwise.</returns>
-		static bool TryReadInstallDirRegistryKey32(string KeySuffix, string ValueName, [NotNullWhen(true)] out DirectoryReference? InstallDir)
-		{
-			foreach (KeyValuePair<RegistryKey, string> InstallRoot in InstallDirRoots)
-			{
-				using (RegistryKey Key = InstallRoot.Key.OpenSubKey(InstallRoot.Value + KeySuffix))
-				{
-					if (Key != null && TryReadDirRegistryKey(Key.Name, ValueName, out InstallDir))
-					{
-						return true;
-					}
-				}
-			}
-			InstallDir = null;
-			return false;
-		}
-
-		/// <summary>
-		/// For each root location relevant to install dirs, look for the given key and add its subkeys to the set of subkeys to return.
-		/// This checks for per-user and machine wide settings, and under the Wow64 virtual keys (HKCU\SOFTWARE, HKLM\SOFTWARE, HKCU\SOFTWARE\Wow6432Node, HKLM\SOFTWARE\Wow6432Node).
-		/// </summary>
-		/// <param name="KeyName">The subkey to look for under each root location</param>
-		/// <returns>A list of unique subkeys found under any of the existing subkeys</returns>
-		static string[] ReadInstallDirSubKeys32(string KeyName)
-		{
-			HashSet<string> AllSubKeys = new HashSet<string>(StringComparer.Ordinal);
-			foreach (KeyValuePair<RegistryKey, string> Root in InstallDirRoots)
-			{
-				using (RegistryKey Key = Root.Key.OpenSubKey(Root.Value + KeyName))
-				{
-					if (Key == null)
-					{
-						continue;
-					}
-
-					foreach (string SubKey in Key.GetSubKeyNames())
-					{
-						AllSubKeys.Add(SubKey);
-					}
-				}
-			}
-			return AllSubKeys.ToArray();
-		}
-
-		/// <summary>
-		/// Attempts to reads a directory name stored in a registry key
-		/// </summary>
-		/// <param name="KeyName">Key to read from</param>
-		/// <param name="ValueName">Value within the key to read</param>
-		/// <param name="Value">The directory read from the registry key</param>
-		/// <returns>True if the key was read, false if it was missing or empty</returns>
-		static bool TryReadDirRegistryKey(string KeyName, string ValueName, [NotNullWhen(true)] out DirectoryReference? Value)
-		{
-			string? StringValue = Registry.GetValue(KeyName, ValueName, null) as string;
-			if (String.IsNullOrEmpty(StringValue))
-			{
-				Value = null;
-				return false;
-			}
-			else
-			{
-				Value = new DirectoryReference(StringValue);
-				return true;
-			}
-		}
-
 		/// <summary>
 		/// Gets the path to MSBuild. This mirrors the logic in GetMSBuildPath.bat.
 		/// </summary>
@@ -1706,25 +1608,25 @@ namespace UnrealBuildTool
 		public static List<DirectoryReference> FindDiaSdkDirs(WindowsCompiler Compiler)
 		{
 			List<DirectoryReference>? DiaSdkDirs;
-			if(!CachedDiaSdkDirs.TryGetValue(Compiler, out DiaSdkDirs))
+			if (!CachedDiaSdkDirs.TryGetValue(Compiler, out DiaSdkDirs))
 			{
 				DiaSdkDirs = new List<DirectoryReference>();
 
 				DirectoryReference PlatformDir;
-				if(UEBuildPlatformSDK.TryGetHostPlatformAutoSDKDir(out PlatformDir))
+				if (UEBuildPlatformSDK.TryGetHostPlatformAutoSDKDir(out PlatformDir))
 				{
-					DirectoryReference DiaSdkDir = DirectoryReference.Combine(PlatformDir, "Win64", "DIA SDK", (Compiler == WindowsCompiler.VisualStudio2019)? "VS2019" : "VS2017");
-					if(IsValidDiaSdkDir(DiaSdkDir))
+					DirectoryReference DiaSdkDir = DirectoryReference.Combine(PlatformDir, "Win64", "DIA SDK", (Compiler == WindowsCompiler.VisualStudio2019) ? "VS2019" : "VS2017");
+					if (IsValidDiaSdkDir(DiaSdkDir))
 					{
 						DiaSdkDirs.Add(DiaSdkDir);
 					}
 				}
 
 				List<DirectoryReference> VisualStudioDirs = FindVisualStudioInstallations(Compiler).ConvertAll(x => x.BaseDir);
-				foreach(DirectoryReference VisualStudioDir in VisualStudioDirs)
+				foreach (DirectoryReference VisualStudioDir in VisualStudioDirs)
 				{
 					DirectoryReference DiaSdkDir = DirectoryReference.Combine(VisualStudioDir, "DIA SDK");
-					if(IsValidDiaSdkDir(DiaSdkDir))
+					if (IsValidDiaSdkDir(DiaSdkDir))
 					{
 						DiaSdkDirs.Add(DiaSdkDir);
 					}
@@ -1743,227 +1645,13 @@ namespace UnrealBuildTool
 			return FileReference.Exists(FileReference.Combine(DiaSdkDir, "bin", "amd64", "msdia140.dll"));
 		}
 
-		/// <summary>
-		/// Updates the CachedWindowsSdkDirs and CachedUniversalCrtDirs variables
-		/// </summary>
-		private static void UpdateCachedWindowsSdks()
-		{
-			Dictionary<VersionNumber, DirectoryReference> WindowsSdkDirs = new Dictionary<VersionNumber, DirectoryReference>();
-			Dictionary<VersionNumber, DirectoryReference> UniversalCrtDirs = new Dictionary<VersionNumber, DirectoryReference>();
-
-			// Enumerate the Windows 8.1 SDK, if present
-			DirectoryReference? InstallDir_8_1;
-			if(TryReadInstallDirRegistryKey32("Microsoft\\Microsoft SDKs\\Windows\\v8.1", "InstallationFolder", out InstallDir_8_1))
-			{
-				if(FileReference.Exists(FileReference.Combine(InstallDir_8_1, "Include", "um", "windows.h")))
-				{
-					Log.TraceLog("Found Windows 8.1 SDK at {0}", InstallDir_8_1);
-					VersionNumber Version_8_1 = new VersionNumber(8, 1);
-					WindowsSdkDirs[Version_8_1] = InstallDir_8_1;
-				}
-			}
-
-			// Find all the root directories for Windows 10 SDKs
-			List<DirectoryReference> InstallDirs_10 = new List<DirectoryReference>();
-			EnumerateSdkRootDirs(InstallDirs_10);
-
-			// Enumerate all the Windows 10 SDKs
-			foreach(DirectoryReference InstallDir_10 in InstallDirs_10.Distinct())
-			{
-				DirectoryReference IncludeRootDir = DirectoryReference.Combine(InstallDir_10, "Include");
-				if(DirectoryReference.Exists(IncludeRootDir))
-				{
-					foreach(DirectoryReference IncludeDir in DirectoryReference.EnumerateDirectories(IncludeRootDir))
-					{
-						VersionNumber? IncludeVersion;
-						if(VersionNumber.TryParse(IncludeDir.GetDirectoryName(), out IncludeVersion))
-						{
-							if(FileReference.Exists(FileReference.Combine(IncludeDir, "um", "windows.h")))
-							{
-								Log.TraceLog("Found Windows 10 SDK version {0} at {1}", IncludeVersion, InstallDir_10);
-								WindowsSdkDirs[IncludeVersion] = InstallDir_10;
-							}
-							if(FileReference.Exists(FileReference.Combine(IncludeDir, "ucrt", "corecrt.h")))
-							{
-								Log.TraceLog("Found Universal CRT version {0} at {1}", IncludeVersion, InstallDir_10);
-								UniversalCrtDirs[IncludeVersion] = InstallDir_10;
-							}
-						}
-					}
-				}
-			}
-
-			CachedWindowsSdkDirs = WindowsSdkDirs;
-			CachedUniversalCrtDirs = UniversalCrtDirs;
-		}
-
-		/// <summary>
-		/// Finds all the installed Windows SDK versions
-		/// </summary>
-		/// <returns>Map of version number to Windows SDK directories</returns>
-		public static IReadOnlyDictionary<VersionNumber, DirectoryReference> FindWindowsSdkDirs()
-		{
-			// Update the cache of install directories, if it's not set
-			if(CachedWindowsSdkDirs == null)
-			{
-				UpdateCachedWindowsSdks();
-			}
-			return CachedWindowsSdkDirs!;
-		}
-
-		/// <summary>
-		/// Finds all the installed Universal CRT versions
-		/// </summary>
-		/// <returns>Map of version number to universal CRT directories</returns>
-		public static IReadOnlyDictionary<VersionNumber, DirectoryReference> FindUniversalCrtDirs()
-		{
-			if(CachedUniversalCrtDirs == null)
-			{
-				UpdateCachedWindowsSdks();
-			}
-			return CachedUniversalCrtDirs!;
-		}
-
-		/// <summary>
-		/// Enumerates all the Windows 10 SDK root directories
-		/// </summary>
-		/// <param name="RootDirs">Receives all the Windows 10 sdk root directories</param>
-		private static void EnumerateSdkRootDirs(List<DirectoryReference> RootDirs)
-		{
-			DirectoryReference? RootDir;
-			if(TryReadInstallDirRegistryKey32("Microsoft\\Windows Kits\\Installed Roots", "KitsRoot10", out RootDir))
-			{
-				Log.TraceLog("Found Windows 10 SDK root at {0} (1)", RootDir);
-				RootDirs.Add(RootDir);
-			}
-			if(TryReadInstallDirRegistryKey32("Microsoft\\Microsoft SDKs\\Windows\\v10.0", "InstallationFolder", out RootDir))
-			{
-				Log.TraceLog("Found Windows 10 SDK root at {0} (2)", RootDir);
-				RootDirs.Add(RootDir);
-			}
 
-			DirectoryReference HostAutoSdkDir;
-			if(UEBuildPlatformSDK.TryGetHostPlatformAutoSDKDir(out HostAutoSdkDir))
-			{
-				DirectoryReference RootDirAutoSdk = DirectoryReference.Combine(HostAutoSdkDir, "Win64", "Windows Kits", "10");
-				if(DirectoryReference.Exists(RootDirAutoSdk))
-				{
-					Log.TraceLog("Found Windows 10 AutoSDK root at {0}", RootDirAutoSdk);
-					RootDirs.Add(RootDirAutoSdk);
-				}
-			}
-		}
 
-		/// <summary>
-		/// Determines the directory containing the Windows SDK toolchain
-		/// </summary>
-		/// <param name="DesiredVersion">The desired Windows SDK version. This may be "Latest", a specific version number, or null. If null, the function will look for DefaultWindowsSdkVersion. Failing that, it will return the latest version.</param>
-		/// <param name="OutSdkVersion">Receives the version number of the selected Windows SDK</param>
-		/// <param name="OutSdkDir">Receives the root directory for the selected SDK</param>
-		/// <returns>True if the toolchain directory was found correctly</returns>
 		public static bool TryGetWindowsSdkDir(string? DesiredVersion, [NotNullWhen(true)] out VersionNumber? OutSdkVersion, [NotNullWhen(true)] out DirectoryReference? OutSdkDir)
 		{
-			// Get a map of Windows SDK versions to their root directories
-			/*IReadOnlyDictionary<VersionNumber, DirectoryReference> WindowsSdkDirs =*/ FindWindowsSdkDirs();
-
-			// Figure out which version number to look for
-			VersionNumber? WindowsSdkVersion = null;
-			if(!string.IsNullOrEmpty(DesiredVersion))
-			{
-				if(String.Compare(DesiredVersion, "Latest", StringComparison.InvariantCultureIgnoreCase) == 0 && CachedWindowsSdkDirs!.Count > 0)
-				{
-					WindowsSdkVersion = CachedWindowsSdkDirs.OrderBy(x => x.Key).Last().Key;
-				}
-				else if(!VersionNumber.TryParse(DesiredVersion, out WindowsSdkVersion))
-				{
-					throw new BuildException("Unable to find requested Windows SDK; '{0}' is an invalid version", DesiredVersion);
-				}
-			}
-			else
-			{
-				WindowsSdkVersion = PreferredWindowsSdkVersions.FirstOrDefault(x => CachedWindowsSdkDirs!.ContainsKey(x));
-				if(WindowsSdkVersion == null && CachedWindowsSdkDirs!.Count > 0)
-				{
-					WindowsSdkVersion = CachedWindowsSdkDirs.OrderBy(x => x.Key).Last().Key;
-				}
-			}
-
-			// Get the actual directory for this version
-			DirectoryReference? SdkDir;
-			if(WindowsSdkVersion != null && CachedWindowsSdkDirs!.TryGetValue(WindowsSdkVersion, out SdkDir))
-			{
-				OutSdkDir = SdkDir;
-				OutSdkVersion = WindowsSdkVersion;
-				return true;
-			}
-			else
-			{
-				OutSdkDir = null;
-				OutSdkVersion = null;
-				return false;
-			}
+			return MicrosoftPlatformSDK.TryGetWindowsSdkDir(DesiredVersion, out OutSdkVersion, out OutSdkDir);
 		}
 
-		/// <summary>
-		/// Gets the installation directory for the NETFXSDK
-		/// </summary>
-		/// <param name="OutInstallDir">Receives the installation directory on success</param>
-		/// <returns>True if the directory was found, false otherwise</returns>
-		public static bool TryGetNetFxSdkInstallDir([NotNullWhen(true)] out DirectoryReference? OutInstallDir)
-		{
-			DirectoryReference HostAutoSdkDir;
-			if(UEBuildPlatformSDK.TryGetHostPlatformAutoSDKDir(out HostAutoSdkDir))
-			{
-				DirectoryReference NetFxDir_4_6 = DirectoryReference.Combine(HostAutoSdkDir, "Win64", "Windows Kits", "NETFXSDK", "4.6");
-				if(FileReference.Exists(FileReference.Combine(NetFxDir_4_6, "Include", "um", "mscoree.h")))
-				{
-					OutInstallDir = NetFxDir_4_6;
-					return true;
-				}
-
-				DirectoryReference NetFxDir_4_6_1 = DirectoryReference.Combine(HostAutoSdkDir, "Win64", "Windows Kits", "NETFXSDK", "4.6.1");
-				if(FileReference.Exists(FileReference.Combine(NetFxDir_4_6_1, "Include", "um", "mscoree.h")))
-				{
-					OutInstallDir = NetFxDir_4_6_1;
-					return true;
-				}
-			}
-
-			string NetFxSDKKeyName = "Microsoft\\Microsoft SDKs\\NETFXSDK";
-			string[] PreferredVersions = new string[] { "4.6.2", "4.6.1", "4.6" };
-			foreach (string PreferredVersion in PreferredVersions)
-			{
-				if (TryReadInstallDirRegistryKey32(NetFxSDKKeyName + "\\" + PreferredVersion, "KitsInstallationFolder", out OutInstallDir))
-				{
-					return true;
-				}
-			}
-
-			// If we didn't find one of our preferred versions for NetFXSDK, use the max version present on the system
-			Version? MaxVersion = null;
-			string? MaxVersionString = null;
-			foreach (string ExistingVersionString in ReadInstallDirSubKeys32(NetFxSDKKeyName))
-			{
-				Version? ExistingVersion;
-				if (!Version.TryParse(ExistingVersionString, out ExistingVersion))
-				{
-					continue;
-				}
-				if (MaxVersion == null || ExistingVersion.CompareTo(MaxVersion) > 0)
-				{
-					MaxVersion = ExistingVersion;
-					MaxVersionString = ExistingVersionString;
-				}
-			}
-
-			if (MaxVersionString != null)
-			{
-				return TryReadInstallDirRegistryKey32(NetFxSDKKeyName + "\\" + MaxVersionString, "KitsInstallationFolder", out OutInstallDir);
-			}
-
-			OutInstallDir = null;
-			return false;
-		}
 
 		/// <summary>
 		/// Gets the platform name that should be used.
diff --git a/Engine/Source/Programs/UnrealBuildTool/Platform/Windows/VCEnvironment.cs b/Engine/Source/Programs/UnrealBuildTool/Platform/Windows/VCEnvironment.cs
--- a/Engine/Source/Programs/UnrealBuildTool/Platform/Windows/VCEnvironment.cs
+++ b/Engine/Source/Programs/UnrealBuildTool/Platform/Windows/VCEnvironment.cs
@@ -421,7 +421,7 @@ namespace UnrealBuildTool
 			// If we're on Visual Studio 2015 and using pre-Windows 10 SDK, we need to find a Windows 10 SDK and add the UCRT include paths
 			if(ToolChain >= WindowsCompiler.VisualStudio2015_DEPRECATED && WindowsSdkVersion < new VersionNumber(10))
 			{
-				KeyValuePair<VersionNumber, DirectoryReference> Pair = WindowsPlatform.FindUniversalCrtDirs().OrderByDescending(x => x.Key).FirstOrDefault();
+				KeyValuePair<VersionNumber, DirectoryReference> Pair = MicrosoftPlatformSDK.FindUniversalCrtDirs().OrderByDescending(x => x.Key).FirstOrDefault();
 				if(Pair.Key == null || Pair.Key < new VersionNumber(10))
 				{
 					throw new BuildException("{0} requires the Universal CRT to be installed.", WindowsPlatform.GetCompilerName(ToolChain));
@@ -436,7 +436,7 @@ namespace UnrealBuildTool
 
 			// Add the NETFXSDK include path. We need this for SwarmInterface.
 			DirectoryReference? NetFxSdkDir;
-			if(WindowsPlatform.TryGetNetFxSdkInstallDir(out NetFxSdkDir))
+			if(MicrosoftPlatformSDK.TryGetNetFxSdkInstallDir(out NetFxSdkDir))
 			{
 				IncludePaths.Add(DirectoryReference.Combine(NetFxSdkDir, "include", "um"));
 				LibraryPaths.Add(DirectoryReference.Combine(NetFxSdkDir, "lib", "um", ArchFolder));
diff --git a/Engine/Source/Programs/UnrealBuildTool/Platform/Windows/WindowsExports.cs b/Engine/Source/Programs/UnrealBuildTool/Platform/Windows/WindowsExports.cs
--- a/Engine/Source/Programs/UnrealBuildTool/Platform/Windows/WindowsExports.cs
+++ b/Engine/Source/Programs/UnrealBuildTool/Platform/Windows/WindowsExports.cs
@@ -81,7 +81,7 @@ namespace UnrealBuildTool
 			}
 
 			// Add all the other directories sorted in reverse order
-			IReadOnlyDictionary<VersionNumber, DirectoryReference> WindowsSdkDirPairs = WindowsPlatform.FindWindowsSdkDirs();
+			IReadOnlyDictionary<VersionNumber, DirectoryReference> WindowsSdkDirPairs = MicrosoftPlatformSDK.FindWindowsSdkDirs();
 			foreach(KeyValuePair<VersionNumber, DirectoryReference> Pair in WindowsSdkDirPairs.OrderByDescending(x => x.Key))
 			{
 				if(!WindowsSdkDirs.Any(x => x.Value == Pair.Value))
