Fix comment about XGE over VPN.

#rnx

diff --git a/Engine/Source/Programs/UnrealBuildTool/Executors/XGE.cs b/Engine/Source/Programs/UnrealBuildTool/Executors/XGE.cs
--- a/Engine/Source/Programs/UnrealBuildTool/Executors/XGE.cs
+++ b/Engine/Source/Programs/UnrealBuildTool/Executors/XGE.cs
@@ -41,7 +41,7 @@ namespace UnrealBuildTool
 		bool bStopXGECompilationAfterErrors = false;
 
 		/// <summary>
-		/// When set to false, XGE will not be When enabled, XGE will stop compiling targets after a compile error occurs.  Recommended, as it saves computing resources for others.
+		/// When set to false, XGE will not be enabled when running connected to the coordinator over VPN. Configure VPN-assigned subnets via the VpnSubnets parameter.
 		/// </summary>
 		[XmlConfigFile(Category = "XGE")]
 		static bool bAllowOverVpn = true;
