Added debug output to help track down
UE-105044

diff --git a/Engine/Source/Programs/UnrealBuildTool/Modes/WriteMetadataMode.cs b/Engine/Source/Programs/UnrealBuildTool/Modes/WriteMetadataMode.cs
--- a/Engine/Source/Programs/UnrealBuildTool/Modes/WriteMetadataMode.cs
+++ b/Engine/Source/Programs/UnrealBuildTool/Modes/WriteMetadataMode.cs
@@ -83,6 +83,8 @@ namespace UnrealBuildTool
 			// Acquire a different mutex to the regular UBT instance, since this mode will be called as part of a build. We need the mutex to ensure that building two modular configurations 
 			// in parallel don't clash over writing shared *.modules files (eg. DebugGame and Development editors).
 			string MutexName = SingleInstanceMutex.GetUniqueMutexForPath("UnrealBuildTool_WriteMetadata", UnrealBuildTool.RootDirectory.FullName);
+			// temporary output to debug https://jira.it.epicgames.com/browse/UE-105044
+			Log.TraceInformation("Locked mutex {0} at {1:HH\\:mm\\:ss\\.fff} when writing metadata.", MutexName, DateTime.Now);
 			using(new SingleInstanceMutex(MutexName, true))
 			{
 				return ExecuteInternal(Arguments);
