Fix issue with IOS builds looking for architecture appended, eg UnrealGamearm64

diff --git a/Engine/Source/Programs/UnrealBuildTool/Platform/IOS/UEBuildIOS.cs b/Engine/Source/Programs/UnrealBuildTool/Platform/IOS/UEBuildIOS.cs
--- a/Engine/Source/Programs/UnrealBuildTool/Platform/IOS/UEBuildIOS.cs
+++ b/Engine/Source/Programs/UnrealBuildTool/Platform/IOS/UEBuildIOS.cs
@@ -756,6 +756,15 @@ namespace UnrealBuildTool
 			Target.bCheckSystemHeadersForModification = false;
 		}
 
+		/// <summary>
+		/// Allows the platform to override whether the architecture name should be appended to the name of binaries.
+		/// </summary>
+		/// <returns>True if the architecture name should be appended to the binary</returns>
+		public override bool RequiresArchitectureSuffix()
+		{
+			return false;
+		}
+
 		/// <summary>
 		/// Determines if the given name is a build product for a target.
 		/// </summary>
