Adding info on which processes has file opened to debug UE-105044. Removed previous temp output as we have determined the mutex locking is working as expected.

diff --git a/Engine/Source/Programs/UnrealBuildTool/Modes/WriteMetadataMode.cs b/Engine/Source/Programs/UnrealBuildTool/Modes/WriteMetadataMode.cs
--- a/Engine/Source/Programs/UnrealBuildTool/Modes/WriteMetadataMode.cs
+++ b/Engine/Source/Programs/UnrealBuildTool/Modes/WriteMetadataMode.cs
@@ -83,8 +83,6 @@ namespace UnrealBuildTool
 			// Acquire a different mutex to the regular UBT instance, since this mode will be called as part of a build. We need the mutex to ensure that building two modular configurations 
 			// in parallel don't clash over writing shared *.modules files (eg. DebugGame and Development editors).
 			string MutexName = SingleInstanceMutex.GetUniqueMutexForPath("UnrealBuildTool_WriteMetadata", UnrealBuildTool.RootDirectory.FullName);
-			// temporary output to debug https://jira.it.epicgames.com/browse/UE-105044
-			Log.TraceInformation("Locked mutex {0} at {1:HH\\:mm\\:ss\\.fff} when writing metadata.", MutexName, DateTime.Now);
 			using(new SingleInstanceMutex(MutexName, true))
 			{
 				return ExecuteInternal(Arguments);
diff --git a/Engine/Source/Programs/UnrealBuildTool/System/ModuleManifest.cs b/Engine/Source/Programs/UnrealBuildTool/System/ModuleManifest.cs
--- a/Engine/Source/Programs/UnrealBuildTool/System/ModuleManifest.cs
+++ b/Engine/Source/Programs/UnrealBuildTool/System/ModuleManifest.cs
@@ -122,9 +122,41 @@ namespace UnrealBuildTool
 		public void Write(FileReference FileName)
 		{
 			DirectoryReference.CreateDirectory(FileName.Directory);
-			using(StreamWriter Writer = new StreamWriter(FileName.FullName))
+			try
 			{
-				Write(Writer);
+				using(StreamWriter Writer = new StreamWriter(FileName.FullName))
+				{
+					Write(Writer);
+				}
+			}
+			catch (IOException Ex)
+			{
+				// for non win32 platforms just rethrow
+				if (Environment.OSVersion.Platform != PlatformID.Win32NT)
+					throw;
+
+				// for win32 we try to determine which other processes has this file opened to improve the error message
+				StringBuilder Message = new StringBuilder($"Unable to write ModuleManifest {FileName.FullName} - {Ex.Message}");
+				List<FileUtils.FileLockInfo_Win32>? LockInfoList;
+				try
+				{
+					LockInfoList = FileUtils.GetFileLockInfo_Win32(FileName.FullName);
+				}
+				catch
+				{
+					LockInfoList = null;
+				}
+
+				if (LockInfoList != null && LockInfoList.Count > 0)
+				{
+					Message.Append("\nProcesses with open handles to file:");
+					foreach (FileUtils.FileLockInfo_Win32 LockInfo in LockInfoList)
+					{
+						Message.Append($"\n  {LockInfo}");
+					}
+				}
+
+				throw new Exception(Message.ToString(), Ex);
 			}
 		}
 
