Change Clang toolchain to impersonate VS2019.

#rnx

diff --git a/Engine/Source/Programs/UnrealBuildTool/Platform/Windows/VCToolChain.cs b/Engine/Source/Programs/UnrealBuildTool/Platform/Windows/VCToolChain.cs
--- a/Engine/Source/Programs/UnrealBuildTool/Platform/Windows/VCToolChain.cs
+++ b/Engine/Source/Programs/UnrealBuildTool/Platform/Windows/VCToolChain.cs
@@ -167,7 +167,7 @@ namespace UnrealBuildTool
 			if (Target.WindowsPlatform.Compiler == WindowsCompiler.Clang)
 			{
 				// FIXME: We should be compatible with the latest MSVC compiler toolkit installed instead of hard-coding it here.
-				Arguments.Add("-fms-compatibility-version=19.1");
+				Arguments.Add("-fms-compatibility-version=19.20");
 
 				if (Target.WindowsPlatform.StaticAnalyzer == WindowsStaticAnalyzer.Default)
 				{
