Bump minimum required VS2019 version to 16.4 (MSVC version 19.24)

#rnx

diff --git a/Engine/Source/Programs/UnrealBuildTool/Platform/Windows/UEBuildWindows.cs b/Engine/Source/Programs/UnrealBuildTool/Platform/Windows/UEBuildWindows.cs
--- a/Engine/Source/Programs/UnrealBuildTool/Platform/Windows/UEBuildWindows.cs
+++ b/Engine/Source/Programs/UnrealBuildTool/Platform/Windows/UEBuildWindows.cs
@@ -707,7 +707,7 @@ namespace UnrealBuildTool
 		static readonly VersionNumber[] PreferredVisualStudioToolChainVersion = new VersionNumber[]
 		{
 			VersionNumber.Parse("14.24.28315"), // VS2019 v16.4.3 (installed to 14.24.28314 folder)
-			VersionNumber.Parse("14.22.27905"), // VS2019 v16.2.3
+//			VersionNumber.Parse("14.22.27905"), // VS2019 v16.2.3
 //			VersionNumber.Parse("14.16.27023.2"), // VS2017 v15.9.15
 //			VersionNumber.Parse("14.16.27023"), // fallback to VS2017 15.9 toolchain, microsoft updates these in places so for local installs only this version number is present
 		};
diff --git a/Engine/Source/Programs/UnrealBuildTool/Platform/Windows/VCToolChain.cs b/Engine/Source/Programs/UnrealBuildTool/Platform/Windows/VCToolChain.cs
--- a/Engine/Source/Programs/UnrealBuildTool/Platform/Windows/VCToolChain.cs
+++ b/Engine/Source/Programs/UnrealBuildTool/Platform/Windows/VCToolChain.cs
@@ -167,7 +167,7 @@ namespace UnrealBuildTool
 			if (Target.WindowsPlatform.Compiler == WindowsCompiler.Clang)
 			{
 				// FIXME: We should be compatible with the latest MSVC compiler toolkit installed instead of hard-coding it here.
-				Arguments.Add("-fms-compatibility-version=19.20");
+				Arguments.Add("-fms-compatibility-version=19.24");
 
 				if (Target.WindowsPlatform.StaticAnalyzer == WindowsStaticAnalyzer.Default)
 				{
diff --git a/Engine/Source/Runtime/Core/Public/Windows/WindowsPlatformCompilerSetup.h b/Engine/Source/Runtime/Core/Public/Windows/WindowsPlatformCompilerSetup.h
--- a/Engine/Source/Runtime/Core/Public/Windows/WindowsPlatformCompilerSetup.h
+++ b/Engine/Source/Runtime/Core/Public/Windows/WindowsPlatformCompilerSetup.h
@@ -16,7 +16,7 @@
 //
 // We require at least Visual Studio 2019 to compile
 //
-static_assert(_MSC_VER >= 1920, "Visual Studio 2019 is required to build Unreal Engine.");
+static_assert(_MSC_VER >= 1924, "Visual Studio 2019 v16.4 or greater is required to build Unreal Engine.");
 
 //
 // Manually enable all warnings as errors, except ones that are explicitly skipped.
