Change MSVC static analyzer to use conformance mode (while disabling the attributes that are non-conforming).

diff --git a/Engine/Source/Programs/UnrealBuildTool/Platform/Windows/VCToolChain.cs b/Engine/Source/Programs/UnrealBuildTool/Platform/Windows/VCToolChain.cs
--- a/Engine/Source/Programs/UnrealBuildTool/Platform/Windows/VCToolChain.cs
+++ b/Engine/Source/Programs/UnrealBuildTool/Platform/Windows/VCToolChain.cs
@@ -296,12 +296,11 @@ namespace UnrealBuildTool
 			// Disable Microsoft extensions on VS2017+ for improved standards compliance.
 			if (Target.WindowsPlatform.Compiler >= WindowsCompiler.VisualStudio2017 && Target.WindowsPlatform.bStrictConformanceMode)
 			{
-				// Do not enable if we are running the MSVC code analyzer as it will fail to compile.
-				if (Target.WindowsPlatform.StaticAnalyzer != WindowsStaticAnalyzer.Default)
-				{
-					Arguments.Add("/permissive-");
-					Arguments.Add("/Zc:strictStrings-"); // Have to disable strict const char* semantics due to Windows headers not being compliant.
-				}
+				// This define is needed to ensure that MSVC static analysis mode doesn't declare attributes that are incompatible with strict conformance mode
+				AddDefinition(Arguments, "SAL_NO_ATTRIBUTE_DECLARATIONS=1");
+				
+				Arguments.Add("/permissive-");
+				Arguments.Add("/Zc:strictStrings-"); // Have to disable strict const char* semantics due to Windows headers not being compliant.
 			}
 
 			// @todo HoloLens: UE4 is non-compliant when it comes to use of %s and %S
