Change single quotes to double quotes for .NET Core on Mac/Linux Gradle calls

#ue5
#android

diff --git a/Engine/Source/Programs/UnrealBuildTool/Platform/Android/UEDeployAndroid.cs b/Engine/Source/Programs/UnrealBuildTool/Platform/Android/UEDeployAndroid.cs
--- a/Engine/Source/Programs/UnrealBuildTool/Platform/Android/UEDeployAndroid.cs
+++ b/Engine/Source/Programs/UnrealBuildTool/Platform/Android/UEDeployAndroid.cs
@@ -1355,6 +1355,17 @@ namespace UnrealBuildTool
 
 		private static int RunCommandLineProgramAndReturnResult(string WorkingDirectory, string Command, string Params, string OverrideDesc = null, bool bUseShellExecute = false)
 		{
+			// Process Arguments follow windows conventions in .NET Core
+			// Which means single quotes ' are not considered quotes.
+			// see [REDACTED_PROFILE_URL]/runtime/issues/29857
+			// also see UE-102580
+			// for rules see https://docs.microsoft.com/en-us/cpp/cpp/main-function-command-line-args
+			if (Params.Contains("\'"))
+			{
+				Params = Params.Replace("\"", "\\\"");
+				Params = Params.Replace('\'', '\"');
+			}
+
 			if (OverrideDesc == null)
 			{
 				Log.TraceInformation("\nRunning: " + Command + " " + Params);
@@ -1382,6 +1393,17 @@ namespace UnrealBuildTool
 
 		private static void RunCommandLineProgramWithException(string WorkingDirectory, string Command, string Params, string OverrideDesc = null, bool bUseShellExecute = false)
 		{
+			// Process Arguments follow windows conventions in .NET Core
+			// Which means single quotes ' are not considered quotes.
+			// see [REDACTED_PROFILE_URL]/runtime/issues/29857
+			// also see UE-102580
+			// for rules see https://docs.microsoft.com/en-us/cpp/cpp/main-function-command-line-args
+			if (Params.Contains("\'"))
+			{
+				Params = Params.Replace("\"", "\\\"");
+				Params = Params.Replace('\'', '\"');
+			}
+
 			if (OverrideDesc == null)
 			{
 				Log.TraceInformation("\nRunning: " + Command + " " + Params);
@@ -1446,6 +1468,17 @@ namespace UnrealBuildTool
 
 		private static void RunCommandLineProgramWithExceptionAndFiltering(string WorkingDirectory, string Command, string Params, string OverrideDesc = null, bool bUseShellExecute = false)
 		{
+			// Process Arguments follow windows conventions in .NET Core
+			// Which means single quotes ' are not considered quotes.
+			// see [REDACTED_PROFILE_URL]/runtime/issues/29857
+			// also see UE-102580
+			// for rules see https://docs.microsoft.com/en-us/cpp/cpp/main-function-command-line-args
+			if (Params.Contains("\'"))
+			{
+				Params = Params.Replace("\"", "\\\"");
+				Params = Params.Replace('\'', '\"');
+			}
+
 			if (OverrideDesc == null)
 			{
 				Log.TraceInformation("\nRunning: " + Command + " " + Params);
