Output self contained deployments of UBT on Linux so that we do not have to install the dotnet runtime

diff --git a/Engine/Source/Programs/UnrealBuildTool/UnrealBuildTool.csproj b/Engine/Source/Programs/UnrealBuildTool/UnrealBuildTool.csproj
--- a/Engine/Source/Programs/UnrealBuildTool/UnrealBuildTool.csproj
+++ b/Engine/Source/Programs/UnrealBuildTool/UnrealBuildTool.csproj
@@ -19,6 +19,11 @@
     <Nullable>enable</Nullable>
   </PropertyGroup>
   <Choose>
+    <When Condition="'$(IsLinux)'=='true'">
+      <PropertyGroup>
+        <RuntimeIdentifier>linux-x64</RuntimeIdentifier>
+      </PropertyGroup>
+    </When>
     <When Condition="'$(IsOSX)'=='true'">
       <PropertyGroup>
         <!-- Set runtie identifier for osx so that it produces a app host executable (consistency with the other platforms) -->
