Add cleanup for old UE4 artifacts when packaging for Android

#ue5
#android

diff --git a/Engine/Plugins/Online/OnlineSubsystemFacebook/Source/Private/Android/OnlineSharingFacebook.cpp b/Engine/Plugins/Online/OnlineSubsystemFacebook/Source/Private/Android/OnlineSharingFacebook.cpp
--- a/Engine/Plugins/Online/OnlineSubsystemFacebook/Source/Private/Android/OnlineSharingFacebook.cpp
+++ b/Engine/Plugins/Online/OnlineSubsystemFacebook/Source/Private/Android/OnlineSharingFacebook.cpp
@@ -195,7 +195,7 @@ bool AndroidThunkCpp_Facebook_RequestReadPermissions(const TArray<FSharingPermis
 	return bSuccess;
 }
 
-JNI_METHOD void Java_com_epicgames_ue4_FacebookLogin_nativeRequestReadPermissionsComplete(JNIEnv* jenv, jobject thiz, jsize responseCode, jstring accessToken)
+JNI_METHOD void Java_com_epicgames_unreal_FacebookLogin_nativeRequestReadPermissionsComplete(JNIEnv* jenv, jobject thiz, jsize responseCode, jstring accessToken)
 {
 	EFacebookLoginResponse LoginResponse = (EFacebookLoginResponse)responseCode;
 	
@@ -249,7 +249,7 @@ bool AndroidThunkCpp_Facebook_RequestPublishPermissions(const TArray<FSharingPer
 	return bSuccess;
 }
 
-JNI_METHOD void Java_com_epicgames_ue4_FacebookLogin_nativeRequestPublishPermissionsComplete(JNIEnv* jenv, jobject thiz, jsize responseCode, jstring accessToken)
+JNI_METHOD void Java_com_epicgames_unreal_FacebookLogin_nativeRequestPublishPermissionsComplete(JNIEnv* jenv, jobject thiz, jsize responseCode, jstring accessToken)
 {
 	EFacebookLoginResponse LoginResponse = (EFacebookLoginResponse)responseCode;
 
diff --git a/Engine/Source/Programs/UnrealBuildTool/Platform/Android/UEDeployAndroid.cs b/Engine/Source/Programs/UnrealBuildTool/Platform/Android/UEDeployAndroid.cs
--- a/Engine/Source/Programs/UnrealBuildTool/Platform/Android/UEDeployAndroid.cs
+++ b/Engine/Source/Programs/UnrealBuildTool/Platform/Android/UEDeployAndroid.cs
@@ -3419,6 +3419,14 @@ namespace UnrealBuildTool
 				}
 			}
 
+			// force cleanup if older UE4 project
+			if (File.Exists(Path.Combine(IntermediateAndroidPath, "arm64", "jni", "arm64-v8a", "libUE4.so")) ||
+				File.Exists(Path.Combine(IntermediateAndroidPath, "x64", "jni", "x86_64", "libUE4.so")))
+			{
+				Log.TraceInformation("Old version of library .so found, forcing clean");
+				bCreateFromScratch = true;
+			}
+
 			// check if the enabled plugins has changed
 			string PluginListFilename = Path.Combine(IntermediateAndroidPath, "ActiveUPL.txt");
 			string PluginListContents = ActiveUPLFiles.ToString();
@@ -3513,6 +3521,7 @@ namespace UnrealBuildTool
 
 				// Make a set of files that are okay to clean up
 				HashSet<string> cleanFiles = new HashSet<string>();
+				cleanFiles.Add("DownloadShim.java");
 				cleanFiles.Add("OBBData.java");
 				foreach (TemplateFile template in templates)
 				{
@@ -3521,8 +3530,11 @@ namespace UnrealBuildTool
 
 				foreach (string filename in files)
 				{
-					if (filename == UnrealDownloadShimFileName)  // we always need the shim, and it'll get rewritten if needed anyway
+					// keep the shim if it is in the right place
+					if (filename == UnrealDownloadShimFileName)
+					{
 						continue;
+					}
 
 					string filePath = Path.GetDirectoryName(filename);  // grab the file's path
 					if (filePath != TemplateDestinationBase)             // and check to make sure it isn't the same as the Template directory we calculated earlier
