Protect VS Andriod config with platfrom check to fix project generation on non-Windows platforms.

#rnx

diff --git a/Engine/Source/Programs/UnrealBuildTool/Platform/Android/AndroidProjectGenerator.cs b/Engine/Source/Programs/UnrealBuildTool/Platform/Android/AndroidProjectGenerator.cs
--- a/Engine/Source/Programs/UnrealBuildTool/Platform/Android/AndroidProjectGenerator.cs
+++ b/Engine/Source/Programs/UnrealBuildTool/Platform/Android/AndroidProjectGenerator.cs
@@ -40,15 +40,18 @@ namespace UnrealBuildTool
 		public AndroidProjectGenerator(CommandLineArguments Arguments)
 			: base(Arguments)
 		{
-			if (Arguments.HasOption("-vsdebugandroid"))
+			if (HostSupportsVSAndroid())
 			{
-				VSDebugCommandLineOptionPresent = true;
-			}
+				if (Arguments.HasOption("-vsdebugandroid"))
+				{
+					VSDebugCommandLineOptionPresent = true;
+				}
 
-			AGDEInstalled = Microsoft.Win32.Registry.LocalMachine.OpenSubKey(@"SOFTWARE\WOW6432Node\Google\AndroidGameDevelopmentExtension")?.ValueCount > 0;
-			if (Arguments.HasOption("-noagde"))
-			{
-				AGDEInstalled = false;
+				AGDEInstalled = Microsoft.Win32.Registry.LocalMachine.OpenSubKey(@"SOFTWARE\WOW6432Node\Google\AndroidGameDevelopmentExtension")?.ValueCount > 0;
+				if (Arguments.HasOption("-noagde"))
+				{
+					AGDEInstalled = false;
+				}
 			}
 		}
 
