ubt: fix for process ids that exist but have exited causing incrementabl build slowdowns.

diff --git a/Engine/Source/Programs/UnrealBuildTool/System/HotReload.cs b/Engine/Source/Programs/UnrealBuildTool/System/HotReload.cs
--- a/Engine/Source/Programs/UnrealBuildTool/System/HotReload.cs
+++ b/Engine/Source/Programs/UnrealBuildTool/System/HotReload.cs
@@ -346,6 +346,34 @@ namespace UnrealBuildTool
 								continue;
 							}
 
+							// bugfix - the editor sometimes doesn't delete its editorrun file due to
+							// crash or debugger stop or whatever. ~eventually~ this should get caught
+							// by the above check where the PID no longer exists, however windows actually
+							// keeps the process table entry around for a ~long~ time (days, across hibernations).
+							//
+							// What ends up happening is we successfully get the Process object, but we throw
+							// an exception trying to retrieve the module handle for the filename, and then
+							// don't delete it.
+							//
+							// On my machine this was ~750 ms _per orphaned file_, and I spoke to someone
+							// with 10 of these in his Engine/Intermediate/EditorRun directory.
+							// 
+							try
+							{
+								if (RunningProcess.HasExited)
+								{
+									FileReference.Delete(EditorInstanceFile);
+									continue;
+								}
+							}
+							catch
+							{
+								// if we got here, we either failed to access the HasExited, or 
+								// failed to delete the file.
+								Log.TraceLog("Failed to delete EditorRun file for exited process: {0}", EditorInstanceFile.GetFileName());
+								continue;
+							}
+
 							FileReference MainModuleFile;
 							try
 							{
