Fix for netcore AutomationTool/UnrealBuildTool binaries

diff --git a/Engine/Build/InstalledEngineBuild.xml b/Engine/Build/InstalledEngineBuild.xml
--- a/Engine/Build/InstalledEngineBuild.xml
+++ b/Engine/Build/InstalledEngineBuild.xml
@@ -507,6 +507,10 @@
 				</ForEach>
 			</Do>
 			
+			<!-- Tag AutomationTool and UnrealBuildTool folders recursively as NET Core dependencies are not currently handled by CsCompile -->
+			<Tag Files="$(CsToolsDir)/Engine/Binaries/DotNET/AutomationTool/..." With="#Build Tools CS Binaries" />
+			<Tag Files="$(CsToolsDir)/Engine/Binaries/DotNET/UnrealBuildTool/..." With="#Build Tools CS Binaries" />
+			
 			<!-- Compile IPP. Mac/Linux bundled versions of Mono do not have Winforms, so we just use the existing one. -->
 			<Property Name="CompileIPP" Value="true"/>
 			<Property Name="CompileIPP" Value="false" If="'$(HostPlatform)' != 'Win64'"/>
diff --git a/Engine/Source/Programs/UnrealBuildTool/ProjectFiles/VisualStudio/VCProject.cs b/Engine/Source/Programs/UnrealBuildTool/ProjectFiles/VisualStudio/VCProject.cs
--- a/Engine/Source/Programs/UnrealBuildTool/ProjectFiles/VisualStudio/VCProject.cs
+++ b/Engine/Source/Programs/UnrealBuildTool/ProjectFiles/VisualStudio/VCProject.cs
@@ -86,7 +86,7 @@ namespace UnrealBuildTool
 
 		/// The name of the Visual C++ platform to use for stub project configurations
 		/// NOTE: We always use Win32 for the stub project's platform, since that is guaranteed to be supported by Visual Studio
-		static public readonly string StubProjectPlatformName = "Win32";
+		static public readonly string StubProjectPlatformName = "Win64";
 
 		/// <summary>
 		/// The Guid representing the project type e.g. C# or C++
@@ -265,7 +265,7 @@ namespace UnrealBuildTool
 		/// This is the platform name that Visual Studio is always guaranteed to support.  We'll use this as
 		/// a platform for any project configurations where our actual platform is not supported by the
 		/// installed version of Visual Studio (e.g, "iOS")
-		public const string DefaultPlatformName = "Win32";
+		public const string DefaultPlatformName = "Win64";
 
 		// This is the GUID that Visual Studio uses to identify a C++ project file in the solution
 		public override string ProjectTypeGUID
