Swap single quotes out for delimited double quotes.

diff --git a/Engine/Source/Programs/UnrealBuildTool/System/BuildHostPlatform.cs b/Engine/Source/Programs/UnrealBuildTool/System/BuildHostPlatform.cs
--- a/Engine/Source/Programs/UnrealBuildTool/System/BuildHostPlatform.cs
+++ b/Engine/Source/Programs/UnrealBuildTool/System/BuildHostPlatform.cs
@@ -295,7 +295,7 @@ namespace UnrealBuildTool
 			string TempFile = Path.Combine("/var/tmp", Path.GetTempFileName());
 			ProcessStartInfo StartInfo = new ProcessStartInfo();
 			StartInfo.FileName = "/bin/sh";
-			StartInfo.Arguments = "-c 'ps -eaw -o pid,comm > " + TempFile + "'";
+			StartInfo.Arguments = "-c \"ps -eaw -o pid,comm > " + TempFile + "\"";
 			StartInfo.CreateNoWindow = true;
 
 			Process Proc = new Process();
