Make sure the InterchangeWorker is always build when we build the editor.

diff --git a/Engine/Source/Programs/UnrealBuildTool/Configuration/TargetRules.cs b/Engine/Source/Programs/UnrealBuildTool/Configuration/TargetRules.cs
--- a/Engine/Source/Programs/UnrealBuildTool/Configuration/TargetRules.cs
+++ b/Engine/Source/Programs/UnrealBuildTool/Configuration/TargetRules.cs
@@ -1648,6 +1648,11 @@ namespace UnrealBuildTool
 			else if(Type == global::UnrealBuildTool.TargetType.Editor)
 			{
 				GlobalDefinitions.Add("UE_EDITOR=1");
+				if (Platform == UnrealTargetPlatform.Win64 || Platform == UnrealTargetPlatform.Linux || Platform == UnrealTargetPlatform.Mac)
+				{
+					List<string> EmptyArguments = new List<string>();
+					PreBuildTargets.Add(new TargetInfo("InterchangeWorker", Platform, Configuration, Architecture, ProjectFile, new CommandLineArguments(EmptyArguments.ToArray())));
+				}
 			}
 			else if(Type == global::UnrealBuildTool.TargetType.Server)
 			{
