UBT: Rename UE4CSharp.prog files to *.uprogram.

diff --git a/Engine/Source/Programs/UnrealBuildTool/ProjectFiles/ProjectFileGenerator.cs b/Engine/Source/Programs/UnrealBuildTool/ProjectFiles/ProjectFileGenerator.cs
--- a/Engine/Source/Programs/UnrealBuildTool/ProjectFiles/ProjectFileGenerator.cs
+++ b/Engine/Source/Programs/UnrealBuildTool/ProjectFiles/ProjectFileGenerator.cs
@@ -555,7 +555,7 @@ namespace UnrealBuildTool
 					break;
 				}
 				// For csproj files, add them to the sln if the UE4CSharp.prog file also exists.
-				if (bIsCsProj && FileReference.Exists(FileReference.Combine(SearchFolder, "UE4CSharp.prog")))
+				if (bIsCsProj && FileReference.Exists(File.ChangeExtension(".uprogram")))
 				{
 					FoundProjects.Add(File);
 				}
