Temporary logging when a parallel executor task fails but doesn't generate any output.
This is to try and catch issues with PVS-Studio in the farm that we are having trouble reproducing in a controlled environment.

Testing: Caused a compile error during a normal build with the if statement disabled, generated output as expected with the GotOutput value never getting set.

#rnx

diff --git a/Engine/Source/Programs/UnrealBuildTool/Executors/ParallelExecutor.cs b/Engine/Source/Programs/UnrealBuildTool/Executors/ParallelExecutor.cs
--- a/Engine/Source/Programs/UnrealBuildTool/Executors/ParallelExecutor.cs
+++ b/Engine/Source/Programs/UnrealBuildTool/Executors/ParallelExecutor.cs
@@ -31,6 +31,10 @@ namespace UnrealBuildTool
 			public List<string> LogLines = new List<string>();
 			public int ExitCode = -1;
 
+			// BEGIN TEMPORARY TO CATCH PVS-STUDIO ISSUES
+			public bool GotOutput = false;
+			// END TEMPORARY
+
 			public BuildAction(LinkedAction Inner)
 			{
 				this.Inner = Inner;
@@ -227,6 +231,15 @@ namespace UnrealBuildTool
 										}
 										else
 										{
+											// BEGIN TEMPORARY TO CATCH PVS-STUDIO ISSUES
+											if (!CompletedAction.GotOutput)
+											{
+												string StatusDescription = CompletedAction.LogLines.Count > 0 ? CompletedAction.LogLines[0] : "";
+												Log.TraceInformation("[{0}/{1}] {2} - Error but no output", NumCompletedActions, InputActions.Count, StatusDescription);
+												Log.TraceInformation("[{0}/{1}] {2} - {3} {4} {5} {6}", NumCompletedActions, InputActions.Count, StatusDescription, CompletedAction.ExitCode, 
+													CompletedAction.Inner.WorkingDirectory, CompletedAction.Inner.CommandPath, CompletedAction.Inner.CommandArguments);
+											}
+											// END TEMPORARY
 											// Update the exit code if it's not already set
 											if(bResult && CompletedAction.ExitCode != 0)
 											{
@@ -265,6 +278,9 @@ namespace UnrealBuildTool
 				Action.LogLines.Add(Action.Inner.StatusDescription);
 			}
 
+			// BEGIN TEMPORARY TO CATCH PVS-STUDIO ISSUES
+			int LogLineCount = Action.LogLines.Count;
+			// END TEMPORARY
 			try
 			{
 				using (ManagedProcess Process = new ManagedProcess(ProcessGroup, Action.Inner.CommandPath.FullName, Action.Inner.CommandArguments, Action.Inner.WorkingDirectory.FullName, null, null, ProcessPriorityClass.BelowNormal))
@@ -272,6 +288,9 @@ namespace UnrealBuildTool
 					Action.LogLines.AddRange(Process.ReadAllLines());
 					Action.ExitCode = Process.ExitCode;
 				}
+				// BEGIN TEMPORARY TO CATCH PVS-STUDIO ISSUES
+				Action.GotOutput = Action.LogLines.Count > LogLineCount;
+				// END TEMPORARY
 			}
 			catch(Exception Ex)
 			{
