Add additional conditions for invalidating previous APK

#ue5
#android
Brandon.Schaefer

diff --git a/Engine/Source/Programs/UnrealBuildTool/Platform/Android/UEDeployAndroid.cs b/Engine/Source/Programs/UnrealBuildTool/Platform/Android/UEDeployAndroid.cs
--- a/Engine/Source/Programs/UnrealBuildTool/Platform/Android/UEDeployAndroid.cs
+++ b/Engine/Source/Programs/UnrealBuildTool/Platform/Android/UEDeployAndroid.cs
@@ -1534,6 +1534,7 @@ namespace UnrealBuildTool
 			StartInfo.Arguments = Params;
 			StartInfo.UseShellExecute = bUseShellExecute;
 			StartInfo.WindowStyle = ProcessWindowStyle.Minimized;
+			StartInfo.RedirectStandardInput = true;
 			StartInfo.RedirectStandardOutput = true;
 			StartInfo.RedirectStandardError = true;
 
@@ -1541,9 +1542,15 @@ namespace UnrealBuildTool
 			Proc.StartInfo = StartInfo;
 			Proc.OutputDataReceived += FilterStdOutErr;
 			Proc.ErrorDataReceived += FilterStdOutErr;
+
 			Proc.Start();
 			Proc.BeginOutputReadLine();
 			Proc.BeginErrorReadLine();
+
+			StreamWriter StreamIn = Proc.StandardInput;
+			StreamIn.WriteLine("yes");
+			StreamIn.Close();
+
 			Proc.WaitForExit();
 
 			// android bat failure
@@ -3490,10 +3497,10 @@ namespace UnrealBuildTool
 				RequiredOBBFiles.Add(PatchFileLocation);
 			}
 
-			// If we are not skipping Gradle build this is done per architecture so store version may be different
-			if (bSkipGradleBuild)
+			if (Arches.Count > 0)
 			{
 				// Generate the OBBData.java file if out of date (can skip rewriting it if packaging inside Apk in some cases)
+				// Note: this may be replaced per architecture later if store version is different
 				WriteJavaOBBDataFile(UnrealOBBDataFileName, PackageName, RequiredOBBFiles, CookFlavor, bPackageDataInsideApk, Arches[0].Substring(1));
 			}
 
