Add another quote to make dotnet happy due to dotnot supporting single quotes as quotes

diff --git a/Engine/Source/Programs/UnrealBuildTool/System/ActionGraphBuilder.cs b/Engine/Source/Programs/UnrealBuildTool/System/ActionGraphBuilder.cs
--- a/Engine/Source/Programs/UnrealBuildTool/System/ActionGraphBuilder.cs
+++ b/Engine/Source/Programs/UnrealBuildTool/System/ActionGraphBuilder.cs
@@ -221,7 +221,7 @@ namespace UnrealBuildTool
 			}
 			else
 			{
-				CopyAction.CommandArguments = String.Format("-c 'cp -f \"{0}\" \"{1}\"'", SourceFile.AbsolutePath, TargetFile.AbsolutePath);
+				CopyAction.CommandArguments = String.Format("-c 'cp -f \"\"{0}\"\" \"\"{1}\"'", SourceFile.AbsolutePath, TargetFile.AbsolutePath);
 			}
 			CopyAction.WorkingDirectory = UnrealBuildTool.EngineSourceDirectory;
 			CopyAction.PrerequisiteItems.Add(SourceFile);
diff --git a/Engine/Source/Programs/UnrealBuildTool/ToolChain/UEToolChain.cs b/Engine/Source/Programs/UnrealBuildTool/ToolChain/UEToolChain.cs
--- a/Engine/Source/Programs/UnrealBuildTool/ToolChain/UEToolChain.cs
+++ b/Engine/Source/Programs/UnrealBuildTool/ToolChain/UEToolChain.cs
@@ -482,22 +482,7 @@ namespace UnrealBuildTool
 				FileItem SourceFileItem = FileItem.GetItemByFileReference(SourceFile);
 				FileItem TargetFileItem = FileItem.GetItemByFileReference(TargetFile);
 
-				Action CopyAction = Graph.CreateAction(ActionType.BuildProject);
-				CopyAction.CommandDescription = "Copy";
-				CopyAction.CommandPath = BuildHostPlatform.Current.Shell;
-				if (BuildHostPlatform.Current.ShellType == ShellType.Cmd)
-				{
-					CopyAction.CommandArguments = String.Format("/C \"copy /Y \"{0}\" \"{1}\" 1>nul\"", SourceFile, TargetFile);
-				}
-				else
-				{
-					CopyAction.CommandArguments = String.Format("-c 'cp -f \"{0}\" \"{1}\"'", SourceFile.FullName, TargetFile.FullName);
-				}
-				CopyAction.WorkingDirectory = UnrealBuildTool.EngineSourceDirectory;
-				CopyAction.PrerequisiteItems.Add(SourceFileItem);
-				CopyAction.ProducedItems.Add(TargetFileItem);
-				CopyAction.StatusDescription = TargetFileItem.Location.GetFileName();
-				CopyAction.bCanExecuteRemotely = false;
+				Action CopyAction = Graph.CreateCopyAction(SourceFileItem, TargetFileItem);
 				CopyAction.bShouldOutputStatusDescription = false;
 
 				Result.GeneratedHeaderFiles.Add(TargetFileItem);
