Create a stable GUID for *.csproj files in the solution if there isn't one in the project itself (as is the case with NET Core).

#rnx

diff --git a/Engine/Source/Programs/UnrealBuildTool/ProjectFiles/VisualStudio/VCProject.cs b/Engine/Source/Programs/UnrealBuildTool/ProjectFiles/VisualStudio/VCProject.cs
--- a/Engine/Source/Programs/UnrealBuildTool/ProjectFiles/VisualStudio/VCProject.cs
+++ b/Engine/Source/Programs/UnrealBuildTool/ProjectFiles/VisualStudio/VCProject.cs
@@ -96,6 +96,8 @@ namespace UnrealBuildTool
 			get { throw new BuildException("Unrecognized type of project file for Visual Studio solution"); }
 		}
 
+		static Guid PathNamespaceGuid { get; } = new Guid("2D8570D5-7FFC-4E6D-A9D7-E860E117D717");
+
 		/// <summary>
 		/// Constructs a new project file object
 		/// </summary>
@@ -116,14 +118,14 @@ namespace UnrealBuildTool
 				catch (Exception)
 				{
 					// Failed to find GUID, so just create a new one
-					ProjectGUID = Guid.NewGuid();
+					ProjectGUID = VCProjectFileGenerator.MakeMd5Guid(PathNamespaceGuid, ProjectFilePath.FullName);
 				}
 			}
 
 			if (ProjectGUID == Guid.Empty)
 			{
 				// Generate a brand new GUID
-				ProjectGUID = Guid.NewGuid();
+				ProjectGUID = VCProjectFileGenerator.MakeMd5Guid(PathNamespaceGuid, ProjectFilePath.FullName);
 			}
 		}
 
diff --git a/Engine/Source/Programs/UnrealBuildTool/ProjectFiles/VisualStudio/VCProjectFileGenerator.cs b/Engine/Source/Programs/UnrealBuildTool/ProjectFiles/VisualStudio/VCProjectFileGenerator.cs
--- a/Engine/Source/Programs/UnrealBuildTool/ProjectFiles/VisualStudio/VCProjectFileGenerator.cs
+++ b/Engine/Source/Programs/UnrealBuildTool/ProjectFiles/VisualStudio/VCProjectFileGenerator.cs
@@ -464,6 +464,19 @@ namespace UnrealBuildTool
 			return new Guid(Hash);
 		}
 
+		public static Guid MakeMd5Guid(Guid Namespace, string Text)
+		{
+			byte[] Input = new byte[16 + Encoding.UTF8.GetByteCount(Text)];
+
+			Namespace.TryWriteBytes(Input.AsSpan(0, 16));
+			Array.Reverse(Input, 0, 4);
+			Array.Reverse(Input, 4, 2);
+			Array.Reverse(Input, 6, 2);
+
+			Encoding.UTF8.GetBytes(Text, 0, Text.Length, Input, 16);
+			return MakeMd5Guid(Input);
+		}
+
 		/// <summary>
 		/// Writes the project files to disk
 		/// </summary>
