When invoking UBT with -Module=X on Mac it runs a deploy step at the end of a build. This requires a build receipt which may not exist if a full build hasn't yet been completed.

#rnx

diff --git a/Engine/Source/Programs/UnrealBuildTool/Configuration/UEBuildTarget.cs b/Engine/Source/Programs/UnrealBuildTool/Configuration/UEBuildTarget.cs
--- a/Engine/Source/Programs/UnrealBuildTool/Configuration/UEBuildTarget.cs
+++ b/Engine/Source/Programs/UnrealBuildTool/Configuration/UEBuildTarget.cs
@@ -1054,7 +1054,7 @@ namespace UnrealBuildTool
 			RulesAssembly = InRulesAssembly;
 			TargetType = Rules.Type;
 			ForeignPlugin = InDescriptor.ForeignPlugin;
-			bDeployAfterCompile = InRules.bDeployAfterCompile && !InRules.bDisableLinking && InDescriptor.SpecificFilesToCompile.Count == 0;
+			bDeployAfterCompile = InRules.bDeployAfterCompile && !InRules.bDisableLinking && InDescriptor.SpecificFilesToCompile.Count == 0 && InDescriptor.OnlyModuleNames.Count == 0;
 
 			// now that we have the platform, we can set the intermediate path to include the platform/architecture name
 			PlatformIntermediateFolder = GetPlatformIntermediateFolder(Platform, Architecture);
