Fix ISPC copy action wanting to avoid deleting a file, revert back to this behaviour

diff --git a/Engine/Source/Programs/UnrealBuildTool/ToolChain/UEToolChain.cs b/Engine/Source/Programs/UnrealBuildTool/ToolChain/UEToolChain.cs
--- a/Engine/Source/Programs/UnrealBuildTool/ToolChain/UEToolChain.cs
+++ b/Engine/Source/Programs/UnrealBuildTool/ToolChain/UEToolChain.cs
@@ -482,7 +482,22 @@ namespace UnrealBuildTool
 				FileItem SourceFileItem = FileItem.GetItemByFileReference(SourceFile);
 				FileItem TargetFileItem = FileItem.GetItemByFileReference(TargetFile);
 
-				Action CopyAction = Graph.CreateCopyAction(SourceFileItem, TargetFileItem);
+				Action CopyAction = Graph.CreateAction(ActionType.BuildProject);
+				CopyAction.CommandDescription = "Copy";
+				CopyAction.CommandPath = BuildHostPlatform.Current.Shell;
+				if (BuildHostPlatform.Current.ShellType == ShellType.Cmd)
+				{
+					CopyAction.CommandArguments = String.Format("/C \"copy /Y \"{0}\" \"{1}\" 1>nul\"", SourceFile, TargetFile);
+				}
+				else
+				{
+					CopyAction.CommandArguments = String.Format("-c 'cp -f \"\"{0}\"\" \"\"{1}\"'", SourceFile.FullName, TargetFile.FullName);
+				}
+				CopyAction.WorkingDirectory = UnrealBuildTool.EngineSourceDirectory;
+				CopyAction.PrerequisiteItems.Add(SourceFileItem);
+				CopyAction.ProducedItems.Add(TargetFileItem);
+				CopyAction.StatusDescription = TargetFileItem.Location.GetFileName();
+				CopyAction.bCanExecuteRemotely = false;
 				CopyAction.bShouldOutputStatusDescription = false;
 
 				Result.GeneratedHeaderFiles.Add(TargetFileItem);
