UBT: Compile marketplace plugin rules into a separate assembly, and do not overwrite files under the engine folder.

#rnx

diff --git a/Engine/Source/Programs/UnrealBuildTool/System/DynamicCompilation.cs b/Engine/Source/Programs/UnrealBuildTool/System/DynamicCompilation.cs
--- a/Engine/Source/Programs/UnrealBuildTool/System/DynamicCompilation.cs
+++ b/Engine/Source/Programs/UnrealBuildTool/System/DynamicCompilation.cs
@@ -34,13 +34,12 @@ namespace UnrealBuildTool
 		/// <returns>True if the assembly needs to be built</returns>
 		private static bool RequiresCompilation(HashSet<FileReference> SourceFiles, FileReference AssemblyManifestFilePath, FileReference OutputAssemblyPath)
 		{
-			// HACK: We cannot do this without breaking marketplace plugins. Need to distribute separate rules assemblies for those.
-//			// Do not compile the file if it's installed
-//			if (UnrealBuildTool.IsFileInstalled(OutputAssemblyPath))
-//			{
-//				Log.TraceLog("Skipping {0}: File is installed", OutputAssemblyPath);
-//				return false;
-//			}
+			// Do not compile the file if it's installed
+			if (UnrealBuildTool.IsFileInstalled(OutputAssemblyPath))
+			{
+				Log.TraceLog("Skipping {0}: File is installed", OutputAssemblyPath);
+				return false;
+			}
 
 			// Check to see if we already have a compiled assembly file on disk
 			FileItem OutputAssemblyInfo = FileItem.GetItemByFileReference(OutputAssemblyPath);
diff --git a/Engine/Source/Programs/UnrealBuildTool/System/RulesCompiler.cs b/Engine/Source/Programs/UnrealBuildTool/System/RulesCompiler.cs
--- a/Engine/Source/Programs/UnrealBuildTool/System/RulesCompiler.cs
+++ b/Engine/Source/Programs/UnrealBuildTool/System/RulesCompiler.cs
@@ -348,14 +348,28 @@ namespace UnrealBuildTool
 		{
 			if (EngineRulesAssembly == null)
 			{
-				List<PluginInfo> IncludedPlugins = new List<PluginInfo>();
+				List<PluginInfo> EnginePlugins = new List<PluginInfo>();
+				List<PluginInfo> MarketplacePlugins = new List<PluginInfo>();
 
-				// search for all engine plugins
-				IncludedPlugins.AddRange(Plugins.ReadEnginePlugins(UnrealBuildTool.EngineDirectory));
+				DirectoryReference MarketplaceDirectory = DirectoryReference.Combine(UnrealBuildTool.EngineDirectory, "Plugins", "Marketplace");
+				foreach (PluginInfo PluginInfo in Plugins.ReadEnginePlugins(UnrealBuildTool.EngineDirectory))
+				{
+					if (PluginInfo.File.IsUnderDirectory(MarketplaceDirectory))
+					{
+						MarketplacePlugins.Add(PluginInfo);
+					}
+					else
+					{
+						EnginePlugins.Add(PluginInfo);
+					}
+				}
 
-				RulesScope EngineScope = new RulesScope("Engine", null);
+				EngineRulesAssembly = CreateEngineRulesAssemblyInternal(UnrealBuildTool.GetExtensionDirs(UnrealBuildTool.EngineDirectory), ProjectFileGenerator.EngineProjectFileNameBase, EnginePlugins, UnrealBuildTool.IsEngineInstalled() || bUsePrecompiled, bSkipCompile, null);
 
-				EngineRulesAssembly = CreateEngineRulesAssemblyInternal(EngineScope, UnrealBuildTool.GetExtensionDirs(UnrealBuildTool.EngineDirectory), ProjectFileGenerator.EngineProjectFileNameBase, IncludedPlugins, UnrealBuildTool.IsEngineInstalled() || bUsePrecompiled, bSkipCompile, null);
+				if (MarketplacePlugins.Count > 0)
+				{
+					EngineRulesAssembly = CreateMarketplaceRulesAssembly(MarketplacePlugins, bSkipCompile, EngineRulesAssembly);
+				}
 			}
 			return EngineRulesAssembly;
 		}
@@ -363,7 +377,6 @@ namespace UnrealBuildTool
 		/// <summary>
 		/// Creates a rules assembly
 		/// </summary>
-		/// <param name="Scope">Scope for items created from this assembly</param>
 		/// <param name="RootDirectories">The root directories to create rules for</param>
 		/// <param name="AssemblyPrefix">A prefix for the assembly file name</param>
 		/// <param name="Plugins">List of plugins to include in this assembly</param>
@@ -371,11 +384,12 @@ namespace UnrealBuildTool
 		/// <param name="bSkipCompile">Whether to skip compilation for this assembly</param>
 		/// <param name="Parent">The parent rules assembly</param>
 		/// <returns>New rules assembly</returns>
-		private static RulesAssembly CreateEngineRulesAssemblyInternal(RulesScope Scope, List<DirectoryReference> RootDirectories, string AssemblyPrefix, IReadOnlyList<PluginInfo> Plugins, bool bReadOnly, bool bSkipCompile, RulesAssembly? Parent)
+		private static RulesAssembly CreateEngineRulesAssemblyInternal(List<DirectoryReference> RootDirectories, string AssemblyPrefix, IReadOnlyList<PluginInfo> Plugins, bool bReadOnly, bool bSkipCompile, RulesAssembly? Parent)
 		{
 			// Scope hierarchy
-			RulesScope PluginsScope = new RulesScope(Scope.Name + " Plugins", Scope);
-			RulesScope ProgramsScope = new RulesScope(Scope.Name + " Programs", PluginsScope);
+			RulesScope Scope= new RulesScope("Engine", null);
+			RulesScope PluginsScope = new RulesScope("Engine Plugins", Scope);
+			RulesScope ProgramsScope = new RulesScope("Engine Programs", PluginsScope);
 
 			// Find the shared modules, excluding the programs directory. These are used to create an assembly with the bContainsEngineModules flag set to true.
 			Dictionary<FileReference, ModuleRulesContext> ModuleFileToContext = new Dictionary<FileReference, ModuleRulesContext>();
@@ -437,6 +451,35 @@ namespace UnrealBuildTool
 			return ProgramAssembly;
 		}
 
+		/// <summary>
+		/// Creates a rules assembly
+		/// </summary>
+		/// <param name="Plugins">List of plugins to include in this assembly</param>
+		/// <param name="bSkipCompile">Whether to skip compilation for this assembly</param>
+		/// <param name="Parent">The parent rules assembly</param>
+		/// <returns>New rules assembly</returns>
+		private static RulesAssembly CreateMarketplaceRulesAssembly(IReadOnlyList<PluginInfo> Plugins, bool bSkipCompile, RulesAssembly Parent)
+		{
+			RulesScope MarketplaceScope = new RulesScope("Marketplace", Parent.Scope);
+
+			// Add all the plugin modules too (don't need to loop over RootDirectories since the plugins come in already found
+			Dictionary<FileReference, ModuleRulesContext> ModuleFileToContext = new Dictionary<FileReference, ModuleRulesContext>();
+			using (Timeline.ScopeEvent("Finding marketplace plugin modules"))
+			{
+				ModuleRulesContext PluginsModuleContext = new ModuleRulesContext(MarketplaceScope, UnrealBuildTool.EngineDirectory);
+				FindModuleRulesForPlugins(Plugins, PluginsModuleContext, ModuleFileToContext);
+			}
+
+			// Create the assembly
+			RulesAssembly Result = Parent;
+			if (ModuleFileToContext.Count > 0)
+			{
+				FileReference AssemblyFileName = FileReference.Combine(UnrealBuildTool.WritableEngineDirectory, "Intermediate", "Build", "BuildRules", "MarketplaceRules.dll");
+				Result = new RulesAssembly(MarketplaceScope, new List<DirectoryReference> { UnrealBuildTool.EngineDirectory }, Plugins, ModuleFileToContext, new List<FileReference>(), AssemblyFileName, bContainsEngineModules: true, DefaultBuildSettings: BuildSettingsVersion.Latest, bReadOnly: false, bSkipCompile: bSkipCompile, Parent: Parent);
+			}
+			return Result;
+		}
+
 		/// <summary>
 		/// Creates a rules assembly with the given parameters.
 		/// </summary>
diff --git a/Engine/Source/Programs/UnrealBuildTool/UnrealBuildTool.cs b/Engine/Source/Programs/UnrealBuildTool/UnrealBuildTool.cs
--- a/Engine/Source/Programs/UnrealBuildTool/UnrealBuildTool.cs
+++ b/Engine/Source/Programs/UnrealBuildTool/UnrealBuildTool.cs
@@ -100,6 +100,38 @@ namespace UnrealBuildTool
 		[Obsolete("Please use UnrealBuildTool.GetExtensionDirs(UnrealBuildTool.EngineDirectory, \"Source/ThirdParty\") instead.")]
 		public static readonly DirectoryReference EngineSourceThirdPartyDirectory = DirectoryReference.Combine(EngineSourceDirectory, "ThirdParty");
 
+		/// <summary>
+		/// Cached copy of the writable engine directory
+		/// </summary>
+		static DirectoryReference? CachedWritableEngineDirectory;
+
+		/// <summary>
+		/// Writable engine directory. Uses the user's settings folder for installed builds.
+		/// </summary>
+		public static DirectoryReference WritableEngineDirectory
+		{
+			get
+			{
+				if (CachedWritableEngineDirectory == null)
+				{
+					DirectoryReference? UserDir = null;
+					if (IsEngineInstalled())
+					{
+						UserDir = Utils.GetUserSettingDirectory();
+					}
+					if (UserDir == null)
+					{
+						CachedWritableEngineDirectory = EngineDirectory;
+					}
+					else
+					{
+						CachedWritableEngineDirectory = DirectoryReference.Combine(UserDir, "UnrealEngine");
+					}
+				}
+				return CachedWritableEngineDirectory;
+			}
+		}
+
 		/// <summary>
 		/// The engine programs directory
 		/// </summary>
