UBT: Add support for ASan on Windows for MSVC versions that support it.
Note that up until VS 16.9, according to the MSVC blog, ASan requires explcit helper libraries to be added to the linker line, which requires knowledge of what type of CRT we're compiling with on the link side.

diff --git a/Engine/Source/Programs/UnrealBuildTool/Configuration/UEBuildTarget.cs b/Engine/Source/Programs/UnrealBuildTool/Configuration/UEBuildTarget.cs
--- a/Engine/Source/Programs/UnrealBuildTool/Configuration/UEBuildTarget.cs
+++ b/Engine/Source/Programs/UnrealBuildTool/Configuration/UEBuildTarget.cs
@@ -3400,7 +3400,8 @@ namespace UnrealBuildTool
 			// UHTLite
 			GlobalCompileEnvironment.bHackHeaderGenerator |= (AppName == "UHTLite");
 
-			GlobalCompileEnvironment.bUseDebugCRT = GlobalCompileEnvironment.Configuration == CppConfiguration.Debug && Rules.bDebugBuildsActuallyUseDebugCRT;
+			bool bUseDebugCRT = GlobalCompileEnvironment.Configuration == CppConfiguration.Debug && Rules.bDebugBuildsActuallyUseDebugCRT;
+			GlobalCompileEnvironment.bUseDebugCRT = bUseDebugCRT;
 			GlobalCompileEnvironment.bEnableOSX109Support = Rules.bEnableOSX109Support;
 			GlobalCompileEnvironment.Definitions.Add(String.Format("IS_PROGRAM={0}", TargetType == TargetType.Program ? "1" : "0"));
 			GlobalCompileEnvironment.Definitions.AddRange(Rules.GlobalDefinitions);
@@ -3432,6 +3433,8 @@ namespace UnrealBuildTool
 			GlobalCompileEnvironment.CppStandard = Rules.CppStandard;
 			GlobalCompileEnvironment.AdditionalArguments = Rules.AdditionalCompilerArguments ?? String.Empty;
 
+			GlobalLinkEnvironment.bUseDebugCRT = bUseDebugCRT;
+			GlobalLinkEnvironment.bUseStaticCRT = Rules.bUseStaticCRT;
 			GlobalLinkEnvironment.bIsBuildingConsoleApplication = Rules.bIsBuildingConsoleApplication;
 			GlobalLinkEnvironment.bOptimizeForSize = Rules.bCompileForSize;
 			GlobalLinkEnvironment.bOmitFramePointers = Rules.bOmitFramePointers;
diff --git a/Engine/Source/Programs/UnrealBuildTool/Platform/Windows/UEBuildWindows.cs b/Engine/Source/Programs/UnrealBuildTool/Platform/Windows/UEBuildWindows.cs
--- a/Engine/Source/Programs/UnrealBuildTool/Platform/Windows/UEBuildWindows.cs
+++ b/Engine/Source/Programs/UnrealBuildTool/Platform/Windows/UEBuildWindows.cs
@@ -193,6 +193,13 @@ namespace UnrealBuildTool
 		[CommandLine("-StaticAnalyzer")]
 		public WindowsStaticAnalyzer StaticAnalyzer = WindowsStaticAnalyzer.None;
 
+		/// <summary>
+		/// Enables address sanitizer (ASan). Only supported for Visual Studio 2019 16.7.0 and up.
+		/// </summary>
+		[XmlConfigFile(Category = "BuildConfiguration", Name = "bEnableAddressSanitizer")]
+		[CommandLine("-EnableASan")]
+		public bool bEnableAddressSanitizer = false;
+
 		/// <summary>
 		/// Whether we should export a file containing .obj to source file mappings.
 		/// </summary>
@@ -492,6 +499,11 @@ namespace UnrealBuildTool
 			get { return Inner.StaticAnalyzer; }
 		}
 
+		public bool bEnableAddressSanitizer
+		{ 
+			get { return Inner.bEnableAddressSanitizer; }
+		}
+
 		public string? ObjSrcMapFile
 		{
 			get { return Inner.ObjSrcMapFile; }
@@ -867,6 +879,12 @@ namespace UnrealBuildTool
 				throw new BuildException("Support for C++20 modules requires Visual Studio 2019 16.8 preview 3 or later.");
 			}
 
+			// Ensure we're using recent enough version of Visual Studio to support ASan builds.
+			if (Target.WindowsPlatform.bEnableAddressSanitizer && Target.WindowsPlatform.Environment.CompilerVersion < new VersionNumber(14, 27, 0))
+			{
+				throw new BuildException("Address sanitizer requires Visual Studio 2019 16.7 or later.");
+			}
+
 //			@Todo: Still getting reports of frequent OOM issues with this enabled as of 15.7.
 //			// Enable fast PDB linking if we're on VS2017 15.7 or later. Previous versions have OOM issues with large projects.
 //			if(!Target.bFormalBuild && !Target.bUseFastPDBLinking.HasValue && Target.WindowsPlatform.Compiler >= WindowsCompiler.VisualStudio2017)
diff --git a/Engine/Source/Programs/UnrealBuildTool/Platform/Windows/VCToolChain.cs b/Engine/Source/Programs/UnrealBuildTool/Platform/Windows/VCToolChain.cs
--- a/Engine/Source/Programs/UnrealBuildTool/Platform/Windows/VCToolChain.cs
+++ b/Engine/Source/Programs/UnrealBuildTool/Platform/Windows/VCToolChain.cs
@@ -344,6 +344,18 @@ namespace UnrealBuildTool
 				Arguments.Add("/Ob0");
 			}
 
+			// Address sanitizer
+			if (Target.WindowsPlatform.bEnableAddressSanitizer)
+			{
+				// Enable address sanitizer. This also requires companion libraries at link time.
+				// Works for clang too.
+				Arguments.Add("-fsanitize=address");
+
+				// Use the CRT allocator so that ASan is able to hook into it for better error
+				// detection.
+				AddDefinition(Arguments, "FORCE_ANSI_ALLOCATOR=1");
+			}
+
 			//
 			//	Debug
 			//
@@ -355,8 +367,11 @@ namespace UnrealBuildTool
 				// Favor code size (especially useful for embedded platforms).
 				Arguments.Add("/Os");
 
-				// Always include runtime error checks
-				Arguments.Add("/RTCs");
+				// Runtime checks and ASan are incompatible.
+				if (!Target.WindowsPlatform.bEnableAddressSanitizer)
+				{
+					Arguments.Add("/RTCs");
+				}
 			}
 			//
 			//	Development and LTCG
@@ -429,7 +444,7 @@ namespace UnrealBuildTool
 			else
 			{
 				// This is required to disable exception handling in VC platform headers.
-				Arguments.Add("/D_HAS_EXCEPTIONS=0");
+				AddDefinition(Arguments, "_HAS_EXCEPTIONS=0");
 			}
 
 			// If enabled, create debug information.
@@ -776,11 +791,16 @@ namespace UnrealBuildTool
 			// Prevents the linker from displaying its logo for each invocation.
 			Arguments.Add("/NOLOGO");
 
-			if (LinkEnvironment.bCreateDebugInfo)
+			// Address sanitizer requires debug info for symbolizing callstacks whether
+			// we're building debug or shipping.
+			if (LinkEnvironment.bCreateDebugInfo || Target.WindowsPlatform.bEnableAddressSanitizer)
 			{
 				// Output debug info for the linked executable.
 				Arguments.Add("/DEBUG");
+			}
 
+			if (LinkEnvironment.bCreateDebugInfo)
+			{
 				// Allow partial PDBs for faster linking
 				if (LinkEnvironment.bUseFastPDBLinking)
 				{
@@ -1256,6 +1276,38 @@ namespace UnrealBuildTool
 			}
 		}
 
+		public override void PrepareRuntimeDependencies(List<RuntimeDependency> RuntimeDependencies, Dictionary<FileReference, FileReference> TargetFileToSourceFile, DirectoryReference ExeDir)
+		{
+			// If ASan is enabled we need to copy the companion helper libraries from the MSVC tools bin folder to the
+			// target executable folder.
+			if (Target.WindowsPlatform.bEnableAddressSanitizer)
+			{
+				DirectoryReference ASanRuntimeDir;
+				String ASanArchSuffix;
+				if (EnvVars.Architecture == WindowsArchitecture.x64)
+				{
+					ASanRuntimeDir = DirectoryReference.Combine(EnvVars.ToolChainDir, "bin", "Hostx64", "x64");
+					ASanArchSuffix = "x86_64";
+				}
+				else
+				{
+					throw new BuildException("Unsupported build architecture for Address Sanitizer");
+				}
+
+
+				String ASanRuntimeDLL = String.Format("clang_rt.asan_dynamic-{0}.dll", ASanArchSuffix);
+				String ASanDebugRuntimeDLL = String.Format("clang_rt.asan_dbg_dynamic-{0}.dll", ASanArchSuffix);
+
+				RuntimeDependencies.Add(new RuntimeDependency(FileReference.Combine(ExeDir, ASanRuntimeDLL), StagedFileType.NonUFS));
+				TargetFileToSourceFile[FileReference.Combine(ExeDir, ASanRuntimeDLL)] = FileReference.Combine(ASanRuntimeDir, ASanRuntimeDLL);
+				if (Target.bDebugBuildsActuallyUseDebugCRT)
+				{
+					RuntimeDependencies.Add(new RuntimeDependency(FileReference.Combine(ExeDir, ASanDebugRuntimeDLL), StagedFileType.NonUFS));
+					TargetFileToSourceFile[FileReference.Combine(ExeDir, ASanDebugRuntimeDLL)] = FileReference.Combine(ASanRuntimeDir, ASanDebugRuntimeDLL);
+				}
+			}
+		}
+
 		public virtual FileReference GetApplicationIcon(FileReference ProjectFile)
 		{
 			return WindowsPlatform.GetWindowsApplicationIcon(ProjectFile);
@@ -1547,6 +1599,45 @@ namespace UnrealBuildTool
 				}
 			}
 
+			// If we're building either an executable or a DLL, make sure we link in the 
+			// correct address sanitizer helper libs.
+			if (!bBuildImportLibraryOnly && !LinkEnvironment.bIsBuildingLibrary && Target.WindowsPlatform.bEnableAddressSanitizer)
+			{
+				String ASanArchSuffix = "";
+				if (EnvVars.Architecture == WindowsArchitecture.x64)
+				{
+					ASanArchSuffix = "x86_64";
+				}
+				else
+				{
+					throw new BuildException("Unsupported build architecture for Address Sanitizer");
+				}
+
+				String ASanDebugInfix = "";
+				if (LinkEnvironment.bUseDebugCRT)
+				{
+					ASanDebugInfix = "_dbg";
+				}
+
+				if (LinkEnvironment.bUseStaticCRT)
+				{
+					if (LinkEnvironment.bIsBuildingDLL)
+					{
+						Arguments.Add(String.Format("/wholearchive:clang_rt.asan{0}_dll_thunk-{1}.lib", ASanDebugInfix, ASanArchSuffix));
+					}
+					else
+					{
+						Arguments.Add(String.Format("/wholearchive:clang_rt.asan{0}-{1}.lib", ASanDebugInfix, ASanArchSuffix));
+						Arguments.Add(String.Format("/wholearchive:clang_rt.asan_cxx{0}-{1}.lib", ASanDebugInfix, ASanArchSuffix));
+					}
+				}
+				else
+				{
+					Arguments.Add(String.Format("/wholearchive:clang_rt.asan{0}_dynamic-{1}.lib", ASanDebugInfix, ASanArchSuffix));
+					Arguments.Add(String.Format("/wholearchive:clang_rt.asan{0}_dynamic_runtime_thunk-{1}.lib", ASanDebugInfix, ASanArchSuffix));
+				}
+			}
+
 			// Enable function level hot-patching
 			if(!bBuildImportLibraryOnly && Target.WindowsPlatform.bCreateHotpatchableImage)
 			{
diff --git a/Engine/Source/Programs/UnrealBuildTool/System/LinkEnvironment.cs b/Engine/Source/Programs/UnrealBuildTool/System/LinkEnvironment.cs
--- a/Engine/Source/Programs/UnrealBuildTool/System/LinkEnvironment.cs
+++ b/Engine/Source/Programs/UnrealBuildTool/System/LinkEnvironment.cs
@@ -147,6 +147,16 @@ namespace UnrealBuildTool
 		/// </summary>
 		public bool bIsBuildingDLL = false;
 
+		/// <summary>
+		/// Whether we should compile using the statically-linked CRT. This is not widely supported for the whole engine, but is required for programs that need to run without dependencies.
+		/// </summary>
+		public bool bUseStaticCRT = false;
+
+		/// <summary>
+		/// Whether to use the debug CRT in debug configurations
+		/// </summary>
+		public bool bUseDebugCRT = false;
+
 		/// <summary>
 		/// True if this is a console application that's being build
 		/// </summary>
@@ -341,6 +351,8 @@ namespace UnrealBuildTool
 			bIsBuildingLibrary = Other.bIsBuildingLibrary;
             bDisableSymbolCache = Other.bDisableSymbolCache;
 			bIsBuildingDLL = Other.bIsBuildingDLL;
+			bUseStaticCRT = Other.bUseStaticCRT;
+			bUseDebugCRT = Other.bUseDebugCRT;
 			bIsBuildingConsoleApplication = Other.bIsBuildingConsoleApplication;
 			WindowsEntryPointOverride = Other.WindowsEntryPointOverride;
 			bIsCrossReferenced = Other.bIsCrossReferenced;
