Fix TVOS provisioning profile name
Increase log level for iPhonePackager provision output

diff --git a/Engine/Source/Programs/UnrealBuildTool/Platform/IOS/UEBuildIOS.cs b/Engine/Source/Programs/UnrealBuildTool/Platform/IOS/UEBuildIOS.cs
--- a/Engine/Source/Programs/UnrealBuildTool/Platform/IOS/UEBuildIOS.cs
+++ b/Engine/Source/Programs/UnrealBuildTool/Platform/IOS/UEBuildIOS.cs
@@ -495,7 +495,7 @@ namespace UnrealBuildTool
                 SigningCertificate = "";
                 MobileProvision = "";
 				MobileProvisionFile = null;
-                Log.TraceLog("Provision not specified or not found for " + ((ProjectFile != null) ? ProjectFile.GetFileNameWithoutAnyExtensions() : "UnrealGame") + ", searching for compatible match...");
+                Log.TraceInformation("Provision not specified or not found for " + ((ProjectFile != null) ? ProjectFile.GetFileNameWithoutAnyExtensions() : "UnrealGame") + ", searching for compatible match...");
                 Process IPPProcess = new Process();
 
 				IPPProcess.OutputDataReceived += new DataReceivedEventHandler(IPPDataReceivedHandler);
@@ -520,7 +520,7 @@ namespace UnrealBuildTool
 				Utils.RunLocalProcess(IPPProcess);
 				if(MobileProvisionFile != null)
 				{
-					Log.TraceLog("Provision found for " + ((ProjectFile != null) ? ProjectFile.GetFileNameWithoutAnyExtensions() : "UnrealGame") + ", Provision: " + MobileProvisionFile + " Certificate: " + SigningCertificate);
+					Log.TraceInformation("Provision found for " + ((ProjectFile != null) ? ProjectFile.GetFileNameWithoutAnyExtensions() : "UnrealGame") + ", Provision: " + MobileProvisionFile + " Certificate: " + SigningCertificate);
 				}
             }
 
@@ -530,11 +530,11 @@ namespace UnrealBuildTool
             // read the provision to get the UUID
 			if(MobileProvisionFile == null)
 			{
-				Log.TraceLog("No matching provision file was discovered for {0}. Please ensure you have a compatible provision installed.", ProjectFile);
+				Log.TraceInformation("No matching provision file was discovered for {0}. Please ensure you have a compatible provision installed.", ProjectFile);
 			}
 			else if(!FileReference.Exists(MobileProvisionFile))
 			{
-				Log.TraceLog("Selected mobile provision for {0} ({1}) was not found. Please ensure you have a compatible provision installed.", ProjectFile, MobileProvisionFile);
+				Log.TraceInformation("Selected mobile provision for {0} ({1}) was not found. Please ensure you have a compatible provision installed.", ProjectFile, MobileProvisionFile);
 			}
 			else
             {
@@ -617,7 +617,7 @@ namespace UnrealBuildTool
 				{
 					MobileProvision = null;
 					SigningCertificate = null;
-					Log.TraceLog("Failed to parse the mobile provisioning profile.");
+					Log.TraceInformation("Failed to parse the mobile provisioning profile.");
 				}
             }
 		}
@@ -626,8 +626,17 @@ namespace UnrealBuildTool
         {
             if ((Line != null) && (Line.Data != null))
             {
-				Log.TraceLog("{0}", Line.Data);
-                if (!string.IsNullOrEmpty(SigningCertificate))
+				if (Line.Data.StartsWith("IPP WARNING:"))
+				{
+					// Don't output IPP warnings to the console as they may not be warnings relevant to the build and could cause build failures.
+					Log.TraceLog("{0}", Line.Data);
+				}
+				else
+				{
+					Log.TraceInformation("{0}", Line.Data);
+				}
+
+				if (!string.IsNullOrEmpty(SigningCertificate))
                 {
                     if (Line.Data.Contains("CERTIFICATE-") && Line.Data.Contains(SigningCertificate))
                     {
