Fix invalid stdin handle causing mabu to fail in the Lumin toolchain

diff --git a/Engine/Source/Programs/UnrealBuildTool/System/Utils.cs b/Engine/Source/Programs/UnrealBuildTool/System/Utils.cs
--- a/Engine/Source/Programs/UnrealBuildTool/System/Utils.cs
+++ b/Engine/Source/Programs/UnrealBuildTool/System/Utils.cs
@@ -442,6 +442,7 @@ namespace UnrealBuildTool
 
 			ProcessStartInfo StartInfo = new ProcessStartInfo(Command, Args);
 			StartInfo.UseShellExecute = false;
+			StartInfo.RedirectStandardInput = true;
 			StartInfo.RedirectStandardOutput = true;
 			StartInfo.RedirectStandardError = true;
 			StartInfo.CreateNoWindow = true;
