UBT: Changes to adaptive unity to expose more non-unity errors

- Exclude source files from unity if their first include is in the working set.
- Compile header files in the working set in their own dedicated source file if there is not an existing source file that includes it first.

diff --git a/Engine/Source/Programs/UnrealBuildTool/Configuration/TargetRules.cs b/Engine/Source/Programs/UnrealBuildTool/Configuration/TargetRules.cs
--- a/Engine/Source/Programs/UnrealBuildTool/Configuration/TargetRules.cs
+++ b/Engine/Source/Programs/UnrealBuildTool/Configuration/TargetRules.cs
@@ -920,6 +920,12 @@ namespace UnrealBuildTool
 		[XmlConfigFile(Category = "BuildConfiguration")]
 		public bool bAdaptiveUnityEnablesEditAndContinue = false;
 
+		/// <summary>
+		/// Creates a dedicated source file for each header file in the working set to detect missing includes in headers.
+		/// </summary>
+		[XmlConfigFile(Category = "BuildConfiguration")]
+		public bool bAdaptiveUnityCompilesHeaderFiles = false;
+
 		/// <summary>
 		/// The number of source files in a game module before unity build will be activated for that module.  This
 		/// allows small game modules to have faster iterative compile times for single files, at the expense of slower full
@@ -2341,6 +2347,11 @@ namespace UnrealBuildTool
 			get { return Inner.bAdaptiveUnityEnablesEditAndContinue; }
 		}
 
+		public bool bAdaptiveUnityCompilesHeaderFiles
+		{
+			get { return Inner.bAdaptiveUnityCompilesHeaderFiles; }
+		}
+
 		public int MinGameModuleSourceFilesForUnityBuild
 		{
 			get { return Inner.MinGameModuleSourceFilesForUnityBuild; }
diff --git a/Engine/Source/Programs/UnrealBuildTool/Configuration/UEBuildModuleCPP.cs b/Engine/Source/Programs/UnrealBuildTool/Configuration/UEBuildModuleCPP.cs
--- a/Engine/Source/Programs/UnrealBuildTool/Configuration/UEBuildModuleCPP.cs
+++ b/Engine/Source/Programs/UnrealBuildTool/Configuration/UEBuildModuleCPP.cs
@@ -473,7 +473,7 @@ namespace UnrealBuildTool
 			List<FileItem> CPPFilesToCompile = InputFiles.CPPFiles;
 			if (bModuleUsesUnityBuild)
 			{
-				CPPFilesToCompile = Unity.GenerateUnityCPPs(Target, CPPFilesToCompile, CompileEnvironment, WorkingSet, Rules.ShortName ?? Name, IntermediateDirectory, Graph, SourceFileToUnityFile);
+				CPPFilesToCompile = Unity.GenerateUnityCPPs(Target, CPPFilesToCompile, InputFiles.HeaderFiles, CompileEnvironment, WorkingSet, Rules.ShortName ?? Name, IntermediateDirectory, Graph, SourceFileToUnityFile);
 				LinkInputFiles.AddRange(CompileUnityFilesWithToolChain(Target, ToolChain, CompileEnvironment, ModuleCompileEnvironment, CPPFilesToCompile, Graph).ObjectFiles);
 			}
 			else
@@ -526,7 +526,7 @@ namespace UnrealBuildTool
 
 					if (bModuleUsesUnityBuild)
 					{
-						GeneratedFileItems = Unity.GenerateUnityCPPs(Target, GeneratedFileItems, GeneratedCPPCompileEnvironment, WorkingSet, (Rules.ShortName ?? Name) + ".gen", IntermediateDirectory, Graph, SourceFileToUnityFile);
+						GeneratedFileItems = Unity.GenerateUnityCPPs(Target, GeneratedFileItems, new List<FileItem>(), GeneratedCPPCompileEnvironment, WorkingSet, (Rules.ShortName ?? Name) + ".gen", IntermediateDirectory, Graph, SourceFileToUnityFile);
 						LinkInputFiles.AddRange(CompileUnityFilesWithToolChain(Target, ToolChain, GeneratedCPPCompileEnvironment, ModuleCompileEnvironment, GeneratedFileItems, Graph).ObjectFiles);
 					}
 					else
diff --git a/Engine/Source/Programs/UnrealBuildTool/System/Unity.cs b/Engine/Source/Programs/UnrealBuildTool/System/Unity.cs
--- a/Engine/Source/Programs/UnrealBuildTool/System/Unity.cs
+++ b/Engine/Source/Programs/UnrealBuildTool/System/Unity.cs
@@ -5,7 +5,6 @@ using System.Collections.Generic;
 using System.Text;
 using System.IO;
 using System.Linq;
-using System.Diagnostics;
 using EpicGames.Core;
 
 namespace UnrealBuildTool
@@ -148,6 +147,7 @@ namespace UnrealBuildTool
 		/// </summary>
 		/// <param name="Target">The target we're building</param>
 		/// <param name="CPPFiles">The C++ files to #include.</param>
+		/// <param name="HeaderFiles">The header files that might correspond to the C++ files.</param>
 		/// <param name="CompileEnvironment">The environment that is used to compile the C++ files.</param>
 		/// <param name="WorkingSet">Interface to query files which belong to the working set</param>
 		/// <param name="BaseName">Base name to use for the Unity files</param>
@@ -158,6 +158,7 @@ namespace UnrealBuildTool
 		public static List<FileItem> GenerateUnityCPPs(
 			ReadOnlyTargetRules Target,
 			List<FileItem> CPPFiles,
+			List<FileItem> HeaderFiles,
 			CppCompileEnvironment CompileEnvironment,
 			ISourceFileWorkingSet WorkingSet,
 			string BaseName,
@@ -188,38 +189,68 @@ namespace UnrealBuildTool
 
 			// Build the list of unity files.
 			List<FileCollection> AllUnityFiles;
+			List<FileItem> AllWorkingSetHeaderFiles;
 			{
-				// Sort the incoming file paths alphabetically, so there will be consistency in unity blobs across multiple machines.
+				// Sort the incoming file paths lexicographically, so there will be consistency in unity blobs across multiple machines.
 				// Note that we're relying on this not only sorting files within each directory, but also the directories
 				// themselves, so the whole list of file paths is the same across computers.
-				List<FileItem> SortedCPPFiles = CPPFiles.GetRange(0, CPPFiles.Count);
-				{
-					// Case-insensitive file path compare, because you never know what is going on with local file systems
-					Comparison<FileItem> FileItemComparer = (FileA, FileB) => { return FileA.AbsolutePath.ToLowerInvariant().CompareTo(FileB.AbsolutePath.ToLowerInvariant()); };
-					SortedCPPFiles.Sort(FileItemComparer);
-				}
-
+				// Case-insensitive file path compare, because you never know what is going on with local file systems.
+				List<FileItem> SortedCPPFiles = CPPFiles.OrderBy(File => File.AbsolutePath, StringComparer.OrdinalIgnoreCase).ToList();
 
 				// Figure out whether we REALLY want to use adaptive unity for this module.  If nearly every file in the module appears in the working
 				// set, we'll just go ahead and let unity build do its thing.
-				HashSet<FileItem> FilesInWorkingSet = new HashSet<FileItem>();
+				HashSet<FileItem> CPPFilesInWorkingSet = new HashSet<FileItem>();
+				HashSet<FileItem> HeaderFilesInWorkingSet = new HashSet<FileItem>();
 				if (bUseAdaptiveUnityBuild)
 				{
+					// Figure out which uniquely-named header files are in the working set. The unique name
+					// is important to avoid confusion about which header a source file is including.
+					Dictionary<string, FileItem> NameToHeaderFileInWorkingSet = new Dictionary<string, FileItem>();
+					List<string> DuplicateHeaderNames = new List<string>();
+					HashSet<string> HeaderNames = new HashSet<string>();
+					foreach (FileItem HeaderFile in HeaderFiles)
+					{
+						string HeaderFileName = HeaderFile.Location.GetFileName();
+						if (!HeaderNames.Add(HeaderFileName))
+						{
+							DuplicateHeaderNames.Add(HeaderFileName);
+						}
+						else if (WorkingSet.Contains(HeaderFile))
+						{
+							NameToHeaderFileInWorkingSet[HeaderFileName] = HeaderFile;
+						}
+					}
+					foreach (string Name in DuplicateHeaderNames)
+					{
+						NameToHeaderFileInWorkingSet.Remove(Name);
+					}
+					HeaderFilesInWorkingSet.UnionWith(NameToHeaderFileInWorkingSet.Values);
+
+					// Remove source files from unity files based on the working set.
 					int CandidateWorkingSetSourceFileCount = 0;
 					int WorkingSetSourceFileCount = 0;
 					foreach (FileItem CPPFile in SortedCPPFiles)
 					{
 						++CandidateWorkingSetSourceFileCount;
 
-						// Don't include writable source files into unity blobs
-						if (WorkingSet.Contains(CPPFile))
+						// Don't create a separate source file for headers in the working set if they are the first include of a source file.
+						bool bHeaderInWorkingSet = false;
+						if (CompileEnvironment.MetadataCache.GetFirstInclude(CPPFile) is string FirstInclude &&
+							NameToHeaderFileInWorkingSet.TryGetValue(FirstInclude, out FileItem? HeaderFile))
+						{
+							bHeaderInWorkingSet = true;
+							HeaderFilesInWorkingSet.Remove(HeaderFile);
+						}
+
+						// Don't include source files into unity blobs if they or their first include is in the working set.
+						if (bHeaderInWorkingSet || WorkingSet.Contains(CPPFile))
 						{
 							++WorkingSetSourceFileCount;
 
 							// Mark this file as part of the working set.  This will be saved into the UBT Makefile so that
 							// the assembler can automatically invalidate the Makefile when the working set changes (allowing this
 							// code to run again, to build up new unity blobs.)
-							FilesInWorkingSet.Add(CPPFile);
+							CPPFilesInWorkingSet.Add(CPPFile);
 							Graph.AddFileToWorkingSet(CPPFile);
 						}
 					}
@@ -232,6 +263,9 @@ namespace UnrealBuildTool
 					}
 				}
 
+				// Header files in the working set that are not the first include of a source file.
+				AllWorkingSetHeaderFiles = HeaderFilesInWorkingSet.OrderBy(File => File.AbsolutePath, StringComparer.OrdinalIgnoreCase).ToList();
+
 				UnityFileBuilder CPPUnityFileBuilder = new UnityFileBuilder(bForceIntoSingleUnityFile ? -1 : Target.NumIncludedBytesPerUnityCPP);
 				StringBuilder AdaptiveUnityBuildInfoString = new StringBuilder();
 				foreach (FileItem CPPFile in SortedCPPFiles)
@@ -242,7 +276,7 @@ namespace UnrealBuildTool
 					}
 
 					// When adaptive unity is enabled, go ahead and exclude any source files that we're actively working with
-					if (bUseAdaptiveUnityBuild && FilesInWorkingSet.Contains(CPPFile))
+					if (bUseAdaptiveUnityBuild && CPPFilesInWorkingSet.Contains(CPPFile))
 					{
 						// Just compile this file normally, not as part of the unity blob
 						NewCPPFiles.Add(CPPFile);
@@ -346,6 +380,28 @@ namespace UnrealBuildTool
 				}
 			}
 
+			// Create a set of CPP files that independently compile headers in the working set to detect issues affecting non-unity builds.
+			if (Target.bAdaptiveUnityCompilesHeaderFiles && AllWorkingSetHeaderFiles.Any())
+			{
+				Graph.AddDiagnostic($"[Adaptive unity build] Headers from {BaseName} compiling individually: " +
+					String.Join(", ", AllWorkingSetHeaderFiles.Select(File => Path.GetFileName(File.AbsolutePath))));
+				foreach (FileItem HeaderFile in AllWorkingSetHeaderFiles)
+				{
+					// Add this header file to the set that will invalidate the makefile
+					Graph.AddCandidateForWorkingSet(HeaderFile);
+
+					StringWriter OutputHeaderCPPWriter = new StringWriter();
+					OutputHeaderCPPWriter.WriteLine("// This file is automatically generated at compile-time to include a modified header file.");
+					OutputHeaderCPPWriter.WriteLine($"#include \"{HeaderFile.AbsolutePath.Replace('\\', '/')}\"");
+
+					string HeaderCPPFileName = string.Format($"{ModulePrefix}{BaseName}.Header.{Path.GetFileNameWithoutExtension(HeaderFile.AbsolutePath)}.cpp");
+					FileReference HeaderCPPFilePath = FileReference.Combine(IntermediateDirectory, HeaderCPPFileName);
+
+					FileItem HeaderCPPFile = Graph.CreateIntermediateTextFile(HeaderCPPFilePath, OutputHeaderCPPWriter.ToString());
+					NewCPPFiles.Add(HeaderCPPFile);
+				}
+			}
+
 			return NewCPPFiles;
 		}
 
