UBT: Sort includes in unity files by InvariantCultureIgnoreCase for consistency with the previous implementation

#rnx

diff --git a/Engine/Source/Programs/UnrealBuildTool/System/Unity.cs b/Engine/Source/Programs/UnrealBuildTool/System/Unity.cs
--- a/Engine/Source/Programs/UnrealBuildTool/System/Unity.cs
+++ b/Engine/Source/Programs/UnrealBuildTool/System/Unity.cs
@@ -195,7 +195,7 @@ namespace UnrealBuildTool
 				// Note that we're relying on this not only sorting files within each directory, but also the directories
 				// themselves, so the whole list of file paths is the same across computers.
 				// Case-insensitive file path compare, because you never know what is going on with local file systems.
-				List<FileItem> SortedCPPFiles = CPPFiles.OrderBy(File => File.AbsolutePath, StringComparer.OrdinalIgnoreCase).ToList();
+				List<FileItem> SortedCPPFiles = CPPFiles.OrderBy(File => File.AbsolutePath, StringComparer.InvariantCultureIgnoreCase).ToList();
 
 				// Figure out whether we REALLY want to use adaptive unity for this module.  If nearly every file in the module appears in the working
 				// set, we'll just go ahead and let unity build do its thing.
