UBT: Apply optimization, precompiled header, and header build behaviors to the working set even when unity is disabled on the module

#rn Ben.Marsh

diff --git a/Engine/Source/Programs/UnrealBuildTool/Configuration/UEBuildModuleCPP.cs b/Engine/Source/Programs/UnrealBuildTool/Configuration/UEBuildModuleCPP.cs
--- a/Engine/Source/Programs/UnrealBuildTool/Configuration/UEBuildModuleCPP.cs
+++ b/Engine/Source/Programs/UnrealBuildTool/Configuration/UEBuildModuleCPP.cs
@@ -470,15 +470,17 @@ namespace UnrealBuildTool
 			Dictionary<FileItem, FileItem> SourceFileToUnityFile = new Dictionary<FileItem, FileItem>();
 
 			// Compile CPP files
-			List<FileItem> CPPFilesToCompile = InputFiles.CPPFiles;
 			if (bModuleUsesUnityBuild)
 			{
-				CPPFilesToCompile = Unity.GenerateUnityCPPs(Target, CPPFilesToCompile, InputFiles.HeaderFiles, CompileEnvironment, WorkingSet, Rules.ShortName ?? Name, IntermediateDirectory, Graph, SourceFileToUnityFile);
-				LinkInputFiles.AddRange(CompileUnityFilesWithToolChain(Target, ToolChain, CompileEnvironment, ModuleCompileEnvironment, CPPFilesToCompile, Graph).ObjectFiles);
+				Unity.GenerateUnityCPPs(Target, InputFiles.CPPFiles, InputFiles.HeaderFiles, CompileEnvironment, WorkingSet, Rules.ShortName ?? Name, IntermediateDirectory, Graph, SourceFileToUnityFile,
+					out List<FileItem> NormalFiles, out List<FileItem> AdaptiveFiles);
+				LinkInputFiles.AddRange(CompileFilesWithToolChain(Target, ToolChain, CompileEnvironment, ModuleCompileEnvironment, NormalFiles, AdaptiveFiles, Graph).ObjectFiles);
 			}
 			else
 			{
-				LinkInputFiles.AddRange(ToolChain.CompileCPPFiles(CompileEnvironment, CPPFilesToCompile, IntermediateDirectory, Name, Graph).ObjectFiles);
+				Unity.GetAdaptiveFiles(Target, InputFiles.CPPFiles, InputFiles.HeaderFiles, CompileEnvironment, WorkingSet, Graph,
+					out List<FileItem> NormalFiles, out List<FileItem> AdaptiveFiles);
+				LinkInputFiles.AddRange(CompileFilesWithToolChain(Target, ToolChain, CompileEnvironment, ModuleCompileEnvironment, NormalFiles, AdaptiveFiles, Graph).ObjectFiles);
 			}
 
 			// Compile all the generated CPP files
@@ -526,8 +528,9 @@ namespace UnrealBuildTool
 
 					if (bModuleUsesUnityBuild)
 					{
-						GeneratedFileItems = Unity.GenerateUnityCPPs(Target, GeneratedFileItems, new List<FileItem>(), GeneratedCPPCompileEnvironment, WorkingSet, (Rules.ShortName ?? Name) + ".gen", IntermediateDirectory, Graph, SourceFileToUnityFile);
-						LinkInputFiles.AddRange(CompileUnityFilesWithToolChain(Target, ToolChain, GeneratedCPPCompileEnvironment, ModuleCompileEnvironment, GeneratedFileItems, Graph).ObjectFiles);
+						Unity.GenerateUnityCPPs(Target, GeneratedFileItems, new List<FileItem>(), CompileEnvironment, WorkingSet, (Rules.ShortName ?? Name) + ".gen", IntermediateDirectory, Graph, SourceFileToUnityFile,
+							out List<FileItem> NormalGeneratedFiles, out List<FileItem> AdaptiveGeneratedFiles);
+						LinkInputFiles.AddRange(CompileFilesWithToolChain(Target, ToolChain, GeneratedCPPCompileEnvironment, ModuleCompileEnvironment, NormalGeneratedFiles, AdaptiveGeneratedFiles, Graph).ObjectFiles);
 					}
 					else
 					{
@@ -862,11 +865,15 @@ namespace UnrealBuildTool
 		/// <summary>
 		/// Compiles the provided CPP unity files. Will
 		/// </summary>
-		private CPPOutput CompileUnityFilesWithToolChain(ReadOnlyTargetRules Target, UEToolChain ToolChain, CppCompileEnvironment CompileEnvironment, CppCompileEnvironment ModuleCompileEnvironment, List<FileItem> SourceFiles, IActionGraphBuilder Graph)
+		private CPPOutput CompileFilesWithToolChain(
+			ReadOnlyTargetRules Target,
+			UEToolChain ToolChain,
+			CppCompileEnvironment CompileEnvironment,
+			CppCompileEnvironment ModuleCompileEnvironment,
+			List<FileItem> NormalFiles,
+			List<FileItem> AdaptiveFiles,
+			IActionGraphBuilder Graph)
 		{
-			List<FileItem> NormalFiles = new List<FileItem>();
-			List<FileItem> AdaptiveFiles = new List<FileItem>();
-
 			bool bAdaptiveUnityDisablesPCH = false;
 			if(Rules.PCHUsage == ModuleRules.PCHUsageMode.UseExplicitOrSharedPCHs)
 			{
@@ -880,24 +887,33 @@ namespace UnrealBuildTool
 				}
 			}
 
-			if ((Target.bAdaptiveUnityDisablesOptimizations || bAdaptiveUnityDisablesPCH || Target.bAdaptiveUnityCreatesDedicatedPCH) && !Target.bStressTestUnity)
+			if (AdaptiveFiles.Count > 0)
 			{
-				foreach (FileItem File in SourceFiles)
+				if (Target.bAdaptiveUnityCreatesDedicatedPCH)
 				{
-					// Basic check as to whether something in this module is/isn't a unity file...
-					if (File.Location.GetFileName().StartsWith(Unity.ModulePrefix))
-					{
-						NormalFiles.Add(File);
-					}
-					else
-					{
-						AdaptiveFiles.Add(File);
-					}
+					Graph.AddDiagnostic("[Adaptive Build] Creating dedicated PCH for each excluded file. Set bAdaptiveUnityCreatesDedicatedPCH to false in BuildConfiguration.xml to change this behavior.");
+				}
+				else if (Target.bAdaptiveUnityDisablesPCH)
+				{
+					Graph.AddDiagnostic("[Adaptive Build] Disabling PCH for excluded files. Set bAdaptiveUnityDisablesPCH to false in BuildConfiguration.xml to change this behavior.");
 				}
+
+				if (Target.bAdaptiveUnityDisablesOptimizations)
+				{
+					Graph.AddDiagnostic("[Adaptive Build] Disabling optimizations for excluded files. Set bAdaptiveUnityDisablesOptimizations to false in BuildConfiguration.xml to change this behavior.");
+				}
+				if (Target.bAdaptiveUnityEnablesEditAndContinue)
+				{
+					Graph.AddDiagnostic("[Adaptive Build] Enabling Edit & Continue for excluded files. Set bAdaptiveUnityEnablesEditAndContinue to false in BuildConfiguration.xml to change this behavior.");
+				}
+
+				Graph.AddDiagnostic($"[Adaptive Build] Excluded from {Name} unity file: " + String.Join(", ", AdaptiveFiles.Select(File => Path.GetFileName(File.AbsolutePath))));
 			}
-			else
+
+			if ((!Target.bAdaptiveUnityDisablesOptimizations && !bAdaptiveUnityDisablesPCH && !Target.bAdaptiveUnityCreatesDedicatedPCH) || Target.bStressTestUnity)
 			{
-				NormalFiles.AddRange(SourceFiles);
+				NormalFiles = NormalFiles.Concat(AdaptiveFiles).ToList();
+				AdaptiveFiles = new List<FileItem>();
 			}
 
 			CPPOutput OutputFiles = new CPPOutput();
diff --git a/Engine/Source/Programs/UnrealBuildTool/System/Unity.cs b/Engine/Source/Programs/UnrealBuildTool/System/Unity.cs
--- a/Engine/Source/Programs/UnrealBuildTool/System/Unity.cs
+++ b/Engine/Source/Programs/UnrealBuildTool/System/Unity.cs
@@ -143,7 +143,7 @@ namespace UnrealBuildTool
 		/// <summary>
 		/// Given a set of C++ files, generates another set of C++ files that #include all the original
 		/// files, the goal being to compile the same code in fewer translation units.
-		/// The "unity" files are written to the CompileEnvironment's OutputDirectory.
+		/// The "unity" files are written to the IntermediateDirectory.
 		/// </summary>
 		/// <param name="Target">The target we're building</param>
 		/// <param name="CPPFiles">The C++ files to #include.</param>
@@ -154,8 +154,9 @@ namespace UnrealBuildTool
 		/// <param name="IntermediateDirectory">Intermediate directory for unity cpp files</param>
 		/// <param name="Graph">The makefile being built</param>
 		/// <param name="SourceFileToUnityFile">Receives a mapping of source file to unity file</param>
-		/// <returns>The "unity" C++ files.</returns>
-		public static List<FileItem> GenerateUnityCPPs(
+		/// <param name="NormalFiles">Receives the files to compile using the normal configuration.</param>
+		/// <param name="AdaptiveFiles">Receives the files to compile using the adaptive unity configuration.</param>
+		public static void GenerateUnityCPPs(
 			ReadOnlyTargetRules Target,
 			List<FileItem> CPPFiles,
 			List<FileItem> HeaderFiles,
@@ -164,32 +165,32 @@ namespace UnrealBuildTool
 			string BaseName,
 			DirectoryReference IntermediateDirectory,
 			IActionGraphBuilder Graph,
-			Dictionary<FileItem, FileItem> SourceFileToUnityFile
-			)
+			Dictionary<FileItem, FileItem> SourceFileToUnityFile,
+			out List<FileItem> NormalFiles,
+			out List<FileItem> AdaptiveFiles)
 		{
 			List<FileItem> NewCPPFiles = new List<FileItem>();
 
-			//UEBuildPlatform BuildPlatform = UEBuildPlatform.GetBuildPlatform(CompileEnvironment.Platform);
-
 			// Figure out size of all input files combined. We use this to determine whether to use larger unity threshold or not.
 			long TotalBytesInCPPFiles = CPPFiles.Sum(F => F.Length);
 
 			// We have an increased threshold for unity file size if, and only if, all files fit into the same unity file. This
 			// is beneficial when dealing with PCH files. The default PCH creation limit is X unity files so if we generate < X 
 			// this could be fairly slow and we'd rather bump the limit a bit to group them all into the same unity file.
-
-
-			// When enabled, UnrealBuildTool will try to determine source files that you are actively iteratively changing, and break those files
-			// out of their unity blobs so that you can compile them as individual translation units, much faster than recompiling the entire
-			// unity blob each time.
-			bool bUseAdaptiveUnityBuild = Target.bUseAdaptiveUnityBuild && !Target.bStressTestUnity;
-
 			// Optimization only makes sense if PCH files are enabled.
 			bool bForceIntoSingleUnityFile = Target.bStressTestUnity || (TotalBytesInCPPFiles < Target.NumIncludedBytesPerUnityCPP * 2 && Target.bUsePCHFiles);
 
+			// Every single file in the module appears in the working set. Don't bother using adaptive unity for this module.
+			// Otherwise it would make full builds really slow.
+			GetAdaptiveFiles(Target, CPPFiles, HeaderFiles, CompileEnvironment, WorkingSet, Graph, out NormalFiles, out AdaptiveFiles);
+			if (NormalFiles.Count == 0)
+			{
+				NormalFiles = CPPFiles;
+				AdaptiveFiles.RemoveAll(new HashSet<FileItem>(NormalFiles).Contains);
+			}
+
 			// Build the list of unity files.
 			List<FileCollection> AllUnityFiles;
-			List<FileItem> AllWorkingSetHeaderFiles;
 			{
 				// Sort the incoming file paths lexicographically, so there will be consistency in unity blobs across multiple machines.
 				// Note that we're relying on this not only sorting files within each directory, but also the directories
@@ -197,77 +198,8 @@ namespace UnrealBuildTool
 				// Case-insensitive file path compare, because you never know what is going on with local file systems.
 				List<FileItem> SortedCPPFiles = CPPFiles.OrderBy(File => File.AbsolutePath, StringComparer.OrdinalIgnoreCase).ToList();
 
-				// Figure out whether we REALLY want to use adaptive unity for this module.  If nearly every file in the module appears in the working
-				// set, we'll just go ahead and let unity build do its thing.
-				HashSet<FileItem> CPPFilesInWorkingSet = new HashSet<FileItem>();
-				HashSet<FileItem> HeaderFilesInWorkingSet = new HashSet<FileItem>();
-				if (bUseAdaptiveUnityBuild)
-				{
-					// Figure out which uniquely-named header files are in the working set. The unique name
-					// is important to avoid confusion about which header a source file is including.
-					Dictionary<string, FileItem> NameToHeaderFileInWorkingSet = new Dictionary<string, FileItem>();
-					List<string> DuplicateHeaderNames = new List<string>();
-					HashSet<string> HeaderNames = new HashSet<string>();
-					foreach (FileItem HeaderFile in HeaderFiles)
-					{
-						string HeaderFileName = HeaderFile.Location.GetFileName();
-						if (!HeaderNames.Add(HeaderFileName))
-						{
-							DuplicateHeaderNames.Add(HeaderFileName);
-						}
-						else if (WorkingSet.Contains(HeaderFile))
-						{
-							NameToHeaderFileInWorkingSet[HeaderFileName] = HeaderFile;
-						}
-					}
-					foreach (string Name in DuplicateHeaderNames)
-					{
-						NameToHeaderFileInWorkingSet.Remove(Name);
-					}
-					HeaderFilesInWorkingSet.UnionWith(NameToHeaderFileInWorkingSet.Values);
-
-					// Remove source files from unity files based on the working set.
-					int CandidateWorkingSetSourceFileCount = 0;
-					int WorkingSetSourceFileCount = 0;
-					foreach (FileItem CPPFile in SortedCPPFiles)
-					{
-						++CandidateWorkingSetSourceFileCount;
-
-						// Don't create a separate source file for headers in the working set if they are the first include of a source file.
-						bool bHeaderInWorkingSet = false;
-						if (CompileEnvironment.MetadataCache.GetFirstInclude(CPPFile) is string FirstInclude &&
-							NameToHeaderFileInWorkingSet.TryGetValue(FirstInclude, out FileItem? HeaderFile))
-						{
-							bHeaderInWorkingSet = true;
-							HeaderFilesInWorkingSet.Remove(HeaderFile);
-						}
-
-						// Don't include source files into unity blobs if they or their first include is in the working set.
-						if (bHeaderInWorkingSet || WorkingSet.Contains(CPPFile))
-						{
-							++WorkingSetSourceFileCount;
-
-							// Mark this file as part of the working set.  This will be saved into the UBT Makefile so that
-							// the assembler can automatically invalidate the Makefile when the working set changes (allowing this
-							// code to run again, to build up new unity blobs.)
-							CPPFilesInWorkingSet.Add(CPPFile);
-							Graph.AddFileToWorkingSet(CPPFile);
-						}
-					}
-
-					if (WorkingSetSourceFileCount >= CandidateWorkingSetSourceFileCount)
-					{
-						// Every single file in the module appears in the working set, so don't bother using adaptive unity for this
-						// module.  Otherwise it would make full builds really slow.
-						bUseAdaptiveUnityBuild = false;
-					}
-				}
-
-				// Header files in the working set that are not the first include of a source file.
-				AllWorkingSetHeaderFiles = HeaderFilesInWorkingSet.OrderBy(File => File.AbsolutePath, StringComparer.OrdinalIgnoreCase).ToList();
-
+				HashSet<FileItem> AdaptiveFileSet = new HashSet<FileItem>(AdaptiveFiles);
 				UnityFileBuilder CPPUnityFileBuilder = new UnityFileBuilder(bForceIntoSingleUnityFile ? -1 : Target.NumIncludedBytesPerUnityCPP);
-				StringBuilder AdaptiveUnityBuildInfoString = new StringBuilder();
 				foreach (FileItem CPPFile in SortedCPPFiles)
 				{
 					if (!bForceIntoSingleUnityFile && CPPFile.AbsolutePath.IndexOf(".GeneratedWrapper.", StringComparison.InvariantCultureIgnoreCase) != -1)
@@ -276,63 +208,21 @@ namespace UnrealBuildTool
 					}
 
 					// When adaptive unity is enabled, go ahead and exclude any source files that we're actively working with
-					if (bUseAdaptiveUnityBuild && CPPFilesInWorkingSet.Contains(CPPFile))
+					if (AdaptiveFileSet.Contains(CPPFile))
 					{
-						// Just compile this file normally, not as part of the unity blob
-						NewCPPFiles.Add(CPPFile);
-
 						// Let the unity file builder know about the file, so that we can retain the existing size of the unity blobs.
 						// This won't actually make the source file part of the unity blob, but it will keep track of how big the
-						// file is so that other existing unity blobs from the same module won't be invalidated.  This prevents much
+						// file is so that other existing unity blobs from the same module won't be invalidated. This prevents much
 						// longer compile times the first time you build after your working file set changes.
 						CPPUnityFileBuilder.AddVirtualFile(CPPFile);
-
-						string CPPFileName = Path.GetFileName(CPPFile.AbsolutePath);
-						if (AdaptiveUnityBuildInfoString.Length == 0)
-						{
-							AdaptiveUnityBuildInfoString.Append(String.Format("[Adaptive unity build] Excluded from {0} unity file: {1}", BaseName, CPPFileName));
-						}
-						else
-						{
-							AdaptiveUnityBuildInfoString.Append(", " + CPPFileName);
-						}
 					}
 					else
 					{
-						// If adaptive unity build is enabled for this module, add this source file to the set that will invalidate the makefile
-						if(bUseAdaptiveUnityBuild)
-						{
-							Graph.AddCandidateForWorkingSet(CPPFile);
-						}
-
 						// Compile this file as part of the unity blob
 						CPPUnityFileBuilder.AddFile(CPPFile);
 					}
 				}
 
-				if (AdaptiveUnityBuildInfoString.Length > 0)
-				{
-					if (Target.bAdaptiveUnityCreatesDedicatedPCH)
-					{
-						Graph.AddDiagnostic("[Adaptive unity build] Creating dedicated PCH for each excluded file. Set bAdaptiveUnityCreatesDedicatedPCH to false in BuildConfiguration.xml to change this behavior.");
-					}
-					else if (Target.bAdaptiveUnityDisablesPCH)
-					{
-						Graph.AddDiagnostic("[Adaptive unity build] Disabling PCH for excluded files. Set bAdaptiveUnityDisablesPCH to false in BuildConfiguration.xml to change this behavior.");
-					}
-
-					if (Target.bAdaptiveUnityDisablesOptimizations)
-					{
-						Graph.AddDiagnostic("[Adaptive unity build] Disabling optimizations for excluded files. Set bAdaptiveUnityDisablesOptimizations to false in BuildConfiguration.xml to change this behavior.");
-					}
-					if (Target.bAdaptiveUnityEnablesEditAndContinue)
-					{
-						Graph.AddDiagnostic("[Adaptive unity build] Enabling Edit & Continue for excluded files. Set bAdaptiveUnityEnablesEditAndContinue to false in BuildConfiguration.xml to change this behavior.");
-					}
-
-					Graph.AddDiagnostic(AdaptiveUnityBuildInfoString.ToString());
-				}
-
 				AllUnityFiles = CPPUnityFileBuilder.GetUnityFiles();
 			}
 
@@ -380,29 +270,79 @@ namespace UnrealBuildTool
 				}
 			}
 
-			// Create a set of CPP files that independently compile headers in the working set to detect issues affecting non-unity builds.
-			if (Target.bAdaptiveUnityCompilesHeaderFiles && AllWorkingSetHeaderFiles.Any())
+			NormalFiles = NewCPPFiles;
+		}
+
+		public static void GetAdaptiveFiles(
+			ReadOnlyTargetRules Target,
+			List<FileItem> CPPFiles,
+			List<FileItem> HeaderFiles,
+			CppCompileEnvironment CompileEnvironment,
+			ISourceFileWorkingSet WorkingSet,
+			IActionGraphBuilder Graph,
+			out List<FileItem> NormalFiles,
+			out List<FileItem> AdaptiveFiles)
+		{
+			NormalFiles = new List<FileItem>();
+			AdaptiveFiles = new List<FileItem>();
+
+			if (!Target.bUseAdaptiveUnityBuild)
 			{
-				Graph.AddDiagnostic($"[Adaptive unity build] Headers from {BaseName} compiling individually: " +
-					String.Join(", ", AllWorkingSetHeaderFiles.Select(File => Path.GetFileName(File.AbsolutePath))));
-				foreach (FileItem HeaderFile in AllWorkingSetHeaderFiles)
-				{
-					// Add this header file to the set that will invalidate the makefile
-					Graph.AddCandidateForWorkingSet(HeaderFile);
+				NormalFiles = CPPFiles;
+				return;
+			}
 
-					StringWriter OutputHeaderCPPWriter = new StringWriter();
-					OutputHeaderCPPWriter.WriteLine("// This file is automatically generated at compile-time to include a modified header file.");
-					OutputHeaderCPPWriter.WriteLine($"#include \"{HeaderFile.AbsolutePath.Replace('\\', '/')}\"");
+			HashSet<FileItem> HeaderFilesInWorkingSet = new HashSet<FileItem>(HeaderFiles.Where(WorkingSet.Contains));
 
-					string HeaderCPPFileName = string.Format($"{ModulePrefix}{BaseName}.Header.{Path.GetFileNameWithoutExtension(HeaderFile.AbsolutePath)}.cpp");
-					FileReference HeaderCPPFilePath = FileReference.Combine(IntermediateDirectory, HeaderCPPFileName);
+			// Figure out which uniquely-named header files are in the working set.
+			// Unique names are important to avoid ambiguity about which header a source file includes.
+			Dictionary<string, FileItem> NameToHeaderFileInWorkingSet = new Dictionary<string, FileItem>();
+			List<string> DuplicateHeaderNames = new List<string>();
+			HashSet<string> HeaderNames = new HashSet<string>();
+			foreach (FileItem HeaderFile in HeaderFiles)
+			{
+				string HeaderFileName = HeaderFile.Location.GetFileName();
+				if (!HeaderNames.Add(HeaderFileName))
+				{
+					DuplicateHeaderNames.Add(HeaderFileName);
+				}
+				else if (HeaderFilesInWorkingSet.Contains(HeaderFile))
+				{
+					NameToHeaderFileInWorkingSet[HeaderFileName] = HeaderFile;
+				}
+			}
+			foreach (string Name in DuplicateHeaderNames)
+			{
+				NameToHeaderFileInWorkingSet.Remove(Name);
+			}
 
-					FileItem HeaderCPPFile = Graph.CreateIntermediateTextFile(HeaderCPPFilePath, OutputHeaderCPPWriter.ToString());
-					NewCPPFiles.Add(HeaderCPPFile);
+			// Add source files to the adaptive set if they or their first included header are in the working set.
+			foreach (FileItem CPPFile in CPPFiles)
+			{
+				bool bHeaderInWorkingSet = false;
+				if (CompileEnvironment.MetadataCache.GetFirstInclude(CPPFile) is string FirstInclude &&
+					NameToHeaderFileInWorkingSet.TryGetValue(FirstInclude, out FileItem? HeaderFile))
+				{
+					bHeaderInWorkingSet = true;
+					HeaderFilesInWorkingSet.Remove(HeaderFile);
 				}
+
+				bool bAdaptive = bHeaderInWorkingSet || WorkingSet.Contains(CPPFile);
+				List<FileItem> Files = bAdaptive ? AdaptiveFiles : NormalFiles;
+				Files.Add(CPPFile);
+			}
+
+			// Add header files in the working set to the adaptive files if they are not the first include of a source file.
+			if (Target.bAdaptiveUnityCompilesHeaderFiles)
+			{
+				AdaptiveFiles.AddRange(HeaderFilesInWorkingSet);
 			}
 
-			return NewCPPFiles;
+			// Add adaptive files to the working set that will invalidate the makefile if it changes.
+			foreach (FileItem File in AdaptiveFiles)
+			{
+				Graph.AddFileToWorkingSet(File);
+			}
 		}
 
 		static void AddUniqueDiagnostic(TargetMakefile Makefile, string Message)
