Fix another mabu redirection error in the Lumin toolchain

diff --git a/Engine/Source/Programs/UnrealBuildTool/Platform/Lumin/LuminToolChain.cs b/Engine/Source/Programs/UnrealBuildTool/Platform/Lumin/LuminToolChain.cs
--- a/Engine/Source/Programs/UnrealBuildTool/Platform/Lumin/LuminToolChain.cs
+++ b/Engine/Source/Programs/UnrealBuildTool/Platform/Lumin/LuminToolChain.cs
@@ -276,6 +276,9 @@ namespace UnrealBuildTool
 			}
 
 			ProcessStartInfo StartInfo = new ProcessStartInfo();
+			StartInfo.RedirectStandardInput = true;
+			StartInfo.RedirectStandardError = true;
+			StartInfo.RedirectStandardOutput = true;
 			StartInfo.WorkingDirectory = WorkingDirectory;
 			StartInfo.FileName = Command;
 			StartInfo.Arguments = Params;
