Make sure that the editor process is actually running for hot reload test.  Fixes situation where PIDs were still building up in Engine\Intermediate\EditorRuns, causing extrememly long build times in no build situations.  Saves about 50 seconds on my current setup.

Ben.Marsh

diff --git a/Engine/Source/Programs/UnrealBuildTool/System/HotReload.cs b/Engine/Source/Programs/UnrealBuildTool/System/HotReload.cs
--- a/Engine/Source/Programs/UnrealBuildTool/System/HotReload.cs
+++ b/Engine/Source/Programs/UnrealBuildTool/System/HotReload.cs
@@ -340,7 +340,7 @@ namespace UnrealBuildTool
 								RunningProcess = null;
 							}
 
-							if(RunningProcess == null)
+							if(RunningProcess == null || RunningProcess.HasExited)
 							{
 								FileReference.Delete(EditorInstanceFile);
 								continue;
@@ -384,7 +384,7 @@ namespace UnrealBuildTool
 								MainModuleFile = null;
 							}
 
-							if(!bIsRunning && EditorLocation == MainModuleFile)
+							if (!bIsRunning && EditorLocation == MainModuleFile)
 							{
 								bIsRunning = true;
 							}
