UBT: Fixed bAdaptiveUnityCompilesHeaderFiles to use generated files to include headers

#rnx

diff --git a/Engine/Source/Programs/UnrealBuildTool/Configuration/UEBuildModuleCPP.cs b/Engine/Source/Programs/UnrealBuildTool/Configuration/UEBuildModuleCPP.cs
--- a/Engine/Source/Programs/UnrealBuildTool/Configuration/UEBuildModuleCPP.cs
+++ b/Engine/Source/Programs/UnrealBuildTool/Configuration/UEBuildModuleCPP.cs
@@ -478,7 +478,7 @@ namespace UnrealBuildTool
 			}
 			else
 			{
-				Unity.GetAdaptiveFiles(Target, InputFiles.CPPFiles, InputFiles.HeaderFiles, CompileEnvironment, WorkingSet, Graph,
+				Unity.GetAdaptiveFiles(Target, InputFiles.CPPFiles, InputFiles.HeaderFiles, CompileEnvironment, WorkingSet, Rules.ShortName ?? Name, IntermediateDirectory, Graph,
 					out List<FileItem> NormalFiles, out List<FileItem> AdaptiveFiles);
 				LinkInputFiles.AddRange(CompileFilesWithToolChain(Target, ToolChain, CompileEnvironment, ModuleCompileEnvironment, NormalFiles, AdaptiveFiles, Graph).ObjectFiles);
 			}
diff --git a/Engine/Source/Programs/UnrealBuildTool/System/Unity.cs b/Engine/Source/Programs/UnrealBuildTool/System/Unity.cs
--- a/Engine/Source/Programs/UnrealBuildTool/System/Unity.cs
+++ b/Engine/Source/Programs/UnrealBuildTool/System/Unity.cs
@@ -182,7 +182,7 @@ namespace UnrealBuildTool
 
 			// Every single file in the module appears in the working set. Don't bother using adaptive unity for this module.
 			// Otherwise it would make full builds really slow.
-			GetAdaptiveFiles(Target, CPPFiles, HeaderFiles, CompileEnvironment, WorkingSet, Graph, out NormalFiles, out AdaptiveFiles);
+			GetAdaptiveFiles(Target, CPPFiles, HeaderFiles, CompileEnvironment, WorkingSet, BaseName, IntermediateDirectory, Graph, out NormalFiles, out AdaptiveFiles);
 			if (NormalFiles.Count == 0)
 			{
 				NormalFiles = CPPFiles;
@@ -279,6 +279,8 @@ namespace UnrealBuildTool
 			List<FileItem> HeaderFiles,
 			CppCompileEnvironment CompileEnvironment,
 			ISourceFileWorkingSet WorkingSet,
+			string BaseName,
+			DirectoryReference IntermediateDirectory,
 			IActionGraphBuilder Graph,
 			out List<FileItem> NormalFiles,
 			out List<FileItem> AdaptiveFiles)
@@ -321,7 +323,7 @@ namespace UnrealBuildTool
 			{
 				bool bHeaderInWorkingSet = false;
 				if (CompileEnvironment.MetadataCache.GetFirstInclude(CPPFile) is string FirstInclude &&
-					NameToHeaderFileInWorkingSet.TryGetValue(FirstInclude, out FileItem? HeaderFile))
+					NameToHeaderFileInWorkingSet.TryGetValue(Path.GetFileName(FirstInclude), out FileItem? HeaderFile))
 				{
 					bHeaderInWorkingSet = true;
 					HeaderFilesInWorkingSet.Remove(HeaderFile);
@@ -335,7 +337,18 @@ namespace UnrealBuildTool
 			// Add header files in the working set to the adaptive files if they are not the first include of a source file.
 			if (Target.bAdaptiveUnityCompilesHeaderFiles)
 			{
-				AdaptiveFiles.AddRange(HeaderFilesInWorkingSet);
+				foreach (FileItem HeaderFile in HeaderFilesInWorkingSet)
+				{
+					StringWriter OutputHeaderCPPWriter = new StringWriter();
+					OutputHeaderCPPWriter.WriteLine("// This file is automatically generated at compile-time to include a modified header file.");
+					OutputHeaderCPPWriter.WriteLine($"#include \"{HeaderFile.AbsolutePath.Replace('\\', '/')}\"");
+
+					string HeaderCPPFileName = string.Format($"{ModulePrefix}{BaseName}.Header.{Path.GetFileNameWithoutExtension(HeaderFile.AbsolutePath)}.cpp");
+					FileReference HeaderCPPFilePath = FileReference.Combine(IntermediateDirectory, HeaderCPPFileName);
+
+					FileItem HeaderCPPFile = Graph.CreateIntermediateTextFile(HeaderCPPFilePath, OutputHeaderCPPWriter.ToString());
+					AdaptiveFiles.Add(HeaderCPPFile);
+				}
 			}
 
 			// Add adaptive files to the working set that will invalidate the makefile if it changes.
