Use the custom linker to build large projects on Lumin

diff --git a/Engine/Extras/Android/UpdateLinkerLumin.bat b/Engine/Extras/Android/UpdateLinkerLumin.bat
new file mode 100755
--- /dev/null
+++ b/Engine/Extras/Android/UpdateLinkerLumin.bat
@@ -0,0 +1,14 @@
+@echo off
+setlocal ENABLEEXTENSIONS
+
+IF EXIST "%MLSDK%" (
+	echo Replacing ld.lld.exe at %MLSDK%\tools\toolchains\llvm-8\bin
+	copy /y ld.lld.exe "%MLSDK%"\tools\toolchains\llvm-8\bin
+) else (
+	echo Unable to locate local Magic Leap SDK location.
+	pause
+	exit /b 1
+)
+
+pause
+exit /b 0
diff --git a/Engine/Extras/Android/UpdateLinkerLumin.command b/Engine/Extras/Android/UpdateLinkerLumin.command
new file mode 100644
--- /dev/null
+++ b/Engine/Extras/Android/UpdateLinkerLumin.command
@@ -0,0 +1,12 @@
+#!/bin/sh
+
+if [ ! -d "$MLSDK" ]; then
+	echo Unable to locate Magic Leap SDK location.
+	read -rsp $'Press any key to continue...\n' -n1 key
+	exit 1
+fi
+echo Replacing ld.lld at $MLSDK/tools/toolchains/llvm-8/bin
+chmod +x ld.lld
+cp ld.lld "%MLSDK%"/tools/toolchains/llvm-8/bin
+
+exit 0
diff --git a/Engine/Source/Programs/UnrealBuildTool/Platform/Lumin/LuminToolChain.cs b/Engine/Source/Programs/UnrealBuildTool/Platform/Lumin/LuminToolChain.cs
--- a/Engine/Source/Programs/UnrealBuildTool/Platform/Lumin/LuminToolChain.cs
+++ b/Engine/Source/Programs/UnrealBuildTool/Platform/Lumin/LuminToolChain.cs
@@ -224,7 +224,7 @@ namespace UnrealBuildTool
 
 		protected override string GetLinkArguments(LinkEnvironment LinkEnvironment, string Architecture)
 		{
-			string Result = "";
+			string Result = " -fuse-ld=lld";
 
 			if (LinkEnvironment.bIsBuildingDLL)
 			{
