Update to latest Visual Studio toolchain.

#rnx

diff --git a/Engine/Source/Programs/UnrealBuildTool/Platform/Windows/UEBuildWindows.cs b/Engine/Source/Programs/UnrealBuildTool/Platform/Windows/UEBuildWindows.cs
--- a/Engine/Source/Programs/UnrealBuildTool/Platform/Windows/UEBuildWindows.cs
+++ b/Engine/Source/Programs/UnrealBuildTool/Platform/Windows/UEBuildWindows.cs
@@ -710,6 +710,7 @@ namespace UnrealBuildTool
 		/// </summary>
 		static readonly VersionNumber[] PreferredVisualStudioToolChainVersion = new VersionNumber[]
 		{
+			VersionNumber.Parse("14.28.29912"), // VS2019 v19.1 (installed to 14.28.29910 folder)
 			VersionNumber.Parse("14.24.28315"), // VS2019 v16.4.3 (installed to 14.24.28314 folder)
 //			VersionNumber.Parse("14.22.27905"), // VS2019 v16.2.3
 //			VersionNumber.Parse("14.16.27023.2"), // VS2017 v15.9.15
