Fix support for C++20 modules.

#rnx

diff --git a/Engine/Source/Programs/UnrealBuildTool/Modes/BuildMode.cs b/Engine/Source/Programs/UnrealBuildTool/Modes/BuildMode.cs
--- a/Engine/Source/Programs/UnrealBuildTool/Modes/BuildMode.cs
+++ b/Engine/Source/Programs/UnrealBuildTool/Modes/BuildMode.cs
@@ -348,12 +348,23 @@ namespace UnrealBuildTool
 					}
 				}
 
+				// Create the C++ dependency cache
+				CppDependencyCache CppDependencies = new CppDependencyCache();
+				for (int TargetIdx = 0; TargetIdx < TargetDescriptors.Count; TargetIdx++)
+				{
+					using (Timeline.ScopeEvent("Reading dependency cache"))
+					{
+						TargetDescriptor TargetDescriptor = TargetDescriptors[TargetIdx];
+						CppDependencies.Mount(TargetDescriptor.ProjectFile, TargetDescriptor.Name, TargetDescriptor.Platform, TargetDescriptor.Configuration, Makefiles[TargetIdx].TargetType, TargetDescriptor.Architecture);
+					}
+				}
+
 				// Pre-process module interfaces to generate dependency files
 				List<LinkedAction> ModuleDependencyActions = PrerequisiteActions.Where(x => x.ActionType == ActionType.GatherModuleDependencies).ToList();
 				if (ModuleDependencyActions.Count > 0)
 				{
 					Dictionary<LinkedAction, bool> PreprocessActionToOutdatedFlag = new Dictionary<LinkedAction, bool>();
-					ActionGraph.GatherAllOutdatedActions(ModuleDependencyActions, History, PreprocessActionToOutdatedFlag, null!, BuildConfiguration.bIgnoreOutdatedImportLibraries);
+					ActionGraph.GatherAllOutdatedActions(ModuleDependencyActions, History, PreprocessActionToOutdatedFlag, CppDependencies, BuildConfiguration.bIgnoreOutdatedImportLibraries);
 
 					List<LinkedAction> PreprocessActions = PreprocessActionToOutdatedFlag.Where(x => x.Value).Select(x => x.Key).ToList();
 					if (PreprocessActions.Count > 0)
@@ -374,28 +385,30 @@ namespace UnrealBuildTool
 				{
 					TargetDescriptor TargetDescriptor = TargetDescriptors[TargetIdx];
 
-					// Create the dependencies cache
-					CppDependencyCache CppDependencies;
-					using (Timeline.ScopeEvent("Reading dependency cache"))
-					{
-						CppDependencies = CppDependencyCache.CreateHierarchy(TargetDescriptor.ProjectFile, TargetDescriptor.Name, TargetDescriptor.Platform, TargetDescriptor.Configuration, Makefiles[TargetIdx].TargetType, TargetDescriptor.Architecture);
-					}
-
 					// Update the module dependencies
 					Dictionary<string, FileItem> ModuleOutputs = new Dictionary<string, FileItem>(StringComparer.Ordinal);
 					if (ModuleDependencyActions.Count > 0)
 					{
+						Dictionary<FileItem, List<string>> ModuleImports = new Dictionary<FileItem, List<string>>();
+
 						List<FileItem> CompiledModuleInterfaces = new List<FileItem>();
-						foreach(LinkedAction PrerequisiteAction in PrerequisiteActions)
+						foreach(LinkedAction ModuleDependencyAction in ModuleDependencyActions)
 						{
-							ICppCompileAction? CppModulesAction = PrerequisiteAction.Inner as ICppCompileAction;
+							ICppCompileAction? CppModulesAction = ModuleDependencyAction.Inner as ICppCompileAction;
 							if (CppModulesAction != null && CppModulesAction.CompiledModuleInterfaceFile != null)
 							{
 								string? ProducedModule;
-								if (CppDependencies.TryGetProducedModule(PrerequisiteAction.DependencyListFile!, out ProducedModule))
+								if (CppDependencies.TryGetProducedModule(ModuleDependencyAction.DependencyListFile!, out ProducedModule))
 								{
 									ModuleOutputs[ProducedModule] = CppModulesAction.CompiledModuleInterfaceFile;
 								}
+
+								List<string>? ImportedModules;
+								if (CppDependencies.TryGetImportedModules(ModuleDependencyAction.DependencyListFile!, out ImportedModules))
+								{
+									ModuleImports[CppModulesAction.CompiledModuleInterfaceFile] = ImportedModules;
+								}
+
 								CompiledModuleInterfaces.Add(CppModulesAction.CompiledModuleInterfaceFile);
 							}
 						}
@@ -405,8 +418,10 @@ namespace UnrealBuildTool
 						{
 							if (PrerequisiteAction.ActionType == ActionType.CompileModuleInterface)
 							{
+								ICppCompileAction CppModulesAction = (ICppCompileAction)PrerequisiteAction.Inner;
+
 								List<string>? ImportedModules;
-								if(CppDependencies.TryGetImportedModules(PrerequisiteAction.DependencyListFile!, out ImportedModules))
+								if(ModuleImports.TryGetValue(CppModulesAction.CompiledModuleInterfaceFile, out ImportedModules))
 								{
 									foreach (string ImportedModule in ImportedModules)
 									{
diff --git a/Engine/Source/Programs/UnrealBuildTool/Platform/Windows/VCCompileAction.cs b/Engine/Source/Programs/UnrealBuildTool/Platform/Windows/VCCompileAction.cs
--- a/Engine/Source/Programs/UnrealBuildTool/Platform/Windows/VCCompileAction.cs
+++ b/Engine/Source/Programs/UnrealBuildTool/Platform/Windows/VCCompileAction.cs
@@ -24,6 +24,11 @@ namespace UnrealBuildTool
 	/// </summary>
 	class VCCompileAction : ICppCompileAction
 	{
+		/// <summary>
+		/// The action type
+		/// </summary>
+		public ActionType ActionType { get; set; } = ActionType.Compile;
+
 		/// <summary>
 		/// Path to the compiler
 		/// </summary>
@@ -146,7 +151,6 @@ namespace UnrealBuildTool
 
 		#region Implementation of IAction
 
-		ActionType IExternalAction.ActionType => ActionType.Compile;
 		IEnumerable<FileItem> IExternalAction.DeleteItems => DeleteItems;
 		public DirectoryReference WorkingDirectory => UnrealBuildTool.EngineSourceDirectory;
 		string IExternalAction.CommandDescription => "Compile";
@@ -207,12 +211,17 @@ namespace UnrealBuildTool
 			}
 		}
 
+		/// <summary>
+		/// Whether to use cl-filter
+		/// </summary>
+		bool UsingClFilter => DependencyListFile != null && !DependencyListFile.HasExtension(".json");
+
 		/// <inheritdoc/>
 		FileReference IExternalAction.CommandPath
 		{
 			get
 			{
-				if (DependencyListFile != null)
+				if (UsingClFilter)
 				{
 					return FileReference.Combine(UnrealBuildTool.EngineDirectory, "Build", "Windows", "cl-filter", "cl-filter.exe");
 				}
@@ -228,7 +237,7 @@ namespace UnrealBuildTool
 		{
 			get
 			{
-				if (DependencyListFile != null)
+				if (UsingClFilter)
 				{
 					return GetClFilterArguments();
 				}
@@ -300,6 +309,7 @@ namespace UnrealBuildTool
 		/// <param name="Reader">Reader to serialize from</param>
 		public VCCompileAction(BinaryArchiveReader Reader)
 		{
+			ActionType = (ActionType)Reader.ReadInt();
 			CompilerExe = Reader.ReadFileItem();
 			CompilerType = (WindowsCompiler)Reader.ReadInt();
 			ToolChainVersion = Reader.ReadString();
@@ -330,6 +340,7 @@ namespace UnrealBuildTool
 		/// <inheritdoc/>
 		public void Write(BinaryArchiveWriter Writer)
 		{
+			Writer.WriteInt((int)ActionType);
 			Writer.WriteFileItem(CompilerExe);
 			Writer.WriteInt((int)CompilerType);
 			Writer.WriteString(ToolChainVersion);
diff --git a/Engine/Source/Programs/UnrealBuildTool/Platform/Windows/VCToolChain.cs b/Engine/Source/Programs/UnrealBuildTool/Platform/Windows/VCToolChain.cs
--- a/Engine/Source/Programs/UnrealBuildTool/Platform/Windows/VCToolChain.cs
+++ b/Engine/Source/Programs/UnrealBuildTool/Platform/Windows/VCToolChain.cs
@@ -1166,21 +1166,34 @@ namespace UnrealBuildTool
 
 				if (SourceFile.HasExtension(".ixx"))
 				{
+					FileItem IfcFile = FileItem.GetItemByFileReference(FileReference.Combine(GetModuleInterfaceDir(OutputDir), SourceFile.Location.ChangeExtension(".ifc").GetFileName()));
+					CompileAction.Arguments.Add("/interface");
+					CompileAction.Arguments.Add(String.Format("/ifcOutput \"{0}\"", IfcFile.Location));
+					CompileAction.CompiledModuleInterfaceFile = IfcFile;
+
 					FileItem IfcDepsFile = FileItem.GetItemByFileReference(FileReference.Combine(OutputDir, SourceFile.Location.GetFileName() + ".md.json"));
 
 					VCCompileAction CompileDepsAction = new VCCompileAction(CompileAction);
+					CompileDepsAction.ActionType = ActionType.GatherModuleDependencies;
+					CompileDepsAction.ResponseFile = FileItem.GetItemByPath(IfcDepsFile + ".response");
+					CompileDepsAction.ObjectFile = null;
+					CompileDepsAction.DependencyListFile = IfcDepsFile;
 					CompileDepsAction.Arguments.Add(String.Format("/sourceDependencies:directives \"{0}\"", IfcDepsFile.Location));
 					CompileDepsAction.AdditionalPrerequisiteItems.Add(SourceFile);
 					CompileDepsAction.AdditionalPrerequisiteItems.AddRange(CompileEnvironment.ForceIncludeFiles);
 					CompileDepsAction.AdditionalPrerequisiteItems.AddRange(CompileEnvironment.AdditionalPrerequisites);
 					CompileDepsAction.AdditionalProducedItems.Add(IfcDepsFile);
+					Graph.AddAction(CompileDepsAction);
 
-					FileItem IfcFile = FileItem.GetItemByFileReference(FileReference.Combine(GetModuleInterfaceDir(OutputDir), SourceFile.Location.ChangeExtension(".ifc").GetFileName()));
-					CompileAction.Arguments.Add("/interface");
-					CompileAction.Arguments.Add(String.Format("/ifcOutput \"{0}\"", IfcFile.Location));
+					if (!ProjectFileGenerator.bGenerateProjectFiles)
+					{
+						CompileDepsAction.WriteResponseFile(Graph);
+					}
+
+					CompileAction.ActionType = ActionType.CompileModuleInterface;
 					CompileAction.AdditionalPrerequisiteItems.Add(IfcDepsFile); // Force the dependencies file into the action graph
 					CompileAction.AdditionalProducedItems.Add(IfcFile);
-					CompileAction.DependencyListFile = IfcDepsFile;
+					CompileAction.DependencyListFile = FileItem.GetItemByFileReference(FileReference.Combine(OutputDir, String.Format("{0}.txt", SourceFile.Location.GetFileName())));
 					CompileAction.CompiledModuleInterfaceFile = IfcFile;
 				}
 				else if (CompileEnvironment.bGenerateDependenciesFile)
diff --git a/Engine/Source/Programs/UnrealBuildTool/System/CppDependencyCache.cs b/Engine/Source/Programs/UnrealBuildTool/System/CppDependencyCache.cs
--- a/Engine/Source/Programs/UnrealBuildTool/System/CppDependencyCache.cs
+++ b/Engine/Source/Programs/UnrealBuildTool/System/CppDependencyCache.cs
@@ -54,57 +54,119 @@ namespace UnrealBuildTool
 			}
 		}
 
-		/// <summary>
-		/// The current file version
-		/// </summary>
-		public const int CurrentVersion = 3;
+		class CachePartition
+		{
+			/// <summary>
+			/// The current file version
+			/// </summary>
+			public const int CurrentVersion = 3;
+
+			/// <summary>
+			/// Location of this dependency cache
+			/// </summary>
+			public FileReference Location;
+
+			/// <summary>
+			/// Directory for files to cache dependencies for.
+			/// </summary>
+			public DirectoryReference BaseDir;
+
+			/// <summary>
+			/// Map from file item to dependency info
+			/// </summary>
+			public ConcurrentDictionary<FileItem, DependencyInfo> DependencyFileToInfo = new ConcurrentDictionary<FileItem, DependencyInfo>();
+
+			/// <summary>
+			/// Whether the cache has been modified and needs to be saved
+			/// </summary>
+			public bool bModified;
+
+			/// <summary>
+			/// Constructs a dependency cache. This method is private; call CppDependencyCache.Create() to create a cache hierarchy for a given project.
+			/// </summary>
+			/// <param name="Location">File to store the cache</param>
+			/// <param name="BaseDir">Base directory for files that this cache should store data for</param>
+			public CachePartition(FileReference Location, DirectoryReference BaseDir)
+			{
+				this.Location = Location;
+				this.BaseDir = BaseDir;
 
-		/// <summary>
-		/// Location of this dependency cache
-		/// </summary>
-		FileReference Location;
+				if (FileReference.Exists(Location))
+				{
+					Read();
+				}
+			}
 
-		/// <summary>
-		/// Directory for files to cache dependencies for.
-		/// </summary>
-		DirectoryReference BaseDir;
+			/// <summary>
+			/// Reads data for this dependency cache from disk
+			/// </summary>
+			public void Read()
+			{
+				try
+				{
+					using (BinaryArchiveReader Reader = new BinaryArchiveReader(Location))
+					{
+						int Version = Reader.ReadInt();
+						if (Version != CurrentVersion)
+						{
+							Log.TraceLog("Unable to read dependency cache from {0}; version {1} vs current {2}", Location, Version, CurrentVersion);
+							return;
+						}
 
-		/// <summary>
-		/// The parent cache.
-		/// </summary>
-		CppDependencyCache? Parent;
+						int Count = Reader.ReadInt();
+						for (int Idx = 0; Idx < Count; Idx++)
+						{
+							FileItem File = Reader.ReadFileItem();
+							DependencyFileToInfo[File] = DependencyInfo.Read(Reader);
+						}
+					}
+				}
+				catch (Exception Ex)
+				{
+					Log.TraceWarning("Unable to read {0}. See log for additional information.", Location);
+					Log.TraceLog("{0}", ExceptionUtils.FormatExceptionDetails(Ex));
+				}
+			}
 
-		/// <summary>
-		/// Map from file item to dependency info
-		/// </summary>
-		ConcurrentDictionary<FileItem, DependencyInfo> DependencyFileToInfo = new ConcurrentDictionary<FileItem, DependencyInfo>();
+			/// <summary>
+			/// Writes data for this dependency cache to disk
+			/// </summary>
+			public void Write()
+			{
+				DirectoryReference.CreateDirectory(Location.Directory);
+				using (FileStream Stream = File.Open(Location.FullName, FileMode.Create, FileAccess.Write, FileShare.Read))
+				{
+					using (BinaryArchiveWriter Writer = new BinaryArchiveWriter(Stream))
+					{
+						Writer.WriteInt(CurrentVersion);
+
+						Writer.WriteInt(DependencyFileToInfo.Count);
+						foreach (KeyValuePair<FileItem, DependencyInfo> Pair in DependencyFileToInfo)
+						{
+							Writer.WriteFileItem(Pair.Key);
+							Pair.Value.Write(Writer);
+						}
+					}
+				}
+				bModified = false;
+			}
+		}
 
 		/// <summary>
-		/// Whether the cache has been modified and needs to be saved
+		/// List of partitions
 		/// </summary>
-		bool bModified;
+		List<CachePartition> Partitions = new List<CachePartition>();
 
 		/// <summary>
 		/// Static cache of all constructed dependency caches
 		/// </summary>
-		static Dictionary<FileReference, CppDependencyCache> Caches = new Dictionary<FileReference, CppDependencyCache>();
+		static Dictionary<FileReference, CachePartition> GlobalPartitions = new Dictionary<FileReference, CachePartition>();
 
 		/// <summary>
 		/// Constructs a dependency cache. This method is private; call CppDependencyCache.Create() to create a cache hierarchy for a given project.
 		/// </summary>
-		/// <param name="Location">File to store the cache</param>
-		/// <param name="BaseDir">Base directory for files that this cache should store data for</param>
-		/// <param name="Parent">The parent cache to use</param>
-		private CppDependencyCache(FileReference Location, DirectoryReference BaseDir, CppDependencyCache? Parent)
+		public CppDependencyCache()
 		{
-			this.Location = Location;
-			this.BaseDir = BaseDir;
-			this.Parent = Parent;
-
-			if (FileReference.Exists(Location))
-			{
-				Read();
-			}
 		}
 
 		/// <summary>
@@ -204,23 +266,25 @@ namespace UnrealBuildTool
 		/// <returns>True if the input file exists and the dependencies were read</returns>
 		private bool TryGetDependencyInfoInternal(FileItem InputFile, [NotNullWhen(true)] out DependencyInfo? OutInfo)
 		{
-			if (Parent != null && !InputFile.Location.IsUnderDirectory(BaseDir))
-			{
-				return Parent.TryGetDependencyInfoInternal(InputFile, out OutInfo);
-			}
-			else
+			foreach(CachePartition Partition in Partitions)
 			{
-				DependencyInfo? Info;
-				if (!DependencyFileToInfo.TryGetValue(InputFile, out Info) || InputFile.LastWriteTimeUtc.Ticks > Info.LastWriteTimeUtc)
+				if (InputFile.Location.IsUnderDirectory(Partition.BaseDir))
 				{
-					Info = ReadDependencyInfo(InputFile);
-					DependencyFileToInfo.TryAdd(InputFile, Info);
-					bModified = true;
-				}
+					DependencyInfo? Info;
+					if (!Partition.DependencyFileToInfo.TryGetValue(InputFile, out Info) || InputFile.LastWriteTimeUtc.Ticks > Info.LastWriteTimeUtc)
+					{
+						Info = ReadDependencyInfo(InputFile);
+						Partition.DependencyFileToInfo.TryAdd(InputFile, Info);
+						Partition.bModified = true;
+					}
 
-				OutInfo = Info;
-				return true;
+					OutInfo = Info;
+					return true;
+				}
 			}
+
+			OutInfo = null;
+			return false;
 		}
 
 		/// <summary>
@@ -233,10 +297,8 @@ namespace UnrealBuildTool
 		/// <param name="TargetType">The target type</param>
 		/// <param name="Architecture">The target architecture</param>
 		/// <returns>Dependency cache hierarchy for the given project</returns>
-		public static CppDependencyCache CreateHierarchy(FileReference? ProjectFile, string TargetName, UnrealTargetPlatform Platform, UnrealTargetConfiguration Configuration, TargetType TargetType, string Architecture)
+		public void Mount(FileReference? ProjectFile, string TargetName, UnrealTargetPlatform Platform, UnrealTargetConfiguration Configuration, TargetType TargetType, string Architecture)
 		{
-			CppDependencyCache Cache = null!;
-
 			if (ProjectFile == null || !UnrealBuildTool.IsEngineInstalled())
 			{
 				string AppName;
@@ -250,16 +312,14 @@ namespace UnrealBuildTool
 				}
 
 				FileReference EngineCacheLocation = FileReference.Combine(UnrealBuildTool.EngineDirectory, UEBuildTarget.GetPlatformIntermediateFolder(Platform, Architecture), AppName, Configuration.ToString(), "DependencyCache.bin");
-				Cache = FindOrAddCache(EngineCacheLocation, UnrealBuildTool.EngineDirectory, Cache);
+				FindOrAddPartition(EngineCacheLocation, UnrealBuildTool.EngineDirectory);
 			}
 
 			if (ProjectFile != null)
 			{
 				FileReference ProjectCacheLocation = FileReference.Combine(ProjectFile.Directory, UEBuildTarget.GetPlatformIntermediateFolder(Platform, Architecture), TargetName, Configuration.ToString(), "DependencyCache.bin");
-				Cache = FindOrAddCache(ProjectCacheLocation, ProjectFile.Directory, Cache);
+				FindOrAddPartition(ProjectCacheLocation, ProjectFile.Directory);
 			}
-
-			return Cache;
 		}
 
 		/// <summary>
@@ -267,24 +327,21 @@ namespace UnrealBuildTool
 		/// </summary>
 		/// <param name="Location">File to store the cache</param>
 		/// <param name="BaseDir">Base directory for files that this cache should store data for</param>
-		/// <param name="Parent">The parent cache to use</param>
 		/// <returns>Reference to a dependency cache with the given settings</returns>
-		static CppDependencyCache FindOrAddCache(FileReference Location, DirectoryReference BaseDir, CppDependencyCache? Parent)
+		void FindOrAddPartition(FileReference Location, DirectoryReference BaseDir)
 		{
-			lock (Caches)
+			lock (GlobalPartitions)
 			{
-				CppDependencyCache? Cache;
-				if (Caches.TryGetValue(Location, out Cache))
-				{
-					Debug.Assert(Cache.BaseDir == BaseDir);
-					Debug.Assert(Cache.Parent == Parent);
-				}
-				else
+				if (!Partitions.Any(x => x.Location == Location))
 				{
-					Cache = new CppDependencyCache(Location, BaseDir, Parent);
-					Caches.Add(Location, Cache);
+					CachePartition? Partition;
+					if (!GlobalPartitions.TryGetValue(Location, out Partition))
+					{
+						Partition = new CachePartition(Location, BaseDir);
+						GlobalPartitions.Add(Location, Partition);
+					}
+					Partitions.Add(Partition);
 				}
-				return Cache;
 			}
 		}
 
@@ -293,62 +350,9 @@ namespace UnrealBuildTool
 		/// </summary>
 		public static void SaveAll()
 		{
-			Parallel.ForEach(Caches.Values, Cache => { if (Cache.bModified) { Cache.Write(); } });
-		}
-
-		/// <summary>
-		/// Reads data for this dependency cache from disk
-		/// </summary>
-		private void Read()
-		{
-			try
-			{
-				using (BinaryArchiveReader Reader = new BinaryArchiveReader(Location))
-				{
-					int Version = Reader.ReadInt();
-					if (Version != CurrentVersion)
-					{
-						Log.TraceLog("Unable to read dependency cache from {0}; version {1} vs current {2}", Location, Version, CurrentVersion);
-						return;
-					}
-
-					int Count = Reader.ReadInt();
-					for (int Idx = 0; Idx < Count; Idx++)
-					{
-						FileItem File = Reader.ReadFileItem();
-						DependencyFileToInfo[File] = DependencyInfo.Read(Reader);
-					}
-				}
-			}
-			catch (Exception Ex)
-			{
-				Log.TraceWarning("Unable to read {0}. See log for additional information.", Location);
-				Log.TraceLog("{0}", ExceptionUtils.FormatExceptionDetails(Ex));
-			}
+			Parallel.ForEach(GlobalPartitions.Values, Cache => { if (Cache.bModified) { Cache.Write(); } });
 		}
 
-		/// <summary>
-		/// Writes data for this dependency cache to disk
-		/// </summary>
-		private void Write()
-		{
-			DirectoryReference.CreateDirectory(Location.Directory);
-			using (FileStream Stream = File.Open(Location.FullName, FileMode.Create, FileAccess.Write, FileShare.Read))
-			{
-				using (BinaryArchiveWriter Writer = new BinaryArchiveWriter(Stream))
-				{
-					Writer.WriteInt(CurrentVersion);
-
-					Writer.WriteInt(DependencyFileToInfo.Count);
-					foreach (KeyValuePair<FileItem, DependencyInfo> Pair in DependencyFileToInfo)
-					{
-						Writer.WriteFileItem(Pair.Key);
-						Pair.Value.Write(Writer);
-					}
-				}
-			}
-			bModified = false;
-		}
 
 		/// <summary>
 		/// Reads dependencies from the given file.
diff --git a/Engine/Source/Programs/UnrealBuildTool/System/HotReload.cs b/Engine/Source/Programs/UnrealBuildTool/System/HotReload.cs
--- a/Engine/Source/Programs/UnrealBuildTool/System/HotReload.cs
+++ b/Engine/Source/Programs/UnrealBuildTool/System/HotReload.cs
@@ -508,7 +508,8 @@ namespace UnrealBuildTool
 		public static List<LinkedAction> PatchActionsForTarget(BuildConfiguration BuildConfiguration, TargetDescriptor TargetDescriptor, TargetMakefile Makefile, List<LinkedAction> PrerequisiteActions, List<LinkedAction> TargetActionsToExecute, Dictionary<FileReference, FileReference> InitialPatchedOldLocationToNewLocation)
 		{
 			// Get the dependency history
-			CppDependencyCache CppDependencies = CppDependencyCache.CreateHierarchy(TargetDescriptor.ProjectFile, TargetDescriptor.Name, TargetDescriptor.Platform, TargetDescriptor.Configuration, Makefile.TargetType, TargetDescriptor.Architecture);
+			CppDependencyCache CppDependencies = new CppDependencyCache();
+			CppDependencies.Mount(TargetDescriptor.ProjectFile, TargetDescriptor.Name, TargetDescriptor.Platform, TargetDescriptor.Configuration, Makefile.TargetType, TargetDescriptor.Architecture);
 
 			ActionHistory History = new ActionHistory();
 			if(TargetDescriptor.ProjectFile != null)
