Plugin reference descriptors do not verify that the plugin's platforms names are known platforms. This means that the project will generate and build correctly even if a user doesn't have a particular platform extension synced.

diff --git a/Engine/Source/Programs/AutomationTool/Scripts/ResaveProjectDescriptors.cs b/Engine/Source/Programs/AutomationTool/Scripts/ResaveProjectDescriptors.cs
--- a/Engine/Source/Programs/AutomationTool/Scripts/ResaveProjectDescriptors.cs
+++ b/Engine/Source/Programs/AutomationTool/Scripts/ResaveProjectDescriptors.cs
@@ -1,4 +1,4 @@
-﻿// Copyright Epic Games, Inc. All Rights Reserved.
+// Copyright Epic Games, Inc. All Rights Reserved.
 
 using System;
 using System.Collections.Generic;
@@ -39,7 +39,7 @@ namespace AutomationTool
 			string RootDirParam = ParseParamValue("RootDir", null);
 			if(RootDirParam == null)
 			{
-				throw new AutomationException("Missing -BaseDir=... parameter");
+				throw new AutomationException("Missing -RootDir=... parameter");
 			}
 
 			string EngineAssociation = ParseParamValue("EngineAssociation", null);
@@ -91,7 +91,7 @@ namespace AutomationTool
 							PluginInfo Info;
 							if(NameToPluginInfo.TryGetValue(Reference.Name, out Info))
 							{
-								Reference.SupportedTargetPlatforms = Info.Descriptor.SupportedTargetPlatforms;
+								Reference.SupportedTargetPlatforms = Info.Descriptor.GetSupportedTargetPlatformNames();
 							}
 						}
 						else
diff --git a/Engine/Source/Programs/UnrealBuildTool/Configuration/UEBuildTarget.cs b/Engine/Source/Programs/UnrealBuildTool/Configuration/UEBuildTarget.cs
--- a/Engine/Source/Programs/UnrealBuildTool/Configuration/UEBuildTarget.cs
+++ b/Engine/Source/Programs/UnrealBuildTool/Configuration/UEBuildTarget.cs
@@ -3111,7 +3111,7 @@ namespace UnrealBuildTool
 						ReferencedNames.Add(Plugin.Name);
 
 						PluginReferenceDescriptor PluginReference = new PluginReferenceDescriptor(Plugin.Name, null, true);
-						PluginReference.SupportedTargetPlatforms = Plugin.Descriptor.SupportedTargetPlatforms;
+						PluginReference.SupportedTargetPlatforms = Plugin.Descriptor.GetSupportedTargetPlatformNames();
 						PluginReference.bOptional = true;
 
 						AddPlugin(PluginReference, "default plugins", ExcludeFolders, NameToInstance, NameToInfo);
diff --git a/Engine/Source/Programs/UnrealBuildTool/System/PluginDescriptor.cs b/Engine/Source/Programs/UnrealBuildTool/System/PluginDescriptor.cs
--- a/Engine/Source/Programs/UnrealBuildTool/System/PluginDescriptor.cs
+++ b/Engine/Source/Programs/UnrealBuildTool/System/PluginDescriptor.cs
@@ -443,5 +443,14 @@ namespace UnrealBuildTool
 		{
 			return SupportedTargetPlatforms == null || SupportedTargetPlatforms.Count == 0 || SupportedTargetPlatforms.Contains(Platform);
 		}
+
+		/// <summary>
+		/// Retrieve the list of supported target platforms as a string list
+		/// </summary>
+		/// <returns>String list of supported target platforms</returns>
+		public string[]? GetSupportedTargetPlatformNames()
+		{
+			return SupportedTargetPlatforms?.Select( P => P.ToString() ).ToArray();
+		}
 	}
 }
diff --git a/Engine/Source/Programs/UnrealBuildTool/System/PluginReferenceDescriptor.cs b/Engine/Source/Programs/UnrealBuildTool/System/PluginReferenceDescriptor.cs
--- a/Engine/Source/Programs/UnrealBuildTool/System/PluginReferenceDescriptor.cs
+++ b/Engine/Source/Programs/UnrealBuildTool/System/PluginReferenceDescriptor.cs
@@ -44,12 +44,12 @@ namespace UnrealBuildTool
 		/// <summary>
 		/// If enabled, list of platforms for which the plugin should be enabled (or all platforms if blank).
 		/// </summary>
-		public List<UnrealTargetPlatform>? WhitelistPlatforms;
+		public string[]? WhitelistPlatforms;
 
 		/// <summary>
 		/// If enabled, list of platforms for which the plugin should be disabled.
 		/// </summary>
-		public List<UnrealTargetPlatform>? BlacklistPlatforms;
+		public string[]? BlacklistPlatforms;
 
 		/// <summary>
 		/// If enabled, list of target configurations for which the plugin should be enabled (or all target configurations if blank).
@@ -74,7 +74,7 @@ namespace UnrealBuildTool
 		/// <summary>
 		/// The list of supported platforms for this plugin. This field is copied from the plugin descriptor, and supplements the user's whitelisted and blacklisted platforms.
 		/// </summary>
-		public List<UnrealTargetPlatform>? SupportedTargetPlatforms;
+		public string[]? SupportedTargetPlatforms;
 
 		/// <summary>
 		/// Constructor
@@ -110,11 +110,11 @@ namespace UnrealBuildTool
 			{
 				Writer.WriteValue("MarketplaceURL", MarketplaceURL);
 			}
-			if(WhitelistPlatforms != null && WhitelistPlatforms.Count > 0)
+			if(WhitelistPlatforms != null && WhitelistPlatforms.Length > 0)
 			{
 				Writer.WriteStringArrayField("WhitelistPlatforms", WhitelistPlatforms.Select(x => x.ToString()).ToArray());
 			}
-			if(BlacklistPlatforms != null && BlacklistPlatforms.Count > 0)
+			if(BlacklistPlatforms != null && BlacklistPlatforms.Length > 0)
 			{
 				Writer.WriteStringArrayField("BlacklistPlatforms", BlacklistPlatforms.Select(x => x.ToString()).ToArray());
 			}
@@ -134,7 +134,7 @@ namespace UnrealBuildTool
 			{
 				Writer.WriteEnumArrayField("BlacklistTargets", BlacklistTargets);
 			}
-			if(SupportedTargetPlatforms != null && SupportedTargetPlatforms.Count > 0)
+			if(SupportedTargetPlatforms != null && SupportedTargetPlatforms.Length > 0)
 			{
 				Writer.WriteStringArrayField("SupportedTargetPlatforms", SupportedTargetPlatforms.Select(x => x.ToString()).ToArray());
 			}
@@ -160,6 +160,25 @@ namespace UnrealBuildTool
 			}
 		}
 
+		/// <summary>
+		/// Checks the given platform names array and logs a message about any that are not known. This may simply be due to not having the platform synced in Engine/Platforms/[PlatformName].
+		/// </summary>
+		/// <param name="PlatformNames"></param>
+		private void VerifyPlatformNames( string[]? PlatformNames )
+		{
+			if (PlatformNames != null)
+			{
+				foreach( string PlatformName in PlatformNames )
+				{
+					UnrealTargetPlatform Platform;
+					if (!UnrealTargetPlatform.TryParse(PlatformName, out Platform))
+					{
+						Log.TraceLogOnce("Ignoring unknown platform '{0}' (referenced via a project's plugin descriptor for '{1}')", PlatformName, Name );
+					}
+				}
+			}
+		}
+
 		/// <summary>
 		/// Construct a PluginReferenceDescriptor from a Json object
 		/// </summary>
@@ -167,10 +186,6 @@ namespace UnrealBuildTool
 		/// <returns>New PluginReferenceDescriptor object</returns>
 		public static PluginReferenceDescriptor FromJsonObject(JsonObject RawObject)
 		{
-			string[]? WhitelistPlatformNames = null;
-			string[]? BlacklistPlatformNames = null;
-			string[]? SupportedTargetPlatformNames = null;
-
 			PluginReferenceDescriptor Descriptor = new PluginReferenceDescriptor(RawObject.GetStringField("Name"), null, RawObject.GetBoolField("Enabled"));
 			RawObject.TryGetBoolField("Optional", out Descriptor.bOptional);
 			RawObject.TryGetStringField("Description", out Descriptor.Description);
@@ -179,37 +194,18 @@ namespace UnrealBuildTool
 			// Only parse platform information if enabled
 			if (Descriptor.bEnabled)
 			{
-				RawObject.TryGetStringArrayField("WhitelistPlatforms", out WhitelistPlatformNames);
-				RawObject.TryGetStringArrayField("BlacklistPlatforms", out BlacklistPlatformNames);
+				RawObject.TryGetStringArrayField("WhitelistPlatforms", out Descriptor.WhitelistPlatforms);
+				RawObject.TryGetStringArrayField("BlacklistPlatforms", out Descriptor.BlacklistPlatforms);
 				RawObject.TryGetEnumArrayField<UnrealTargetConfiguration>("WhitelistTargetConfigurations", out Descriptor.WhitelistTargetConfigurations);
 				RawObject.TryGetEnumArrayField<UnrealTargetConfiguration>("BlacklistTargetConfigurations", out Descriptor.BlacklistTargetConfigurations);
 				RawObject.TryGetEnumArrayField<TargetType>("WhitelistTargets", out Descriptor.WhitelistTargets);
 				RawObject.TryGetEnumArrayField<TargetType>("BlacklistTargets", out Descriptor.BlacklistTargets);
-				RawObject.TryGetStringArrayField("SupportedTargetPlatforms", out SupportedTargetPlatformNames);
-			}
+				RawObject.TryGetStringArrayField("SupportedTargetPlatforms", out Descriptor.SupportedTargetPlatforms);
 
-			try
-			{
-				// convert string array to UnrealTargetPlatform arrays
-				if (WhitelistPlatformNames != null)
-				{
-					Descriptor.WhitelistPlatforms = WhitelistPlatformNames.Select(x => UnrealTargetPlatform.Parse(x)).ToList();
-				}
-				if (BlacklistPlatformNames != null)
-				{
-					Descriptor.BlacklistPlatforms = BlacklistPlatformNames.Select(x => UnrealTargetPlatform.Parse(x)).ToList();
-				}
-				if (SupportedTargetPlatformNames != null)
-				{
-					Descriptor.SupportedTargetPlatforms = SupportedTargetPlatformNames.Select(x => UnrealTargetPlatform.Parse(x)).ToList();
-				}
+				Descriptor.VerifyPlatformNames(Descriptor.WhitelistPlatforms);
+				Descriptor.VerifyPlatformNames(Descriptor.BlacklistPlatforms);
+				Descriptor.VerifyPlatformNames(Descriptor.SupportedTargetPlatforms);
 			}
-			catch (BuildException Ex)
-			{
-				ExceptionUtils.AddContext(Ex, "while parsing PluginReferenceDescriptor {0}", Descriptor.Name);
-				throw;
-			}
-
 
 			return Descriptor;
 		}
@@ -225,11 +221,11 @@ namespace UnrealBuildTool
 			{
 				return false;
 			}
-			if (WhitelistPlatforms != null && WhitelistPlatforms.Count > 0 && !WhitelistPlatforms.Contains(Platform))
+			if (WhitelistPlatforms != null && WhitelistPlatforms.Length > 0 && !WhitelistPlatforms.Contains(Platform.ToString()))
 			{
 				return false;
 			}
-			if (BlacklistPlatforms != null && BlacklistPlatforms.Contains(Platform))
+			if (BlacklistPlatforms != null && BlacklistPlatforms.Contains(Platform.ToString()))
 			{
 				return false;
 			}
@@ -287,7 +283,7 @@ namespace UnrealBuildTool
 		/// <returns>True if the plugin for this target platform</returns>
 		public bool IsSupportedTargetPlatform(UnrealTargetPlatform Platform)
 		{
-			return SupportedTargetPlatforms == null || SupportedTargetPlatforms.Count == 0 || SupportedTargetPlatforms.Contains(Platform);
+			return SupportedTargetPlatforms == null || SupportedTargetPlatforms.Length == 0 || SupportedTargetPlatforms.Contains(Platform.ToString());
 		}
 	}
 }
diff --git a/Engine/Source/Programs/UnrealBuildTool/System/Plugins.cs b/Engine/Source/Programs/UnrealBuildTool/System/Plugins.cs
--- a/Engine/Source/Programs/UnrealBuildTool/System/Plugins.cs
+++ b/Engine/Source/Programs/UnrealBuildTool/System/Plugins.cs
@@ -448,7 +448,7 @@ namespace UnrealBuildTool
 							{
 								if (ChildPluginReference.WhitelistPlatforms != null)
 								{
-									ParentPluginReference.WhitelistPlatforms = ParentPluginReference.WhitelistPlatforms.Union(ChildPluginReference.WhitelistPlatforms).ToList();
+									ParentPluginReference.WhitelistPlatforms = ParentPluginReference.WhitelistPlatforms.Union(ChildPluginReference.WhitelistPlatforms).ToArray();
 								}
 							}
 
@@ -461,7 +461,7 @@ namespace UnrealBuildTool
 								}
 								else
 								{
-									ParentPluginReference.BlacklistPlatforms = ParentPluginReference.BlacklistPlatforms.Union(ChildPluginReference.BlacklistPlatforms).ToList();
+									ParentPluginReference.BlacklistPlatforms = ParentPluginReference.BlacklistPlatforms.Union(ChildPluginReference.BlacklistPlatforms).ToArray();
 								}
 							}
 						}
