Fix for wrongful makefile invalidation

Remove code that was causing many source files to be treated as internal dependencies where they should not have been.

diff --git a/Engine/Source/Programs/UnrealBuildTool/Configuration/UEBuildTarget.cs b/Engine/Source/Programs/UnrealBuildTool/Configuration/UEBuildTarget.cs
--- a/Engine/Source/Programs/UnrealBuildTool/Configuration/UEBuildTarget.cs
+++ b/Engine/Source/Programs/UnrealBuildTool/Configuration/UEBuildTarget.cs
@@ -2001,16 +2001,6 @@ namespace UnrealBuildTool
 			// Add any toolchain dependencies
 			TargetToolChain.GetExternalDependencies(Makefile.ExternalDependencies);
 
-			// Add any leaf dependencies (eg. response files) to the dependencies list
-			IEnumerable<FileItem> LeafPrerequisiteItems = Makefile.Actions.SelectMany(x => x.PrerequisiteItems).Except(Makefile.Actions.SelectMany(x => x.ProducedItems));
-			foreach (FileItem LeafPrerequisiteItem in LeafPrerequisiteItems)
-			{
-				if (LeafPrerequisiteItem.Exists)
-				{
-					Makefile.InternalDependencies.Add(LeafPrerequisiteItem);
-				}
-			}
-
 			// Write a header containing public definitions for this target
 			if (Rules.ExportPublicHeader != null)
 			{
