Remove unnecessary Gradle steps if some files did not change

#ue5
#android

diff --git a/Engine/Source/Programs/UnrealBuildTool/Platform/Android/UEDeployAndroid.cs b/Engine/Source/Programs/UnrealBuildTool/Platform/Android/UEDeployAndroid.cs
--- a/Engine/Source/Programs/UnrealBuildTool/Platform/Android/UEDeployAndroid.cs
+++ b/Engine/Source/Programs/UnrealBuildTool/Platform/Android/UEDeployAndroid.cs
@@ -660,6 +660,104 @@ namespace UnrealBuildTool
 			}
 		}
 
+		private bool BinaryFileEquals(string SourceFilename, string DestFilename)
+		{
+			if (!File.Exists(SourceFilename))
+			{
+				return false;
+			}
+			if (!File.Exists(DestFilename))
+			{
+				return false;
+			}
+
+			FileInfo SourceInfo = new FileInfo(SourceFilename);
+			FileInfo DestInfo = new FileInfo(DestFilename);
+			if (SourceInfo.Length != DestInfo.Length)
+			{
+				return false;
+			}
+
+			using (FileStream SourceStream = new FileStream(SourceFilename, FileMode.Open, FileAccess.Read, FileShare.Read))
+			using (BinaryReader SourceReader = new BinaryReader(SourceStream))
+			using (FileStream DestStream = new FileStream(DestFilename, FileMode.Open, FileAccess.Read, FileShare.Read))
+			using (BinaryReader DestReader = new BinaryReader(DestStream))
+			{
+				while (true)
+				{
+					byte[] SourceData = SourceReader.ReadBytes(4096);
+					byte[] DestData = DestReader.ReadBytes(4096);
+					if (SourceData.Length != DestData.Length)
+					{
+						return false;
+					}
+					if (SourceData.Length == 0)
+					{
+						return true;
+					}
+					if (!SourceData.SequenceEqual(DestData))
+					{
+						return false;
+					}
+				}
+			}
+		}
+
+		private bool CopyIfDifferent(string SourceFilename, string DestFilename, bool bLog, bool bContentCompare)
+		{
+			if (!File.Exists(SourceFilename))
+			{
+				return false;
+			}
+
+			bool bDestFileAlreadyExists = File.Exists(DestFilename);
+			bool bNeedCopy = !bDestFileAlreadyExists;
+
+			if (!bNeedCopy)
+			{
+				if (bContentCompare)
+				{
+					bNeedCopy = !BinaryFileEquals(SourceFilename, DestFilename);
+				}
+				else
+				{
+					FileInfo SourceInfo = new FileInfo(SourceFilename);
+					FileInfo DestInfo = new FileInfo(DestFilename);
+
+					if (SourceInfo.Length != DestInfo.Length)
+					{
+						bNeedCopy = true;
+					}
+					else if (File.GetLastWriteTimeUtc(DestFilename) < File.GetLastWriteTimeUtc(SourceFilename))
+					{
+						// destination file older than source
+						bNeedCopy = true;
+					}
+				}
+			}
+
+			if (bNeedCopy)
+			{
+				if (bLog)
+				{
+					Log.TraceInformation("Copying {0} to {1}", SourceFilename, DestFilename);
+				}
+
+				if (bDestFileAlreadyExists)
+				{
+					SafeDeleteFile(DestFilename, false);
+				}
+				File.Copy(SourceFilename, DestFilename);
+				File.SetLastWriteTimeUtc(DestFilename, File.GetLastWriteTimeUtc(SourceFilename));
+
+				// did copy
+				return true;
+			}
+
+			// did not copy
+			return false;
+		}
+
 		private void CleanCopyDirectory(string SourceDir, string DestDir, string[] Excludes = null)
 		{
 			if (!Directory.Exists(SourceDir))
@@ -1226,6 +1324,7 @@ namespace UnrealBuildTool
 
 				// make sure it's writable if the source was readonly (e.g. autosdks)
 				new FileInfo(FinalSTLSOName).IsReadOnly = false;
+				File.SetLastWriteTimeUtc(FinalSTLSOName, File.GetLastWriteTimeUtc(SourceSTLSOName));
 			}
 		}
 
@@ -1254,6 +1353,7 @@ namespace UnrealBuildTool
 
 							Log.TraceInformation("Copying {0} to {1}", MaliLibSrcPath, MaliLibDstPath);
 							File.Copy(MaliLibSrcPath, MaliLibDstPath, true);
+							File.SetLastWriteTimeUtc(MaliLibDstPath, File.GetLastWriteTimeUtc(MaliLibSrcPath));
 
 							string MaliVkLayerLibSrcPath = Path.Combine(MaliGraphicsDebuggerPath, "target", "android", "arm", "rooted", NDKArch, "libGLES_aga.so");
 							if (File.Exists(MaliVkLayerLibSrcPath))
@@ -1261,6 +1361,7 @@ namespace UnrealBuildTool
 								string MaliVkLayerLibDstPath = Path.Combine(UnrealBuildPath, "libs", NDKArch, "libVkLayerAGA.so");
 								Log.TraceInformation("Copying {0} to {1}", MaliVkLayerLibSrcPath, MaliVkLayerLibDstPath);
 								File.Copy(MaliVkLayerLibSrcPath, MaliVkLayerLibDstPath, true);
+								File.SetLastWriteTimeUtc(MaliVkLayerLibDstPath, File.GetLastWriteTimeUtc(MaliVkLayerLibSrcPath));
 							}
 						}
 					}
@@ -3845,12 +3946,7 @@ namespace UnrealBuildTool
 						Log.TraceInformation("Obb file exists...");
 						string DestFileName = Path.Combine(ObbFileDestination, "main.obb.png"); // Need a rename to turn off compression
 						string SrcFileName = ObbFileLocation;
-						if (!File.Exists(DestFileName) || File.GetLastWriteTimeUtc(DestFileName) < File.GetLastWriteTimeUtc(SrcFileName))
-						{
-							Log.TraceInformation("Copying {0} to {1}", SrcFileName, DestFileName);
-							SafeDeleteFile(DestFileName);
-							File.Copy(SrcFileName, DestFileName);
-						}
+						CopyIfDifferent(SrcFileName, DestFileName, true, false);
 					}
 				}
 				else // try to remove the file it we aren't packaging inside the APK
@@ -3868,16 +3964,7 @@ namespace UnrealBuildTool
 					Directory.CreateDirectory(UnrealBuildPath);
 					Directory.CreateDirectory(Path.Combine(UnrealBuildPath, "assets"));
 					Console.WriteLine("UnrealCommandLine.txt exists...");
-					bool bDestFileAlreadyExists = File.Exists(CommandLineDestFileName);
-					if (!bDestFileAlreadyExists || File.GetLastWriteTimeUtc(CommandLineDestFileName) < File.GetLastWriteTimeUtc(CommandLineSourceFileName))
-					{
-						Console.WriteLine("Copying {0} to {1}", CommandLineSourceFileName, CommandLineDestFileName);
-						if (bDestFileAlreadyExists)
-						{
-							SafeDeleteFile(CommandLineDestFileName, false);
-						}
-						File.Copy(CommandLineSourceFileName, CommandLineDestFileName);
-					}
+					CopyIfDifferent(CommandLineSourceFileName, CommandLineDestFileName, true, true);
 				}
 				else // try to remove the file if we aren't packaging one
 				{
@@ -4054,7 +4141,7 @@ namespace UnrealBuildTool
 				// stage files into gradle app directory
 				string GradleManifest = Path.Combine(UnrealBuildGradleMainPath, "AndroidManifest.xml");
 				MakeDirectoryIfRequired(GradleManifest);
-				File.Copy(Path.Combine(UnrealBuildPath, "AndroidManifest.xml"), GradleManifest, true);
+				CopyIfDifferent(Path.Combine(UE4BuildPath, "AndroidManifest.xml"), GradleManifest, true, true);
 
 				string[] Excludes;
 				switch (NDKArch)
