Fix build.

#rnx

diff --git a/Engine/Source/Programs/UnrealBuildTool/Platform/Android/UEDeployAndroid.cs b/Engine/Source/Programs/UnrealBuildTool/Platform/Android/UEDeployAndroid.cs
--- a/Engine/Source/Programs/UnrealBuildTool/Platform/Android/UEDeployAndroid.cs
+++ b/Engine/Source/Programs/UnrealBuildTool/Platform/Android/UEDeployAndroid.cs
@@ -4141,7 +4141,7 @@ namespace UnrealBuildTool
 				// stage files into gradle app directory
 				string GradleManifest = Path.Combine(UnrealBuildGradleMainPath, "AndroidManifest.xml");
 				MakeDirectoryIfRequired(GradleManifest);
-				CopyIfDifferent(Path.Combine(UE4BuildPath, "AndroidManifest.xml"), GradleManifest, true, true);
+				CopyIfDifferent(Path.Combine(UnrealBuildPath, "AndroidManifest.xml"), GradleManifest, true, true);
 
 				string[] Excludes;
 				switch (NDKArch)
