Use the legacy Visual Studio headers when building with Clang.

Latest version has some issues: https://developercommunity.visualstudio.com/t/bug-in-compiling-std-with-clan-clang-cl-in-latest/1168772

diff --git a/Engine/Source/Programs/UnrealBuildTool/Platform/Windows/VCEnvironment.cs b/Engine/Source/Programs/UnrealBuildTool/Platform/Windows/VCEnvironment.cs
--- a/Engine/Source/Programs/UnrealBuildTool/Platform/Windows/VCEnvironment.cs
+++ b/Engine/Source/Programs/UnrealBuildTool/Platform/Windows/VCEnvironment.cs
@@ -435,7 +435,12 @@ namespace UnrealBuildTool
 			DirectoryReference? SelectedToolChainDir;
 			if(Compiler == WindowsCompiler.Clang || Compiler == WindowsCompiler.Intel)
 			{
-				if (WindowsPlatform.TryGetToolChainDir(WindowsCompiler.VisualStudio2019, null, out SelectedToolChainVersion, out SelectedToolChainDir))
+				const string DefaultClangVisualStudioVersion = "14.24.28315";
+				if (WindowsPlatform.TryGetToolChainDir(WindowsCompiler.VisualStudio2019, DefaultClangVisualStudioVersion, out SelectedToolChainVersion, out SelectedToolChainDir))
+				{
+					ToolChain = WindowsCompiler.VisualStudio2019;
+				}
+				else if (WindowsPlatform.TryGetToolChainDir(WindowsCompiler.VisualStudio2019, null, out SelectedToolChainVersion, out SelectedToolChainDir))
 				{
 					ToolChain = WindowsCompiler.VisualStudio2019;
 				}
