Fall back to LD linker if LLD is outdated in the Lumin toolchain

diff --git a/Engine/Source/Programs/UnrealBuildTool/Platform/Lumin/LuminToolChain.cs b/Engine/Source/Programs/UnrealBuildTool/Platform/Lumin/LuminToolChain.cs
--- a/Engine/Source/Programs/UnrealBuildTool/Platform/Lumin/LuminToolChain.cs
+++ b/Engine/Source/Programs/UnrealBuildTool/Platform/Lumin/LuminToolChain.cs
@@ -21,6 +21,7 @@ namespace UnrealBuildTool
 
 		//static private string[] LibrariesToSkip = new string[] { "nvToolsExt", "nvToolsExtStub", "oculus", "vrapi", "ovrkernel", "systemutils", "openglloader", "gpg", };
 
+		private bool bUseLLD = false;
 		private List<string> AdditionalGPUArches;
 
 		protected string StripPath;
@@ -131,6 +132,18 @@ namespace UnrealBuildTool
 
 			ToolchainLinkParamsArm64 = ToolchainParamsArm64;
 			ToolchainLinkParamsx64 = ToolchainParamsx64;
+
+			string LLDPath = Path.Combine(MLSDKPath, "tools", "toolchains", "llvm-8", "bin", "ld.lld" + GetHostPlatformBinarySuffix());
+			string LLDBuild = Utils.RunLocalProcessAndReturnStdOut(LLDPath, "--version");
+			if (!string.IsNullOrEmpty(LLDBuild))
+			{
+				string[] Split = LLDBuild.Split(' ');
+				if (Split.Length < 2)
+				{
+					throw new BuildException("Unexpected output from lld in --version: '{0}'", LLDBuild);
+				}
+				bUseLLD = VersionNumber.Parse(Split[1]) >= VersionNumber.Parse("11.0.0");
+			}
 		}
 
 		public bool UseVulkan()
@@ -224,7 +237,17 @@ namespace UnrealBuildTool
 
 		protected override string GetLinkArguments(LinkEnvironment LinkEnvironment, string Architecture)
 		{
-			string Result = " -fuse-ld=lld";
+			string Result = "";
+			
+			if (bUseLLD)
+			{
+				Result += " -fuse-ld=lld";
+			}
+			else
+			{
+				Log.TraceWarning("The linker in the Magic Leap LLVM toolchain is known to fail when linking large projects." +
+					"  Run the UpdateLinkerLumin script in Engine/Extras/Android to update the linker if the build command fails.");
+			}
 
 			if (LinkEnvironment.bIsBuildingDLL)
 			{
@@ -383,5 +406,15 @@ namespace UnrealBuildTool
 				CompileOrLinkAction.CommandArguments = String.Format("-c \'{0} {1}\'", QuotedCommandPath.Replace("\'", "\\\'"), CommandArguments.Replace("\'", "\\\'"));
 			}
 		}
+
+		private string GetHostPlatformBinarySuffix()
+		{
+			if (BuildHostPlatform.Current.Platform == UnrealTargetPlatform.Win64)
+			{
+				return ".exe";
+			}
+
+			return "";
+		}
 	};
 }
