Fix for copying Vulkan validation layers

#ue4
#android

diff --git a/Engine/Source/Programs/UnrealBuildTool/Platform/Android/UEDeployAndroid.cs b/Engine/Source/Programs/UnrealBuildTool/Platform/Android/UEDeployAndroid.cs
--- a/Engine/Source/Programs/UnrealBuildTool/Platform/Android/UEDeployAndroid.cs
+++ b/Engine/Source/Programs/UnrealBuildTool/Platform/Android/UEDeployAndroid.cs
@@ -1413,7 +1413,10 @@ namespace UnrealBuildTool
 					Directory.CreateDirectory(DestDir);
 					string SourceFilename = Path.Combine(VulkanLayersDir, ANDROID_VULKAN_VALIDATION_LAYER);
 					string DestFilename = Path.Combine(DestDir, ANDROID_VULKAN_VALIDATION_LAYER);
+					SafeDeleteFile(DestFilename);
 					File.Copy(SourceFilename, DestFilename);
+					FileInfo DestFileInfo = new FileInfo(DestFilename);
+					DestFileInfo.Attributes = DestFileInfo.Attributes & ~FileAttributes.ReadOnly;
 					File.SetLastWriteTimeUtc(DestFilename, File.GetLastWriteTimeUtc(SourceFilename));
 				}
 			}
