Update ASan support for MSVC 16.9 by disabling explicit linking of support libraries.

Also fixed the flag format, for consistency, now that they support the '/' variant for it, which early builds didn't.

#rnx

diff --git a/Engine/Source/Programs/UnrealBuildTool/Platform/Windows/VCToolChain.cs b/Engine/Source/Programs/UnrealBuildTool/Platform/Windows/VCToolChain.cs
--- a/Engine/Source/Programs/UnrealBuildTool/Platform/Windows/VCToolChain.cs
+++ b/Engine/Source/Programs/UnrealBuildTool/Platform/Windows/VCToolChain.cs
@@ -349,7 +349,7 @@ namespace UnrealBuildTool
 			{
 				// Enable address sanitizer. This also requires companion libraries at link time.
 				// Works for clang too.
-				Arguments.Add("-fsanitize=address");
+				Arguments.Add("/fsanitize=address");
 
 				// Use the CRT allocator so that ASan is able to hook into it for better error
 				// detection.
@@ -1619,7 +1619,9 @@ namespace UnrealBuildTool
 
 			// If we're building either an executable or a DLL, make sure we link in the 
 			// correct address sanitizer helper libs.
-			if (!bBuildImportLibraryOnly && !LinkEnvironment.bIsBuildingLibrary && Target.WindowsPlatform.bEnableAddressSanitizer)
+			// Note: As of MSVC 16.9, this is automatically done if the /fsanitize=address flag is used.
+			if (!bBuildImportLibraryOnly && !LinkEnvironment.bIsBuildingLibrary && Target.WindowsPlatform.bEnableAddressSanitizer && 
+			    EnvVars.CompilerVersion < new VersionNumber(14, 28, 0))
 			{
 				String ASanArchSuffix = "";
 				if (EnvVars.Architecture == WindowsArchitecture.x64)
