Make it somewhat easier to bump C++ versions:
* Re-defined CppStandardVersion.Default in terms of existing enums and removed it from maps and switches etc.
* Added C++20 flags for all the recognized compilers.

#rnx

diff --git a/Engine/Source/Programs/UnrealBuildTool/Platform/Android/AndroidToolChain.cs b/Engine/Source/Programs/UnrealBuildTool/Platform/Android/AndroidToolChain.cs
--- a/Engine/Source/Programs/UnrealBuildTool/Platform/Android/AndroidToolChain.cs
+++ b/Engine/Source/Programs/UnrealBuildTool/Platform/Android/AndroidToolChain.cs
@@ -766,16 +766,20 @@ namespace UnrealBuildTool
 			return Result;
 		}
 
-		static string GetCppStandardCompileArgument(CppCompileEnvironment CompileEnvironment)
+		static string GetCppStandardCompileArgument(CppStandardVersion Version)
 		{
-			var Mapping = new Dictionary<CppStandardVersion, string>
-			{
-				{ CppStandardVersion.Cpp14, " -std=c++14" },
-				{ CppStandardVersion.Cpp17, " -std=c++17" },
-				{ CppStandardVersion.Latest, " -std=c++17" },
-				{ CppStandardVersion.Default, " -std=c++14" }
-			};
-			return Mapping[CompileEnvironment.CppStandard];
+			switch (Version)
+			{
+				case CppStandardVersion.Cpp14:
+					return " -std=c++14";
+				case CppStandardVersion.Latest:
+				case CppStandardVersion.Cpp17:
+					return " -std=c++17";
+				case CppStandardVersion.Cpp20:
+					return " -std=c++20";
+				default:
+					throw new BuildException($"Unsupported C++ standard type set: {Version}");
+			}
 		}
 
 		static string GetCompileArguments_CPP(CppCompileEnvironment CompileEnvironment, bool bDisableOptimizations)
@@ -783,7 +787,7 @@ namespace UnrealBuildTool
 			string Result = "";
 
 			Result += " -x c++";
-			Result += GetCppStandardCompileArgument(CompileEnvironment);
+			Result += GetCppStandardCompileArgument(CompileEnvironment.CppStandard);
 
 			// optimization level
 			if (bDisableOptimizations)
@@ -822,7 +826,7 @@ namespace UnrealBuildTool
 			string Result = "";
 
 			Result += " -x c++-header";
-			Result += GetCppStandardCompileArgument(CompileEnvironment);
+			Result += GetCppStandardCompileArgument(CompileEnvironment.CppStandard);
 
 			// optimization level
 			if (bDisableOptimizations)
diff --git a/Engine/Source/Programs/UnrealBuildTool/Platform/IOS/IOSToolChain.cs b/Engine/Source/Programs/UnrealBuildTool/Platform/IOS/IOSToolChain.cs
--- a/Engine/Source/Programs/UnrealBuildTool/Platform/IOS/IOSToolChain.cs
+++ b/Engine/Source/Programs/UnrealBuildTool/Platform/IOS/IOSToolChain.cs
@@ -335,35 +335,35 @@ namespace UnrealBuildTool
 			return Result;
 		}
 
-		static string GetCompileArguments_CPP()
+		static string GetCompileArguments_CPP(CppCompileEnvironment CompileEnvironment)
 		{
 			string Result = "";
 			Result += " -x objective-c++";
 			Result += " -fobjc-abi-version=2";
 			Result += " -fobjc-legacy-dispatch";
-			Result += " -std=c++14";
+			Result += GetCppStandardCompileArgument(CompileEnvironment.CppStandard);
 			Result += " -stdlib=libc++";
 			return Result;
 		}
 
-		static string GetCompileArguments_MM()
+		static string GetCompileArguments_MM(CppCompileEnvironment CompileEnvironment)
 		{
 			string Result = "";
 			Result += " -x objective-c++";
 			Result += " -fobjc-abi-version=2";
 			Result += " -fobjc-legacy-dispatch";
-			Result += " -std=c++14";
+			Result += GetCppStandardCompileArgument(CompileEnvironment.CppStandard);
 			Result += " -stdlib=libc++";
 			return Result;
 		}
 
-		static string GetCompileArguments_M()
+		static string GetCompileArguments_M(CppCompileEnvironment CompileEnvironment)
 		{
 			string Result = "";
 			Result += " -x objective-c";
 			Result += " -fobjc-abi-version=2";
 			Result += " -fobjc-legacy-dispatch";
-			Result += " -std=c++14";
+			Result += GetCppStandardCompileArgument(CompileEnvironment.CppStandard);
 			Result += " -stdlib=libc++";
 			return Result;
 		}
@@ -375,11 +375,11 @@ namespace UnrealBuildTool
 			return Result;
 		}
 
-		static string GetCompileArguments_PCH()
+		static string GetCompileArguments_PCH(CppCompileEnvironment CompileEnvironment)
 		{
 			string Result = "";
 			Result += " -x objective-c++-header";
-			Result += " -std=c++14";
+			Result += GetCppStandardCompileArgument(CompileEnvironment.CppStandard);
 			Result += " -stdlib=libc++";
 			return Result;
 		}
@@ -621,7 +621,7 @@ namespace UnrealBuildTool
 				if (CompileEnvironment.PrecompiledHeaderAction == PrecompiledHeaderAction.Create)
 				{
 					// Compile the file as a C++ PCH.
-					FileArguments += GetCompileArguments_PCH();
+					FileArguments += GetCompileArguments_PCH(CompileEnvironment);
 					FileArguments += GetRTTIFlag(CompileEnvironment);
 					FileArguments += GetObjCExceptionsFlag(CompileEnvironment);
 				}
@@ -633,20 +633,20 @@ namespace UnrealBuildTool
 				else if (Extension == ".MM")
 				{
 					// Compile the file as Objective-C++ code.
-					FileArguments += GetCompileArguments_MM();
+					FileArguments += GetCompileArguments_MM(CompileEnvironment);
 					FileArguments += GetRTTIFlag(CompileEnvironment);
 					FileArguments += GetObjCExceptionsFlag(CompileEnvironment);
 				}
 				else if (Extension == ".M")
 				{
 					// Compile the file as Objective-C code.
-					FileArguments += GetCompileArguments_M();
+					FileArguments += GetCompileArguments_M(CompileEnvironment);
 					FileArguments += GetObjCExceptionsFlag(CompileEnvironment);
 				}
 				else
 				{
 					// Compile the file as C++ code.
-					FileArguments += GetCompileArguments_CPP();
+					FileArguments += GetCompileArguments_CPP(CompileEnvironment);
 					FileArguments += GetRTTIFlag(CompileEnvironment);
 					FileArguments += GetObjCExceptionsFlag(CompileEnvironment);
 
diff --git a/Engine/Source/Programs/UnrealBuildTool/Platform/Linux/LinuxToolChain.cs b/Engine/Source/Programs/UnrealBuildTool/Platform/Linux/LinuxToolChain.cs
--- a/Engine/Source/Programs/UnrealBuildTool/Platform/Linux/LinuxToolChain.cs
+++ b/Engine/Source/Programs/UnrealBuildTool/Platform/Linux/LinuxToolChain.cs
@@ -872,30 +872,27 @@ namespace UnrealBuildTool
 				: string.Format("{0}={1}", myKey, myValue);
 		}
 
-		static string GetCompilerStandardVersion_CPP(CppCompileEnvironment CompileEnvironment)
+		static string GetCompilerStandardVersion_CPP(CppStandardVersion Version)
 		{
-			if (CompileEnvironment.CppStandard == CppStandardVersion.Cpp14 || CompileEnvironment.CppStandard == CppStandardVersion.Default)
-			{
-				return " -std=c++14";
-			}
-			else if (CompileEnvironment.CppStandard == CppStandardVersion.Cpp17)
-			{
-				return " -std=c++17";
+			switch (Version)
+			{
+				case CppStandardVersion.Cpp14:
+					return " -std=c++14";
+				case CppStandardVersion.Cpp17:
+					return " -std=c++17";
+				case CppStandardVersion.Latest:
+				case CppStandardVersion.Cpp20:
+					return " -std=c++20";
+				default:
+					throw new BuildException($"Unsupported C++ standard type set: {Version}");
 			}
-			else if (CompileEnvironment.CppStandard == CppStandardVersion.Latest)
-			{
-				return " -std=c++17";
-			}
-
-			throw new BuildException(
-			string.Format("Unknown C++ standard type set: {0}", CompileEnvironment.CppStandard));
 		}
 
 		static string GetCompileArguments_CPP(CppCompileEnvironment CompileEnvironment)
 		{
 			string Result = "";
 			Result += " -x c++";
-			Result += GetCompilerStandardVersion_CPP(CompileEnvironment);
+			Result += GetCompilerStandardVersion_CPP(CompileEnvironment.CppStandard);
 			return Result;
 		}
 
@@ -912,7 +909,7 @@ namespace UnrealBuildTool
 			Result += " -x objective-c++";
 			Result += " -fobjc-abi-version=2";
 			Result += " -fobjc-legacy-dispatch";
-			Result += GetCompilerStandardVersion_CPP(CompileEnvironment);
+			Result += GetCompilerStandardVersion_CPP(CompileEnvironment.CppStandard);
 			return Result;
 		}
 
@@ -941,7 +938,7 @@ namespace UnrealBuildTool
 			Result += " -x objective-c";
 			Result += " -fobjc-abi-version=2";
 			Result += " -fobjc-legacy-dispatch";
-			Result += GetCompilerStandardVersion_CPP(CompileEnvironment);
+			Result += GetCompilerStandardVersion_CPP(CompileEnvironment.CppStandard);
 			return Result;
 		}
 
@@ -949,7 +946,7 @@ namespace UnrealBuildTool
 		{
 			string Result = "";
 			Result += " -x c++-header";
-			Result += GetCompilerStandardVersion_CPP(CompileEnvironment);
+			Result += GetCompilerStandardVersion_CPP(CompileEnvironment.CppStandard);
 			return Result;
 		}
 
diff --git a/Engine/Source/Programs/UnrealBuildTool/Platform/Mac/MacToolChain.cs b/Engine/Source/Programs/UnrealBuildTool/Platform/Mac/MacToolChain.cs
--- a/Engine/Source/Programs/UnrealBuildTool/Platform/Mac/MacToolChain.cs
+++ b/Engine/Source/Programs/UnrealBuildTool/Platform/Mac/MacToolChain.cs
@@ -292,26 +292,14 @@ namespace UnrealBuildTool
 
 			return Result;
 		}
-
-		static string GetCppStandardCompileArgument(CppCompileEnvironment CompileEnvironment)
-		{
-			var Mapping = new Dictionary<CppStandardVersion, string>
-			{
-				{ CppStandardVersion.Cpp14, " -std=c++14" },
-				{ CppStandardVersion.Cpp17, " -std=c++17" },
-				{ CppStandardVersion.Latest, " -std=c++17" },
-				{ CppStandardVersion.Default, " -std=c++14" }
-			};
-			return Mapping[CompileEnvironment.CppStandard];
-		}
-
+		
 		static string GetCompileArguments_CPP(CppCompileEnvironment CompileEnvironment)
 		{
 			string Result = "";
 			Result += " -x objective-c++";
 			Result += " -fobjc-abi-version=2";
 			Result += " -fobjc-legacy-dispatch";
-			Result += GetCppStandardCompileArgument(CompileEnvironment);
+			Result += GetCppStandardCompileArgument(CompileEnvironment.CppStandard);
 			Result += " -stdlib=libc++";
 			return Result;
 		}
@@ -322,7 +310,7 @@ namespace UnrealBuildTool
 			Result += " -x objective-c++";
 			Result += " -fobjc-abi-version=2";
 			Result += " -fobjc-legacy-dispatch";
-			Result += GetCppStandardCompileArgument(CompileEnvironment);
+			Result += GetCppStandardCompileArgument(CompileEnvironment.CppStandard);
 			Result += " -stdlib=libc++";
 			return Result;
 		}
@@ -350,7 +338,7 @@ namespace UnrealBuildTool
 			Result += " -x objective-c++-header";
 			Result += " -fobjc-abi-version=2";
 			Result += " -fobjc-legacy-dispatch";
-			Result += GetCppStandardCompileArgument(CompileEnvironment);
+			Result += GetCppStandardCompileArgument(CompileEnvironment.CppStandard);
 			Result += " -stdlib=libc++";
 			return Result;
 		}
diff --git a/Engine/Source/Programs/UnrealBuildTool/ProjectFiles/VisualStudio/VCProject.cs b/Engine/Source/Programs/UnrealBuildTool/ProjectFiles/VisualStudio/VCProject.cs
--- a/Engine/Source/Programs/UnrealBuildTool/ProjectFiles/VisualStudio/VCProject.cs
+++ b/Engine/Source/Programs/UnrealBuildTool/ProjectFiles/VisualStudio/VCProject.cs
@@ -578,15 +578,18 @@ namespace UnrealBuildTool
 		{
 			switch (Version)
 			{
-				case CppStandardVersion.Default:
 				case CppStandardVersion.Cpp14:
 					return "/std:c++14";
 				case CppStandardVersion.Cpp17:
 					return "/std:c++17";
 				case CppStandardVersion.Latest:
 					return "/std:c++latest";
+				// Will be added when MSVC is feature-complete.
+				// https://docs.microsoft.com/en-us/cpp/build/reference/std-specify-language-standard-version?view=msvc-160
+				// case CppStandardVersion.Cpp20:
+				//	return "/std:c++20";
 				default:
-					throw new ArgumentOutOfRangeException(nameof(Version), Version, "Please update switch above with new C++ version");
+					throw new BuildException($"Unsupported C++ standard type set: {Version}");
 			}
 		}
 
diff --git a/Engine/Source/Programs/UnrealBuildTool/System/CppCompileEnvironment.cs b/Engine/Source/Programs/UnrealBuildTool/System/CppCompileEnvironment.cs
--- a/Engine/Source/Programs/UnrealBuildTool/System/CppCompileEnvironment.cs
+++ b/Engine/Source/Programs/UnrealBuildTool/System/CppCompileEnvironment.cs
@@ -27,11 +27,6 @@ namespace UnrealBuildTool
 	/// </summary>
 	public enum CppStandardVersion
 	{
-		/// <summary>
-		/// Use the default standard version
-		/// </summary>
-		Default,
-
 		/// <summary>
 		/// Supports C++14
 		/// </summary>
@@ -46,11 +41,16 @@ namespace UnrealBuildTool
 		/// Supports C++20
 		/// </summary>
 		Cpp20,
-
+		
 		/// <summary>
 		/// Latest standard supported by the compiler
 		/// </summary>
 		Latest,
+		
+		/// <summary>
+		/// Use the default standard version
+		/// </summary>
+		Default = Cpp14,
 	}
 
 	/// <summary>
diff --git a/Engine/Source/Programs/UnrealBuildTool/ToolChain/AppleToolChain.cs b/Engine/Source/Programs/UnrealBuildTool/ToolChain/AppleToolChain.cs
--- a/Engine/Source/Programs/UnrealBuildTool/ToolChain/AppleToolChain.cs
+++ b/Engine/Source/Programs/UnrealBuildTool/ToolChain/AppleToolChain.cs
@@ -212,6 +212,22 @@ namespace UnrealBuildTool
 			return FullClangVersion;
 		}
 
+		protected static string GetCppStandardCompileArgument(CppStandardVersion Version)
+		{
+			switch (Version)
+			{
+				case CppStandardVersion.Cpp14:
+					return " -std=c++14";
+				case CppStandardVersion.Cpp17:
+					return " -std=c++17";
+				case CppStandardVersion.Latest:
+				case CppStandardVersion.Cpp20:
+					return " -std=c++20";
+				default:
+					throw new BuildException($"Unsupported C++ standard type set: {Version}");
+			}
+		}
+
 		protected string GetDsymutilPath(out string ExtraOptions, bool bIsForLTOBuild=false)
 		{
 			FileReference DsymutilLocation = new FileReference("/usr/bin/dsymutil");
