Update C++17 (and C++20) support for all platforms that were missing from CL 15801501.

#rnx

diff --git a/Engine/Source/Programs/UnrealBuildTool/Modes/GenerateClangDatabase.cs b/Engine/Source/Programs/UnrealBuildTool/Modes/GenerateClangDatabase.cs
--- a/Engine/Source/Programs/UnrealBuildTool/Modes/GenerateClangDatabase.cs
+++ b/Engine/Source/Programs/UnrealBuildTool/Modes/GenerateClangDatabase.cs
@@ -97,13 +97,20 @@ namespace UnrealBuildTool
 								StringBuilder CommandBuilder = new StringBuilder();
 								CommandBuilder.AppendFormat("\"{0}\"", ClangPath.FullName);
 
-								if (ModuleCompileEnvironment.CppStandard >= CppStandardVersion.Cpp17)
+								switch (ModuleCompileEnvironment.CppStandard)
 								{
-									CommandBuilder.AppendFormat(" -std=c++17");
-								}
-								else if (ModuleCompileEnvironment.CppStandard >= CppStandardVersion.Cpp14)
-								{
-									CommandBuilder.AppendFormat(" -std=c++14");
+									case CppStandardVersion.Cpp14:
+										CommandBuilder.AppendFormat(" -std=c++14");
+										break;
+									case CppStandardVersion.Latest:
+									case CppStandardVersion.Cpp17:
+										CommandBuilder.AppendFormat(" -std=c++17");
+										break;
+									case CppStandardVersion.Cpp20:
+										CommandBuilder.AppendFormat(" -std=c++20");
+										break;
+									default:
+										throw new BuildException($"Unsupported C++ standard type set: {ModuleCompileEnvironment.CppStandard}");
 								}
 
 								foreach (FileItem ForceIncludeFile in ModuleCompileEnvironment.ForceIncludeFiles)
diff --git a/Engine/Source/Programs/UnrealBuildTool/Platform/HoloLens/HoloLensToolChain.cs b/Engine/Source/Programs/UnrealBuildTool/Platform/HoloLens/HoloLensToolChain.cs
--- a/Engine/Source/Programs/UnrealBuildTool/Platform/HoloLens/HoloLensToolChain.cs
+++ b/Engine/Source/Programs/UnrealBuildTool/Platform/HoloLens/HoloLensToolChain.cs
@@ -255,18 +255,26 @@ namespace UnrealBuildTool
 			// we've hit problems where types are somehow in windows.winmd on some installations but not others, leading to either
 			// missing or duplicated type references.
 			Arguments.Add("/ZW:nostdlib");
-
-			if (CompileEnvironment.CppStandard >= CppStandardVersion.Latest)
-			{
-				Arguments.Add("/std:c++latest");
-			}
-			else if (CompileEnvironment.CppStandard >= CppStandardVersion.Cpp17)
-			{
-				Arguments.Add("/std:c++17");
-			}
-			else if (CompileEnvironment.CppStandard >= CppStandardVersion.Cpp14)
+			
+			// Note: Should be shared with all Windows-like platforms.
+			switch (CompileEnvironment.CppStandard)
 			{
-				Arguments.Add("/std:c++14");
+				case CppStandardVersion.Cpp14:
+					Arguments.Add("/std:c++14");
+					break;
+				case CppStandardVersion.Cpp17:
+					Arguments.Add("/std:c++17");
+					break;
+				case CppStandardVersion.Latest:
+					Arguments.Add("/std:c++latest");
+					break;
+				// Will be added when MSVC is feature-complete.
+				// https://docs.microsoft.com/en-us/cpp/build/reference/std-specify-language-standard-version?view=msvc-160
+				// case CppStandardVersion.Cpp20:
+				//	Arguments.Add("/std:c++20");
+				//  break;
+				default:
+					throw new BuildException($"Unsupported C++ standard type set: {CompileEnvironment.CppStandard}");
 			}
 
 			// Explicitly compile the file as C++.
diff --git a/Engine/Source/Programs/UnrealBuildTool/Platform/Windows/VCToolChain.cs b/Engine/Source/Programs/UnrealBuildTool/Platform/Windows/VCToolChain.cs
--- a/Engine/Source/Programs/UnrealBuildTool/Platform/Windows/VCToolChain.cs
+++ b/Engine/Source/Programs/UnrealBuildTool/Platform/Windows/VCToolChain.cs
@@ -660,17 +660,25 @@ namespace UnrealBuildTool
 			}
 			else
 			{
-				if(CompileEnvironment.CppStandard >= CppStandardVersion.Cpp20)
-				{
-					Arguments.Add("/std:c++latest");
-				}
-				else if(CompileEnvironment.CppStandard >= CppStandardVersion.Cpp17)
-				{
-					Arguments.Add("/std:c++17");
-				}
-				else if(CompileEnvironment.CppStandard >= CppStandardVersion.Cpp14)
-				{
-					Arguments.Add("/std:c++14");
+				// Note: Should be shared with all Windows-like platforms.
+				switch (CompileEnvironment.CppStandard)
+				{
+					case CppStandardVersion.Cpp14:
+						Arguments.Add("/std:c++14");
+						break;
+					case CppStandardVersion.Cpp17:
+						Arguments.Add("/std:c++17");
+						break;
+					case CppStandardVersion.Latest:
+						Arguments.Add("/std:c++latest");
+						break;
+					// Will be added when MSVC is feature-complete.
+					// https://docs.microsoft.com/en-us/cpp/build/reference/std-specify-language-standard-version?view=msvc-160
+					// case CppStandardVersion.Cpp20:
+					//	Arguments.Add("/std:c++20");
+					//  break;
+					default:
+						throw new BuildException($"Unsupported C++ standard type set: {CompileEnvironment.CppStandard}");
 				}
 			}
 
diff --git a/Engine/Source/Programs/UnrealBuildTool/ProjectFiles/VisualStudio/VCProject.cs b/Engine/Source/Programs/UnrealBuildTool/ProjectFiles/VisualStudio/VCProject.cs
--- a/Engine/Source/Programs/UnrealBuildTool/ProjectFiles/VisualStudio/VCProject.cs
+++ b/Engine/Source/Programs/UnrealBuildTool/ProjectFiles/VisualStudio/VCProject.cs
@@ -574,7 +574,7 @@ namespace UnrealBuildTool
 		/// <summary>
 		/// Gets compiler switch for specifying in AdditionalOptions in .vcxproj file for specific C++ version
 		/// </summary>
-		private static string GetIntelliSenseSwitchForCppVersion(CppStandardVersion Version)
+		private static string GetCppStandardCompileArgument(CppStandardVersion Version)
 		{
 			switch (Version)
 			{
@@ -1059,7 +1059,7 @@ namespace UnrealBuildTool
 				VCProjectFileContent.AppendLine("    <NMakeForcedIncludes>$(NMakeForcedIncludes)</NMakeForcedIncludes>");
 				VCProjectFileContent.AppendLine("    <NMakeAssemblySearchPath>$(NMakeAssemblySearchPath)</NMakeAssemblySearchPath>");
 				VCProjectFileContent.AppendLine("    <AdditionalOptions>{0}</AdditionalOptions>",
-					GetIntelliSenseSwitchForCppVersion(GetIntelliSenseCppVersion()));
+					GetCppStandardCompileArgument(GetIntelliSenseCppVersion()));
 				VCProjectFileContent.AppendLine("  </PropertyGroup>");
 			}
 
@@ -1674,19 +1674,7 @@ namespace UnrealBuildTool
 					VCProjectFileContent.AppendLine("    <NMakeReBuildCommandLine>{0} {1}</NMakeReBuildCommandLine>", EscapePath(NormalizeProjectPath(FileReference.Combine(BatchFilesDirectory, "Rebuild.bat"))), BuildArguments.ToString());
 					VCProjectFileContent.AppendLine("    <NMakeCleanCommandLine>{0} {1}</NMakeCleanCommandLine>", EscapePath(NormalizeProjectPath(FileReference.Combine(BatchFilesDirectory, "Clean.bat"))), BuildArguments.ToString());
 					VCProjectFileContent.AppendLine("    <NMakeOutput>{0}</NMakeOutput>", NormalizeProjectPath(NMakePath.FullName));
-
-					if (TargetRulesObject.CppStandard >= CppStandardVersion.Latest)
-					{
-						VCProjectFileContent.AppendLine("    <AdditionalOptions>/std:c++latest</AdditionalOptions>");
-					}
-					else if (TargetRulesObject.CppStandard >= CppStandardVersion.Cpp17)
-					{
-						VCProjectFileContent.AppendLine("    <AdditionalOptions>/std:c++17</AdditionalOptions>");
-					}
-					else if (TargetRulesObject.CppStandard >= CppStandardVersion.Cpp14)
-					{
-						VCProjectFileContent.AppendLine("    <AdditionalOptions>/std:c++14</AdditionalOptions>");
-					}
+					VCProjectFileContent.AppendLine("    <AdditionalOptions>{0}</AdditionalOptions>", GetCppStandardCompileArgument(TargetRulesObject.CppStandard));
 
 					if (TargetRulesObject.Type == TargetType.Game || TargetRulesObject.Type == TargetType.Client || TargetRulesObject.Type == TargetType.Server)
 					{
