Remove a UE4 reference from a comment.

diff --git a/Engine/Source/Programs/UnrealBuildTool/UnrealBuildTool.cs b/Engine/Source/Programs/UnrealBuildTool/UnrealBuildTool.cs
--- a/Engine/Source/Programs/UnrealBuildTool/UnrealBuildTool.cs
+++ b/Engine/Source/Programs/UnrealBuildTool/UnrealBuildTool.cs
@@ -57,7 +57,7 @@ namespace UnrealBuildTool
 		public static readonly FileReference UnrealBuildToolPath = FileReference.FindCorrectCase(new FileReference(Path.ChangeExtension(Assembly.GetExecutingAssembly().GetOriginalLocation(), null)));
 #endif
 		/// <summary>
-		/// The full name of the Root UE4 directory
+		/// The full name of the root UE directory
 		/// </summary>
 		public static readonly DirectoryReference RootDirectory = DirectoryReference.Combine(UnrealBuildToolPath.Directory, "..", "..", "..", "..");
 		/// <summary>
