UE5EA - Fix UE4 reference

diff --git a/Engine/Source/Programs/UnrealBuildTool/ProjectFiles/Xcode/XcodeProject.cs b/Engine/Source/Programs/UnrealBuildTool/ProjectFiles/Xcode/XcodeProject.cs
--- a/Engine/Source/Programs/UnrealBuildTool/ProjectFiles/Xcode/XcodeProject.cs
+++ b/Engine/Source/Programs/UnrealBuildTool/ProjectFiles/Xcode/XcodeProject.cs
@@ -640,8 +640,8 @@ namespace UnrealBuildTool
 
 		private void AppendLegacyTargetSection(StringBuilder Content, string TargetName, string TargetGuid, string TargetBuildConfigGuid, FileReference UProjectPath, bool bHasEditorConfiguration)
 		{
-			string UE4Dir = ConvertPath(Path.GetFullPath(Directory.GetCurrentDirectory() + "../../.."));
-			string BuildToolPath = UE4Dir + "/Engine/Build/BatchFiles/Mac/XcodeBuild.sh";
+			string UEDir = ConvertPath(Path.GetFullPath(Directory.GetCurrentDirectory() + "../../.."));
+			string BuildToolPath = UEDir + "/Engine/Build/BatchFiles/Mac/XcodeBuild.sh";
 
 			Content.Append("/* Begin PBXLegacyTarget section */" + ProjectFileGenerator.NewLine);
 
@@ -654,7 +654,7 @@ namespace UnrealBuildTool
 			Content.Append("\t\t\tbuildPhases = (" + ProjectFileGenerator.NewLine);
 			Content.Append("\t\t\t);" + ProjectFileGenerator.NewLine);
 			Content.Append("\t\t\tbuildToolPath = \"" + BuildToolPath + "\";" + ProjectFileGenerator.NewLine);
-			Content.Append("\t\t\tbuildWorkingDirectory = \"" + UE4Dir + "\";" + ProjectFileGenerator.NewLine);
+			Content.Append("\t\t\tbuildWorkingDirectory = \"" + UEDir + "\";" + ProjectFileGenerator.NewLine);
 			Content.Append("\t\t\tdependencies = (" + ProjectFileGenerator.NewLine);
 			Content.Append("\t\t\t);" + ProjectFileGenerator.NewLine);
 			Content.Append("\t\t\tname = \"" + TargetName + "\";" + ProjectFileGenerator.NewLine);
@@ -974,7 +974,7 @@ namespace UnrealBuildTool
 
 		private void AppendPlatformConfiguration(StringBuilder Content, FileReference MacExecutablePath, string TargetName, FileReference ProjectFile, bool bSupportMac, bool bSupportIOS, bool bSupportTVOS, out string IOSRunTimeVersion, out string TVOSRunTimeVersion, string BinariesSubDir = "/Payload")
 		{
-			string UE4Dir = ConvertPath(Path.GetFullPath(Directory.GetCurrentDirectory() + "../../.."));
+			string UEDir = ConvertPath(Path.GetFullPath(Directory.GetCurrentDirectory() + "../../.."));
 			string MacExecutableDir = bSupportMac ? ConvertPath(MacExecutablePath.Directory.FullName) : "";
 			string MacExecutableFileName = bSupportMac ? MacExecutablePath.GetFileName() : "";
 
@@ -1105,11 +1105,11 @@ namespace UnrealBuildTool
 				{
 					if (IOSRunTimeVersion != null)
 					{
-						Content.Append("\t\t\t\t\"CONFIGURATION_BUILD_DIR[sdk=iphoneos*]\" = \"" + UE4Dir + "/Engine/Binaries/IOS" + BinariesSubDir + "\";" + ProjectFileGenerator.NewLine);
+						Content.Append("\t\t\t\t\"CONFIGURATION_BUILD_DIR[sdk=iphoneos*]\" = \"" + UEDir + "/Engine/Binaries/IOS" + BinariesSubDir + "\";" + ProjectFileGenerator.NewLine);
 					}
 					if (TVOSRunTimeVersion != null)
 					{
-						Content.Append("\t\t\t\t\"CONFIGURATION_BUILD_DIR[sdk=appletvos*]\" = \"" + UE4Dir + "/Engine/Binaries/TVOS" + BinariesSubDir + "\";" + ProjectFileGenerator.NewLine);
+						Content.Append("\t\t\t\t\"CONFIGURATION_BUILD_DIR[sdk=appletvos*]\" = \"" + UEDir + "/Engine/Binaries/TVOS" + BinariesSubDir + "\";" + ProjectFileGenerator.NewLine);
 					}
 				}
 				else if (ProjectFile != null)
@@ -1127,11 +1127,11 @@ namespace UnrealBuildTool
 				{
 					if (IOSRunTimeVersion != null)
 					{
-						Content.Append("\t\t\t\t\"CONFIGURATION_BUILD_DIR[sdk=iphoneos*]\" = \"" + UE4Dir + "/Engine/Binaries/IOS" + BinariesSubDir + "\";" + ProjectFileGenerator.NewLine);
+						Content.Append("\t\t\t\t\"CONFIGURATION_BUILD_DIR[sdk=iphoneos*]\" = \"" + UEDir + "/Engine/Binaries/IOS" + BinariesSubDir + "\";" + ProjectFileGenerator.NewLine);
 					}
 					if (TVOSRunTimeVersion != null)
 					{
-						Content.Append("\t\t\t\t\"CONFIGURATION_BUILD_DIR[sdk=appletvos*]\" = \"" + UE4Dir + "/Engine/Binaries/TVOS" + BinariesSubDir + "\";" + ProjectFileGenerator.NewLine);
+						Content.Append("\t\t\t\t\"CONFIGURATION_BUILD_DIR[sdk=appletvos*]\" = \"" + UEDir + "/Engine/Binaries/TVOS" + BinariesSubDir + "\";" + ProjectFileGenerator.NewLine);
 					}
 				}
 
@@ -1153,7 +1153,7 @@ namespace UnrealBuildTool
 			Content.Append("\t\t\tisa = XCBuildConfiguration;" + ProjectFileGenerator.NewLine);
 			Content.Append("\t\t\tbuildSettings = {" + ProjectFileGenerator.NewLine);
 
-			string UE4Dir = ConvertPath(Path.GetFullPath(Directory.GetCurrentDirectory() + "../../.."));
+			string UEDir = ConvertPath(Path.GetFullPath(Directory.GetCurrentDirectory() + "../../.."));
 			//string MacExecutableDir = ConvertPath(Config.MacExecutablePath.Directory.FullName);
 			string MacExecutableFileName = Config.MacExecutablePath.GetFileName();
 
@@ -1175,17 +1175,17 @@ namespace UnrealBuildTool
 				string TVOSEntitlementPath = null;
 				if (bIsUnrealGame)
 				{
-					IOSInfoPlistPath = UE4Dir + "/Engine/Intermediate/IOS/" + Config.BuildTarget + "-Info.plist";
-					TVOSInfoPlistPath = UE4Dir + "/Engine/Intermediate/TVOS/" + Config.BuildTarget + "-Info.plist";
-					MacInfoPlistPath = UE4Dir + "/Engine/Intermediate/Mac/" + MacExecutableFileName + "-Info.plist";
+					IOSInfoPlistPath = UEDir + "/Engine/Intermediate/IOS/" + Config.BuildTarget + "-Info.plist";
+					TVOSInfoPlistPath = UEDir + "/Engine/Intermediate/TVOS/" + Config.BuildTarget + "-Info.plist";
+					MacInfoPlistPath = UEDir + "/Engine/Intermediate/Mac/" + MacExecutableFileName + "-Info.plist";
 					IOSEntitlementPath = "";
 					TVOSEntitlementPath = "";
 				}
 				else if (bIsUE4Client)
 				{
-					IOSInfoPlistPath = UE4Dir + "/Engine/Intermediate/IOS/UnrealGame-Info.plist";
-					TVOSInfoPlistPath = UE4Dir + "/Engine/Intermediate/TVOS/UnrealGame-Info.plist";
-					MacInfoPlistPath = UE4Dir + "/Engine/Intermediate/Mac/" + MacExecutableFileName + "-Info.plist";
+					IOSInfoPlistPath = UEDir + "/Engine/Intermediate/IOS/UnrealGame-Info.plist";
+					TVOSInfoPlistPath = UEDir + "/Engine/Intermediate/TVOS/UnrealGame-Info.plist";
+					MacInfoPlistPath = UEDir + "/Engine/Intermediate/Mac/" + MacExecutableFileName + "-Info.plist";
 					IOSEntitlementPath = "";
 					TVOSEntitlementPath = "";
 					}
@@ -1201,9 +1201,9 @@ namespace UnrealBuildTool
 				{
 					if (GamePath == null)
 					{
-						IOSInfoPlistPath = UE4Dir + "/Engine/Intermediate/IOS/" + Config.BuildTarget + "-Info.plist";
-						TVOSInfoPlistPath = UE4Dir + "/Engine/Intermediate/TVOS/" + Config.BuildTarget + "-Info.plist";
-						MacInfoPlistPath = UE4Dir + "/Engine/Intermediate/Mac/" + MacExecutableFileName + "-Info.plist";
+						IOSInfoPlistPath = UEDir + "/Engine/Intermediate/IOS/" + Config.BuildTarget + "-Info.plist";
+						TVOSInfoPlistPath = UEDir + "/Engine/Intermediate/TVOS/" + Config.BuildTarget + "-Info.plist";
+						MacInfoPlistPath = UEDir + "/Engine/Intermediate/Mac/" + MacExecutableFileName + "-Info.plist";
 					}
 					else
 					{
@@ -1248,7 +1248,7 @@ namespace UnrealBuildTool
 					if (bCreateMacInfoPlist || bCreateIOSInfoPlist || bCreateTVOSInfoPlist)
 					{
 						DirectoryReference ProjectPath = GameDir;
-						DirectoryReference EngineDir = DirectoryReference.Combine(new DirectoryReference(UE4Dir), "Engine");
+						DirectoryReference EngineDir = DirectoryReference.Combine(new DirectoryReference(UEDir), "Engine");
 						string GameName = Config.BuildTarget;
 						bool bIsClient = false;
 						if (ProjectPath == null)
@@ -1314,7 +1314,7 @@ namespace UnrealBuildTool
                 Content.Append("\t\t\t\tSDKROOT = macosx;" + ProjectFileGenerator.NewLine);
             }
             Content.Append("\t\t\t\tGCC_PRECOMPILE_PREFIX_HEADER = YES;" + ProjectFileGenerator.NewLine);
-			Content.Append("\t\t\t\tGCC_PREFIX_HEADER = \"" + UE4Dir + "/Engine/Source/Editor/UnrealEd/Public/UnrealEd.h\";" + ProjectFileGenerator.NewLine);
+			Content.Append("\t\t\t\tGCC_PREFIX_HEADER = \"" + UEDir + "/Engine/Source/Editor/UnrealEd/Public/UnrealEd.h\";" + ProjectFileGenerator.NewLine);
 			Content.Append("\t\t\t};" + ProjectFileGenerator.NewLine);
 			Content.Append("\t\t\tname = \"" + Config.DisplayName + "\";" + ProjectFileGenerator.NewLine);
 			Content.Append("\t\t};" + ProjectFileGenerator.NewLine);
