Merging CL 15807489 from //UE4/Main to Release-5.0-EarlyAccess

Use the LLVM 8 STL on Lumin

We've already switched to the LLVM 8 toolchain, but were still using the old STL headers.

diff --git a/Engine/Source/Programs/UnrealBuildTool/Platform/Lumin/UEBuildLumin.cs b/Engine/Source/Programs/UnrealBuildTool/Platform/Lumin/UEBuildLumin.cs
--- a/Engine/Source/Programs/UnrealBuildTool/Platform/Lumin/UEBuildLumin.cs
+++ b/Engine/Source/Programs/UnrealBuildTool/Platform/Lumin/UEBuildLumin.cs
@@ -178,11 +178,12 @@ namespace UnrealBuildTool
 			CompileEnvironment.Definitions.Add("WITH_OGGVORBIS=1");
 
 			DirectoryReference MLSDKDir = new DirectoryReference(Environment.GetEnvironmentVariable("MLSDK"));
-			CompileEnvironment.SystemIncludePaths.Add(DirectoryReference.Combine(MLSDKDir, "lumin/stl/libc++/include"));
+			CompileEnvironment.SystemIncludePaths.Add(DirectoryReference.Combine(MLSDKDir, "lumin/stl/libc++-8/include"));
 			CompileEnvironment.SystemIncludePaths.Add(DirectoryReference.Combine(MLSDKDir, "include"));
+			CompileEnvironment.SystemIncludePaths.Add(DirectoryReference.Combine(MLSDKDir, "tools/toolchains/llvm-8/sysroot/usr/include"));
 
 			LinkEnvironment.SystemLibraryPaths.Add(DirectoryReference.Combine(MLSDKDir, "lib/lumin"));
-			LinkEnvironment.SystemLibraryPaths.Add(DirectoryReference.Combine(MLSDKDir, "lumin/stl/libc++/lib"));
+			LinkEnvironment.SystemLibraryPaths.Add(DirectoryReference.Combine(MLSDKDir, "lumin/stl/libc++-8/lib"));
 
 			LinkEnvironment.SystemLibraries.Add("GLESv3");
 			LinkEnvironment.SystemLibraries.Add("EGL");
