Fix Xcode build deprecation warning

diff --git a/Engine/Source/Programs/UnrealBuildTool/ProjectFiles/Xcode/XcodeProjectFileGenerator.cs b/Engine/Source/Programs/UnrealBuildTool/ProjectFiles/Xcode/XcodeProjectFileGenerator.cs
--- a/Engine/Source/Programs/UnrealBuildTool/ProjectFiles/Xcode/XcodeProjectFileGenerator.cs
+++ b/Engine/Source/Programs/UnrealBuildTool/ProjectFiles/Xcode/XcodeProjectFileGenerator.cs
@@ -158,7 +158,19 @@ namespace UnrealBuildTool
 			return WriteFileIfChanged(Path, WorkspaceSettingsContent.ToString(), new UTF8Encoding());
 		}
 
-		
+		private bool WriteWorkspaceSharedSettingsFile(string Path)
+		{
+			StringBuilder WorkspaceSettingsContent = new StringBuilder();
+			WorkspaceSettingsContent.Append("<?xml version=\"1.0\" encoding=\"UTF-8\"?>" + ProjectFileGenerator.NewLine);
+			WorkspaceSettingsContent.Append("<!DOCTYPE plist PUBLIC \"-//Apple//DTD PLIST 1.0//EN\" \"http://www.apple.com/DTDs/PropertyList-1.0.dtd\">" + ProjectFileGenerator.NewLine);
+			WorkspaceSettingsContent.Append("<plist version=\"1.0\">" + ProjectFileGenerator.NewLine);
+			WorkspaceSettingsContent.Append("<dict>" + ProjectFileGenerator.NewLine);
+			WorkspaceSettingsContent.Append("\t<key>DisableBuildSystemDeprecationWarning</key>" + ProjectFileGenerator.NewLine);
+			WorkspaceSettingsContent.Append("\t<true/>" + ProjectFileGenerator.NewLine);
+			WorkspaceSettingsContent.Append("</dict>" + ProjectFileGenerator.NewLine);
+			WorkspaceSettingsContent.Append("</plist>" + ProjectFileGenerator.NewLine);
+			return WriteFileIfChanged(Path, WorkspaceSettingsContent.ToString(), new UTF8Encoding());
+		}
 
 		private bool WriteXcodeWorkspace()
 		{
@@ -249,8 +261,11 @@ namespace UnrealBuildTool
 			{
 				string WorkspaceSettingsFilePath = MasterProjectPath + "/" + ProjectName + ".xcworkspace/xcuserdata/" + Environment.UserName + ".xcuserdatad/WorkspaceSettings.xcsettings";
 				bSuccess = WriteWorkspaceSettingsFile(WorkspaceSettingsFilePath);
+				string WorkspaceSharedSettingsFilePath = MasterProjectPath + "/" + ProjectName + ".xcworkspace/xcshareddata/WorkspaceSettings.xcsettings";
+				bSuccess = WriteWorkspaceSharedSettingsFile(WorkspaceSharedSettingsFilePath);
 			}
 
+
 			return bSuccess;
 		}
 
