[UBT][Mac] Fix support for uproject pathnames containing spaces.

- Also, fixup Link phase to use `env' to avoid quotation discrepancies between .NET and shell environments.

#rnx

diff --git a/Engine/Build/BatchFiles/Mac/GenerateProjectFiles.sh b/Engine/Build/BatchFiles/Mac/GenerateProjectFiles.sh
--- a/Engine/Build/BatchFiles/Mac/GenerateProjectFiles.sh
+++ b/Engine/Build/BatchFiles/Mac/GenerateProjectFiles.sh
@@ -22,7 +22,7 @@ fi
 
 source "$BASE_PATH/SetupEnvironment.sh" -dotnet "$BASE_PATH"
 if [ -f "$BASE_PATH/../../../Source/Programs/UnrealBuildTool/UnrealBuildTool.csproj" ]; then
-  dotnet msbuild /restore /target:build /property:Configuration=Development /nologo $BASE_PATH/../../../Source/Programs/UnrealBuildTool/UnrealBuildTool.csproj /verbosity:quiet
+  dotnet msbuild /restore /target:build /property:Configuration=Development /nologo "$BASE_PATH/../../../Source/Programs/UnrealBuildTool/UnrealBuildTool.csproj" /verbosity:quiet
 
   if [ $? -ne 0 ]; then
     echo GenerateProjectFiles ERROR: Failed to build UnrealBuildTool
@@ -30,5 +30,16 @@ if [ -f "$BASE_PATH/../../../Source/Programs/UnrealBuildTool/UnrealBuildTool.csp
   fi
 fi
 
+if [ $# -lt 1 ]; then
+	# Create an empty string argument for UBT; otherwise we see an index out-of-range error.
+	QuotedArgs="\"\""
+else
+	# Apply quote-escape to all arguments
+	QuotedArgs="$(printf " %q" "$@")"
+
+	# Trim leading whitespace
+	QuotedArgs="${QuotedArgs#"${QuotedArgs%%[![:space:]]*}"}"
+fi
+
 # pass all parameters to UBT
-dotnet "$BASE_PATH/../../../Binaries/DotNET/UnrealBuildTool/UnrealBuildTool.dll" -projectfiles "$@"
+dotnet "$BASE_PATH/../../../Binaries/DotNET/UnrealBuildTool/UnrealBuildTool.dll" -projectfiles "${QuotedArgs}"
diff --git a/Engine/Source/Programs/UnrealBuildTool/Platform/Mac/MacToolChain.cs b/Engine/Source/Programs/UnrealBuildTool/Platform/Mac/MacToolChain.cs
--- a/Engine/Source/Programs/UnrealBuildTool/Platform/Mac/MacToolChain.cs
+++ b/Engine/Source/Programs/UnrealBuildTool/Platform/Mac/MacToolChain.cs
@@ -370,7 +370,7 @@ namespace UnrealBuildTool
 				Result += " -F \"" + Path.GetDirectoryName(Path.GetFullPath(FrameworkName)) + "\"";
 				FrameworkName = Path.GetFileNameWithoutExtension(FrameworkName);
 			}
-			Result += " " + Arg + " \"\"" + FrameworkName + "\"\"";
+			Result += " " + Arg + " \"" + FrameworkName + "\"";
 			return Result;
 		}
 
@@ -764,7 +764,7 @@ namespace UnrealBuildTool
 			Action LinkAction = Graph.CreateAction(ActionType.Link);
 
 			LinkAction.WorkingDirectory = GetMacDevSrcRoot();
-			LinkAction.CommandPath = BuildHostPlatform.Current.Shell;
+			LinkAction.CommandPath = new FileReference("/usr/bin/env");
 			LinkAction.CommandDescription = "Link";
 			LinkAction.CommandVersion = GetFullClangVersion();
 
@@ -839,7 +839,7 @@ namespace UnrealBuildTool
 					}
 					else  if (string.IsNullOrEmpty(Path.GetDirectoryName(AdditionalLibrary)) && string.IsNullOrEmpty(Path.GetExtension(AdditionalLibrary)))
 					{
-						LinkCommand += string.Format(" -l{0}", AdditionalLibrary);
+						LinkCommand += string.Format(" -l\"{0}\"", AdditionalLibrary);
 					}
 					else
 					{
@@ -922,7 +922,7 @@ namespace UnrealBuildTool
 				{
 					InstallName = string.Format("{0}/{1}", DylibsPath, Path.GetFileName(OutputFile.AbsolutePath));
 				}
-				LinkCommand += string.Format(" -install_name {0}", InstallName);
+				LinkCommand += string.Format(" -install_name \"{0}\"", InstallName);
 			}
 
 			if (!bIsBuildingLibrary)
@@ -965,7 +965,7 @@ namespace UnrealBuildTool
 			// Add the additional arguments specified by the environment.
 			LinkCommand += LinkEnvironment.AdditionalArguments;
 
-			LinkAction.CommandArguments = "-c '" + LinkCommand + "'";
+			LinkAction.CommandArguments = "-- " + LinkCommand;
 
 			// Only execute linking on the local Mac.
 			LinkAction.bCanExecuteRemotely = false;
diff --git a/GenerateProjectFiles.sh b/GenerateProjectFiles.sh
--- a/GenerateProjectFiles.sh
+++ b/GenerateProjectFiles.sh
@@ -7,7 +7,7 @@ cd "`dirname "$0"`"
 
 if [ ! -f Engine/Build/BatchFiles/Mac/GenerateProjectFiles.sh ]; then
 	echo "GenerateProjectFiles ERROR: This script does not appear to be located \
-       in the root UE4 directory and must be run from there."
+       in the root Unreal Engine directory and must be run from there."
   exit 1
 fi 
 
