UBT: Set a list of preferred ranges for compiler toolchains, rather than just specific versions. Avoids always falling back to the oldest available.

diff --git a/Engine/Source/Programs/UnrealBuildTool/Platform/Windows/UEBuildWindows.cs b/Engine/Source/Programs/UnrealBuildTool/Platform/Windows/UEBuildWindows.cs
--- a/Engine/Source/Programs/UnrealBuildTool/Platform/Windows/UEBuildWindows.cs
+++ b/Engine/Source/Programs/UnrealBuildTool/Platform/Windows/UEBuildWindows.cs
@@ -679,42 +679,77 @@ namespace UnrealBuildTool
 		class ToolChainInstallation
 		{
 			/// <summary>
-			/// Base directory for the toolchain
+			/// The version "family" (ie. the nominal version number of the directory this toolchain is installed to)
 			/// </summary>
-			public DirectoryReference BaseDir;
+			public VersionNumber Family { get; }
+
+			/// <summary>
+			/// Index into the preferred version range
+			/// </summary>
+			public int FamilyRank { get; }
+
+			/// <summary>
+			/// The actual version number of this toolchain
+			/// </summary>
+			public VersionNumber Version { get; }
+
+			/// <summary>
+			/// Whether this is a 64-bit toolchain
+			/// </summary>
+			public bool Is64Bit { get; }
 
 			/// <summary>
 			/// Whether it's a pre-release version of the toolchain.
 			/// </summary>
-			public bool bPreview;
+			public bool IsPreview { get; }
+
+			/// <summary>
+			/// Reason for this toolchain not being compatible
+			/// </summary>
+			public string? Error { get; }
+
+			/// <summary>
+			/// Base directory for the toolchain
+			/// </summary>
+			public DirectoryReference BaseDir { get; }
 
 			/// <summary>
 			/// Constructor
 			/// </summary>
+			/// <param name="Family"></param>
+			/// <param name="FamilyRank"></param>
+			/// <param name="Version"></param>
+			/// <param name="Is64Bit"></param>
+			/// <param name="IsPreview">Whether it's a pre-release version of the toolchain</param>
+			/// <param name="Error"></param>
 			/// <param name="BaseDir">Base directory for the toolchain</param>
-			/// <param name="bPreview">Whether it's a pre-release version of the toolchain</param>
-			public ToolChainInstallation(DirectoryReference BaseDir, bool bPreview)
-			{
+			public ToolChainInstallation(VersionNumber Family, int FamilyRank, VersionNumber Version, bool Is64Bit, bool IsPreview, string? Error, DirectoryReference BaseDir)
+			{
+				this.Family = Family;
+				this.FamilyRank = FamilyRank;
+				this.Version = Version;
+				this.Is64Bit = Is64Bit;
+				this.IsPreview = IsPreview;
+				this.Error = Error;
 				this.BaseDir = BaseDir;
-				this.bPreview = bPreview;
 			}
 		}
 
 		/// <summary>
 		/// The default compiler version to be used, if installed. 
 		/// </summary>
-		static readonly VersionNumber DefaultClangToolChainVersion = VersionNumber.Parse("10.0.0");
+		static readonly VersionNumberRange[] PreferredClangVersions =
+		{
+			VersionNumberRange.Parse("10.0.0", "10.999")
+		};
 
 		/// <summary>
-		/// The compiler toolchains to be used if installed, the first installed in the list will be used.
+		/// Ranges of tested compiler toolchains to be used, in order of preference. If multiple toolchains in a range are present, the latest version will be preferred.
+		/// Note that the numbers here correspond to the installation *folders* rather than precise executable versions. 
 		/// </summary>
-		static readonly VersionNumber[] PreferredVisualStudioToolChainVersion = new VersionNumber[]
+		static readonly VersionNumberRange[] PreferredVisualCppVersions = new VersionNumberRange[]
 		{
-			VersionNumber.Parse("14.28.29912"), // VS2019 v19.1 (installed to 14.28.29910 folder)
-			VersionNumber.Parse("14.24.28315"), // VS2019 v16.4.3 (installed to 14.24.28314 folder)
-//			VersionNumber.Parse("14.22.27905"), // VS2019 v16.2.3
-//			VersionNumber.Parse("14.16.27023.2"), // VS2017 v15.9.15
-//			VersionNumber.Parse("14.16.27023"), // fallback to VS2017 15.9 toolchain, microsoft updates these in places so for local installs only this version number is present
+			VersionNumberRange.Parse("14.24.27905", "14.24.29910") // VS2019 16.2.3 -> 19.1
 		};
 
 		/// <summary>
@@ -725,7 +760,7 @@ namespace UnrealBuildTool
 		/// <summary>
 		/// Cache of Visual C++ installation directories
 		/// </summary>
-		private static Dictionary<WindowsCompiler, Dictionary<VersionNumber, ToolChainInstallation>> CachedToolChainInstallations = new Dictionary<WindowsCompiler, Dictionary<VersionNumber, ToolChainInstallation>>();
+		private static Dictionary<WindowsCompiler, List<ToolChainInstallation>> CachedToolChainInstallations = new Dictionary<WindowsCompiler, List<ToolChainInstallation>>();
 
 		/// <summary>
 		/// Cache of DIA SDK installation directories
@@ -1105,12 +1140,12 @@ namespace UnrealBuildTool
 		/// </summary>
 		/// <param name="Compiler">Major version of the compiler to use</param>
 		/// <returns>Map of version number to directories</returns>
-		static Dictionary<VersionNumber, ToolChainInstallation> FindToolChainInstallations(WindowsCompiler Compiler)
+		static List<ToolChainInstallation> FindToolChainInstallations(WindowsCompiler Compiler)
 		{
-			Dictionary<VersionNumber, ToolChainInstallation>? ToolChainInstallations;
-			if(!CachedToolChainInstallations.TryGetValue(Compiler, out ToolChainInstallations))
+			List<ToolChainInstallation>? ToolChains;
+			if(!CachedToolChainInstallations.TryGetValue(Compiler, out ToolChains))
 			{
-				ToolChainInstallations = new Dictionary<VersionNumber, ToolChainInstallation>();
+				ToolChains = new List<ToolChainInstallation>();
 			    if (BuildHostPlatform.Current.Platform == UnrealTargetPlatform.Win64)
 			    {
 					if(Compiler == WindowsCompiler.Clang)
@@ -1124,7 +1159,7 @@ namespace UnrealBuildTool
 							{
 								FileVersionInfo VersionInfo = FileVersionInfo.GetVersionInfo(CompilerFile.FullName);
 								VersionNumber Version = new VersionNumber(VersionInfo.FileMajorPart, VersionInfo.FileMinorPart, VersionInfo.FileBuildPart);
-								ToolChainInstallations[Version] = new ToolChainInstallation(InstallDir, false);
+								AddClangToolChain(Version, InstallDir, ToolChains);
 							}
 						}
 
@@ -1140,7 +1175,7 @@ namespace UnrealBuildTool
 									VersionNumber? Version;
 									if(VersionNumber.TryParse(ToolChainDir.GetDirectoryName(), out Version) && IsValidToolChainDirClang(ToolChainDir))
 									{
-										ToolChainInstallations[Version] = new ToolChainInstallation(ToolChainDir, false);
+										AddClangToolChain(Version, ToolChainDir, ToolChains);
 									}
 								}
 							}
@@ -1158,7 +1193,7 @@ namespace UnrealBuildTool
 								{
 									FileVersionInfo VersionInfo = FileVersionInfo.GetVersionInfo(CompilerFile.FullName);
 									VersionNumber Version = new VersionNumber(VersionInfo.FileMajorPart, VersionInfo.FileMinorPart, VersionInfo.FileBuildPart);
-									ToolChainInstallations[Version] = new ToolChainInstallation(LLVMPathDir, false);
+									AddClangToolChain(Version, LLVMPathDir, ToolChains);
 								}
 							}
 						}
@@ -1174,7 +1209,7 @@ namespace UnrealBuildTool
 							{
 								FileVersionInfo VersionInfo = FileVersionInfo.GetVersionInfo(IclPath.FullName);
 								VersionNumber Version = new VersionNumber(VersionInfo.FileMajorPart, VersionInfo.FileMinorPart, VersionInfo.FileBuildPart);
-								ToolChainInstallations[Version] = new ToolChainInstallation(InstallDir, false);
+								ToolChains.Add(new ToolChainInstallation(Version, 0, Version, true, false, null, InstallDir));
 							}
 						}
 					}
@@ -1187,7 +1222,8 @@ namespace UnrealBuildTool
 							DirectoryReference ToolChainBaseDir = DirectoryReference.Combine(Installation.BaseDir, "VC");
 							if(IsValidToolChainDir2015(ToolChainBaseDir))
 							{
-								ToolChainInstallations[new VersionNumber(14, 0)] = new ToolChainInstallation(ToolChainBaseDir, Installation.bPreview);
+								VersionNumber Version = new VersionNumber(14, 0);
+								ToolChains.Add(new ToolChainInstallation(Version, 0, Version, Has64BitToolChain(ToolChainBaseDir), Installation.bPreview, null, ToolChainBaseDir));
 							}
 					    }
 				    }
@@ -1198,7 +1234,7 @@ namespace UnrealBuildTool
 					    foreach(VisualStudioInstallation Installation in Installations)
 					    {
 						    DirectoryReference ToolChainBaseDir = DirectoryReference.Combine(Installation.BaseDir, "VC", "Tools", "MSVC");
-							FindVisualStudioToolChains(ToolChainBaseDir, Installation.bPreview, ToolChainInstallations);
+							FindVisualStudioToolChains(ToolChainBaseDir, Installation.bPreview, ToolChains);
 					    }
 
 						// Enumerate all the AutoSDK toolchains
@@ -1206,10 +1242,10 @@ namespace UnrealBuildTool
 						if (UEBuildPlatformSDK.TryGetHostPlatformAutoSDKDir(out PlatformDir))
 						{
 							DirectoryReference ReleaseBaseDir = DirectoryReference.Combine(PlatformDir, "Win64", (Compiler == WindowsCompiler.VisualStudio2019) ? "VS2019" : "VS2017");
-							FindVisualStudioToolChains(ReleaseBaseDir, false, ToolChainInstallations);
+							FindVisualStudioToolChains(ReleaseBaseDir, false, ToolChains);
 
 							DirectoryReference PreviewBaseDir = DirectoryReference.Combine(PlatformDir, "Win64", (Compiler == WindowsCompiler.VisualStudio2019) ? "VS2019-Preview" : "VS2017-Preview");
-							FindVisualStudioToolChains(PreviewBaseDir, true, ToolChainInstallations);
+							FindVisualStudioToolChains(PreviewBaseDir, true, ToolChains);
 						}
 					}
 					else
@@ -1217,9 +1253,9 @@ namespace UnrealBuildTool
 					    throw new BuildException("Unsupported compiler version ({0})", Compiler);
 				    }
 				}
-				CachedToolChainInstallations.Add(Compiler, ToolChainInstallations);
+				CachedToolChainInstallations.Add(Compiler, ToolChains);
 			}
-			return ToolChainInstallations;
+			return ToolChains;
 		}
 
 		/// <summary>
@@ -1227,8 +1263,8 @@ namespace UnrealBuildTool
 		/// </summary>
 		/// <param name="BaseDir">Base directory to search</param>
 		/// <param name="bPreview">Whether this is a preview installation</param>
-		/// <param name="ToolChainInstallations">Map of tool chain version to installation info</param>
-		static void FindVisualStudioToolChains(DirectoryReference BaseDir, bool bPreview, Dictionary<VersionNumber, ToolChainInstallation> ToolChainInstallations)
+		/// <param name="ToolChains">Map of tool chain version to installation info</param>
+		static void FindVisualStudioToolChains(DirectoryReference BaseDir, bool bPreview, List<ToolChainInstallation> ToolChains)
 		{
 			if (DirectoryReference.Exists(BaseDir))
 			{
@@ -1237,27 +1273,54 @@ namespace UnrealBuildTool
 					VersionNumber? Version;
 					if (IsValidToolChainDir2017or2019(ToolChainDir, out Version))
 					{
-						
-						Log.TraceLog("Found Visual Studio toolchain: {0} (Version={1})", ToolChainDir, Version);
-						if (!ToolChainInstallations.ContainsKey(Version))
-						{
-							//Add VS installation based on cl.exe internal Product Version
-							ToolChainInstallations[Version] = new ToolChainInstallation(ToolChainDir, bPreview);
-						}
-
-						VersionNumber? VersionFolder;
-						if ( (VersionNumber.TryParse(ToolChainDir.GetDirectoryName(), out VersionFolder)) && (VersionFolder != Version) && (!ToolChainInstallations.ContainsKey(VersionFolder)))
-						{
-							//Add VS installation based on the version number in the installation path
-							Log.TraceLog("Found Visual Studio toolchain: {0} (Version={1})", ToolChainDir, VersionFolder);
-							ToolChainInstallations[VersionFolder] = new ToolChainInstallation(ToolChainDir, bPreview);
-						}
-
+						AddVisualCppToolChain(Version, bPreview, ToolChainDir, ToolChains);
 					}
 				}
 			}
 		}
 
+		/// <summary>
+		/// Add a Clang toolchain
+		/// </summary>
+		/// <param name="Version"></param>
+		/// <param name="ToolChainDir"></param>
+		/// <param name="ToolChains"></param>
+		static void AddClangToolChain(VersionNumber Version, DirectoryReference ToolChainDir, List<ToolChainInstallation> ToolChains)
+		{
+			int Rank = PreferredClangVersions.TakeWhile(x => !x.Contains(Version)).Count();
+			Log.TraceLog("Found Clang toolchain: {0} (Version={1}, Rank={2})", ToolChainDir, Version, Rank);
+			ToolChains.Add(new ToolChainInstallation(Version, Rank, Version, true, false, null, ToolChainDir));
+		}
+
+		/// <summary>
+		/// Adds a Visual C++ toolchain to a list of installations
+		/// </summary>
+		/// <param name="Version"></param>
+		/// <param name="bPreview"></param>
+		/// <param name="ToolChainDir"></param>
+		/// <param name="ToolChains"></param>
+		static void AddVisualCppToolChain(VersionNumber Version, bool bPreview, DirectoryReference ToolChainDir, List<ToolChainInstallation> ToolChains)
+		{
+			bool Is64Bit = Has64BitToolChain(ToolChainDir);
+
+			VersionNumber? Family;
+			if (!VersionNumber.TryParse(ToolChainDir.GetDirectoryName(), out Family))
+			{
+				Family = Version;
+			}
+
+			int FamilyRank = PreferredVisualCppVersions.TakeWhile(x => !x.Contains(Family)).Count();
+
+			string? Error = null;
+			if (Version >= new VersionNumber(14, 23, 0) && Version < new VersionNumber(14, 23, 28107))
+			{
+				Error = String.Format("The Visual C++ 14.23 toolchain is known to have code-generation issues with UE4. Please install a later toolchain from the Visual Studio installer. See here for more information: https://developercommunity.visualstudio.com/content/problem/734585/msvc-142328019-compilation-bug.html");
+			}
+
+			Log.TraceLog("Found Visual Studio toolchain: {0} (Family={1}, FamilyRank={2}, Version={3}, Is64Bit={4}, Preview={5}, Error={6})", ToolChainDir, Family, FamilyRank, Version, Is64Bit, bPreview, Error != null);
+			ToolChains.Add(new ToolChainInstallation(Family, FamilyRank, Version, Is64Bit, bPreview, Error, ToolChainDir));
+		}
+
 		/// <summary>
 		/// Checks if the given directory contains a valid Clang toolchain
 		/// </summary>
@@ -1368,6 +1431,29 @@ namespace UnrealBuildTool
 			}
 		}
 
+		/// <summary>
+		/// Select which toolchain to use, combining a custom preference with a default sort order
+		/// </summary>
+		/// <param name="ToolChains"></param>
+		/// <param name="Preference"></param>
+		/// <returns></returns>
+		static ToolChainInstallation? SelectToolChain(IEnumerable<ToolChainInstallation> ToolChains, Func<IOrderedEnumerable<ToolChainInstallation>, IOrderedEnumerable<ToolChainInstallation>> Preference)
+		{
+			ToolChainInstallation? ToolChain = Preference(ToolChains.OrderByDescending(x => x.Error == null))
+				.ThenByDescending(x => x.Is64Bit)
+				.ThenBy(x => x.IsPreview)
+				.ThenBy(x => x.FamilyRank)
+				.ThenByDescending(x => x.Version)
+				.FirstOrDefault();
+
+			if (ToolChain?.Error != null)
+			{
+				throw new BuildException(ToolChain.Error);
+			}
+
+			return ToolChain;
+		}
+
 		/// <summary>
 		/// Determines the directory containing the MSVC toolchain
 		/// </summary>
@@ -1379,83 +1465,53 @@ namespace UnrealBuildTool
 		public static bool TryGetToolChainDir(WindowsCompiler Compiler, string? CompilerVersion, [NotNullWhen(true)] out VersionNumber? OutToolChainVersion, [NotNullWhen(true)] out DirectoryReference? OutToolChainDir)
 		{
 			// Find all the installed toolchains
-			Dictionary<VersionNumber, ToolChainInstallation> ToolChainInstallations = FindToolChainInstallations(Compiler);
+			List<ToolChainInstallation> ToolChains = FindToolChainInstallations(Compiler);
 
 			// Figure out the actual version number that we want
-			VersionNumber? ToolChainVersion = null;
-			if(CompilerVersion != null)
+			ToolChainInstallation? ToolChain = null;
+			if (CompilerVersion == null)
 			{
-				if(String.Compare(CompilerVersion, "Latest", StringComparison.InvariantCultureIgnoreCase) == 0 && ToolChainInstallations.Count > 0)
-				{
-					ToolChainVersion = ToolChainInstallations.OrderBy(x => IsCompatibleVisualCppToolChain(x.Key)).ThenBy(x => Has64BitToolChain(x.Value.BaseDir)).ThenBy(x => !x.Value.bPreview).ThenBy(x => x.Key).Last().Key;
-					if (ToolChainVersion == null)
-					{
-						throw new BuildException("Unable to find C++ toolchain for {0}", Compiler);
-					}
-				}
-				else if(String.Compare(CompilerVersion, "Preview", StringComparison.InvariantCultureIgnoreCase) == 0)
+				ToolChain = SelectToolChain(ToolChains, x => x);
+				if (ToolChain == null)
 				{
-					ToolChainVersion = ToolChainInstallations.Where(x => x.Value.bPreview).OrderBy(x => IsCompatibleVisualCppToolChain(x.Key)).ThenBy(x => Has64BitToolChain(x.Value.BaseDir)).ThenBy(x => x.Key).Select(x => x.Key).LastOrDefault();
-					if (ToolChainVersion == null)
-					{
-						throw new BuildException("Unable to find preview toolchain for {0}", Compiler);
-					}
-				}
-				else if(!VersionNumber.TryParse(CompilerVersion, out ToolChainVersion))
-				{
-					throw new BuildException("Unable to find {0} toolchain; '{1}' is an invalid version", GetCompilerName(Compiler), CompilerVersion);
+					OutToolChainVersion = null;
+					OutToolChainDir = null;
+					return false;
 				}
 			}
-			else
+			else if (String.Compare(CompilerVersion, "Latest", StringComparison.InvariantCultureIgnoreCase) == 0)
 			{
-				VersionNumber[] PotentialToolchains;
-				if(Compiler == WindowsCompiler.Clang)
+				ToolChain = SelectToolChain(ToolChains, x => x.ThenByDescending(x => x.Version));
+				if (ToolChain == null)
 				{
-					PotentialToolchains = new VersionNumber[] { DefaultClangToolChainVersion};
+					throw new BuildException("Unable to find C++ toolchain for {0}", Compiler);
 				}
-				else
+			}
+			else if (String.Compare(CompilerVersion, "Preview", StringComparison.InvariantCultureIgnoreCase) == 0)
+			{
+				ToolChain = SelectToolChain(ToolChains, x => x.ThenByDescending(x => x.IsPreview));
+				if (ToolChain == null || !ToolChain.IsPreview)
 				{
-					PotentialToolchains = PreferredVisualStudioToolChainVersion;
+					throw new BuildException("Unable to find preview toolchain for {0}", Compiler);
 				}
-
-				foreach (VersionNumber PreferredToolchain in PotentialToolchains)
-				{
-					ToolChainInstallation? DefaultToolChain;
-					if (ToolChainInstallations.TryGetValue(PreferredToolchain, out DefaultToolChain) && (Compiler == WindowsCompiler.Clang || Has64BitToolChain(DefaultToolChain.BaseDir)))
-					{
-						ToolChainVersion = PreferredToolchain;
-						break;
-					}
-				}
-
-				// if we failed to find any of our preferred toolchains we pick the newest (highest version number)
-				if (ToolChainVersion == null && ToolChainInstallations.Count > 0)
+			}
+			else if (VersionNumber.TryParse(CompilerVersion, out VersionNumber? ToolChainVersion))
+			{
+				ToolChain = SelectToolChain(ToolChains, x => x.ThenByDescending(x => x.Version == ToolChainVersion).ThenByDescending(x => x.Family == ToolChainVersion));
+				if (ToolChain == null)
 				{
-					ToolChainVersion = ToolChainInstallations.OrderBy(x => IsCompatibleVisualCppToolChain(x.Key)).ThenBy(x => Has64BitToolChain(x.Value.BaseDir)).ThenBy(x => !x.Value.bPreview).ThenBy(x => x.Key).Last().Key;
+					throw new BuildException("Unable to find {0} toolchain for {1}", ToolChainVersion, Compiler);
 				}
 			}
-
-			// Check it's valid
-			string? Message;
-			if (ToolChainVersion != null && !IsCompatibleVisualCppToolChain(ToolChainVersion, out Message))
+			else
 			{
-				throw new BuildException(Message);
+				throw new BuildException("Unable to find {0} toolchain; '{1}' is an invalid version", GetCompilerName(Compiler), CompilerVersion);
 			}
 
 			// Get the actual directory for this version
-			ToolChainInstallation? Installation;
-			if (ToolChainVersion != null && ToolChainInstallations.TryGetValue(ToolChainVersion, out Installation))
-			{
-				OutToolChainVersion = ToolChainVersion;
-				OutToolChainDir = Installation.BaseDir;
-				return true;
-			}
-			else
-			{
-				OutToolChainVersion = null;
-				OutToolChainDir = null;
-				return false;
-			}
+			OutToolChainVersion = ToolChain.Version;
+			OutToolChainDir = ToolChain.BaseDir;
+			return true;
 		}
 
 		/// <summary>
diff --git a/Engine/Source/Programs/UnrealBuildTool/System/VersionNumberRange.cs b/Engine/Source/Programs/UnrealBuildTool/System/VersionNumberRange.cs
new file mode 100644
--- /dev/null
+++ b/Engine/Source/Programs/UnrealBuildTool/System/VersionNumberRange.cs
@@ -0,0 +1,62 @@
+// Copyright Epic Games, Inc. All Rights Reserved.
+
+using System;
+using System.Collections.Generic;
+using System.Text;
+
+namespace UnrealBuildTool
+{
+	/// <summary>
+	/// Range of version numbers
+	/// </summary>
+	class VersionNumberRange
+	{
+		/// <summary>
+		/// Minimum version number
+		/// </summary>
+		VersionNumber Min { get; }
+
+		/// <summary>
+		/// Maximum version number
+		/// </summary>
+		VersionNumber Max { get; }
+
+		/// <summary>
+		/// Constructor
+		/// </summary>
+		/// <param name="Min"></param>
+		/// <param name="Max"></param>
+		public VersionNumberRange(VersionNumber Min, VersionNumber Max)
+		{
+			this.Min = Min;
+			this.Max = Max;
+		}
+
+		/// <summary>
+		/// Tests whether this range contains the given version
+		/// </summary>
+		/// <param name="Version"></param>
+		/// <returns></returns>
+		public bool Contains(VersionNumber Version)
+		{
+			return Version >= Min && Version <= Max;
+		}
+
+		/// <summary>
+		/// Parse a version range from two strings
+		/// </summary>
+		/// <param name="MinText"></param>
+		/// <param name="MaxText"></param>
+		/// <returns></returns>
+		public static VersionNumberRange Parse(string MinText, string MaxText)
+		{
+			return new VersionNumberRange(VersionNumber.Parse(MinText), VersionNumber.Parse(MaxText));
+		}
+
+		/// <inheritdoc/>
+		public override string ToString()
+		{
+			return String.Format("{0}-{1}", Min, Max);
+		}
+	}
+}
