Enumerate Clang installations bundled with Visual Studio.

diff --git a/Engine/Source/Programs/UnrealBuildTool/Platform/Windows/UEBuildWindows.cs b/Engine/Source/Programs/UnrealBuildTool/Platform/Windows/UEBuildWindows.cs
--- a/Engine/Source/Programs/UnrealBuildTool/Platform/Windows/UEBuildWindows.cs
+++ b/Engine/Source/Programs/UnrealBuildTool/Platform/Windows/UEBuildWindows.cs
@@ -117,6 +117,7 @@ namespace UnrealBuildTool
 		[CommandLine("-2015", Value = "VisualStudio2015")]
 		[CommandLine("-2017", Value = "VisualStudio2017")]
 		[CommandLine("-2019", Value = "VisualStudio2019")]
+		[CommandLine("-Compiler=")]
 		public WindowsCompiler Compiler = WindowsCompiler.Default;
 
 		/// <summary>
@@ -1150,17 +1151,21 @@ namespace UnrealBuildTool
 			    {
 					if(Compiler == WindowsCompiler.Clang)
 					{
-						// Check for a manual installation
-						DirectoryReference InstallDir = DirectoryReference.Combine(DirectoryReference.GetSpecialFolder(Environment.SpecialFolder.ProgramFiles), "LLVM");
-						if(IsValidToolChainDirClang(InstallDir))
+						// Check for a manual installation to the default directory
+						DirectoryReference ManualInstallDir = DirectoryReference.Combine(DirectoryReference.GetSpecialFolder(Environment.SpecialFolder.ProgramFiles), "LLVM");
+						AddClangToolChain(ManualInstallDir, ToolChains);
+
+						// Check for a manual installation to a custom directory
+						string? LlvmPath = Environment.GetEnvironmentVariable("LLVM_PATH");
+						if (!String.IsNullOrEmpty(LlvmPath))
 						{
-							FileReference CompilerFile = FileReference.Combine(InstallDir, "bin", "clang-cl.exe");
-							if(FileReference.Exists(CompilerFile))
-							{
-								FileVersionInfo VersionInfo = FileVersionInfo.GetVersionInfo(CompilerFile.FullName);
-								VersionNumber Version = new VersionNumber(VersionInfo.FileMajorPart, VersionInfo.FileMinorPart, VersionInfo.FileBuildPart);
-								AddClangToolChain(Version, InstallDir, ToolChains);
-							}
+							AddClangToolChain(new DirectoryReference(LlvmPath), ToolChains);
+						}
+
+						// Check for installations bundled with Visual Studio 2019
+						foreach (VisualStudioInstallation Installation in FindVisualStudioInstallations(WindowsCompiler.VisualStudio2019))
+						{
+							AddClangToolChain(DirectoryReference.Combine(Installation.BaseDir, "VC", "Tools", "Llvm"), ToolChains);
 						}
 
 						// Check for AutoSDK paths
@@ -1172,28 +1177,7 @@ namespace UnrealBuildTool
 							{
 								foreach(DirectoryReference ToolChainDir in DirectoryReference.EnumerateDirectories(ClangBaseDir))
 								{
-									VersionNumber? Version;
-									if(VersionNumber.TryParse(ToolChainDir.GetDirectoryName(), out Version) && IsValidToolChainDirClang(ToolChainDir))
-									{
-										AddClangToolChain(Version, ToolChainDir, ToolChains);
-									}
-								}
-							}
-						}
-
-						// Check for LLVM_PATH environment variable.
-						string? LLVMPath = Environment.GetEnvironmentVariable("LLVM_PATH");
-						if (!String.IsNullOrEmpty(LLVMPath))
-						{
-							DirectoryReference LLVMPathDir = new DirectoryReference(LLVMPath);
-							if (IsValidToolChainDirClang(LLVMPathDir))
-							{
-								FileReference CompilerFile = FileReference.Combine(LLVMPathDir, "bin", "clang-cl.exe");
-								if (FileReference.Exists(CompilerFile))
-								{
-									FileVersionInfo VersionInfo = FileVersionInfo.GetVersionInfo(CompilerFile.FullName);
-									VersionNumber Version = new VersionNumber(VersionInfo.FileMajorPart, VersionInfo.FileMinorPart, VersionInfo.FileBuildPart);
-									AddClangToolChain(Version, LLVMPathDir, ToolChains);
+									AddClangToolChain(ToolChainDir, ToolChains);
 								}
 							}
 						}
@@ -1282,14 +1266,20 @@ namespace UnrealBuildTool
 		/// <summary>
 		/// Add a Clang toolchain
 		/// </summary>
-		/// <param name="Version"></param>
 		/// <param name="ToolChainDir"></param>
 		/// <param name="ToolChains"></param>
-		static void AddClangToolChain(VersionNumber Version, DirectoryReference ToolChainDir, List<ToolChainInstallation> ToolChains)
+		static void AddClangToolChain(DirectoryReference ToolChainDir, List<ToolChainInstallation> ToolChains)
 		{
-			int Rank = PreferredClangVersions.TakeWhile(x => !x.Contains(Version)).Count();
-			Log.TraceLog("Found Clang toolchain: {0} (Version={1}, Rank={2})", ToolChainDir, Version, Rank);
-			ToolChains.Add(new ToolChainInstallation(Version, Rank, Version, true, false, null, ToolChainDir));
+			FileReference CompilerFile = FileReference.Combine(ToolChainDir, "bin", "clang-cl.exe");
+			if (FileReference.Exists(CompilerFile))
+			{
+				FileVersionInfo VersionInfo = FileVersionInfo.GetVersionInfo(CompilerFile.FullName);
+				VersionNumber Version = new VersionNumber(VersionInfo.FileMajorPart, VersionInfo.FileMinorPart, VersionInfo.FileBuildPart);
+
+				int Rank = PreferredClangVersions.TakeWhile(x => !x.Contains(Version)).Count();
+				Log.TraceLog("Found Clang toolchain: {0} (Version={1}, Rank={2})", ToolChainDir, Version, Rank);
+				ToolChains.Add(new ToolChainInstallation(Version, Rank, Version, true, false, null, ToolChainDir));
+			}
 		}
 
 		/// <summary>
