UBT: Add support for compiling a list of files passed via the -FileList=... argument.

diff --git a/Engine/Source/Programs/UnrealBuildTool/Configuration/TargetDescriptor.cs b/Engine/Source/Programs/UnrealBuildTool/Configuration/TargetDescriptor.cs
--- a/Engine/Source/Programs/UnrealBuildTool/Configuration/TargetDescriptor.cs
+++ b/Engine/Source/Programs/UnrealBuildTool/Configuration/TargetDescriptor.cs
@@ -34,9 +34,16 @@ namespace UnrealBuildTool
 		[CommandLine("-Module=")]
 		public HashSet<string> OnlyModuleNames = new HashSet<string>(StringComparer.OrdinalIgnoreCase);
 
+		/// <summary>
+		/// Lists of files to compile
+		/// </summary>
+		[CommandLine("-FileList=")]
+		public List<FileReference> FileLists = new List<FileReference>(); 
+
 		/// <summary>
 		/// Individual file(s) to compile
 		/// </summary>
+		[CommandLine("-File=")]
 		[CommandLine("-SingleFile=")]
 		public List<FileReference> SpecificFilesToCompile = new List<FileReference>();
 
@@ -101,6 +108,19 @@ namespace UnrealBuildTool
 				// Apply the arguments to this object
 				Arguments.ApplyTo(this);
 
+				// Read the file lists
+				foreach (FileReference FileList in FileLists)
+				{
+					string[] Files = FileReference.ReadAllLines(FileList);
+					foreach (string File in Files)
+					{
+						if (!String.IsNullOrWhiteSpace(File))
+						{
+							SpecificFilesToCompile.Add(new FileReference(File));
+						}
+					}
+				}
+
 				// Parse all the hot-reload module names
 				foreach(string ModuleWithSuffix in Arguments.GetValues("-ModuleWithSuffix="))
 				{
@@ -337,7 +357,7 @@ namespace UnrealBuildTool
 
 							if (ProjectFile == null)
 							{
-								throw new BuildException("-TargetType=... requires a project file to be specified");
+								TargetNames.Add(RulesCompiler.CreateEngineRulesAssembly(bUsePrecompiled, bSkipRulesCompile).GetTargetNameByType(TargetType, Platform, Configuration, Architecture, null));
 							}
 							else
 							{
@@ -403,27 +423,6 @@ namespace UnrealBuildTool
 			return false;
 		}
 
-		/// <summary>
-		/// Parse a single argument value, of the form -Foo=Bar
-		/// </summary>
-		/// <param name="Argument">The argument to parse</param>
-		/// <param name="Prefix">The argument prefix, eg. "-Foo="</param>
-		/// <param name="Value">Receives the value of the argument</param>
-		/// <returns>True if the argument could be parsed, false otherwise</returns>
-		private static bool ParseArgumentValue(string Argument, string Prefix, [NotNullWhen(true)] out string? Value)
-		{
-			if(Argument.StartsWith(Prefix, StringComparison.InvariantCultureIgnoreCase))
-			{
-				Value = Argument.Substring(Prefix.Length);
-				return true;
-			}
-			else
-			{
-				Value = null;
-				return false;
-			}
-		}
-
 		/// <summary>
 		/// Format this object for the debugger
 		/// </summary>
diff --git a/Engine/Source/Programs/UnrealBuildTool/System/RulesAssembly.cs b/Engine/Source/Programs/UnrealBuildTool/System/RulesAssembly.cs
--- a/Engine/Source/Programs/UnrealBuildTool/System/RulesAssembly.cs
+++ b/Engine/Source/Programs/UnrealBuildTool/System/RulesAssembly.cs
@@ -673,7 +673,7 @@ namespace UnrealBuildTool
 		/// <param name="Architecture">The architecture being built</param>
 		/// <param name="ProjectFile">Project file for the target being built</param>
 		/// <returns>Name of the target for the given type</returns>
-		public string GetTargetNameByType(TargetType Type, UnrealTargetPlatform Platform, UnrealTargetConfiguration Configuration, string Architecture, FileReference ProjectFile)
+		public string GetTargetNameByType(TargetType Type, UnrealTargetPlatform Platform, UnrealTargetConfiguration Configuration, string Architecture, FileReference? ProjectFile)
 		{
 			// Create all the targets in this assembly 
 			List<string> Matches = new List<string>();
