Prefer 64-bit Clang toolchains.

diff --git a/Engine/Source/Programs/UnrealBuildTool/Platform/Windows/UEBuildWindows.cs b/Engine/Source/Programs/UnrealBuildTool/Platform/Windows/UEBuildWindows.cs
--- a/Engine/Source/Programs/UnrealBuildTool/Platform/Windows/UEBuildWindows.cs
+++ b/Engine/Source/Programs/UnrealBuildTool/Platform/Windows/UEBuildWindows.cs
@@ -11,6 +11,7 @@ using EpicGames.Core;
 using Microsoft.VisualStudio.Setup.Configuration;
 using System.Runtime.InteropServices;
 using System.Diagnostics.CodeAnalysis;
+using System.Buffers.Binary;
 
 namespace UnrealBuildTool
 {
@@ -1277,8 +1278,51 @@ namespace UnrealBuildTool
 				VersionNumber Version = new VersionNumber(VersionInfo.FileMajorPart, VersionInfo.FileMinorPart, VersionInfo.FileBuildPart);
 
 				int Rank = PreferredClangVersions.TakeWhile(x => !x.Contains(Version)).Count();
-				Log.TraceLog("Found Clang toolchain: {0} (Version={1}, Rank={2})", ToolChainDir, Version, Rank);
-				ToolChains.Add(new ToolChainInstallation(Version, Rank, Version, true, false, null, ToolChainDir));
+				bool Is64Bit = Is64BitExecutable(CompilerFile);
+				Log.TraceLog("Found Clang toolchain: {0} (Version={1}, Is64Bit={2}, Rank={3})", ToolChainDir, Version, Is64Bit, Rank);
+				ToolChains.Add(new ToolChainInstallation(Version, Rank, Version, Is64Bit, false, null, ToolChainDir));
+			}
+		}
+
+		/// <summary>
+		/// Test whether an executable is 64-bit
+		/// </summary>
+		/// <param name="File">Executable to test</param>
+		/// <returns></returns>
+		static bool Is64BitExecutable(FileReference File)
+		{
+			using (FileStream Stream = new FileStream(File.FullName, FileMode.Open, FileAccess.Read, FileShare.Read | FileShare.Delete))
+			{
+				byte[] Header = new byte[64];
+				if (Stream.Read(Header, 0, Header.Length) != Header.Length)
+				{
+					return false;
+				}
+				if (Header[0] != (byte)'M' || Header[1] != (byte)'Z')
+				{
+					return false;
+				}
+
+				int Offset = BinaryPrimitives.ReadInt32LittleEndian(Header.AsSpan(0x3c));
+				if (Stream.Seek(Offset, SeekOrigin.Begin) != Offset)
+				{
+					return false;
+				}
+
+				byte[] PeHeader = new byte[6];
+				if(Stream.Read(PeHeader, 0, PeHeader.Length) != PeHeader.Length)
+				{
+					return false;
+				}
+				if (BinaryPrimitives.ReadUInt32BigEndian(PeHeader.AsSpan()) != 0x50450000)
+				{
+					return false;
+				}
+
+				ushort MachineType = BinaryPrimitives.ReadUInt16LittleEndian(PeHeader.AsSpan(4));
+
+				const ushort IMAGE_FILE_MACHINE_AMD64 = 0x8664;
+				return MachineType == IMAGE_FILE_MACHINE_AMD64;
 			}
 		}
 
