UBT: Fix missing response file dependency when determining whether object files are up to date.

diff --git a/Engine/Source/Programs/UnrealBuildTool/Platform/Windows/VCCompileAction.cs b/Engine/Source/Programs/UnrealBuildTool/Platform/Windows/VCCompileAction.cs
--- a/Engine/Source/Programs/UnrealBuildTool/Platform/Windows/VCCompileAction.cs
+++ b/Engine/Source/Programs/UnrealBuildTool/Platform/Windows/VCCompileAction.cs
@@ -164,6 +164,10 @@ namespace UnrealBuildTool
 		{
 			get
 			{
+				if (ResponseFile != null)
+				{
+					yield return ResponseFile;
+				}
 				if (SourceFile != null)
 				{
 					yield return SourceFile;
