Sync clang's compatibliity flag with the actual MSVC toolchain version.

#rnx

diff --git a/Engine/Source/Programs/UnrealBuildTool/Platform/Windows/VCToolChain.cs b/Engine/Source/Programs/UnrealBuildTool/Platform/Windows/VCToolChain.cs
--- a/Engine/Source/Programs/UnrealBuildTool/Platform/Windows/VCToolChain.cs
+++ b/Engine/Source/Programs/UnrealBuildTool/Platform/Windows/VCToolChain.cs
@@ -166,9 +166,10 @@ namespace UnrealBuildTool
 			// This will show the cl.exe options that map to clang.exe ones, which ones are ignored and which ones are unsupported.
 			if (Target.WindowsPlatform.Compiler == WindowsCompiler.Clang)
 			{
-				// FIXME: We should be compatible with the latest MSVC compiler toolkit installed instead of hard-coding it here.
-				Arguments.Add("-fms-compatibility-version=19.24");
-
+				// Sync the compatibility version with the MSVC toolchain version (14.xx which maps to advertised
+				// compiler version of 19.xx).
+				Arguments.Add($"-fms-compatibility-version=19.{EnvVars.ToolChainVersion.GetComponent(1)}");
+				
 				if (Target.WindowsPlatform.StaticAnalyzer == WindowsStaticAnalyzer.Default)
 				{
 					// Enable the static analyzer but only via the backend. Using the frontend
