Enable C++17 support across the board in UE5.

diff --git a/Engine/Source/Programs/UnrealBuildTool/System/CppCompileEnvironment.cs b/Engine/Source/Programs/UnrealBuildTool/System/CppCompileEnvironment.cs
--- a/Engine/Source/Programs/UnrealBuildTool/System/CppCompileEnvironment.cs
+++ b/Engine/Source/Programs/UnrealBuildTool/System/CppCompileEnvironment.cs
@@ -50,7 +50,7 @@ namespace UnrealBuildTool
 		/// <summary>
 		/// Use the default standard version
 		/// </summary>
-		Default = Cpp14,
+		Default = Cpp17,
 	}
 
 	/// <summary>
