Disable pragma-pack warnings globally with Clang on Windows. CL 1586136 fixed unbracketed warning push/pop for clang which uncovered this warning.

# See discussion on https://stackoverflow.com/questions/52865610/suppressing-warning-about-pragma-pack-in-included-file

#rnx

diff --git a/Engine/Source/Programs/UnrealBuildTool/Platform/Windows/VCToolChain.cs b/Engine/Source/Programs/UnrealBuildTool/Platform/Windows/VCToolChain.cs
--- a/Engine/Source/Programs/UnrealBuildTool/Platform/Windows/VCToolChain.cs
+++ b/Engine/Source/Programs/UnrealBuildTool/Platform/Windows/VCToolChain.cs
@@ -719,6 +719,10 @@ namespace UnrealBuildTool
 				{
 					Arguments.Add(" -Wundef" + (CompileEnvironment.bUndefinedIdentifierWarningsAsErrors ? "" : " -Wno-error=undef"));
 				}
+				
+				// This is disabled because clang explicitly warns about changing pack alignment in a header and not
+				// restoring it afterwards, which is something we do with the Pre/PostWindowsApi.h headers.
+				Arguments.Add("-Wno-pragma-pack");
 
 				// @todo clang: Kind of a shame to turn these off.  We'd like to catch unused variables, but it is tricky with how our assertion macros work.
 				Arguments.Add("-Wno-inconsistent-missing-override");
diff --git a/Engine/Source/Runtime/Core/Public/Windows/PostWindowsApi.h b/Engine/Source/Runtime/Core/Public/Windows/PostWindowsApi.h
--- a/Engine/Source/Runtime/Core/Public/Windows/PostWindowsApi.h
+++ b/Engine/Source/Runtime/Core/Public/Windows/PostWindowsApi.h
@@ -79,11 +79,6 @@ THIRD_PARTY_INCLUDES_END
 // Restore the struct packing setting
 PRAGMA_POP_PLATFORM_DEFAULT_PACKING
 
-// Restore the warning that the pack size is changed in this header.
-#ifdef __clang__
-	#pragma clang diagnostic pop
-#endif	// __clang__
-
 // Redefine CDECL to our version of the #define.  <AJS> Is this really necessary?
 #define CDECL	    __cdecl					/* Standard C function */
 
diff --git a/Engine/Source/Runtime/Core/Public/Windows/PreWindowsApi.h b/Engine/Source/Runtime/Core/Public/Windows/PreWindowsApi.h
--- a/Engine/Source/Runtime/Core/Public/Windows/PreWindowsApi.h
+++ b/Engine/Source/Runtime/Core/Public/Windows/PreWindowsApi.h
@@ -1,10 +1,9 @@
 // Copyright Epic Games, Inc. All Rights Reserved.
 
-// Disable the warning that the pack size is changed in this header.
-#ifdef __clang__
-	#pragma clang diagnostic push
-	#pragma clang diagnostic ignored "-Wpragma-pack"
-#else	// __clang__
+// Disable the warning that the pack size is changed in this header. We do this globally for
+// clang, since the mechanism of using a sub-header to modify packing generates a -Wpragma-pack
+// warning about modifying packing alignemnt in a header.
+#if !defined(__clang__)
 	#pragma warning(disable:4103)
 #endif	// __clang__
 
diff --git a/Engine/Source/Runtime/Core/Public/Windows/WindowsPlatformCompilerPreSetup.h b/Engine/Source/Runtime/Core/Public/Windows/WindowsPlatformCompilerPreSetup.h
--- a/Engine/Source/Runtime/Core/Public/Windows/WindowsPlatformCompilerPreSetup.h
+++ b/Engine/Source/Runtime/Core/Public/Windows/WindowsPlatformCompilerPreSetup.h
@@ -161,14 +161,12 @@
 	#ifdef _WIN64
 		#ifndef PRAGMA_PUSH_PLATFORM_DEFAULT_PACKING
 			#define PRAGMA_PUSH_PLATFORM_DEFAULT_PACKING \
-				__pragma(pack(push)) \
-				__pragma(pack(16))
+				__pragma(pack(push, 16))
 		#endif // PRAGMA_PUSH_PLATFORM_DEFAULT_PACKING
 	#else // _WIN64
 		#ifndef PRAGMA_PUSH_PLATFORM_DEFAULT_PACKING
 			#define PRAGMA_PUSH_PLATFORM_DEFAULT_PACKING \
-				__pragma(pack(push)) \
-				__pragma(pack(8))
+				__pragma(pack(push, 8))
 		#endif // PRAGMA_PUSH_PLATFORM_DEFAULT_PACKING
 	#endif // _WIN64
 	
