Fix incorrect preferred max version for VC toolchain and add Cpp20 standard as supported.

#rnx

diff --git a/Engine/Source/Programs/UnrealBuildTool/Platform/Windows/UEBuildWindows.cs b/Engine/Source/Programs/UnrealBuildTool/Platform/Windows/UEBuildWindows.cs
--- a/Engine/Source/Programs/UnrealBuildTool/Platform/Windows/UEBuildWindows.cs
+++ b/Engine/Source/Programs/UnrealBuildTool/Platform/Windows/UEBuildWindows.cs
@@ -751,7 +751,7 @@ namespace UnrealBuildTool
 		/// </summary>
 		static readonly VersionNumberRange[] PreferredVisualCppVersions = new VersionNumberRange[]
 		{
-			VersionNumberRange.Parse("14.24.27905", "14.24.29910") // VS2019 16.2.3 -> 19.1
+			VersionNumberRange.Parse("14.24.27905", "14.28.29910") // VS2019 16.2.3 -> 19.1
 		};
 
 		/// <summary>
@@ -914,7 +914,7 @@ namespace UnrealBuildTool
 			// If we're enabling support for C++ modules, make sure the compiler supports it. VS 16.8 changed which command line arguments are used to enable modules support.
 			if (Target.bEnableCppModules && !ProjectFileGenerator.bGenerateProjectFiles && Target.WindowsPlatform.Environment.CompilerVersion < new VersionNumber(14, 28, 29304))
 			{
-				throw new BuildException("Support for C++20 modules requires Visual Studio 2019 16.8 preview 3 or later.");
+				throw new BuildException("Support for C++20 modules requires Visual Studio 2019 16.8 preview 3 or later. Compiler Version Targeted: {0}", Target.WindowsPlatform.Environment.CompilerVersion);
 			}
 
 			// Ensure we're using recent enough version of Visual Studio to support ASan builds.
diff --git a/Engine/Source/Programs/UnrealBuildTool/Platform/Windows/VCToolChain.cs b/Engine/Source/Programs/UnrealBuildTool/Platform/Windows/VCToolChain.cs
--- a/Engine/Source/Programs/UnrealBuildTool/Platform/Windows/VCToolChain.cs
+++ b/Engine/Source/Programs/UnrealBuildTool/Platform/Windows/VCToolChain.cs
@@ -670,6 +670,7 @@ namespace UnrealBuildTool
 					case CppStandardVersion.Cpp17:
 						Arguments.Add("/std:c++17");
 						break;
+					case CppStandardVersion.Cpp20:
 					case CppStandardVersion.Latest:
 						Arguments.Add("/std:c++latest");
 						break;
