Change VC max preferred version to 14.24.28314 to correct internal compiler error until that can be resolved.

#rnx

diff --git a/Engine/Source/Programs/UnrealBuildTool/Platform/Windows/UEBuildWindows.cs b/Engine/Source/Programs/UnrealBuildTool/Platform/Windows/UEBuildWindows.cs
--- a/Engine/Source/Programs/UnrealBuildTool/Platform/Windows/UEBuildWindows.cs
+++ b/Engine/Source/Programs/UnrealBuildTool/Platform/Windows/UEBuildWindows.cs
@@ -751,7 +751,7 @@ namespace UnrealBuildTool
 		/// </summary>
 		static readonly VersionNumberRange[] PreferredVisualCppVersions = new VersionNumberRange[]
 		{
-			VersionNumberRange.Parse("14.24.27905", "14.28.29910") // VS2019 16.2.3 -> 19.1
+			VersionNumberRange.Parse("14.24.27905", "14.24.28314") // VS2019 16.2.3 -> 19.1
 		};
 
 		/// <summary>
