Make GetDumpEncodeDebugCommand protected virtual

diff --git a/Engine/Source/Programs/UnrealBuildTool/Platform/Linux/LinuxToolChain.cs b/Engine/Source/Programs/UnrealBuildTool/Platform/Linux/LinuxToolChain.cs
--- a/Engine/Source/Programs/UnrealBuildTool/Platform/Linux/LinuxToolChain.cs
+++ b/Engine/Source/Programs/UnrealBuildTool/Platform/Linux/LinuxToolChain.cs
@@ -254,7 +254,7 @@ namespace UnrealBuildTool
 			}
 		}
 
-		internal string GetDumpEncodeDebugCommand(LinkEnvironment LinkEnvironment, FileItem OutputFile)
+		protected internal virtual string GetDumpEncodeDebugCommand(LinkEnvironment LinkEnvironment, FileItem OutputFile)
 		{
 			bool bUseCmdExe = BuildHostPlatform.Current.Platform == UnrealTargetPlatform.Win64;
 			string DumpCommand = bUseCmdExe ? "\"{0}\" \"{1}\" \"{2}\" 2>NUL" : "\"{0}\" -c -o \"{2}\" \"{1}\"";
