Restore CL 15315089: Protect VS Andriod config with platfrom check to fix project generation on non-Windows platforms.

#rnx

diff --git a/Engine/Source/Programs/UnrealBuildTool/Platform/Android/AndroidProjectGenerator.cs b/Engine/Source/Programs/UnrealBuildTool/Platform/Android/AndroidProjectGenerator.cs
--- a/Engine/Source/Programs/UnrealBuildTool/Platform/Android/AndroidProjectGenerator.cs
+++ b/Engine/Source/Programs/UnrealBuildTool/Platform/Android/AndroidProjectGenerator.cs
@@ -53,12 +53,6 @@ namespace UnrealBuildTool
 					AGDEInstalled = false;
 				}
 			}
-
-			AGDEInstalled = Microsoft.Win32.Registry.LocalMachine.OpenSubKey(@"SOFTWARE\WOW6432Node\Google\AndroidGameDevelopmentExtension")?.ValueCount > 0;
-			if (Arguments.HasOption("-noagde"))
-			{
-				AGDEInstalled = false;
-			}
 		}
 
 		bool HostSupportsVSAndroid()
