Fix for iOS Packaging failing when debug symbols are enabled

diff --git a/Engine/Source/Programs/UnrealBuildTool/Platform/IOS/IOSToolChain.cs b/Engine/Source/Programs/UnrealBuildTool/Platform/IOS/IOSToolChain.cs
--- a/Engine/Source/Programs/UnrealBuildTool/Platform/IOS/IOSToolChain.cs
+++ b/Engine/Source/Programs/UnrealBuildTool/Platform/IOS/IOSToolChain.cs
@@ -1074,7 +1074,7 @@ namespace UnrealBuildTool
 			GenDebugAction.WorkingDirectory = DirectoryReference.Combine(UnrealBuildTool.EngineDirectory, "Binaries", "Mac");
 
 			GenDebugAction.CommandPath = BuildHostPlatform.Current.Shell;
-			GenDebugAction.CommandArguments = string.Format("-c 'rm -rf \"{1}\"; dwarfdump --uuid \"{3}\" | cut -d\" \" -f2; chmod 777 ./DsymExporter; ./DsymExporter -UUID=$(dwarfdump --uuid \"{3}\" | cut -d\" \" -f2) \"{0}\" \"{2}\"'",
+			GenDebugAction.CommandArguments = string.Format("-c 'rm -rf \"{1}\"; dwarfdump --uuid \"{3}\" | cut -d\\  -f2; chmod 777 ./DsymExporter; ./DsymExporter -UUID=$(dwarfdump --uuid \"{3}\" | cut -d\\  -f2) \"{0}\" \"{2}\"'",
 					DWARFFile.AbsolutePath,
 					OutputFile.AbsolutePath,
 					Path.GetDirectoryName(OutputFile.AbsolutePath),
