Add an explicit RPATH in LinuxToolchain.cs, due to how our RPATH generation works. Needs to be fixed at a later time, but is prone to breaking edge cases

diff --git a/Engine/Source/Programs/UnrealBuildTool/Platform/Linux/LinuxToolChain.cs b/Engine/Source/Programs/UnrealBuildTool/Platform/Linux/LinuxToolChain.cs
--- a/Engine/Source/Programs/UnrealBuildTool/Platform/Linux/LinuxToolChain.cs
+++ b/Engine/Source/Programs/UnrealBuildTool/Platform/Linux/LinuxToolChain.cs
@@ -1036,6 +1036,7 @@ namespace UnrealBuildTool
 
 			// @FIXME: Workaround for generating RPATHs for launching on devices UE-54136
 			Result += " -Wl,-rpath=${ORIGIN}/../../../Engine/Binaries/ThirdParty/PhysX3/Linux/x86_64-unknown-linux-gnu";
+			Result += " -Wl,-rpath=${ORIGIN}/../../../Engine/Binaries/ThirdParty/Intel/Embree/Embree2140/Linux/x86_64-unknown-linux-gnu/lib";
 
 			// Some OS ship ld with new ELF dynamic tags, which use DT_RUNPATH vs DT_RPATH. Since DT_RUNPATH do not propagate to dlopen()ed DSOs,
 			// this breaks the editor on such systems. See https://kenai.com/projects/maxine/lists/users/archive/2011-01/message/12 for details
