UnrealBuildTool: Limit number of concurrent actions based on free memory availalble, defaulted to requiring 1.5GB per running action.

Refactor GetMaxActionsToExecuteInParallel to Utils from LocalExecutor and share with ParallelExecutor
Remove logic to use average of physical and logical cores if hyper-threading is enabled

diff --git a/Engine/Source/Programs/UnrealBuildTool/Configuration/BuildConfiguration.cs b/Engine/Source/Programs/UnrealBuildTool/Configuration/BuildConfiguration.cs
--- a/Engine/Source/Programs/UnrealBuildTool/Configuration/BuildConfiguration.cs
+++ b/Engine/Source/Programs/UnrealBuildTool/Configuration/BuildConfiguration.cs
@@ -95,7 +95,7 @@ namespace UnrealBuildTool
 
 		/// <summary>
 		/// Number of actions that can be executed in parallel. If 0 then code will pick a default based
-		/// on the number of cores available. Only applies to the ParallelExecutor
+		/// on the number of cores and memory available. Applies to the ParallelExecutor, HybridExecutor, and LocalExecutor
 		/// </summary>
 		[XmlConfigFile]
 		[CommandLine("-MaxParallelActions")]
diff --git a/Engine/Source/Programs/UnrealBuildTool/Executors/LocalExecutor.cs b/Engine/Source/Programs/UnrealBuildTool/Executors/LocalExecutor.cs
--- a/Engine/Source/Programs/UnrealBuildTool/Executors/LocalExecutor.cs
+++ b/Engine/Source/Programs/UnrealBuildTool/Executors/LocalExecutor.cs
@@ -197,19 +197,27 @@ namespace UnrealBuildTool
 
 	class LocalExecutor : ActionExecutor
 	{
+		/// <summary>
+		/// Maximum processor count for local execution. 
+		/// </summary>
+		[XmlConfigFile]
+		int MaxProcessorCount = int.MaxValue;
+
 		/// <summary>
 		/// Processor count multiplier for local execution. Can be below 1 to reserve CPU for other tasks.
 		/// When using the local executor (not XGE), run a single action on each CPU core. Note that you can set this to a larger value
 		/// to get slightly faster build times in many cases, but your computer's responsiveness during compiling may be much worse.
+		/// This value is ignored if the CPU does not support hyper-threading.
 		/// </summary>
 		[XmlConfigFile]
 		double ProcessorCountMultiplier = 1.0;
 
 		/// <summary>
-		/// Maximum processor count for local execution. 
+		/// Free memory per action in bytes, used to limit the number of parallel actions if the machine is memory starved.
+		/// Set to 0 to disable free memory checking.
 		/// </summary>
 		[XmlConfigFile]
-		int MaxProcessorCount = int.MaxValue;
+		double MemoryPerActionBytes = 1.5 * 1024 * 1024 * 1024;
 
 		/// <summary>
 		/// How many processes that will be executed in parallel
@@ -232,7 +240,7 @@ namespace UnrealBuildTool
 			else
 			{
 				// Figure out how many processors to use
-				NumParallelProcesses = GetMaxActionsToExecuteInParallel();
+				NumParallelProcesses = Utils.GetMaxActionsToExecuteInParallel(MaxProcessorCount, ProcessorCountMultiplier, Convert.ToInt64(MemoryPerActionBytes));
 			}
 		}
 
@@ -241,46 +249,6 @@ namespace UnrealBuildTool
 			get { return "Local"; }
 		}
 
-		/// <summary>
-		/// Determines the maximum number of actions to execute in parallel, taking into account the resources available on this machine.
-		/// </summary>
-		/// <returns>Max number of actions to execute in parallel</returns>
-		public virtual int GetMaxActionsToExecuteInParallel()
-		{
-			// Get the number of logical processors
-			int NumLogicalCores = Utils.GetLogicalProcessorCount();
-
-			// Use WMI to figure out physical cores, excluding hyper threading.
-			int NumPhysicalCores = Utils.GetPhysicalProcessorCount();
-			if (NumPhysicalCores == -1)
-			{
-				NumPhysicalCores = NumLogicalCores;
-			}
-
-			// The number of actions to execute in parallel is trying to keep the CPU busy enough in presence of I/O stalls.
-			int MaxActionsToExecuteInParallel = 0;
-			if (NumPhysicalCores < NumLogicalCores && ProcessorCountMultiplier != 1.0)
-			{
-				// The CPU has more logical cores than physical ones, aka uses hyper-threading. 
-				// Use multiplier if provided
-				MaxActionsToExecuteInParallel = (int)(NumPhysicalCores * ProcessorCountMultiplier);
-			}
-			else if (NumPhysicalCores < NumLogicalCores && NumPhysicalCores > 4)
-			{
-				// The CPU has more logical cores than physical ones, aka uses hyper-threading. 
-				// Use average of logical and physical if we have "lots of cores"
-				MaxActionsToExecuteInParallel = Math.Max((int)(NumPhysicalCores + NumLogicalCores) / 2, NumLogicalCores - 4);
-			}
-			// No hyper-threading. Only kicking off a task per CPU to keep machine responsive.
-			else
-			{
-				MaxActionsToExecuteInParallel = NumPhysicalCores;
-			}
-
-			MaxActionsToExecuteInParallel = Math.Max(1, Math.Min(MaxActionsToExecuteInParallel, MaxProcessorCount));
-			return MaxActionsToExecuteInParallel;
-		}
-
 		/// <summary>
 		/// Executes the specified actions locally.
 		/// </summary>
diff --git a/Engine/Source/Programs/UnrealBuildTool/Executors/ParallelExecutor.cs b/Engine/Source/Programs/UnrealBuildTool/Executors/ParallelExecutor.cs
--- a/Engine/Source/Programs/UnrealBuildTool/Executors/ParallelExecutor.cs
+++ b/Engine/Source/Programs/UnrealBuildTool/Executors/ParallelExecutor.cs
@@ -49,10 +49,20 @@ namespace UnrealBuildTool
 
 		/// <summary>
 		/// Processor count multiplier for local execution. Can be below 1 to reserve CPU for other tasks.
+		/// When using the local executor (not XGE), run a single action on each CPU core. Note that you can set this to a larger value
+		/// to get slightly faster build times in many cases, but your computer's responsiveness during compiling may be much worse.
+		/// This value is ignored if the CPU does not support hyper-threading.
 		/// </summary>
 		[XmlConfigFile]
 		double ProcessorCountMultiplier = 1.0;
 
+		/// <summary>
+		/// Free memory per action in bytes, used to limit the number of parallel actions if the machine is memory starved.
+		/// Set to 0 to disable free memory checking.
+		/// </summary>
+		[XmlConfigFile]
+		double MemoryPerActionBytes = 1.5 * 1024 * 1024 * 1024;
+
 		/// <summary>
 		/// When enabled, will stop compiling targets after a compile error occurs.
 		/// </summary>
@@ -80,7 +90,7 @@ namespace UnrealBuildTool
 			else
 			{
 				// Figure out how many processors to use
-				NumParallelProcesses = Math.Min((int)(Environment.ProcessorCount * ProcessorCountMultiplier), MaxProcessorCount);
+				NumParallelProcesses = Utils.GetMaxActionsToExecuteInParallel(MaxProcessorCount, ProcessorCountMultiplier, Convert.ToInt64(MemoryPerActionBytes));
 			}
 		}
 
diff --git a/Engine/Source/Programs/UnrealBuildTool/System/Utils.cs b/Engine/Source/Programs/UnrealBuildTool/System/Utils.cs
--- a/Engine/Source/Programs/UnrealBuildTool/System/Utils.cs
+++ b/Engine/Source/Programs/UnrealBuildTool/System/Utils.cs
@@ -1158,6 +1158,94 @@ namespace UnrealBuildTool
 			return -1;
 		}
 
+		/// <summary>
+		/// Gets the total memory bytes available, based on what is known to the garbage collector.
+		/// </summary>
+		/// <remarks>This will return a max of 2GB for a 32bit application.</remarks>
+		/// <returns>The total memory available, in bytes.</returns>
+		public static long GetAvailableMemoryBytes()
+		{
+			GCMemoryInfo MemoryInfo = GC.GetGCMemoryInfo();
+			// TotalAvailableMemoryBytes will be 0 if garbage collection has not run yet
+			return MemoryInfo.TotalAvailableMemoryBytes != 0 ? MemoryInfo.TotalAvailableMemoryBytes :- 1;
+		}
+
+		/// <summary>
+		/// Gets the total memory bytes free, based on what is known to the garbage collector.
+		/// </summary>
+		/// <remarks>This will return a max of 2GB for a 32bit application.</remarks>
+		/// <returns>The total memory free, in bytes.</returns>
+		public static long GetFreeMemoryBytes()
+		{
+			GCMemoryInfo MemoryInfo = GC.GetGCMemoryInfo();
+			// TotalAvailableMemoryBytes will be 0 if garbage collection has not run yet
+			return MemoryInfo.TotalAvailableMemoryBytes != 0 ? MemoryInfo.TotalAvailableMemoryBytes - MemoryInfo.MemoryLoadBytes : -1;
+		}
+
+		/// <summary>
+		/// Determines the maximum number of actions to execute in parallel, taking into account the resources available on this machine.
+		/// </summary>
+		/// <param name="MaxProcessorCount">How many actions to execute in parallel. When 0 a default will be chosen based on system resources</param>
+		/// <param name="ProcessorCountMultiplier">Processor count multiplier for local execution. Can be below 1 to reserve CPU for other tasks.</param>
+		/// <param name="MemoryPerActionBytes"></param>
+		/// <returns>Max number of actions to execute in parallel</returns>
+		public static int GetMaxActionsToExecuteInParallel(int MaxProcessorCount, double ProcessorCountMultiplier, long MemoryPerActionBytes)
+		{
+			// Get the number of logical processors
+			int NumLogicalCores = Utils.GetLogicalProcessorCount();
+
+			// Use WMI to figure out physical cores, excluding hyper threading.
+			int NumPhysicalCores = Utils.GetPhysicalProcessorCount();
+			if (NumPhysicalCores == -1)
+			{
+				NumPhysicalCores = NumLogicalCores;
+			}
+
+			Log.TraceInformation($"Determining max actions to execute in parallel ({NumPhysicalCores} physical cores, {NumLogicalCores} logical cores)");
+
+			// The number of actions to execute in parallel is trying to keep the CPU busy enough in presence of I/O stalls.
+			int MaxActionsToExecuteInParallel;
+			if (NumPhysicalCores < NumLogicalCores && ProcessorCountMultiplier != 1.0)
+			{
+				// The CPU has more logical cores than physical ones, aka uses hyper-threading. 
+				// Use multiplier if provided
+				MaxActionsToExecuteInParallel = (int)(NumPhysicalCores * ProcessorCountMultiplier);
+				Log.TraceInformation($"  Requested {ProcessorCountMultiplier} process count multiplier: limiting max parallel actions to {MaxActionsToExecuteInParallel}");
+			}
+			else if (NumPhysicalCores < NumLogicalCores)
+			{
+				// The CPU has more logical cores than physical ones, aka uses hyper-threading. 
+				MaxActionsToExecuteInParallel = NumLogicalCores;
+			}
+			// No hyper-threading. Only kicking off a task per CPU to keep machine responsive.
+			else
+			{
+				MaxActionsToExecuteInParallel = NumPhysicalCores;
+			}
+
+			// Limit number of actions to execute if the system is memory starved.
+			if (MemoryPerActionBytes > 0)
+			{
+				long FreeMemoryBytes = GetFreeMemoryBytes();
+				if (FreeMemoryBytes != -1)
+				{
+					int TotalMemoryActions = Convert.ToInt32(FreeMemoryBytes / MemoryPerActionBytes);
+					if (TotalMemoryActions < MaxActionsToExecuteInParallel)
+					{
+						MaxActionsToExecuteInParallel = Math.Max(1, Math.Min(MaxActionsToExecuteInParallel, TotalMemoryActions));
+						Log.TraceInformation($"  Requested {StringUtils.FormatBytesString(MemoryPerActionBytes)} free memory per action, {StringUtils.FormatBytesString(FreeMemoryBytes)} available: limiting max parallel actions to {MaxActionsToExecuteInParallel}");
+					}
+				}
+			}
+
+			if (MaxProcessorCount < MaxActionsToExecuteInParallel)
+			{
+				MaxActionsToExecuteInParallel = Math.Max(1, Math.Min(MaxActionsToExecuteInParallel, MaxProcessorCount));
+				Log.TraceInformation($"  Requested max {MaxProcessorCount} action(s): limiting max parallel actions to {MaxActionsToExecuteInParallel}");
+			}
+			return MaxActionsToExecuteInParallel;
+		}
+
 		/// <summary>
 		/// Executes a list of custom build step scripts
 		/// </summary>
