UnrealBuildTool: Print detailed information when two Actions conflict

diff --git a/Engine/Source/Programs/UnrealBuildTool/System/ActionGraph.cs b/Engine/Source/Programs/UnrealBuildTool/System/ActionGraph.cs
--- a/Engine/Source/Programs/UnrealBuildTool/System/ActionGraph.cs
+++ b/Engine/Source/Programs/UnrealBuildTool/System/ActionGraph.cs
@@ -9,11 +9,29 @@ using System.Runtime.Serialization;
 using EpicGames.Core;
 using System.Text;
 using System.Threading.Tasks;
+using System.Text.Json;
+using System.Text.Json.Serialization;
 
 namespace UnrealBuildTool
 {
 	static class ActionGraph
 	{
+		/// <summary>
+		/// Enum describing why an Action is in conflict with another Action
+		/// </summary>
+		[Flags]
+		internal enum ActionConflictReasonFlags : byte
+		{
+			None = 0,
+			ActionType = 1 << 0,
+			PrerequisiteItems = 1 << 1,
+			DeleteItems = 1 << 2,
+			DependencyListFile = 1 << 3,
+			WorkingDirectory = 1 << 4,
+			CommandPath = 1 << 5,
+			CommandArguments = 1 << 6,
+		};
+
 		/// <summary>
 		/// Links the actions together and sets up their dependencies
 		/// </summary>
@@ -87,60 +105,135 @@ namespace UnrealBuildTool
 		/// </summary>
 		/// <param name="A">The first action</param>
 		/// <param name="B">The second action</param>
-		/// <returns>True if any conflicts were found, false otherwise.</returns>
+		/// <returns>True if no conflicts were found, false otherwise.</returns>
 		public static bool CheckForConflicts(IExternalAction A, IExternalAction B)
 		{
-			bool bResult = true;
+			ActionConflictReasonFlags Reason = ActionConflictReasonFlags.None;
 			if (A.ActionType != B.ActionType)
 			{
-				LogConflict(A, "action type is different", A.ActionType.ToString(), B.ActionType.ToString());
-				bResult = false;
+				Reason |= ActionConflictReasonFlags.ActionType;
 			}
 			if (!Enumerable.SequenceEqual(A.PrerequisiteItems, B.PrerequisiteItems))
 			{
-				LogConflict(A, "prerequisites are different", String.Join(", ", A.PrerequisiteItems.Select(x => x.Location)), String.Join(", ", B.PrerequisiteItems.Select(x => x.Location)));
-				bResult = false;
+				Reason |= ActionConflictReasonFlags.PrerequisiteItems;
 			}
 			if (!Enumerable.SequenceEqual(A.DeleteItems, B.DeleteItems))
 			{
-				LogConflict(A, "deleted items are different", String.Join(", ", A.DeleteItems.Select(x => x.Location)), String.Join(", ", B.DeleteItems.Select(x => x.Location)));
-				bResult = false;
+				Reason |= ActionConflictReasonFlags.DeleteItems;
 			}
 			if (A.DependencyListFile != B.DependencyListFile)
 			{
-				LogConflict(A, "dependency list is different", (A.DependencyListFile == null) ? "(none)" : A.DependencyListFile.AbsolutePath, (B.DependencyListFile == null) ? "(none)" : B.DependencyListFile.AbsolutePath);
-				bResult = false;
+				Reason |= ActionConflictReasonFlags.DependencyListFile;
 			}
 			if (A.WorkingDirectory != B.WorkingDirectory)
 			{
-				LogConflict(A, "working directory is different", A.WorkingDirectory.FullName, B.WorkingDirectory.FullName);
-				bResult = false;
+				Reason |= ActionConflictReasonFlags.WorkingDirectory;
 			}
 			if (A.CommandPath != B.CommandPath)
 			{
-				LogConflict(A, "command path is different", A.CommandPath.FullName, B.CommandPath.FullName);
-				bResult = false;
+				Reason |= ActionConflictReasonFlags.CommandPath;
 			}
 			if (A.CommandArguments != B.CommandArguments)
 			{
-				LogConflict(A, "command arguments are different", A.CommandArguments, B.CommandArguments);
-				bResult = false;
+				Reason |= ActionConflictReasonFlags.CommandArguments;
+			}
+
+			if (Reason != ActionConflictReasonFlags.None)
+			{
+				LogConflict(A, B, Reason);
+				return false;
+			}
+			return true;
+		}
+
+		internal class LogActionActionTypeConverter : JsonConverter<ActionType>
+		{
+			public override ActionType Read(ref Utf8JsonReader reader, Type typeToConvert, JsonSerializerOptions options)
+			{
+				throw new NotImplementedException();
+			}
+
+			public override void Write(Utf8JsonWriter writer, ActionType value, JsonSerializerOptions options)
+			{
+				writer.WriteStringValue(value.ToString());
+			}
+		}
+
+		internal class LogActionFileItemConverter : JsonConverter<FileItem>
+		{
+			public override FileItem Read(ref Utf8JsonReader reader, Type typeToConvert, JsonSerializerOptions options)
+			{
+				throw new NotImplementedException();
+			}
+
+			public override void Write(Utf8JsonWriter writer, FileItem value, JsonSerializerOptions options)
+			{
+				writer.WriteStringValue(value.FullName);
+			}
+		}
+
+		internal class LogActionDirectoryReferenceConverter : JsonConverter<DirectoryReference>
+		{
+			public override DirectoryReference Read(ref Utf8JsonReader reader, Type typeToConvert, JsonSerializerOptions options)
+			{
+				throw new NotImplementedException();
+			}
+
+			public override void Write(Utf8JsonWriter writer, DirectoryReference value, JsonSerializerOptions options)
+			{
+				writer.WriteStringValue(value.FullName);
+			}
+		}
+
+		internal class LogActionFileReferenceConverter : JsonConverter<FileReference>
+		{
+			public override FileReference Read(ref Utf8JsonReader reader, Type typeToConvert, JsonSerializerOptions options)
+			{
+				throw new NotImplementedException();
+			}
+
+			public override void Write(Utf8JsonWriter writer, FileReference value, JsonSerializerOptions options)
+			{
+				writer.WriteStringValue(value.FullName);
 			}
-			return bResult;
 		}
 
 		/// <summary>
 		/// Adds the description of a merge error to an output message
 		/// </summary>
-		/// <param name="Action">The action with the conflict</param>
-		/// <param name="Description">Description of the difference</param>
-		/// <param name="OldValue">Previous value for the field</param>
-		/// <param name="NewValue">Conflicting value for the field</param>
-		static void LogConflict(IExternalAction Action, string Description, string OldValue, string NewValue)
+		/// <param name="A">The first action with the conflict</param>
+		/// <param name="B">The second action with the conflict</param>
+		/// <param name="Reason">Enum flags for which properties are in conflict</param>
+		static void LogConflict(IExternalAction A, IExternalAction B, ActionConflictReasonFlags Reason)
 		{
-			Log.TraceError("Unable to merge actions producing {0}: {1}", Action.ProducedItems.First().Location.GetFileName(), Description);
-			Log.TraceLog("  Previous: {0}", OldValue);
-			Log.TraceLog("  Conflict: {0}", NewValue);
+			// Convert some complex types in IExternalAction to strings when printing json
+			JsonSerializerOptions Options = new JsonSerializerOptions
+			{
+				WriteIndented = true,
+				IgnoreNullValues = true,
+				Converters =
+				{
+					new LogActionActionTypeConverter(),
+					new LogActionFileItemConverter(),
+					new LogActionDirectoryReferenceConverter(),
+					new LogActionFileReferenceConverter(),
+				},
+			};
+
+			string AJson = JsonSerializer.Serialize(A, Options);
+			string BJson = JsonSerializer.Serialize(B, Options);
+			string AJsonPath = Path.Combine(Path.GetTempPath(), "UnrealBuildTool", Path.ChangeExtension(Path.GetRandomFileName(), "json"));
+			string BJsonPath = Path.Combine(Path.GetTempPath(), "UnrealBuildTool", Path.ChangeExtension(Path.GetRandomFileName(), "json"));
+
+			Directory.CreateDirectory(Path.Combine(Path.GetTempPath(), "UnrealBuildTool"));
+			File.WriteAllText(AJsonPath, AJson);
+			File.WriteAllText(BJsonPath, BJson);
+
+			Log.TraceError($"Unable to merge actions '{A.StatusDescription}' and '{B.StatusDescription}': {Reason} are different");
+			Log.TraceInformation($"  First Action: {AJson}");
+			Log.TraceInformation($"  Second Action: {BJson}");
+			Log.TraceInformation($"  First Action json written to '{AJsonPath}'");
+			Log.TraceInformation($"  Second Action json written to '{BJsonPath}'");
 		}
 
 		/// <summary>
