HotReload:

Make ShouldDoHotloadFromIDE() more robust against the possiblity that the PID of a long-gone Editor instance has been recycled as a system process.

diff --git a/Engine/Source/Programs/UnrealBuildTool/System/HotReload.cs b/Engine/Source/Programs/UnrealBuildTool/System/HotReload.cs
--- a/Engine/Source/Programs/UnrealBuildTool/System/HotReload.cs
+++ b/Engine/Source/Programs/UnrealBuildTool/System/HotReload.cs
@@ -340,10 +340,28 @@ namespace UnrealBuildTool
 								RunningProcess = null;
 							}
 
-							if(RunningProcess == null || RunningProcess.HasExited)
+							bool bFileShouldBeDeleted = false;
+
+							if(RunningProcess == null)
 							{
-								FileReference.Delete(EditorInstanceFile);
-								continue;
+								bFileShouldBeDeleted = true;
+							}
+							else
+							{
+								try
+								{
+									if (RunningProcess.HasExited)
+									{
+										bFileShouldBeDeleted = true;
+									}
+								}
+								catch
+								{
+									// if the PID represents an editor that has exited, and is now reused as the pid of a system process,
+									// RunningProcess.HasExited may fail with "Access is denied."
+									// If we can't determine if the process has exited, let's assume that the file should be deleted.
+									bFileShouldBeDeleted = true;
+								}
 							}
 
 							// bugfix - the editor sometimes doesn't delete its editorrun file due to
@@ -358,22 +376,6 @@ namespace UnrealBuildTool
 							// On my machine this was ~750 ms _per orphaned file_, and I spoke to someone
 							// with 10 of these in his Engine/Intermediate/EditorRun directory.
 							// 
-							try
-							{
-								if (RunningProcess.HasExited)
-								{
-									FileReference.Delete(EditorInstanceFile);
-									continue;
-								}
-							}
-							catch
-							{
-								// if we got here, we either failed to access the HasExited, or 
-								// failed to delete the file.
-								Log.TraceLog("Failed to delete EditorRun file for exited process: {0}", EditorInstanceFile.GetFileName());
-								continue;
-							}
-
 							FileReference MainModuleFile;
 							try
 							{
@@ -382,6 +384,20 @@ namespace UnrealBuildTool
 							catch
 							{
 								MainModuleFile = null;
+								bFileShouldBeDeleted = true;
+							}
+
+							if (bFileShouldBeDeleted)
+							{
+								try
+								{
+									FileReference.Delete(EditorInstanceFile);
+								}
+								catch
+								{
+									Log.TraceLog("Failed to delete EditorRun file for exited process: {0}", EditorInstanceFile.GetFileName());
+								}
+								continue;
 							}
 
 							if (!bIsRunning && EditorLocation == MainModuleFile)
