Remove single quotes to better support spaces in project path.

#rnx

diff --git a/Engine/Source/Programs/UnrealBuildTool/Platform/IOS/IOSToolChain.cs b/Engine/Source/Programs/UnrealBuildTool/Platform/IOS/IOSToolChain.cs
--- a/Engine/Source/Programs/UnrealBuildTool/Platform/IOS/IOSToolChain.cs
+++ b/Engine/Source/Programs/UnrealBuildTool/Platform/IOS/IOSToolChain.cs
@@ -491,7 +491,7 @@ namespace UnrealBuildTool
 			Result += GetArchitectureArgument(LinkEnvironment.Configuration, LinkEnvironment.Architecture);
 
 			bool bIsDevice = (LinkEnvironment.Architecture != "-simulator");
-			Result += String.Format(" -isysroot {0}Platforms/{1}.platform/Developer/SDKs/{1}{2}.sdk",
+			Result += String.Format(" -isysroot \\\"{0}Platforms/{1}.platform/Developer/SDKs/{1}{2}.sdk\\\"",
 				Settings.Value.XcodeDeveloperDir, bIsDevice ? Settings.Value.DevicePlatformName : Settings.Value.SimulatorPlatformName, Settings.Value.IOSSDKVersion);
 
 			if (IsBitcodeCompilingEnabled(LinkEnvironment.Configuration))
@@ -797,13 +797,13 @@ namespace UnrealBuildTool
 				// Add the library paths to the argument list.
 				foreach (DirectoryReference LibraryPath in LinkEnvironment.SystemLibraryPaths)
 				{
-					LinkCommandArguments += string.Format(" -L\"{0}\"", LibraryPath.FullName);
+					LinkCommandArguments += string.Format(" -L\\\"{0}\\\"", LibraryPath.FullName);
 				}
 
 				// Add the additional libraries to the argument list.
 				foreach (string AdditionalLibrary in LinkEnvironment.SystemLibraries)
 				{
-					LinkCommandArguments += string.Format(" -l\"{0}\"", AdditionalLibrary);
+					LinkCommandArguments += string.Format(" -l\\\"{0}\\\"", AdditionalLibrary);
 				}
 
 				foreach(FileReference Library in LinkEnvironment.Libraries)
@@ -814,7 +814,7 @@ namespace UnrealBuildTool
 					LinkAction.PrerequisiteItems.Add(LibFile);
 
 					// and add to the commandline
-					LinkCommandArguments += string.Format(" \"{0}\"", Library.FullName);
+					LinkCommandArguments += string.Format(" \\\"{0}\\\"", Library.FullName);
 				}
 			}
 
@@ -836,7 +836,7 @@ namespace UnrealBuildTool
 			if ((!LinkEnvironment.bIsBuildingLibrary || LinkEnvironment.bIsBuildingDLL) && LinkEnvironment.bCreateMapFile)
 			{
 				FileItem MapFile = FileItem.GetItemByFileReference(new FileReference(OutputFile.Location.FullName + ".map"));
-				LinkCommandArguments += string.Format(" -Wl,-map,\"{0}\"", MapFile.Location.FullName);
+				LinkCommandArguments += string.Format(" -Wl,-map,\\\"{0}\\\"", MapFile.Location.FullName);
 				LinkAction.ProducedItems.Add(MapFile);
 			}
 
@@ -863,10 +863,10 @@ namespace UnrealBuildTool
 			}
 			else
 			{
-				bool bIsUE4Game = LinkEnvironment.OutputFilePath.FullName.Contains("UnrealGame");
-				FileReference ResponsePath = FileReference.Combine(((!bIsUE4Game && ProjectFile != null) ? ProjectFile.Directory : UnrealBuildTool.EngineDirectory), "Intermediate", "Build", LinkEnvironment.Platform.ToString(), "LinkFileList_" + LinkEnvironment.OutputFilePath.GetFileNameWithoutExtension() + ".tmp");
+				bool bIsUnrealGame = LinkEnvironment.OutputFilePath.FullName.Contains("UnrealGame");
+				FileReference ResponsePath = FileReference.Combine(((!bIsUnrealGame && ProjectFile != null) ? ProjectFile.Directory : UnrealBuildTool.EngineDirectory), "Intermediate", "Build", LinkEnvironment.Platform.ToString(), "LinkFileList_" + LinkEnvironment.OutputFilePath.GetFileNameWithoutExtension() + ".tmp");
 				Graph.CreateIntermediateTextFile(ResponsePath, InputFileNames);
-				LinkCommandArguments += string.Format(" @\"{0}\"", ResponsePath.FullName);
+				LinkCommandArguments += string.Format(" \\\"@{0}\\\"", ResponsePath.FullName);
 			}
 
 			// if we are making an LTO build, write the lto file next to build so dsymutil can find it
@@ -882,11 +882,11 @@ namespace UnrealBuildTool
 				{
 					LtoObjectFile = OutputFile.AbsolutePath + ".lto.o";
 				}
-				LinkCommandArguments += string.Format(" -flto -Xlinker -object_path_lto -Xlinker \"{0}\"", LtoObjectFile);
+				LinkCommandArguments += string.Format(" -flto -Xlinker -object_path_lto -Xlinker \\\"{0}\\\"", LtoObjectFile);
 			}
 
 			// Add the output file to the command-line.
-			LinkCommandArguments += string.Format(" -o \"{0}\"", OutputFile.AbsolutePath);
+			LinkCommandArguments += string.Format(" -o \\\"{0}\\\"", OutputFile.AbsolutePath);
 
 			// Add the additional arguments specified by the environment.
 			LinkCommandArguments += LinkEnvironment.AdditionalArguments;
@@ -905,20 +905,20 @@ namespace UnrealBuildTool
 				// when re-packaging and no source files change because the linker skips symbol generation and dsymutil will 
 				// recreate a new .dsym file from a symboless exe file. It's just sad. To make things happy we need to delete 
 				// the output file to force the linker to recreate it with symbols again.
-				string linkCommandArguments = "-c '";
+				string linkCommandArguments = "-c \"";
 
-				linkCommandArguments += string.Format("rm -f \"{0}\";", OutputFile.AbsolutePath);
-				linkCommandArguments += string.Format("rm -f \"{0}\\*.bcsymbolmap\";", Path.GetDirectoryName(OutputFile.AbsolutePath));
-				linkCommandArguments += "\"" + LinkerPath + "\" " + LinkCommandArguments + ";";
+				linkCommandArguments += string.Format("rm -f \\\"{0}\\\";", OutputFile.AbsolutePath);
+				linkCommandArguments += string.Format("rm -f \\\"{0}\\*.bcsymbolmap\\\";", Path.GetDirectoryName(OutputFile.AbsolutePath));
+				linkCommandArguments += "\\\"" + LinkerPath + "\\\" " + LinkCommandArguments + ";";
 
-				linkCommandArguments += "'";
+				linkCommandArguments += "\"";
 
 				LinkAction.CommandArguments = linkCommandArguments;
 			}
 			else
 			{
 				// This is not a shipping build so no need to delete the output file since symbols will not have been stripped from it.
-				LinkAction.CommandArguments = string.Format("-c '\"{0}\" {1}'", LinkerPath, LinkCommandArguments);
+				LinkAction.CommandArguments = string.Format("-c \"\\\"{0}\\\" {1}\"", LinkerPath, LinkCommandArguments);
 			}
 
 			return OutputFile;
@@ -959,11 +959,11 @@ namespace UnrealBuildTool
 			Arguments.Append(" --output-format human-readable-text");
 			Arguments.Append(" --notices");
 			Arguments.Append(" --warnings");
-			Arguments.AppendFormat(" --output-partial-info-plist '{0}/assetcatalog_generated_info.plist'", InputDir);
+			Arguments.AppendFormat(" --output-partial-info-plist \"{0}/assetcatalog_generated_info.plist\"", InputDir);
 			if (Platform == UnrealTargetPlatform.TVOS)
 			{
-				Arguments.Append(" --app-icon 'App Icon & Top Shelf Image'");
-				Arguments.Append(" --launch-image 'Launch Image'");
+				Arguments.Append(" --app-icon \"App Icon & Top Shelf Image\"");
+				Arguments.Append(" --launch-image \"Launch Image\"");
 				Arguments.Append(" --filter-for-device-model AppleTV5,3");
 				Arguments.Append(" --target-device tv");
 				Arguments.Append(" --minimum-deployment-target 12.0");
@@ -979,8 +979,8 @@ namespace UnrealBuildTool
 				Arguments.Append(" --platform iphoneos");
 			}
 			Arguments.Append(" --enable-on-demand-resources YES");
-			Arguments.AppendFormat(" --compile '{0}'", OutputDir);
-			Arguments.AppendFormat(" '{0}/Assets.xcassets'", InputDir);
+			Arguments.AppendFormat(" --compile \"{0}\"", OutputDir);
+			Arguments.AppendFormat(" \"{0}/Assets.xcassets\"", InputDir);
 			return Arguments.ToString();
 		}
 
@@ -1026,7 +1026,7 @@ namespace UnrealBuildTool
 			string DsymutilPath = GetDsymutilPath(out ExtraOptions, bIsForLTOBuild: bIsForLTOBuild);
 			if (ProjectSettings.bGeneratedSYMBundle)
 			{
-				GenDebugAction.CommandArguments = string.Format("-c 'rm -rf \"{2}\"; \"{0}\" \"{1}\" {4} -o \"{2}\"; cd \"{2}/..\"; zip -r -y -1 {3}.zip {3}'",
+				GenDebugAction.CommandArguments = string.Format("-c \"rm -rf \\\"{2}\\\"; \\\"{0}\\\" \\\"{1}\\\" {4} -o \\\"{2}\\\"; cd \\\"{2}/..\\\"; zip -r -y -1 {3}.zip {3}\"",
 					DsymutilPath,
 					Executable.AbsolutePath,
 					OutputFile.AbsolutePath,
@@ -1037,7 +1037,7 @@ namespace UnrealBuildTool
 			}
 			else
 			{
-				GenDebugAction.CommandArguments = string.Format("-c 'rm -rf \"{2}\"; \"{0}\" \"{1}\" {3} -f -o \"{2}\"'",
+				GenDebugAction.CommandArguments = string.Format("-c \"rm -rf \\\"{2}\\\"; \\\"{0}\\\" \\\"{1}\\\" {3} -f -o \\\"{2}\\\"\"",
 						DsymutilPath,
 						Executable.AbsolutePath,
 						OutputFile.AbsolutePath,
@@ -1074,7 +1074,7 @@ namespace UnrealBuildTool
 			GenDebugAction.WorkingDirectory = DirectoryReference.Combine(UnrealBuildTool.EngineDirectory, "Binaries", "Mac");
 
 			GenDebugAction.CommandPath = BuildHostPlatform.Current.Shell;
-			GenDebugAction.CommandArguments = string.Format("-c 'rm -rf \"{1}\"; dwarfdump --uuid \"{3}\" | cut -d\\  -f2; chmod 777 ./DsymExporter; ./DsymExporter -UUID=$(dwarfdump --uuid \"{3}\" | cut -d\\  -f2) \"{0}\" \"{2}\"'",
+			GenDebugAction.CommandArguments = string.Format("-c \"rm -rf \\\"{1}\\\"; dwarfdump --uuid \\\"{3}\\\" | cut -d\\  -f2; chmod 777 ./DsymExporter; ./DsymExporter -UUID=$(dwarfdump --uuid \\\"{3}\\\" | cut -d\\  -f2) \\\"{0}\\\" \\\"{2}\\\"\"",
 					DWARFFile.AbsolutePath,
 					OutputFile.AbsolutePath,
 					Path.GetDirectoryName(OutputFile.AbsolutePath),
@@ -1413,7 +1413,7 @@ namespace UnrealBuildTool
 				Action StripAction = Graph.CreateAction(ActionType.CreateAppBundle);
 				StripAction.WorkingDirectory = GetMacDevSrcRoot();
 				StripAction.CommandPath = BuildHostPlatform.Current.Shell;
-				StripAction.CommandArguments = String.Format("-c '\"{0}strip\" {1} \"{2}\" && touch \"{3}\"'", Settings.Value.ToolchainDir, StripArguments, Executable.Location, StripCompleteFile);
+				StripAction.CommandArguments = String.Format("-c \"\\\"{0}strip\\\" {1} \\\"{2}\\\" && touch \\\"{3}\\\"\"", Settings.Value.ToolchainDir, StripArguments, Executable.Location, StripCompleteFile);
 				StripAction.PrerequisiteItems.Add(Executable);
 				StripAction.PrerequisiteItems.AddRange(OutputFiles);
 				StripAction.ProducedItems.Add(StripCompleteFile);
@@ -1595,7 +1595,7 @@ namespace UnrealBuildTool
 				Process FabricProcess = new Process();
 				FabricProcess.StartInfo.WorkingDirectory = Path.GetDirectoryName(DsymZip);
 				FabricProcess.StartInfo.FileName = "/bin/sh";
-				FabricProcess.StartInfo.Arguments = string.Format("-c 'chmod 777 \"{0}/Fabric.framework/upload-symbols\"; \"{0}/Fabric.framework/upload-symbols\" -a 7a4cebd0324af21696e5e321802c5e26ba541cad -p ios {1}'",
+				FabricProcess.StartInfo.Arguments = string.Format("-c \"chmod 777 \\\"{0}/Fabric.framework/upload-symbols\\\"; \\\"{0}/Fabric.framework/upload-symbols\\\" -a 7a4cebd0324af21696e5e321802c5e26ba541cad -p ios {1}\"",
 					FabricPath,
 					DsymZip);
 
@@ -2135,7 +2135,7 @@ namespace UnrealBuildTool
 
 			FileItem TargetItem = FileItem.GetItemByPath(TargetPath);
 
-			CopyAction.CommandArguments = string.Format("-c 'cp -f -R \"{0}\" \"{1}\"; touch -c \"{2}\"'", SourcePath, Path.GetDirectoryName(TargetPath).Replace('\\', '/') + "/", TargetPath.Replace('\\', '/'));
+			CopyAction.CommandArguments = string.Format("-c \"cp -f -R \\\"{0}\\\" \\\"{1}\\\"; touch -c \\\"{2}\\\"\"", SourcePath, Path.GetDirectoryName(TargetPath).Replace('\\', '/') + "/", TargetPath.Replace('\\', '/'));
 			CopyAction.PrerequisiteItems.Add(Executable);
 			CopyAction.ProducedItems.Add(TargetItem);
 			CopyAction.bShouldOutputStatusDescription = Resource.bShouldLog;
diff --git a/Engine/Source/Programs/UnrealBuildTool/Platform/Mac/MacToolChain.cs b/Engine/Source/Programs/UnrealBuildTool/Platform/Mac/MacToolChain.cs
--- a/Engine/Source/Programs/UnrealBuildTool/Platform/Mac/MacToolChain.cs
+++ b/Engine/Source/Programs/UnrealBuildTool/Platform/Mac/MacToolChain.cs
@@ -1167,7 +1167,7 @@ namespace UnrealBuildTool
 			// The script is deleted after it's executed so it's empty when we start appending link commands for the next executable.
 			FileItem FixDylibDepsScript = FileItem.GetItemByFileReference(FileReference.Combine(LinkEnvironment.LocalShadowDirectory, "FixDylibDependencies.sh"));
 
-			FixDylibAction.CommandArguments = "-c 'chmod +x \"" + FixDylibDepsScript.AbsolutePath + "\"; \"" + FixDylibDepsScript.AbsolutePath + "\"; if [[ $? -ne 0 ]]; then exit 1; fi; ";
+			FixDylibAction.CommandArguments = "-c \"chmod +x \\\"" + FixDylibDepsScript.AbsolutePath + "\\\"; \\\"" + FixDylibDepsScript.AbsolutePath + "\\\"; if [[ $? -ne 0 ]]; then exit 1; fi; ";
 
 			// Make sure this action is executed after all the dylibs and the main executable are created
 
@@ -1181,8 +1181,8 @@ namespace UnrealBuildTool
 
 			FileItem OutputFile = FileItem.GetItemByFileReference(FileReference.Combine(LinkEnvironment.LocalShadowDirectory, Path.GetFileNameWithoutExtension(Executable.AbsolutePath) + ".link"));
 
-			FixDylibAction.CommandArguments += "echo \"Dummy\" >> \"" + OutputFile.AbsolutePath + "\"";
-			FixDylibAction.CommandArguments += "'";
+			FixDylibAction.CommandArguments += "echo Dummy >> \\\"" + OutputFile.AbsolutePath + "\\\"";
+			FixDylibAction.CommandArguments += "\"";
 
 			FixDylibAction.ProducedItems.Add(OutputFile);
 
@@ -1243,7 +1243,7 @@ namespace UnrealBuildTool
 			Contents = Contents.Replace("\n", Environment.NewLine);
 
 			Utils.WriteFileIfChanged(DsymShellFile, Contents, StringComparison.InvariantCultureIgnoreCase);
-			GenDebugAction.CommandArguments = string.Format("-c 'chmod +x \"{0}\"; \"{0}\"'", DsymShellFileItem.AbsolutePath);
+			GenDebugAction.CommandArguments = string.Format("-c \"chmod +x \\\"{0}\\\"; \\\"{0}\\\"\"", DsymShellFileItem.AbsolutePath);
 
 			if (LinkEnvironment.bIsCrossReferenced)
 			{
@@ -1302,7 +1302,7 @@ namespace UnrealBuildTool
 
 			FileItem TargetItem = FileItem.GetItemByPath(TargetPath);
 
-			CopyAction.CommandArguments = string.Format("-c 'cp -f -R \"{0}\" \"{1}\"; touch -c \"{2}\"'", SourcePath, Path.GetDirectoryName(TargetPath).Replace('\\', '/') + "/", TargetPath.Replace('\\', '/'));
+			CopyAction.CommandArguments = string.Format("-c \"cp -f -R \\\"{0}\\\" \\\"{1}\\\"; touch -c \\\"{2}\\\"\"", SourcePath, Path.GetDirectoryName(TargetPath).Replace('\\', '/') + "/", TargetPath.Replace('\\', '/'));
 			CopyAction.PrerequisiteItems.Add(Executable);
 			CopyAction.ProducedItems.Add(TargetItem);
 			CopyAction.bShouldOutputStatusDescription = Resource.bShouldLog;
