Definitions.h: Drop some breadcrumbs

As an aid to the unfamiliar searcher (i.e. me), make the name of generated files "Definitions.h" appear as a literal in UBT code. Also, in those generated files, include a reference back to what generates them.

diff --git a/Engine/Source/Programs/UnrealBuildTool/Configuration/UEBuildModuleCPP.cs b/Engine/Source/Programs/UnrealBuildTool/Configuration/UEBuildModuleCPP.cs
--- a/Engine/Source/Programs/UnrealBuildTool/Configuration/UEBuildModuleCPP.cs
+++ b/Engine/Source/Programs/UnrealBuildTool/Configuration/UEBuildModuleCPP.cs
@@ -1123,15 +1123,14 @@ namespace UnrealBuildTool
 		{
 			if(CompileEnvironment.Definitions.Count > 0)
 			{
-				StringBuilder PrivateDefinitionsName = new StringBuilder("Definitions");
+				string PrivateDefinitionsName = "Definitions.h";
+
 				if(!String.IsNullOrEmpty(HeaderSuffix))
 				{
-					PrivateDefinitionsName.Append('.');
-					PrivateDefinitionsName.Append(HeaderSuffix);
+					PrivateDefinitionsName = $"Definitions.{HeaderSuffix}.h";
 				}
-				PrivateDefinitionsName.Append(".h");
 
-				FileReference PrivateDefinitionsFile = FileReference.Combine(IntermediateDirectory, PrivateDefinitionsName.ToString());
+				FileReference PrivateDefinitionsFile = FileReference.Combine(IntermediateDirectory, PrivateDefinitionsName);
 				using (StringWriter Writer = new StringWriter())
 				{
 					WriteDefinitions(CompileEnvironment.Definitions, Writer);
@@ -1199,6 +1198,7 @@ namespace UnrealBuildTool
 		/// <param name="Writer">Writer to receive output</param>
 		static void WriteDefinitions(IEnumerable<string> Definitions, TextWriter Writer)
 		{
+			Writer.WriteLine("// Generated by UnrealBuildTool (UEBuildModuleCPP.cs)");
 			foreach(string Definition in Definitions)
 			{
 				int EqualsIdx = Definition.IndexOf('=');
