Fix issue with Windows launch on IOS giving actool error
Fix hang deploying files to IOS from Windows

diff --git a/Engine/Source/Programs/AutomationTool/IOS/IOSPlatform.Automation.cs b/Engine/Source/Programs/AutomationTool/IOS/IOSPlatform.Automation.cs
--- a/Engine/Source/Programs/AutomationTool/IOS/IOSPlatform.Automation.cs
+++ b/Engine/Source/Programs/AutomationTool/IOS/IOSPlatform.Automation.cs
@@ -1710,7 +1710,7 @@ public class IOSPlatform : Platform
 				Utils.RunLocalProcessAndReturnStdOut(DeviceFS, IdeviceFSArgs, out ExitCode, true);
 				if (ExitCode != 0)
 				{
-					LogError("Deployer error : failed to push files. See above log for more details");
+					throw new AutomationException("Failed to deploy manifest to mobile device.");
 				}
 
 				string[] ManifestFiles = Directory.GetFiles(CombinePaths(Params.BaseStageDirectory, PlatformName), "*_Manifest_UFS*.txt");
@@ -1751,7 +1751,8 @@ public class IOSPlatform : Platform
 		}
 		if (!File.Exists(ExecWithPath) || ExecWithPath == "")
 		{
-			LogError("Failed to locate Libimobile exec. Deploment will fail.");
+			throw new AutomationException("Failed to locate LibiMobileDevice executable.");
+
 		}
 		return ExecWithPath;
 	}
@@ -1794,10 +1795,22 @@ public class IOSPlatform : Platform
 		string IdeviceFSArgs = "-u " + Params.DeviceNames[0] + " -b " + BundleIdentifier + " -x " + "\"" + Directory.GetCurrentDirectory() + "\\CommandsToPush.txt" + "\"";
 		IdeviceFSArgs = GetLibimobileDeviceNetworkedArgument(IdeviceFSArgs, Params.DeviceNames[0]);
 
-		Utils.RunLocalProcessAndReturnStdOut(DeviceFS, IdeviceFSArgs, out ExitCode, true);
-		if (ExitCode != 0)
+		using (Process IDeviceFSProcess = new Process())
 		{
-			LogError("Deployer error : failed to push files. See above log for more details");
+			DataReceivedEventHandler StdOutHandler = (E, Args) => { if (Args.Data != null) { Log.TraceInformation("{0}", Args.Data); } };
+			DataReceivedEventHandler StdErrHandler = (E, Args) => { if (Args.Data != null) { Log.TraceError("{0}", Args.Data); } };
+
+			IDeviceFSProcess.StartInfo.FileName = DeviceFS;
+			IDeviceFSProcess.StartInfo.Arguments = IdeviceFSArgs;
+			IDeviceFSProcess.OutputDataReceived += StdOutHandler;
+			IDeviceFSProcess.ErrorDataReceived += StdErrHandler;
+
+			ExitCode = Utils.RunLocalProcess(IDeviceFSProcess);
+
+			if (ExitCode != 0)
+			{
+				throw new AutomationException("Failed to push content to mobile device.");
+			}
 		}
 
 		File.Delete(Directory.GetCurrentDirectory() + "\\CommandsToPush.txt");
diff --git a/Engine/Source/Programs/UnrealBuildTool/ToolChain/RemoteMac.cs b/Engine/Source/Programs/UnrealBuildTool/ToolChain/RemoteMac.cs
--- a/Engine/Source/Programs/UnrealBuildTool/ToolChain/RemoteMac.cs
+++ b/Engine/Source/Programs/UnrealBuildTool/ToolChain/RemoteMac.cs
@@ -1095,11 +1095,14 @@ namespace UnrealBuildTool
 		{
 			using(Process RsyncProcess = new Process())
 			{
+				DataReceivedEventHandler OutputHandler = (E, Args) => { RsyncOutput(Args, false); };
+				DataReceivedEventHandler ErrorHandler = (E, Args) => { RsyncOutput(Args, true); };
+
 				RsyncProcess.StartInfo.FileName = RsyncExe.FullName;
 				RsyncProcess.StartInfo.Arguments = Arguments;
 				RsyncProcess.StartInfo.WorkingDirectory = SshExe.Directory.FullName;
-				RsyncProcess.OutputDataReceived += RsyncOutput;
-				RsyncProcess.ErrorDataReceived += RsyncOutput;
+				RsyncProcess.OutputDataReceived += OutputHandler;
+				RsyncProcess.ErrorDataReceived += ErrorHandler;
 
 				Log.TraceLog("[Rsync] {0} {1}", Utils.MakePathSafeToUseWithCommandLine(RsyncProcess.StartInfo.FileName), RsyncProcess.StartInfo.Arguments);
 				return Utils.RunLocalProcess(RsyncProcess);
@@ -1109,13 +1112,20 @@ namespace UnrealBuildTool
 		/// <summary>
 		/// Handles data output by rsync
 		/// </summary>
-		/// <param name="Sender">The object sending the messag</param>
 		/// <param name="Args">The received data</param>e
-		private void RsyncOutput(object Sender, DataReceivedEventArgs Args)
+		/// <param name="bStdErr">whether the data was received on stderr</param>
+		private void RsyncOutput(DataReceivedEventArgs Args, bool bStdErr)
 		{
 			if(Args.Data != null)
 			{
-				Log.TraceInformation("  {0}", Args.Data);
+				if (bStdErr)
+				{
+					Log.TraceError("  {0}", Args.Data);
+				}
+				else
+				{
+					Log.TraceInformation("  {0}", Args.Data);
+				}
 			}
 		}
 
@@ -1141,13 +1151,14 @@ namespace UnrealBuildTool
 			string FullCommand = String.Format("cd {0} && {1}", EscapeShellArgument(WorkingDirectory), Command);
 			using(Process SSHProcess = new Process())
 			{
-				DataReceivedEventHandler OutputHandler = (E, Args) => { SshOutput(Args); };
+				DataReceivedEventHandler OutputHandler = (E, Args) => { SshOutput(Args, false); };
+				DataReceivedEventHandler ErrorHandler = (E, Args) => { SshOutput(Args, true); };
 
 				SSHProcess.StartInfo.FileName = SshExe.FullName;
 				SSHProcess.StartInfo.WorkingDirectory = SshExe.Directory.FullName;
-				SSHProcess.StartInfo.Arguments = String.Format("{0} {1}", String.Join(" ", CommonSshArguments), FullCommand.Replace("\"", "\\\""));
+				SSHProcess.StartInfo.Arguments = String.Format("{0} {1}", String.Join(" ", CommonSshArguments), FullCommand);
 				SSHProcess.OutputDataReceived += OutputHandler;
-				SSHProcess.ErrorDataReceived += OutputHandler;
+				SSHProcess.ErrorDataReceived += ErrorHandler;
 
 				Log.TraceLog("[SSH] {0} {1}", Utils.MakePathSafeToUseWithCommandLine(SSHProcess.StartInfo.FileName), SSHProcess.StartInfo.Arguments);
 				return Utils.RunLocalProcess(SSHProcess);
@@ -1158,12 +1169,20 @@ namespace UnrealBuildTool
 		/// Handler for output from running remote SSH commands
 		/// </summary>
 		/// <param name="Args"></param>
-		private void SshOutput(DataReceivedEventArgs Args)
+		/// <param name="bStdErr">whether the data was received on stderr</param>
+		private void SshOutput(DataReceivedEventArgs Args, bool bStdErr)
 		{
 			if(Args.Data != null)
 			{
 				string FormattedOutput = ConvertRemotePathsToLocal(Args.Data);
-				Log.TraceInformation("  {0}", FormattedOutput);
+				if (bStdErr)
+				{
+					Log.TraceError("  {0}", FormattedOutput);
+				}
+				else
+				{
+					Log.TraceInformation("  {0}", FormattedOutput);
+				}
 			}
 		}
 
