Fix for warning LNK4075 :  Ignoring /INCREMENTAL due to /RELEASE or /OPT:ICF in HoloLens and Windows builds

diff --git a/Engine/Source/Programs/UnrealBuildTool/Platform/HoloLens/HoloLensToolChain.cs b/Engine/Source/Programs/UnrealBuildTool/Platform/HoloLens/HoloLensToolChain.cs
--- a/Engine/Source/Programs/UnrealBuildTool/Platform/HoloLens/HoloLensToolChain.cs
+++ b/Engine/Source/Programs/UnrealBuildTool/Platform/HoloLens/HoloLensToolChain.cs
@@ -415,8 +415,8 @@ namespace UnrealBuildTool
 				Arguments.Add("/OPT:NOICF");
 			}
 
-			// Enable incremental linking if wanted.
-			if (LinkEnvironment.bUseIncrementalLinking)
+			// Enable incremental linking if wanted. ( avoid /INCREMENTAL getting ignored (LNK4075) due to /LTCG, /RELEASE, and /OPT:ICF )
+			if (LinkEnvironment.bUseIncrementalLinking && LinkEnvironment.Configuration != CppConfiguration.Shipping)
 			{
 				Arguments.Add("/INCREMENTAL");
 			}
diff --git a/Engine/Source/Programs/UnrealBuildTool/Platform/Windows/VCToolChain.cs b/Engine/Source/Programs/UnrealBuildTool/Platform/Windows/VCToolChain.cs
--- a/Engine/Source/Programs/UnrealBuildTool/Platform/Windows/VCToolChain.cs
+++ b/Engine/Source/Programs/UnrealBuildTool/Platform/Windows/VCToolChain.cs
@@ -941,8 +941,11 @@ namespace UnrealBuildTool
 				Arguments.Add("/OPT:NOICF");
 			}
 
-			// Enable incremental linking if wanted.
-			if (LinkEnvironment.bUseIncrementalLinking)
+			// Enable incremental linking if wanted. ( avoid /INCREMENTAL getting ignored (LNK4075) due to /LTCG, /RELEASE, and /OPT:ICF )
+			if (LinkEnvironment.bUseIncrementalLinking && 
+				LinkEnvironment.Configuration != CppConfiguration.Shipping && 
+				!Target.WindowsPlatform.bMergeIdenticalCOMDATs &&
+				!LinkEnvironment.bAllowLTCG)
 			{
 				Arguments.Add("/INCREMENTAL");
 				Arguments.Add("/verbose:incr");
