Fix-up translated process detection.

#rnx

diff --git a/Engine/Source/Programs/UnrealBuildTool/Platform/Mac/MacExports.cs b/Engine/Source/Programs/UnrealBuildTool/Platform/Mac/MacExports.cs
--- a/Engine/Source/Programs/UnrealBuildTool/Platform/Mac/MacExports.cs
+++ b/Engine/Source/Programs/UnrealBuildTool/Platform/Mac/MacExports.cs
@@ -106,7 +106,7 @@ namespace UnrealBuildTool
 				if (!IsRunningUnderRosettaVar.HasValue)
 				{
 					string TranslatedOutput = Utils.RunLocalProcessAndReturnStdOut("/usr/sbin/sysctl", "sysctl");
-					IsRunningUnderRosettaVar = TranslatedOutput.StartsWith("sysctl.proc_translated: 1");
+					IsRunningUnderRosettaVar = TranslatedOutput.Contains("sysctl.proc_translated: 1");
 				}
 
 				return IsRunningUnderRosettaVar.Value;
