- Fix linux compile error in webhelper app. Clang needs this to be a reference for the iterator (which makes sense)
- Include EpicWebHelper-Linux-Shipping in the "RequiredTools" section of the default makefile so it is compiled by default for linux users

diff --git a/Engine/Source/Programs/EpicWebHelper/Private/EpicWebHelperRemoteScripting.cpp b/Engine/Source/Programs/EpicWebHelper/Private/EpicWebHelperRemoteScripting.cpp
--- a/Engine/Source/Programs/EpicWebHelper/Private/EpicWebHelperRemoteScripting.cpp
+++ b/Engine/Source/Programs/EpicWebHelper/Private/EpicWebHelperRemoteScripting.cpp
@@ -366,7 +366,7 @@ void FEpicWebHelperRemoteScripting::OnContextCreated(CefRefPtr<CefBrowser> Brows
 
 void FEpicWebHelperRemoteScripting::OnContextReleased(CefRefPtr<CefBrowser> Browser, CefRefPtr<CefFrame> Frame, CefRefPtr<CefV8Context> Context)
 {
-	for(const TPair<FGuid,CefRefPtr<FEpicWebHelperRemoteObject>> RemoteObject : RemoteObjects)
+	for(const TPair<FGuid,CefRefPtr<FEpicWebHelperRemoteObject>> &RemoteObject : RemoteObjects)
 	{
 		if (RemoteObject.Value->Browser->IsSame(Browser)
 			&&  RemoteObject.Value->CreationContext->IsSame(Context))
diff --git a/Engine/Source/Programs/UnrealBuildTool/ProjectFiles/Make/MakefileGenerator.cs b/Engine/Source/Programs/UnrealBuildTool/ProjectFiles/Make/MakefileGenerator.cs
--- a/Engine/Source/Programs/UnrealBuildTool/ProjectFiles/Make/MakefileGenerator.cs
+++ b/Engine/Source/Programs/UnrealBuildTool/ProjectFiles/Make/MakefileGenerator.cs
@@ -128,7 +128,7 @@ namespace UnrealBuildTool
 
 			MakefileContent.Append("\n\n" + BuildCommand + ProjectBuildCommand + "\n" +
 				"all: StandardSet\n\n" +
-				"RequiredTools: CrashReportClient-Linux-Shipping CrashReportClientEditor-Linux-Shipping ShaderCompileWorker UnrealLightmass\n\n" +
+				"RequiredTools: CrashReportClient-Linux-Shipping CrashReportClientEditor-Linux-Shipping ShaderCompileWorker UnrealLightmass EpicWebHelper-Linux-Shipping\n\n" +
 				"StandardSet: RequiredTools UnrealFrontend UnrealEditor UnrealInsights\n\n" +
 				"DebugSet: RequiredTools UnrealFrontend-Linux-Debug UnrealEditor-Linux-Debug\n\n"
 			);
