Fix dependency checking for PCHs on Clang. Clang was configured to generate PCHs but not to consume them - though ultimately uses them anyway due to automatic translation of -include options.

diff --git a/Engine/Source/Programs/UnrealBuildTool/Platform/Windows/VCCompileAction.cs b/Engine/Source/Programs/UnrealBuildTool/Platform/Windows/VCCompileAction.cs
--- a/Engine/Source/Programs/UnrealBuildTool/Platform/Windows/VCCompileAction.cs
+++ b/Engine/Source/Programs/UnrealBuildTool/Platform/Windows/VCCompileAction.cs
@@ -420,7 +420,7 @@ namespace UnrealBuildTool
 				Arguments.Add(String.Format("/Yc\"{0}\"", PchThroughHeaderFile.FullName));
 				Arguments.Add(String.Format("/Fp\"{0}\"", CreatePchFile.FullName));
 			}
-			if (UsingPchFile != null)
+			if (UsingPchFile != null && CompilerType != WindowsCompiler.Clang)
 			{
 				Arguments.Add(String.Format("/Yu\"{0}\"", PchThroughHeaderFile.FullName));
 				Arguments.Add(String.Format("/Fp\"{0}\"", UsingPchFile.FullName));
diff --git a/Engine/Source/Programs/UnrealBuildTool/Platform/Windows/VCToolChain.cs b/Engine/Source/Programs/UnrealBuildTool/Platform/Windows/VCToolChain.cs
--- a/Engine/Source/Programs/UnrealBuildTool/Platform/Windows/VCToolChain.cs
+++ b/Engine/Source/Programs/UnrealBuildTool/Platform/Windows/VCToolChain.cs
@@ -1023,18 +1023,11 @@ namespace UnrealBuildTool
 			// If we're using precompiled headers, set that up now
 			if (CompileEnvironment.PrecompiledHeaderAction == PrecompiledHeaderAction.Include)
 			{
-				if (Target.WindowsPlatform.Compiler == WindowsCompiler.Clang)
-				{
-					BaseCompileAction.ForceIncludeFiles.Insert(0, FileItem.GetItemByFileReference(CompileEnvironment.PrecompiledHeaderFile.Location.ChangeExtension(null)));
-				}
-				else
-				{
-					FileItem IncludeHeader = FileItem.GetItemByFileReference(CompileEnvironment.PrecompiledHeaderIncludeFilename);
-					BaseCompileAction.ForceIncludeFiles.Insert(0, IncludeHeader);
+				FileItem IncludeHeader = FileItem.GetItemByFileReference(CompileEnvironment.PrecompiledHeaderIncludeFilename);
+				BaseCompileAction.ForceIncludeFiles.Insert(0, IncludeHeader);
 
-					BaseCompileAction.UsingPchFile = CompileEnvironment.PrecompiledHeaderFile;
-					BaseCompileAction.PchThroughHeaderFile = IncludeHeader;
-				}
+				BaseCompileAction.UsingPchFile = CompileEnvironment.PrecompiledHeaderFile;
+				BaseCompileAction.PchThroughHeaderFile = IncludeHeader;
 			}
 
 			// Generate the timing info
