Only set Clang Windows target for Windows.

#rnx

diff --git a/Engine/Source/Programs/UnrealBuildTool/Platform/Windows/VCToolChain.cs b/Engine/Source/Programs/UnrealBuildTool/Platform/Windows/VCToolChain.cs
--- a/Engine/Source/Programs/UnrealBuildTool/Platform/Windows/VCToolChain.cs
+++ b/Engine/Source/Programs/UnrealBuildTool/Platform/Windows/VCToolChain.cs
@@ -243,7 +243,7 @@ namespace UnrealBuildTool
 			Arguments.Add("/Oi");
 
 
-			if (Target.WindowsPlatform.Compiler == WindowsCompiler.Clang)
+			if (Target.WindowsPlatform.Compiler == WindowsCompiler.Clang && Target.Platform.IsInGroup(UnrealPlatformGroup.Windows))
 			{
 				// Tell the Clang compiler to generate 64-bit code
 				Arguments.Add("--target=x86_64-pc-windows-msvc");
