Removal some references to Mono

#trivial

diff --git a/Engine/Source/Programs/AutomationTool/Android/AndroidPlatform.Automation.cs b/Engine/Source/Programs/AutomationTool/Android/AndroidPlatform.Automation.cs
--- a/Engine/Source/Programs/AutomationTool/Android/AndroidPlatform.Automation.cs
+++ b/Engine/Source/Programs/AutomationTool/Android/AndroidPlatform.Automation.cs
@@ -422,7 +422,7 @@ public class AndroidPlatform : Platform
 
 	public static string GetStorageQueryCommand(bool bForcePC = false)
     {
-		if (!bForcePC && Utils.IsRunningOnMono)
+		if (!bForcePC && !Utils.IsRunningOnWindows)
 		{
 			return "shell 'echo $EXTERNAL_STORAGE'";
 		}
@@ -991,7 +991,7 @@ public class AndroidPlatform : Platform
 						File.WriteAllLines(SymbolizeBatchName, BatchLines);
 					}
 
-					if (Utils.IsRunningOnMono)
+					if (!Utils.IsRunningOnWindows)
 					{
 						if (bHaveAPK)
 						{
@@ -1560,7 +1560,7 @@ public class AndroidPlatform : Platform
 
 	public static IProcessResult RunAdbCommand(ProjectParams Params, string SerialNumber, string Args, string Input = null, ERunOptions Options = ERunOptions.Default, bool bShouldLogCommand = false)
 	{
-		string AdbCommand = Environment.ExpandEnvironmentVariables("%ANDROID_HOME%/platform-tools/adb" + (Utils.IsRunningOnMono ? "" : ".exe"));
+		string AdbCommand = Environment.ExpandEnvironmentVariables("%ANDROID_HOME%/platform-tools/adb" + (Utils.IsRunningOnWindows ? ".exe" : ""));
 		if (Options.HasFlag(ERunOptions.AllowSpew) || Options.HasFlag(ERunOptions.SpewIsVerbose))
 		{
 			LastSpewFilename = "";
@@ -1571,7 +1571,7 @@ public class AndroidPlatform : Platform
 
 	private string RunAndLogAdbCommand(ProjectParams Params, string SerialNumber, string Args, out int SuccessCode)
 	{
-		string AdbCommand = Environment.ExpandEnvironmentVariables("%ANDROID_HOME%/platform-tools/adb" + (Utils.IsRunningOnMono ? "" : ".exe"));
+		string AdbCommand = Environment.ExpandEnvironmentVariables("%ANDROID_HOME%/platform-tools/adb" + (Utils.IsRunningOnWindows ? ".exe" : ""));
 		LastSpewFilename = "";
 		return RunAndLog(CmdEnv, AdbCommand, GetAdbCommandLine(Params, SerialNumber, Args), out SuccessCode, SpewFilterCallback: new ProcessResult.SpewFilterCallbackType(ADBSpewFilter));
 	}
@@ -2419,7 +2419,7 @@ public class AndroidPlatform : Platform
 		uint BestVersion = 0;
 		foreach (string CandidateDir in Subdirs)
 		{
-			string AaptFilename = Path.Combine(CandidateDir, Utils.IsRunningOnMono ? "aapt" : "aapt.exe");
+			string AaptFilename = Path.Combine(CandidateDir, Utils.IsRunningOnWindows ? "aapt.exe" : "aapt");
 			uint RevisionValue = 0;
 
 			if (File.Exists(AaptFilename))
diff --git a/Engine/Source/Programs/AutomationTool/AutomationUtils/CommandEnvironment.cs b/Engine/Source/Programs/AutomationTool/AutomationUtils/CommandEnvironment.cs
--- a/Engine/Source/Programs/AutomationTool/AutomationUtils/CommandEnvironment.cs
+++ b/Engine/Source/Programs/AutomationTool/AutomationUtils/CommandEnvironment.cs
@@ -127,7 +127,7 @@ namespace AutomationTool
 
 			RobocopyExe = GetSystemExePath("robocopy.exe");
 			MountExe = GetSystemExePath("mount.exe");
-			CmdExe = Utils.IsRunningOnMono ? "/bin/sh" : GetSystemExePath("cmd.exe");
+			CmdExe = Utils.IsRunningOnWindows ? GetSystemExePath("cmd.exe") : "/bin/sh";
 			MallocNanoZone = "0";
 			CommandUtils.SetEnvVar(EnvVarNames.MacMallocNanoZone, MallocNanoZone);
 
diff --git a/Engine/Source/Programs/AutomationTool/AutomationUtils/CommandUtils.cs b/Engine/Source/Programs/AutomationTool/AutomationUtils/CommandUtils.cs
--- a/Engine/Source/Programs/AutomationTool/AutomationUtils/CommandUtils.cs
+++ b/Engine/Source/Programs/AutomationTool/AutomationUtils/CommandUtils.cs
@@ -2063,7 +2063,7 @@ namespace AutomationTool
                 }
                 else
                 {
-                    CachedRootBuildStorageDirectory = Utils.IsRunningOnMono ? "/Volumes/Builds" : CombinePaths("P:", "Builds");
+                    CachedRootBuildStorageDirectory = Utils.IsRunningOnWindows ? CombinePaths("P:", "Builds") : "/Volumes/Builds";
                 }
             }
             return CachedRootBuildStorageDirectory;
@@ -2233,7 +2233,7 @@ namespace AutomationTool
 		public static IEnumerable<string> LegacyUnzipFiles(string ZipFileName, string BaseDirectory)
 		{
 			List<string> OutputFileNames = new List<string>();
-			if (Utils.IsRunningOnMono)
+			if (!Utils.IsRunningOnWindows)
 			{
 				CommandUtils.CreateDirectory(BaseDirectory);
 
@@ -2933,9 +2933,9 @@ namespace AutomationTool
 		/// </summary>
 		public static void SignSingleExecutableIfEXEOrDLL(string Filename, bool bIgnoreExtension = false)
 		{
-            if (UnrealBuildTool.Utils.IsRunningOnMono)
+            if (!Utils.IsRunningOnWindows)
             {
-                CommandUtils.LogLog(String.Format("Can't sign '{0}', we are running under mono.", Filename));
+                CommandUtils.LogLog(String.Format("Can't sign '{0}' on non-Windows platform.", Filename));
                 return;
             }
             if (!CommandUtils.FileExists(Filename))
@@ -3126,9 +3126,9 @@ namespace AutomationTool
 
 		public static void SignMultipleFilesIfEXEOrDLL(List<FileReference> Files, bool bIgnoreExtension = false)
 		{
-			if (UnrealBuildTool.Utils.IsRunningOnMono)
+			if (!Utils.IsRunningOnWindows)
 			{
-				CommandUtils.LogLog(String.Format("Can't sign we are running under mono."));
+				CommandUtils.LogLog(String.Format("Can't sign on non-Windows platform."));
 				return;
 			}
 			List<FileReference> FinalFiles = new List<FileReference>();
diff --git a/Engine/Source/Programs/AutomationTool/IOS/IOSPlatform.Automation.cs b/Engine/Source/Programs/AutomationTool/IOS/IOSPlatform.Automation.cs
--- a/Engine/Source/Programs/AutomationTool/IOS/IOSPlatform.Automation.cs
+++ b/Engine/Source/Programs/AutomationTool/IOS/IOSPlatform.Automation.cs
@@ -1488,7 +1488,7 @@ public class IOSPlatform : Platform
 			PlatformGameConfig.GetBool("/Script/IOSRuntimeSettings.IOSRuntimeSettings", "bGenerateXCArchive", out bXCArchive);
 		}
 
-		if (bXCArchive && Utils.IsRunningOnMono)
+		if (bXCArchive && !Utils.IsRunningOnWindows)
 		{
 			// Always put the archive in the current user's Library/Developer/Xcode/Archives path if not on the build machine
 			WindowsIdentity id = WindowsIdentity.GetCurrent();
@@ -1656,7 +1656,7 @@ public class IOSPlatform : Platform
 			Text.AppendLine("</plist>");
 			File.WriteAllText(Path.Combine(ArchiveName, "Info.plist"), Text.ToString());
 		}
-		else if (bXCArchive && !Utils.IsRunningOnMono)
+		else if (bXCArchive && Utils.IsRunningOnWindows)
 		{
 			LogWarning("Can not produce an XCArchive on windows");
 		}
diff --git a/Engine/Source/Programs/AutomationTool/Linux/LinuxPlatform.Automation.cs b/Engine/Source/Programs/AutomationTool/Linux/LinuxPlatform.Automation.cs
--- a/Engine/Source/Programs/AutomationTool/Linux/LinuxPlatform.Automation.cs
+++ b/Engine/Source/Programs/AutomationTool/Linux/LinuxPlatform.Automation.cs
@@ -133,7 +133,7 @@ public abstract class BaseLinuxPlatform : Platform
 		// write out the 
 		FileReference.WriteAllText(IntermediateFile, Script.ToString());
 
-		if (Utils.IsRunningOnMono)
+		if (!Utils.IsRunningOnWindows)
 		{
 			var Result = CommandUtils.Run("env", string.Format("-- \"chmod\" \"+x\" \"{0}\"", IntermediateFile.ToString().Replace("'", "'\"'\"'")));
 			if (Result.ExitCode != 0)
diff --git a/Engine/Source/Programs/AutomationTool/Program.cs b/Engine/Source/Programs/AutomationTool/Program.cs
--- a/Engine/Source/Programs/AutomationTool/Program.cs
+++ b/Engine/Source/Programs/AutomationTool/Program.cs
@@ -122,7 +122,7 @@ namespace AutomationTool
                 NoThrow(() => CommandUtils.ClearDirStack(), "Clear Dir Stack");
 
                 // Try to kill process before app domain exits to leave the other KillAll call to extreme edge cases
-                NoThrow(() => { if (ShouldKillProcesses && !Utils.IsRunningOnMono) ProcessManager.KillAll(); }, "Kill All Processes");
+                NoThrow(() => { if (ShouldKillProcesses && Utils.IsRunningOnWindows) ProcessManager.KillAll(); }, "Kill All Processes");
 
 				// Write the exit code
                 Log.TraceInformation("AutomationTool exiting with ExitCode={0} ({1})", (int)ReturnCode, ReturnCode);
@@ -165,7 +165,7 @@ namespace AutomationTool
 		static void Domain_ProcessExit(object sender, EventArgs e)
 		{
 			// Kill all spawned processes (Console instead of Log because logging is closed at this time anyway)
-			if (ShouldKillProcesses && !Utils.IsRunningOnMono)
+			if (ShouldKillProcesses && Utils.IsRunningOnWindows)
 			{			
 				ProcessManager.KillAll();
 			}
diff --git a/Engine/Source/Programs/AutomationTool/Scripts/BuildCMakeLib.Automation.cs b/Engine/Source/Programs/AutomationTool/Scripts/BuildCMakeLib.Automation.cs
--- a/Engine/Source/Programs/AutomationTool/Scripts/BuildCMakeLib.Automation.cs
+++ b/Engine/Source/Programs/AutomationTool/Scripts/BuildCMakeLib.Automation.cs
@@ -564,23 +564,20 @@ public sealed class BuildCMakeLib : BuildCommand
 
 	private void SetupBuildEnvironment()
 	{
-		if (!Utils.IsRunningOnMono)
-		{
-			// ================================================================================
-			// ThirdPartyNotUE
-			// NOTE: these are Windows executables
-			if (BuildHostPlatform.Current.Platform == UnrealTargetPlatform.Win64)
-			{
-				string CMakePath = DirectoryReference.Combine(TargetPlatform.CMakeRootDirectory, "bin").FullName;
-				string MakePath = DirectoryReference.Combine(TargetPlatform.MakeRootDirectory, "bin").FullName;
-
-				string PrevPath = Environment.GetEnvironmentVariable("PATH");
-				// mixing bundled make and cygwin make is no good. Try to detect and remove cygwin paths.
-				string PathWithoutCygwin = RemoveOtherMakeAndCygwinFromPath(PrevPath);
-				Environment.SetEnvironmentVariable("PATH", CMakePath + ";" + MakePath + ";" + PathWithoutCygwin);
-				Environment.SetEnvironmentVariable("PATH", CMakePath + ";" + MakePath + ";" + Environment.GetEnvironmentVariable("PATH"));
-				LogInformation("set {0}={1}", "PATH", Environment.GetEnvironmentVariable("PATH"));
-			}
+		// ================================================================================
+		// ThirdPartyNotUE
+		// NOTE: these are Windows executables
+		if (BuildHostPlatform.Current.Platform == UnrealTargetPlatform.Win64)
+		{
+			string CMakePath = DirectoryReference.Combine(TargetPlatform.CMakeRootDirectory, "bin").FullName;
+			string MakePath = DirectoryReference.Combine(TargetPlatform.MakeRootDirectory, "bin").FullName;
+
+			string PrevPath = Environment.GetEnvironmentVariable("PATH");
+			// mixing bundled make and cygwin make is no good. Try to detect and remove cygwin paths.
+			string PathWithoutCygwin = RemoveOtherMakeAndCygwinFromPath(PrevPath);
+			Environment.SetEnvironmentVariable("PATH", CMakePath + ";" + MakePath + ";" + PathWithoutCygwin);
+			Environment.SetEnvironmentVariable("PATH", CMakePath + ";" + MakePath + ";" + Environment.GetEnvironmentVariable("PATH"));
+			LogInformation("set {0}={1}", "PATH", Environment.GetEnvironmentVariable("PATH"));
 		}
 	}
 
diff --git a/Engine/Source/Programs/AutomationTool/Scripts/BuildPhysX.Automation.cs b/Engine/Source/Programs/AutomationTool/Scripts/BuildPhysX.Automation.cs
--- a/Engine/Source/Programs/AutomationTool/Scripts/BuildPhysX.Automation.cs
+++ b/Engine/Source/Programs/AutomationTool/Scripts/BuildPhysX.Automation.cs
@@ -600,24 +600,21 @@ public sealed class BuildPhysX : BuildCommand
 
 	private void SetupBuildEnvironment()
 	{
-		if (!Utils.IsRunningOnMono)
+		// ================================================================================
+		// ThirdPartyNotUE
+		// NOTE: these are Windows executables
+		if (BuildHostPlatform.Current.Platform == UnrealTargetPlatform.Win64)
 		{
-			// ================================================================================
-			// ThirdPartyNotUE
-			// NOTE: these are Windows executables
-			if (BuildHostPlatform.Current.Platform == UnrealTargetPlatform.Win64)
-			{
-				DirectoryReference ThirdPartyNotUERootDirectory = DirectoryReference.Combine(RootDirectory, "Engine/Extras/ThirdPartyNotUE");
-				string CMakePath = DirectoryReference.Combine(ThirdPartyNotUERootDirectory, "CMake/bin").FullName;
-				string MakePath = DirectoryReference.Combine(ThirdPartyNotUERootDirectory, "GNU_Make/make-3.81/bin").FullName;
-
-				string PrevPath = Environment.GetEnvironmentVariable("PATH");
-				// mixing bundled make and cygwin make is no good. Try to detect and remove cygwin paths.
-				string PathWithoutCygwin = RemoveOtherMakeAndCygwinFromPath(PrevPath);
-				Environment.SetEnvironmentVariable("PATH", CMakePath + ";" + MakePath + ";" + PathWithoutCygwin);
-				Environment.SetEnvironmentVariable("PATH", CMakePath + ";" + MakePath + ";" + Environment.GetEnvironmentVariable("PATH"));
-				LogInformation("set {0}={1}", "PATH", Environment.GetEnvironmentVariable("PATH"));
-			}
+			DirectoryReference ThirdPartyNotUERootDirectory = DirectoryReference.Combine(RootDirectory, "Engine/Extras/ThirdPartyNotUE");
+			string CMakePath = DirectoryReference.Combine(ThirdPartyNotUERootDirectory, "CMake/bin").FullName;
+			string MakePath = DirectoryReference.Combine(ThirdPartyNotUERootDirectory, "GNU_Make/make-3.81/bin").FullName;
+
+			string PrevPath = Environment.GetEnvironmentVariable("PATH");
+			// mixing bundled make and cygwin make is no good. Try to detect and remove cygwin paths.
+			string PathWithoutCygwin = RemoveOtherMakeAndCygwinFromPath(PrevPath);
+			Environment.SetEnvironmentVariable("PATH", CMakePath + ";" + MakePath + ";" + PathWithoutCygwin);
+			Environment.SetEnvironmentVariable("PATH", CMakePath + ";" + MakePath + ";" + Environment.GetEnvironmentVariable("PATH"));
+			LogInformation("set {0}={1}", "PATH", Environment.GetEnvironmentVariable("PATH"));
 		}
 	}
 
diff --git a/Engine/Source/Programs/AutomationTool/Scripts/RebuildHLODCommand.Automation.cs b/Engine/Source/Programs/AutomationTool/Scripts/RebuildHLODCommand.Automation.cs
--- a/Engine/Source/Programs/AutomationTool/Scripts/RebuildHLODCommand.Automation.cs
+++ b/Engine/Source/Programs/AutomationTool/Scripts/RebuildHLODCommand.Automation.cs
@@ -318,10 +318,8 @@ namespace AutomationScripts.Automation
                 Message.Attachments.Add()*/
 			    try
 			    {
-				    #pragma warning disable CS0618 // Mono 4.6.x obsoletes this class
 				    SmtpClient MailClient = new SmtpClient("smtp.epicgames.net");
 				    MailClient.Send(Message);
-				    #pragma warning restore CS0618
 			    }
 			    catch (Exception Ex)
 			    {
diff --git a/Engine/Source/Programs/AutomationTool/Scripts/RebuildLightMapsCommand.Automation.cs b/Engine/Source/Programs/AutomationTool/Scripts/RebuildLightMapsCommand.Automation.cs
--- a/Engine/Source/Programs/AutomationTool/Scripts/RebuildLightMapsCommand.Automation.cs
+++ b/Engine/Source/Programs/AutomationTool/Scripts/RebuildLightMapsCommand.Automation.cs
@@ -278,10 +278,8 @@ namespace AutomationScripts.Automation
 				Message.Attachments.Add()*/
 				try
 				{
-#pragma warning disable CS0618 // Mono 4.6.x obsoletes this class
 					SmtpClient MailClient = new SmtpClient("smtp.epicgames.net");
 					MailClient.Send(Message);
-#pragma warning restore CS0618
 				}
 				catch (Exception Ex)
 				{
diff --git a/Engine/Source/Programs/AutomationTool/Scripts/ResavePackagesCommand.Automation.cs b/Engine/Source/Programs/AutomationTool/Scripts/ResavePackagesCommand.Automation.cs
--- a/Engine/Source/Programs/AutomationTool/Scripts/ResavePackagesCommand.Automation.cs
+++ b/Engine/Source/Programs/AutomationTool/Scripts/ResavePackagesCommand.Automation.cs
@@ -273,10 +273,8 @@ namespace AutomationScripts.Automation
 				Message.Attachments.Add()*/
 				try
 				{
-#pragma warning disable CS0618 // Mono 4.6.x obsoletes this class
 					SmtpClient MailClient = new SmtpClient("smtp.epicgames.net");
 					MailClient.Send(Message);
-#pragma warning restore CS0618
 				}
 				catch (Exception Ex)
 				{
diff --git a/Engine/Source/Programs/AutomationTool/Scripts/ZipUtils.Automation.cs b/Engine/Source/Programs/AutomationTool/Scripts/ZipUtils.Automation.cs
--- a/Engine/Source/Programs/AutomationTool/Scripts/ZipUtils.Automation.cs
+++ b/Engine/Source/Programs/AutomationTool/Scripts/ZipUtils.Automation.cs
@@ -158,7 +158,7 @@ public class ZipUtils : BuildCommand
 					}
 				}
 
-				if (UnrealBuildTool.Utils.IsRunningOnMono && CommandUtils.IsProbablyAMacOrIOSExe(OutputFileName))
+				if (!Utils.IsRunningOnWindows && CommandUtils.IsProbablyAMacOrIOSExe(OutputFileName))
 				{
 					FixUnixFilePermissions(OutputFileName);
 				}
diff --git a/Engine/Source/Programs/AutomationTool/Win/WinResources.cs b/Engine/Source/Programs/AutomationTool/Win/WinResources.cs
--- a/Engine/Source/Programs/AutomationTool/Win/WinResources.cs
+++ b/Engine/Source/Programs/AutomationTool/Win/WinResources.cs
@@ -6,10 +6,6 @@ using System.Text;
 using System.IO;
 using System.Runtime.InteropServices;
 
-#if !__MonoCS__
-#pragma warning disable CS1591
-#endif
-
 public enum ResourceType
 {
 	Icon = 3,
@@ -287,6 +283,3 @@ public class ModuleResourceUpdate : IDisposable
 	}
 }
 
-#if !__MonoCS__
-#pragma warning restore CS1591
-#endif
diff --git a/Engine/Source/Programs/Shared/EpicGames.Core/ManagedProcess.cs b/Engine/Source/Programs/Shared/EpicGames.Core/ManagedProcess.cs
--- a/Engine/Source/Programs/Shared/EpicGames.Core/ManagedProcess.cs
+++ b/Engine/Source/Programs/Shared/EpicGames.Core/ManagedProcess.cs
@@ -592,7 +592,7 @@ namespace EpicGames.Core
 		/// <param name="Flags">Flags for the new process</param>
 		private void CreateManagedProcessPortable(ManagedProcessGroup? Group, string FileName, string CommandLine, string? WorkingDirectory, IReadOnlyDictionary<string, string>? Environment, ProcessPriorityClass Priority, ManagedProcessFlags Flags)
 		{
-			// Fallback for Mono platforms
+			// for non-Windows platforms
 			FrameworkProcess = new Process();
 			FrameworkProcess.StartInfo.FileName = FileName;
 			FrameworkProcess.StartInfo.Arguments = CommandLine;
diff --git a/Engine/Source/Programs/UnrealBuildTool/Configuration/TargetRules.cs b/Engine/Source/Programs/UnrealBuildTool/Configuration/TargetRules.cs
--- a/Engine/Source/Programs/UnrealBuildTool/Configuration/TargetRules.cs
+++ b/Engine/Source/Programs/UnrealBuildTool/Configuration/TargetRules.cs
@@ -1877,9 +1877,7 @@ namespace UnrealBuildTool
 		/// Accessors for fields on the inner TargetRules instance
 		/// </summary>
 		#region Read-only accessor properties
-		#if !__MonoCS__
 		#pragma warning disable CS1591
-		#endif
 
 		public string Name
 		{
@@ -2837,9 +2835,7 @@ namespace UnrealBuildTool
 			get { return Inner.OptedInModulePlatforms; } 
 		}
 
-#if !__MonoCS__
 #pragma warning restore C1591
-#endif
 		#endregion
 
 		/// <summary>
diff --git a/Engine/Source/Programs/UnrealBuildTool/Configuration/UEBuildTarget.cs b/Engine/Source/Programs/UnrealBuildTool/Configuration/UEBuildTarget.cs
--- a/Engine/Source/Programs/UnrealBuildTool/Configuration/UEBuildTarget.cs
+++ b/Engine/Source/Programs/UnrealBuildTool/Configuration/UEBuildTarget.cs
@@ -3741,7 +3741,7 @@ namespace UnrealBuildTool
 				UEBuildBinary ExecutableBinary = Binaries[0];
 
 				bool IsCurrentPlatform;
-				if (Utils.IsRunningOnMono)
+				if (!Utils.IsRunningOnWindows)
 				{
 					IsCurrentPlatform = Platform == UnrealTargetPlatform.Mac || (UEBuildPlatform.IsPlatformInGroup(Platform, UnrealPlatformGroup.Unix) && Platform == BuildHostPlatform.Current.Platform);
 				}
diff --git a/Engine/Source/Programs/UnrealBuildTool/Platform/Android/AndroidPlatformSDK.cs b/Engine/Source/Programs/UnrealBuildTool/Platform/Android/AndroidPlatformSDK.cs
--- a/Engine/Source/Programs/UnrealBuildTool/Platform/Android/AndroidPlatformSDK.cs
+++ b/Engine/Source/Programs/UnrealBuildTool/Platform/Android/AndroidPlatformSDK.cs
@@ -216,8 +216,8 @@ namespace UnrealBuildTool
 					}
 				}
 
-				// If we are on Mono and we are still missing a key then go and find it from the .bash_profile
-				if (Utils.IsRunningOnMono && !EnvVarNames.All(s => AndroidEnv.ContainsKey(s.Key)))
+				// If we are not running on Windows and we are still missing a key then go and find it from the .bash_profile
+				if (!Utils.IsRunningOnWindows && !EnvVarNames.All(s => AndroidEnv.ContainsKey(s.Key)))
 				{
 					string BashProfilePath = Path.Combine(Environment.GetFolderPath(Environment.SpecialFolder.Personal), ".bash_profile");
 					if (!File.Exists(BashProfilePath))
diff --git a/Engine/Source/Programs/UnrealBuildTool/Platform/Android/AndroidToolChain.cs b/Engine/Source/Programs/UnrealBuildTool/Platform/Android/AndroidToolChain.cs
--- a/Engine/Source/Programs/UnrealBuildTool/Platform/Android/AndroidToolChain.cs
+++ b/Engine/Source/Programs/UnrealBuildTool/Platform/Android/AndroidToolChain.cs
@@ -1291,7 +1291,7 @@ namespace UnrealBuildTool
 			foreach (string Definition in CompileEnvironment.Definitions)
 			{
 				// We must change the \" to escaped double quotes to save it properly for clang .rsp
-				BaseArguments += string.Format(" -D \"{0}\"", Definition.Replace("\"", Utils.IsRunningOnMono ? "\\\"" : "\"\""));
+				BaseArguments += string.Format(" -D \"{0}\"", Definition.Replace("\"", Utils.IsRunningOnWindows ? "\"\"" : "\\\""));
 			}
 
 			//LUMIN_MERGE
@@ -2157,7 +2157,7 @@ namespace UnrealBuildTool
 					// To work around this we call clang to just get the command it would execute and generate a
 					// second response file to directly call ld with the right arguments instead of calling through clang.
 /* disable while tracking down some linker errors this introduces
-					if (!Utils.IsRunningOnMono)
+					if (Utils.IsRunningOnWindows)
 					{
 						// capture the actual link command without running it
 						ProcessStartInfo StartInfo = new ProcessStartInfo();
diff --git a/Engine/Source/Programs/UnrealBuildTool/Platform/Android/UEBuildAndroid.cs b/Engine/Source/Programs/UnrealBuildTool/Platform/Android/UEBuildAndroid.cs
--- a/Engine/Source/Programs/UnrealBuildTool/Platform/Android/UEBuildAndroid.cs
+++ b/Engine/Source/Programs/UnrealBuildTool/Platform/Android/UEBuildAndroid.cs
@@ -51,9 +51,7 @@ namespace UnrealBuildTool
 		/// Accessors for fields on the inner TargetRules instance
 		/// </summary>
 		#region Read-only accessor properties 
-		#if !__MonoCS__
 		#pragma warning disable CS1591
-		#endif
 
 		public IReadOnlyList<string> Architectures
 		{
@@ -65,9 +63,7 @@ namespace UnrealBuildTool
 			get { return Inner.GPUArchitectures.AsReadOnly(); }
 		}
 
-		#if !__MonoCS__
 		#pragma warning restore CS1591
-		#endif
 		#endregion
 	}
 
diff --git a/Engine/Source/Programs/UnrealBuildTool/Platform/Android/UEDeployAndroid.cs b/Engine/Source/Programs/UnrealBuildTool/Platform/Android/UEDeployAndroid.cs
--- a/Engine/Source/Programs/UnrealBuildTool/Platform/Android/UEDeployAndroid.cs
+++ b/Engine/Source/Programs/UnrealBuildTool/Platform/Android/UEDeployAndroid.cs
@@ -328,7 +328,7 @@ namespace UnrealBuildTool
 			uint BestVersion = 0;
 			foreach (string CandidateDir in Subdirs)
 			{
-				string AaptFilename = Path.Combine(CandidateDir, Utils.IsRunningOnMono ? "aapt" : "aapt.exe");
+				string AaptFilename = Path.Combine(CandidateDir, Utils.IsRunningOnWindows ? "aapt.exe" : "aapt");
 				string RevisionString = "";
 				uint RevisionValue = 0;
 
@@ -3621,7 +3621,7 @@ namespace UnrealBuildTool
 			string BuildToolsVersion = GetBuildToolsVersion();
 
 			// cache some tools paths
-			//string NDKBuildPath = Environment.ExpandEnvironmentVariables("%NDKROOT%/ndk-build" + (Utils.IsRunningOnMono ? "" : ".cmd"));
+			//string NDKBuildPath = Environment.ExpandEnvironmentVariables("%NDKROOT%/ndk-build" + (Utils.IsRunningOnWindows ? ".cmd" : ""));
 			//bool HasNDKPath = File.Exists(NDKBuildPath);
 
 			// set up some directory info
@@ -4200,7 +4200,7 @@ namespace UnrealBuildTool
 				if (!bSkipGradleBuild)
 				{
 					string GradleScriptPath = Path.Combine(UnrealBuildGradlePath, "gradlew");
-					if (Utils.IsRunningOnMono)
+					if (!Utils.IsRunningOnWindows)
 					{
 						// fix permissions for Mac/Linux
 						RunCommandLineProgramWithException(UnrealBuildGradlePath, "/bin/sh", string.Format("-c 'chmod 0755 \"{0}\"'", GradleScriptPath.Replace("'", "'\"'\"'")), "Fix gradlew permissions");
@@ -4233,9 +4233,9 @@ namespace UnrealBuildTool
 							Directory.CreateDirectory(Path.GetDirectoryName(DestApkName));
 
 							// Use gradle to build the .apk file
-							string ShellExecutable = Utils.IsRunningOnMono ? "/bin/sh" : "cmd.exe";
-							string ShellParametersBegin = Utils.IsRunningOnMono ? "-c '" : "/c ";
-							string ShellParametersEnd = Utils.IsRunningOnMono ? "'" : "";
+							string ShellExecutable = Utils.IsRunningOnWindows ? "cmd.exe" : "/bin/sh";
+							string ShellParametersBegin = Utils.IsRunningOnWindows ? "/c " : "-c '";
+							string ShellParametersEnd = Utils.IsRunningOnWindows ? "" : "'";
 							RunCommandLineProgramWithExceptionAndFiltering(UnrealBuildGradlePath, ShellExecutable, ShellParametersBegin + "\"" + GradleScriptPath + "\" " + GradleOptions + ShellParametersEnd, "Making .apk with Gradle...");
 
 							if (GradleSecondCallOptions != "")
@@ -4581,7 +4581,7 @@ namespace UnrealBuildTool
 					File.WriteAllText(GradlePropertiesFilename, GradlePropertiesContent);
 
 					string GradleScriptPath = Path.Combine(UnrealBuildGradlePath, "gradlew");
-					if (Utils.IsRunningOnMono)
+					if (!Utils.IsRunningOnWindows)
 					{
 						// fix permissions for Mac/Linux
 						RunCommandLineProgramWithException(UnrealBuildGradlePath, "/bin/sh", string.Format("-c 'chmod 0755 \"{0}\"'", GradleScriptPath.Replace("'", "'\"'\"'")), "Fix gradlew permissions");
@@ -4608,9 +4608,9 @@ namespace UnrealBuildTool
 					//Directory.CreateDirectory(Path.GetDirectoryName(DestApkName));
 
 					// Use gradle to build the .apk file
-					string ShellExecutable = Utils.IsRunningOnMono ? "/bin/sh" : "cmd.exe";
-					string ShellParametersBegin = Utils.IsRunningOnMono ? "-c '" : "/c ";
-					string ShellParametersEnd = Utils.IsRunningOnMono ? "'" : "";
+					string ShellExecutable = Utils.IsRunningOnWindows ? "cmd.exe" : "/bin/sh";
+					string ShellParametersBegin = Utils.IsRunningOnWindows ? "/c " : "-c '";
+					string ShellParametersEnd = Utils.IsRunningOnWindows ? "" : "'";
 					RunCommandLineProgramWithExceptionAndFiltering(UnrealBuildGradlePath, ShellExecutable, ShellParametersBegin + "\"" + GradleScriptPath + "\" " + GradleOptions + ShellParametersEnd, "Making .aab with Gradle...");
 
 					if (GradleSecondCallOptions != "")
@@ -4644,7 +4644,7 @@ namespace UnrealBuildTool
 						string UnrealBuildGradlePath = Path.Combine(UnrealBuildPath, "gradle");
 
 						string GradleScriptPath = Path.Combine(UnrealBuildGradlePath, "gradlew");
-						if (Utils.IsRunningOnMono)
+						if (!Utils.IsRunningOnWindows)
 						{
 							// fix permissions for Mac/Linux
 							RunCommandLineProgramWithException(UnrealBuildGradlePath, "/bin/sh", string.Format("-c 'chmod 0755 \"{0}\"'", GradleScriptPath.Replace("'", "'\"'\"'")), "Fix gradlew permissions");
@@ -4671,9 +4671,9 @@ namespace UnrealBuildTool
 						//Directory.CreateDirectory(Path.GetDirectoryName(DestApkName));
 
 						// Use gradle to build the .apk file
-						string ShellExecutable = Utils.IsRunningOnMono ? "/bin/sh" : "cmd.exe";
-						string ShellParametersBegin = Utils.IsRunningOnMono ? "-c '" : "/c ";
-						string ShellParametersEnd = Utils.IsRunningOnMono ? "'" : "";
+						string ShellExecutable = Utils.IsRunningOnWindows ? "cmd.exe" : "/bin/sh";
+						string ShellParametersBegin = Utils.IsRunningOnWindows ? "/c " : "-c '";
+						string ShellParametersEnd = Utils.IsRunningOnWindows ? "" : "'";
 						RunCommandLineProgramWithExceptionAndFiltering(UnrealBuildGradlePath, ShellExecutable, ShellParametersBegin + "\"" + GradleScriptPath + "\" " + GradleOptions + ShellParametersEnd, "Making .aab with Gradle...");
 
 						if (GradleSecondCallOptions != "")
diff --git a/Engine/Source/Programs/UnrealBuildTool/Platform/HoloLens/UEBuildHoloLens.cs b/Engine/Source/Programs/UnrealBuildTool/Platform/HoloLens/UEBuildHoloLens.cs
--- a/Engine/Source/Programs/UnrealBuildTool/Platform/HoloLens/UEBuildHoloLens.cs
+++ b/Engine/Source/Programs/UnrealBuildTool/Platform/HoloLens/UEBuildHoloLens.cs
@@ -110,9 +110,7 @@ namespace UnrealBuildTool
 		/// Accessors for fields on the inner TargetRules instance
 		/// </summary>
 		#region Read-only accessor properties 
-#if !__MonoCS__
 #pragma warning disable CS1591
-#endif
 		public WindowsCompiler Compiler
 		{
 			get { return Inner.Compiler; }
@@ -150,9 +148,7 @@ namespace UnrealBuildTool
 		{
 			get { return Inner.NativeCodeAnalysisRuleset; }
 		}
-#if !__MonoCS__
 #pragma warning restore CS1591
-#endif
 		#endregion
 	}
 
diff --git a/Engine/Source/Programs/UnrealBuildTool/Platform/IOS/IOSToolChain.cs b/Engine/Source/Programs/UnrealBuildTool/Platform/IOS/IOSToolChain.cs
--- a/Engine/Source/Programs/UnrealBuildTool/Platform/IOS/IOSToolChain.cs
+++ b/Engine/Source/Programs/UnrealBuildTool/Platform/IOS/IOSToolChain.cs
@@ -204,7 +204,7 @@ namespace UnrealBuildTool
 			Result += " -pipe";
 			Result += " -fpascal-strings";
 
-			if (!Utils.IsRunningOnMono)
+			if (Utils.IsRunningOnWindows)
 			{
 				Result += " -fdiagnostics-format=msvc";
 			}
diff --git a/Engine/Source/Programs/UnrealBuildTool/Platform/IOS/UEBuildIOS.cs b/Engine/Source/Programs/UnrealBuildTool/Platform/IOS/UEBuildIOS.cs
--- a/Engine/Source/Programs/UnrealBuildTool/Platform/IOS/UEBuildIOS.cs
+++ b/Engine/Source/Programs/UnrealBuildTool/Platform/IOS/UEBuildIOS.cs
@@ -105,9 +105,7 @@ namespace UnrealBuildTool
 		/// Accessors for fields on the inner TargetRules instance
 		/// </summary>
 		#region Read-only accessor properties 
-#if !__MonoCS__
 #pragma warning disable CS1591
-#endif
 		public bool bStripSymbols
 		{
 			get { return Inner.bStripSymbols; }
@@ -163,9 +161,7 @@ namespace UnrealBuildTool
 			get { return float.Parse(Inner.ProjectSettings.RuntimeVersion, System.Globalization.CultureInfo.InvariantCulture); }
 		}
 		
-#if !__MonoCS__
 #pragma warning restore CS1591
-#endif
 		#endregion
 	}
 
diff --git a/Engine/Source/Programs/UnrealBuildTool/Platform/Linux/UEBuildLinux.cs b/Engine/Source/Programs/UnrealBuildTool/Platform/Linux/UEBuildLinux.cs
--- a/Engine/Source/Programs/UnrealBuildTool/Platform/Linux/UEBuildLinux.cs
+++ b/Engine/Source/Programs/UnrealBuildTool/Platform/Linux/UEBuildLinux.cs
@@ -91,9 +91,7 @@ namespace UnrealBuildTool
 		/// Accessors for fields on the inner TargetRules instance
 		/// </summary>
 		#region Read-only accessor properties 
-		#if !__MonoCS__
 		#pragma warning disable CS1591
-		#endif
 
 		public bool bPreservePSYM
 		{
@@ -125,9 +123,7 @@ namespace UnrealBuildTool
 			get { return Inner.bEnableThinLTO; }
 		}
 
-		#if !__MonoCS__
 		#pragma warning restore CS1591
-		#endif
 		#endregion
 	}
 
diff --git a/Engine/Source/Programs/UnrealBuildTool/Platform/Lumin/LuminToolChain.cs b/Engine/Source/Programs/UnrealBuildTool/Platform/Lumin/LuminToolChain.cs
--- a/Engine/Source/Programs/UnrealBuildTool/Platform/Lumin/LuminToolChain.cs
+++ b/Engine/Source/Programs/UnrealBuildTool/Platform/Lumin/LuminToolChain.cs
@@ -35,7 +35,7 @@ namespace UnrealBuildTool
 			{
 				if (MabuPath_ == null)
 				{
-					MabuPath_ = Environment.ExpandEnvironmentVariables("%MLSDK%/mabu" + (Utils.IsRunningOnMono ? "" : ".cmd"));
+					MabuPath_ = Environment.ExpandEnvironmentVariables("%MLSDK%/mabu" + (Utils.IsRunningOnWindows ? ".cmd" : ""));
 
 					MabuPath_ = MabuPath_.Replace("\"", "");
 
diff --git a/Engine/Source/Programs/UnrealBuildTool/Platform/Lumin/UEBuildLumin.cs b/Engine/Source/Programs/UnrealBuildTool/Platform/Lumin/UEBuildLumin.cs
--- a/Engine/Source/Programs/UnrealBuildTool/Platform/Lumin/UEBuildLumin.cs
+++ b/Engine/Source/Programs/UnrealBuildTool/Platform/Lumin/UEBuildLumin.cs
@@ -41,9 +41,7 @@ namespace UnrealBuildTool
 		/// Accessors for fields on the inner TargetRules instance
 		/// </summary>
 		#region Read-only accessor properties 
-#if !__MonoCS__
 #pragma warning disable CS1591
-#endif
 
 		public IReadOnlyList<string> GPUArchitectures
 		{
@@ -54,9 +52,7 @@ namespace UnrealBuildTool
 			get { return Inner.bForDistribution; }
 		}
 
-#if !__MonoCS__
 #pragma warning restore CS1591
-#endif
 		#endregion
 	}
 
diff --git a/Engine/Source/Programs/UnrealBuildTool/Platform/Mac/MacExports.cs b/Engine/Source/Programs/UnrealBuildTool/Platform/Mac/MacExports.cs
--- a/Engine/Source/Programs/UnrealBuildTool/Platform/Mac/MacExports.cs
+++ b/Engine/Source/Programs/UnrealBuildTool/Platform/Mac/MacExports.cs
@@ -114,7 +114,7 @@ namespace UnrealBuildTool
 		}
 
 		/// <summary>
-		/// Returns true if we're running on Apple architecture (either natively which mono will do, or under Rosetta)
+		/// Returns true if we're running on Apple architecture (either natively which mono/dotnet will do, or under Rosetta)
 		/// </summary>
 		/// <returns></returns>
 		public static bool IsRunningOnAppleArchitecture
diff --git a/Engine/Source/Programs/UnrealBuildTool/Platform/Mac/MacToolChain.cs b/Engine/Source/Programs/UnrealBuildTool/Platform/Mac/MacToolChain.cs
--- a/Engine/Source/Programs/UnrealBuildTool/Platform/Mac/MacToolChain.cs
+++ b/Engine/Source/Programs/UnrealBuildTool/Platform/Mac/MacToolChain.cs
@@ -782,7 +782,7 @@ namespace UnrealBuildTool
 
 			// Tell the action that we're building an import library here and it should conditionally be
 			// ignored as a prerequisite for other actions
-			LinkAction.bProducesImportLibrary = !Utils.IsRunningOnMono && (bBuildImportLibraryOnly || LinkEnvironment.bIsBuildingDLL);
+			LinkAction.bProducesImportLibrary = Utils.IsRunningOnWindows && (bBuildImportLibraryOnly || LinkEnvironment.bIsBuildingDLL);
 			
 			// Add the output file as a production of the link action.
 			FileItem OutputFile = FileItem.GetItemByFileReference(LinkEnvironment.OutputFilePath);
@@ -927,7 +927,7 @@ namespace UnrealBuildTool
 
 			if (!bIsBuildingLibrary)
 			{
-				if (UnrealBuildTool.IsEngineInstalled() || (Utils.IsRunningOnMono && LinkEnvironment.bIsCrossReferenced == false))
+				if (UnrealBuildTool.IsEngineInstalled() || (!Utils.IsRunningOnWindows && LinkEnvironment.bIsCrossReferenced == false))
 				{
 					foreach (string Library in EngineAndGameLibraries)
 					{
@@ -1002,7 +1002,7 @@ namespace UnrealBuildTool
 
 				StreamWriter FixDylibDepsScript = File.AppendText(FixDylibDepsScriptPath.FullName);
 
-				if (LinkEnvironment.bIsCrossReferenced || !Utils.IsRunningOnMono)
+				if (LinkEnvironment.bIsCrossReferenced || Utils.IsRunningOnWindows)
 				{
 					string EngineAndGameLibrariesString = "";
 					foreach (string Library in EngineAndGameLibraries)
diff --git a/Engine/Source/Programs/UnrealBuildTool/Platform/Mac/UEBuildMac.cs b/Engine/Source/Programs/UnrealBuildTool/Platform/Mac/UEBuildMac.cs
--- a/Engine/Source/Programs/UnrealBuildTool/Platform/Mac/UEBuildMac.cs
+++ b/Engine/Source/Programs/UnrealBuildTool/Platform/Mac/UEBuildMac.cs
@@ -70,9 +70,7 @@ namespace UnrealBuildTool
 		/// Accessors for fields on the inner TargetRules instance
 		/// </summary>
 		#region Read-only accessor properties 
-		#if !__MonoCS__
 		#pragma warning disable CS1591
-		#endif
 
 		public bool? bGenerateDsymFile
 		{
@@ -94,9 +92,7 @@ namespace UnrealBuildTool
 			get { return Inner.bEnableUndefinedBehaviorSanitizer; }
 		}		
 
-		#if !__MonoCS__
 		#pragma warning restore CS1591
-		#endif
 		#endregion
 	}
 
diff --git a/Engine/Source/Programs/UnrealBuildTool/Platform/Windows/MicrosoftPlatformSDK.cs b/Engine/Source/Programs/UnrealBuildTool/Platform/Windows/MicrosoftPlatformSDK.cs
--- a/Engine/Source/Programs/UnrealBuildTool/Platform/Windows/MicrosoftPlatformSDK.cs
+++ b/Engine/Source/Programs/UnrealBuildTool/Platform/Windows/MicrosoftPlatformSDK.cs
@@ -42,7 +42,7 @@ namespace UnrealBuildTool
 
 		public override string? GetInstalledSDKVersion()
 		{
-			if (Utils.IsRunningOnMono)
+			if (!Utils.IsRunningOnWindows)
 			{
 				return null;
 			}
diff --git a/Engine/Source/Programs/UnrealBuildTool/Platform/Windows/UEBuildWindows.cs b/Engine/Source/Programs/UnrealBuildTool/Platform/Windows/UEBuildWindows.cs
--- a/Engine/Source/Programs/UnrealBuildTool/Platform/Windows/UEBuildWindows.cs
+++ b/Engine/Source/Programs/UnrealBuildTool/Platform/Windows/UEBuildWindows.cs
@@ -437,9 +437,7 @@ namespace UnrealBuildTool
 		/// Accessors for fields on the inner TargetRules instance
 		/// </summary>
 		#region Read-only accessor properties 
-		#if !__MonoCS__
 		#pragma warning disable CS1591
-		#endif
 
 		public WindowsCompiler Compiler
 		{
@@ -636,9 +634,7 @@ namespace UnrealBuildTool
 			return WindowsExports.GetArchitectureSubpath(Architecture);
 		}
 
-		#if !__MonoCS__
 		#pragma warning restore CS1591
-		#endif
 		#endregion
 	}
 
diff --git a/Engine/Source/Programs/UnrealBuildTool/System/BuildVersion.cs b/Engine/Source/Programs/UnrealBuildTool/System/BuildVersion.cs
--- a/Engine/Source/Programs/UnrealBuildTool/System/BuildVersion.cs
+++ b/Engine/Source/Programs/UnrealBuildTool/System/BuildVersion.cs
@@ -272,9 +272,7 @@ namespace UnrealBuildTool
 		/// Accessors for fields on the inner BuildVersion instance
 		/// </summary>
 		#region Read-only accessor properties 
-		#if !__MonoCS__
 		#pragma warning disable CS1591
-		#endif
 
 		public int MajorVersion
 		{
@@ -326,9 +324,7 @@ namespace UnrealBuildTool
 			get { return Inner.BuildVersionString; }
 		}
 
-		#if !__MonoCS__
 		#pragma warning restore C1591
-		#endif
 		#endregion
 	}
 }
diff --git a/Engine/Source/Programs/UnrealBuildTool/ToolChain/AppleToolChain.cs b/Engine/Source/Programs/UnrealBuildTool/ToolChain/AppleToolChain.cs
--- a/Engine/Source/Programs/UnrealBuildTool/ToolChain/AppleToolChain.cs
+++ b/Engine/Source/Programs/UnrealBuildTool/ToolChain/AppleToolChain.cs
@@ -145,7 +145,7 @@ namespace UnrealBuildTool
 			}
 
 			// make sure we have a valid SDK directory
-			if (Utils.IsRunningOnMono && !Directory.Exists(Path.Combine(BaseSDKDir, OSPrefix + PlatformSDKVersion + ".sdk")))
+			if (!Utils.IsRunningOnWindows && !Directory.Exists(Path.Combine(BaseSDKDir, OSPrefix + PlatformSDKVersion + ".sdk")))
 			{
 				throw new BuildException("Invalid SDK {0}{1}.sdk, not found in {2}", OSPrefix, PlatformSDKVersion, BaseSDKDir);
 			}
