ProjectFileGenerator.cs:

IsRunningOnMono -> !IsRunningOnWindows

#trivial

diff --git a/Engine/Source/Programs/UnrealBuildTool/ProjectFiles/ProjectFileGenerator.cs b/Engine/Source/Programs/UnrealBuildTool/ProjectFiles/ProjectFileGenerator.cs
--- a/Engine/Source/Programs/UnrealBuildTool/ProjectFiles/ProjectFileGenerator.cs
+++ b/Engine/Source/Programs/UnrealBuildTool/ProjectFiles/ProjectFileGenerator.cs
@@ -1776,7 +1776,7 @@ namespace UnrealBuildTool
 		{
 			if(ShouldGenerateIntelliSenseData() && Targets.Count > 0)
 			{
-				string ProgressInfoText = Utils.IsRunningOnMono ? "Generating data for project indexing..." : "Binding IntelliSense data...";
+				string ProgressInfoText = Utils.IsRunningOnWindows ? "Binding IntelliSense data..." : "Generating data for project indexing...";
 				using(ProgressWriter Progress = new ProgressWriter(ProgressInfoText, true))
 				{
 					for(int TargetIndex = 0; TargetIndex < Targets.Count; ++TargetIndex)
