UnrealBuildTool/System/Utils.cs:

Remove special case for old mono version.

#trivial

diff --git a/Engine/Source/Programs/UnrealBuildTool/System/Utils.cs b/Engine/Source/Programs/UnrealBuildTool/System/Utils.cs
--- a/Engine/Source/Programs/UnrealBuildTool/System/Utils.cs
+++ b/Engine/Source/Programs/UnrealBuildTool/System/Utils.cs
@@ -776,23 +776,6 @@ namespace UnrealBuildTool
 				RelativePath = RelativePath.Substring(0, RelativePath.Length - 1);
 			}
 
-			// Uri.MakeRelativeUri is broken in Mono 2.x and sometimes returns broken path
-			if (IsRunningOnMono)
-			{
-				// Check if result is correct
-				string TestPath = Path.GetFullPath(Path.Combine(AbsoluteRelativeDirectory, RelativePath));
-				string AbsoluteTestPath = CollapseRelativeDirectories(AbsolutePath);
-				if (TestPath != AbsoluteTestPath)
-				{
-					TestPath += "/";
-					if (TestPath != AbsoluteTestPath)
-					{
-						// Fix the path. @todo Mac: replace this hack with something better
-						RelativePath = "../" + RelativePath;
-					}
-				}
-			}
-
 			return RelativePath;
 		}
 
@@ -935,10 +918,6 @@ namespace UnrealBuildTool
 			// due to multithreading on Windows, lock the object
 			lock (p)
 			{
-				// Mono has a specific requirement if testing for an alive process
-				if (IsRunningOnMono)
-					return p.Handle != IntPtr.Zero; // native handle to the process
-				// on Windows, simply test the process ID to be non-zero. 
 				// note that this can fail and have a race condition in threads, but the framework throws an exception when this occurs.
 				try
 				{
@@ -1063,8 +1042,8 @@ namespace UnrealBuildTool
 		/// <returns>The number of logical cores.</returns>
 		public static int GetLogicalProcessorCount()
 		{
-			// This function uses Windows P/Invoke calls; if we're on Mono, just return the default.
-			if(!Utils.IsRunningOnMono)
+			// This function uses Windows P/Invoke calls; if we're not running on Windows, just return the default.
+			if(Utils.IsRunningOnWindows)
 			{
 				const int ERROR_INSUFFICIENT_BUFFER = 122;
 
@@ -1116,8 +1095,8 @@ namespace UnrealBuildTool
 		/// <returns>The number of physical cores, or -1 if it could not be obtained</returns>
 		public static int GetPhysicalProcessorCount()
 		{
-			// This function uses Windows P/Invoke calls; if we're on Mono, just fail.
-			if (Utils.IsRunningOnMono)
+			// This function uses Windows P/Invoke calls; if we're not running on Windows, just fail.
+			if (!Utils.IsRunningOnWindows)
 			{
 				return -1;
 			}
