VSCodeProjectFileGenerator.cs:

IsRunningOnMono -> !IsRunningOnWindows

#trivial

diff --git a/Engine/Source/Programs/UnrealBuildTool/ProjectFiles/VisualStudioCode/VSCodeProjectFileGenerator.cs b/Engine/Source/Programs/UnrealBuildTool/ProjectFiles/VisualStudioCode/VSCodeProjectFileGenerator.cs
--- a/Engine/Source/Programs/UnrealBuildTool/ProjectFiles/VisualStudioCode/VSCodeProjectFileGenerator.cs
+++ b/Engine/Source/Programs/UnrealBuildTool/ProjectFiles/VisualStudioCode/VSCodeProjectFileGenerator.cs
@@ -938,7 +938,7 @@ namespace UnrealBuildTool
 						{
 							OutFile.AddField("label", TaskName);
 							OutFile.AddField("group", "build");
-							if (Utils.IsRunningOnMono)
+							if (!Utils.IsRunningOnWindows)
 							{
 								OutFile.AddField("command", MakePathString(FileReference.Combine(UE4ProjectRoot, "Engine", "Build", "BatchFiles", HostPlatform.ToString(), "RunDotnet.sh")));
 							}
@@ -1495,7 +1495,7 @@ namespace UnrealBuildTool
 					WorkspaceFile.AddUnnamedField("ms-dotnettools.csharp");
 
 					// If the platform we run the generator on uses mono, there are additional debugging extensions to add.
-					if (Utils.IsRunningOnMono)
+					if (!Utils.IsRunningOnWindows)
 					{
 						WorkspaceFile.AddUnnamedField("vadimcn.vscode-lldb");
 						WorkspaceFile.AddUnnamedField("ms-vscode.mono-debug");
