UnrealBuildTool: Remove SupportsDotNetCoreProjects()

VS2019 is the minimum required VS version, and has support for NET Core projects - we no longer need this check.

diff --git a/Engine/Source/Programs/UnrealBuildTool/ProjectFiles/ProjectFileGenerator.cs b/Engine/Source/Programs/UnrealBuildTool/ProjectFiles/ProjectFileGenerator.cs
--- a/Engine/Source/Programs/UnrealBuildTool/ProjectFiles/ProjectFileGenerator.cs
+++ b/Engine/Source/Programs/UnrealBuildTool/ProjectFiles/ProjectFileGenerator.cs
@@ -427,8 +427,6 @@ namespace UnrealBuildTool
 				if (FileReference.Exists(ProjectFile))
 				{
 					VCSharpProjectFile Project = new VCSharpProjectFile(ProjectFile);
-					if (Project.IsDotNETCoreProject() && !SupportsDotnetCoreProjects())
-						continue;
 
 					Project.ShouldBuildForAllSolutionTargets = false;//true;
 					AddExistingProjectFile(Project, bForceDevelopmentConfiguration: true);
@@ -449,36 +447,28 @@ namespace UnrealBuildTool
 		/// </summary>
 		void AddSharedDotNetModules(MasterProjectFolder ProgramsFolder)
 		{
-			if (SupportsDotnetCoreProjects())
-			{
 				DirectoryInfo DotNetDir = new DirectoryInfo(DirectoryReference.Combine(UnrealBuildTool.EngineSourceDirectory, "Programs", "Shared").FullName);
-				if (DotNetDir.Exists)
+			if (DotNetDir.Exists)
+			{
+				List<FileInfo> ProjectFiles = new List<FileInfo>();
+				foreach (DirectoryInfo ProjectDir in DotNetDir.EnumerateDirectories())
 				{
-					List<FileInfo> ProjectFiles = new List<FileInfo>();
-					foreach (DirectoryInfo ProjectDir in DotNetDir.EnumerateDirectories())
-					{
-						ProjectFiles.AddRange(ProjectDir.EnumerateFiles("*.csproj"));
-					}
-					if (ProjectFiles.Count > 0)
+					ProjectFiles.AddRange(ProjectDir.EnumerateFiles("*.csproj"));
+				}
+				if (ProjectFiles.Count > 0)
+				{
+					MasterProjectFolder Folder = ProgramsFolder.AddSubFolder("Shared");
+					foreach (FileInfo ProjectFile in ProjectFiles)
 					{
-						MasterProjectFolder Folder = ProgramsFolder.AddSubFolder("Shared");
-						foreach (FileInfo ProjectFile in ProjectFiles)
-						{
-							VCSharpProjectFile Project = new VCSharpProjectFile(new FileReference(ProjectFile));
-							Project.ShouldBuildForAllSolutionTargets = false;
-							AddExistingProjectFile(Project, bForceDevelopmentConfiguration: true);
-							Folder.ChildProjects.Add(Project);
-						}
+						VCSharpProjectFile Project = new VCSharpProjectFile(new FileReference(ProjectFile));
+						Project.ShouldBuildForAllSolutionTargets = false;
+						AddExistingProjectFile(Project, bForceDevelopmentConfiguration: true);
+						Folder.ChildProjects.Add(Project);
 					}
 				}
 			}
 		}
 
-		protected virtual bool SupportsDotnetCoreProjects()
-		{
-			return true;
-		}
-
 		/// <summary>
 		/// Creates a .props file next to each automation project which specifies the path to the engine directory
 		/// </summary>
@@ -541,9 +531,6 @@ namespace UnrealBuildTool
 						{
 							VCSharpProjectFile Project = new VCSharpProjectFile(FoundProject);
 
-							if (Project.IsDotNETCoreProject() && !SupportsDotnetCoreProjects())
-								continue;
-
 							Project.ShouldBuildForAllSolutionTargets = true;
 							Project.ShouldBuildByDefaultForSolutionTargets = true;
 
@@ -1629,10 +1616,6 @@ namespace UnrealBuildTool
 		/// </summary>
 		private void AddUnrealBuildToolProject(MasterProjectFolder ProgramsFolder)
 		{
-			// UBT is a dotnet core project so we can not add it if the project generator does not support dotnet core projects
-			if (!SupportsDotnetCoreProjects())
-				return;
-
 			List<string> ProjectDirectoryNames = new List<string>();
 			ProjectDirectoryNames.Add("UnrealBuildTool");
 
@@ -1682,9 +1665,6 @@ namespace UnrealBuildTool
 			{
 				Project = new VCSharpProjectFile(ProjectFileName);
 
-				if (Project.IsDotNETCoreProject() && !SupportsDotnetCoreProjects())
-					return null;
-
 				Project.ShouldBuildForAllSolutionTargets = bShouldBuildForAllSolutionTargets;
 				Project.ShouldBuildByDefaultForSolutionTargets = bShouldBuildByDefaultForSolutionTargets;
 				AddExistingProjectFile(Project, bForceDevelopmentConfiguration: bForceDevelopmentConfiguration);
diff --git a/Engine/Source/Programs/UnrealBuildTool/ProjectFiles/VisualStudio/VCProjectFileGenerator.cs b/Engine/Source/Programs/UnrealBuildTool/ProjectFiles/VisualStudio/VCProjectFileGenerator.cs
--- a/Engine/Source/Programs/UnrealBuildTool/ProjectFiles/VisualStudio/VCProjectFileGenerator.cs
+++ b/Engine/Source/Programs/UnrealBuildTool/ProjectFiles/VisualStudio/VCProjectFileGenerator.cs
@@ -259,25 +259,6 @@ namespace UnrealBuildTool
 			return string.Empty;
 		}
 
-		protected override bool SupportsDotnetCoreProjects()
-		{
-			// we use dotnet 3 which requires visual studio 2019 or newer
-			switch (Settings.ProjectFileFormat)
-			{
-				case VCProjectFileFormat.VisualStudio2019:
-					return true;
-				case VCProjectFileFormat.Default:
-				case VCProjectFileFormat.VisualStudio2012:
-				case VCProjectFileFormat.VisualStudio2013:
-				case VCProjectFileFormat.VisualStudio2015:
-				case VCProjectFileFormat.VisualStudio2017:
-					return false;
-				default:
-					throw new ArgumentOutOfRangeException(nameof(ProjectFileFormat), Settings.ProjectFileFormat, null);
-			}
-		}
-
-
 		static public void AppendPlatformToolsetProperty(StringBuilder VCProjectFileContent, VCProjectFileFormat ProjectFileFormat)
 		{
 			VCProjectFileContent.AppendLine("    <PlatformToolset>{0}</PlatformToolset>", GetProjectFilePlatformToolsetVersionString(ProjectFileFormat));
