UBT: Fix exception trying to read makefile with null directoryreference objects.

diff --git a/Engine/Source/Programs/Shared/EpicGames.Core/DirectoryReference.cs b/Engine/Source/Programs/Shared/EpicGames.Core/DirectoryReference.cs
--- a/Engine/Source/Programs/Shared/EpicGames.Core/DirectoryReference.cs
+++ b/Engine/Source/Programs/Shared/EpicGames.Core/DirectoryReference.cs
@@ -473,26 +473,9 @@ namespace EpicGames.Core
 		/// </summary>
 		/// <param name="Writer">The writer to output data to</param>
 		/// <param name="Directory">The item to write</param>
-		public static void WriteDirectoryReference(this BinaryArchiveWriter Writer, DirectoryReference Directory)
+		public static void WriteDirectoryReference(this BinaryArchiveWriter Writer, DirectoryReference? Directory)
 		{
-			Writer.WriteNullableDirectoryReference(Directory);
-		}
-
-		/// <summary>
-		/// Writes a directory reference  to a binary archive
-		/// </summary>
-		/// <param name="Writer">The writer to output data to</param>
-		/// <param name="Directory">The item to write</param>
-		public static void WriteNullableDirectoryReference(this BinaryArchiveWriter Writer, DirectoryReference? Directory)
-		{
-			if (Directory == null)
-			{
-				Writer.WriteString(null);
-			}
-			else
-			{
-				Writer.WriteString(Directory.FullName);
-			}
+			Writer.WriteString(Directory?.FullName);
 		}
 
 		/// <summary>
diff --git a/Engine/Source/Programs/UnrealBuildTool/System/ConfigValueTracker.cs b/Engine/Source/Programs/UnrealBuildTool/System/ConfigValueTracker.cs
--- a/Engine/Source/Programs/UnrealBuildTool/System/ConfigValueTracker.cs
+++ b/Engine/Source/Programs/UnrealBuildTool/System/ConfigValueTracker.cs
@@ -24,7 +24,7 @@ namespace UnrealBuildTool
 		/// <summary>
 		/// Project directory to read config files from
 		/// </summary>
-		public DirectoryReference ProjectDir;
+		public DirectoryReference? ProjectDir;
 		
 		/// <summary>
 		/// The platform being built
@@ -49,7 +49,7 @@ namespace UnrealBuildTool
 		/// <param name="Platform">The platform being built</param>
 		/// <param name="SectionName">The section name</param>
 		/// <param name="KeyName">The key name</param>
-		public ConfigDependencyKey(ConfigHierarchyType Type, DirectoryReference ProjectDir, UnrealTargetPlatform Platform, string SectionName, string KeyName)
+		public ConfigDependencyKey(ConfigHierarchyType Type, DirectoryReference? ProjectDir, UnrealTargetPlatform Platform, string SectionName, string KeyName)
 		{
 			this.Type = Type;
 			this.ProjectDir = ProjectDir;
@@ -65,7 +65,7 @@ namespace UnrealBuildTool
 		public ConfigDependencyKey(BinaryArchiveReader Reader)
 		{
 			Type = (ConfigHierarchyType)Reader.ReadInt();
-			ProjectDir = Reader.ReadDirectoryReferenceNotNull();
+			ProjectDir = Reader.ReadDirectoryReference();
 			Platform = Reader.ReadUnrealTargetPlatform();
 			SectionName = Reader.ReadString();
 			KeyName = Reader.ReadString();
