package localization fixes

diff --git a/Engine/Source/Programs/UnrealBuildTool/Platform/Windows/VCManifestGenerator.cs b/Engine/Source/Programs/UnrealBuildTool/Platform/Windows/VCManifestGenerator.cs
--- a/Engine/Source/Programs/UnrealBuildTool/Platform/Windows/VCManifestGenerator.cs
+++ b/Engine/Source/Programs/UnrealBuildTool/Platform/Windows/VCManifestGenerator.cs
@@ -22,6 +22,7 @@ namespace UnrealBuildTool
 
 		protected virtual string IniSection_PlatformTargetSettings { get { return string.Format( "/Script/{0}PlatformEditor.{0}TargetSettings", Platform.ToString() ); } }
 		protected virtual string IniSection_GeneralProjectSettings { get { return "/Script/EngineSettings.GeneralProjectSettings"; } }
+		protected virtual bool bPlatformSupportsPriResourcePacks { get { return false; } }
 
 		protected const string BuildResourceSubPath = "Resources";
 		protected const string EngineResourceSubPath = "DefaultImages";
@@ -217,7 +218,6 @@ namespace UnrealBuildTool
 			if (ReturnVal == null || ReturnVal.Length <= 0)
 			{
 				Log.TraceError("Invalid application ID {0}. Application IDs must only contain letters and numbers. And they must begin with a letter.", InApplicationId);
-				Log.TraceError("Consider using the setting [{0}]:PackageName to provide a specific value.", IniSection_PlatformTargetSettings);
 			}
 			return ReturnVal;
 		}
@@ -593,7 +593,8 @@ namespace UnrealBuildTool
 		protected string GetIdentityPackageName()
 		{
             // Read the PackageName from config
-            var PackageName = Regex.Replace(GetConfigString("PackageName", "ProjectName", "DefaultUE4Project"), "[^-.A-Za-z0-9]", "");
+			var DefaultName = (ProjectFile != null) ? ProjectFile.GetFileNameWithoutAnyExtensions() : "DefaultUEProject";
+            var PackageName = Regex.Replace(GetConfigString("PackageName", "ProjectName", DefaultName), "[^-.A-Za-z0-9]", "");
             if (string.IsNullOrWhiteSpace(PackageName))
             {
                 Log.TraceError("Invalid package name {0}. Package names must only contain letters, numbers, dash, and period and must be at least one character long.", PackageName);
@@ -659,6 +660,11 @@ namespace UnrealBuildTool
 
 		protected abstract string GetMakePriBinaryPath();
 
+		protected virtual string GetMakePriExtraCommandLine()
+		{
+			return "";
+		}
+
 		protected abstract XElement GetManifest(List<UnrealTargetConfiguration> TargetConfigs, List<string> Executables, out string IdentityName);
 
 		protected virtual void ProcessManifest(List<UnrealTargetConfiguration> TargetConfigs, List<string> Executables, string ManifestName, string ManifestTargetPath, string ManifestIntermediatePath)
@@ -850,15 +856,16 @@ namespace UnrealBuildTool
 				string ResourceConfigFile = Path.Combine(IntermediatePath, "priconfig.xml");
 				bool bEnableAutoResourcePacks = false;
 				EngineIni.GetBool(IniSection_PlatformTargetSettings, "bEnableAutoResourcePacks", out bEnableAutoResourcePacks);
+				bEnableAutoResourcePacks &= bPlatformSupportsPriResourcePacks;
 
 				// If the game is not going to support language resource packs then merge the culture qualifiers.
 				if (bEnableAutoResourcePacks || CulturesToStage.Count <= 1)
 				{
-					RunExternalProgram(PriExecutable, "createconfig /cf \"" + ResourceConfigFile + "\" /dq " + DefaultCulture + " /o");
+					RunExternalProgram(PriExecutable, "createconfig /cf \"" + ResourceConfigFile + "\" /dq " + DefaultCulture + " /o " + GetMakePriExtraCommandLine() );
 				}
 				else
 				{
-					RunExternalProgram(PriExecutable, "createconfig /cf \"" + ResourceConfigFile + "\" /dq " + String.Join("_", CulturesToStage) + " /o");
+					RunExternalProgram(PriExecutable, "createconfig /cf \"" + ResourceConfigFile + "\" /dq " + String.Join("_", CulturesToStage) + " /o " + GetMakePriExtraCommandLine() );
 				}
 
 				// Modify configuration to restrict indexing to the Resources directory (saves time and space)
