Fixed two cases where blocked directories in ExcludedIncludePaths in BuildConfiguration.xml weren't always followed.

1) If the user specified a path as Bar\ then "Foo\Bar" would still be allowed (but not Foo\Bar\File.h) due to the trailing directory slash.

2) Code that deals with platform specific includes (ALiasedFiles) would omit the included file but still include the path in the AdditionalIncludeDirectories entry

diff --git a/Engine/Source/Programs/UnrealBuildTool/ProjectFiles/VisualStudio/VCProject.cs b/Engine/Source/Programs/UnrealBuildTool/ProjectFiles/VisualStudio/VCProject.cs
--- a/Engine/Source/Programs/UnrealBuildTool/ProjectFiles/VisualStudio/VCProject.cs
+++ b/Engine/Source/Programs/UnrealBuildTool/ProjectFiles/VisualStudio/VCProject.cs
@@ -729,11 +729,19 @@ namespace UnrealBuildTool
 				FilteredList = CleanPaths.ToArray();
 			}
 
+			// The user might have specified Foo, \Foo, Foo\, or \Foo\ as excludes.
+			// Directory paths don't contain a trailing slash so add one to the string
+			// we'll compare to our list so we can Contains and EndsWith to catch:
+			// <Path>\Foo\Dir
+			// <Path>\Foo
+			string IncludePathWithSeparator = IncludePath.FullName + Path.DirectorySeparatorChar;
+
 			if (FilteredList.Length > 0)
 			{
 				foreach (string Entry in FilteredList)
 				{
-					if (IncludePath.FullName.Contains(Entry))
+					if (IncludePathWithSeparator.Contains(Entry, StringComparison.OrdinalIgnoreCase) 
+						|| IncludePathWithSeparator.EndsWith(Entry, StringComparison.OrdinalIgnoreCase))
 					{
 						return true;
 					}
@@ -1118,13 +1126,19 @@ namespace UnrealBuildTool
 						FiltersFileIsNeeded = EnsureFilterPathExists(AliasedFile.ProjectPath, VCFiltersFileContent, FilterDirectories);
 					}
 
+					// get the filetype as represented to Visual Studio
 					string VCFileType = GetVCFileType(AliasedFile.FileSystemPath);
+
+					// if the filetype is an include and its path is filtered out, skip it entirely (should we do this for any type of
+					// file? Possibly, but not today due to potential fallout)
+					if (VCFileType == "ClInclude" && IncludePathIsFilteredOut(new DirectoryReference(AliasedFile.FileSystemPath)))
+					{
+						continue;
+					}
+
 					if (VCFileType != "ClCompile")
 					{
-						if (!IncludePathIsFilteredOut(new DirectoryReference(AliasedFile.FileSystemPath)))
-						{
-							VCProjectFileContent.AppendLine("    <{0} Include=\"{1}\"/>", VCFileType, EscapeFileName(AliasedFile.FileSystemPath));
-						}
+						VCProjectFileContent.AppendLine("    <{0} Include=\"{1}\"/>", VCFileType, EscapeFileName(AliasedFile.FileSystemPath));
 					}
 					else
 					{
