Fix HoloLens build errors in Development builds (missing PIX imports).

diff --git a/Engine/Source/Programs/UnrealBuildTool/Platform/HoloLens/UEBuildHoloLens.cs b/Engine/Source/Programs/UnrealBuildTool/Platform/HoloLens/UEBuildHoloLens.cs
--- a/Engine/Source/Programs/UnrealBuildTool/Platform/HoloLens/UEBuildHoloLens.cs
+++ b/Engine/Source/Programs/UnrealBuildTool/Platform/HoloLens/UEBuildHoloLens.cs
@@ -499,11 +499,13 @@ namespace UnrealBuildTool
 			{
 				Rules.PublicDefinitions.Add("D3D11_WITH_DWMAPI=0");
 				Rules.PublicDefinitions.Add("WITH_DX_PERF=0");
+				Rules.PrivateDependencyModuleNames.Add("WinPixEventRuntime");
 			}
 			else if (ModuleName == "D3D12RHI")
 			{
 				// To enable platform specific D3D12 RHI Types
 				Rules.PrivateIncludePaths.Add("Runtime/D3D12RHI/Private/HoloLens");
+				Rules.PrivateDependencyModuleNames.Add("WinPixEventRuntime");
 			}
 			else if (ModuleName == "DX11")
 			{
@@ -516,6 +518,7 @@ namespace UnrealBuildTool
 				Rules.PublicAdditionalLibraries.Clear();
 				Rules.PublicAdditionalLibraries.Remove("X3DAudio.lib");
 				Rules.PublicAdditionalLibraries.Remove("XAPOFX.lib");
+				Rules.PrivateDependencyModuleNames.Add("WinPixEventRuntime");
 			}
 			else if (ModuleName == "XAudio2")
 			{
