Avoid having bCompileAPEX and bCompileNvCloth enabled always on Linux build targets if the architecture started with x86_64.

diff --git a/Engine/Source/Programs/UnrealBuildTool/Platform/Linux/UEBuildLinux.cs b/Engine/Source/Programs/UnrealBuildTool/Platform/Linux/UEBuildLinux.cs
--- a/Engine/Source/Programs/UnrealBuildTool/Platform/Linux/UEBuildLinux.cs
+++ b/Engine/Source/Programs/UnrealBuildTool/Platform/Linux/UEBuildLinux.cs
@@ -230,8 +230,8 @@ namespace UnrealBuildTool
 			}
 
 			// depends on arch, APEX cannot be as of November'16 compiled for AArch32/64
-			Target.bCompileAPEX = Target.Architecture.StartsWith("x86_64");
-			Target.bCompileNvCloth = Target.Architecture.StartsWith("x86_64");
+			Target.bCompileAPEX = Target.bCompileAPEX && Target.Architecture.StartsWith("x86_64");
+			Target.bCompileNvCloth = Target.bCompileNvCloth && Target.Architecture.StartsWith("x86_64");
 
 			if (Target.GlobalDefinitions.Contains("USE_NULL_RHI=1"))
 			{				
