UnrealBuildTool: Slightly optimize generated visual studio project to reduce memory usage

* Filter out numerous source files that match */Source/ThirdParty/* that should have already been filtered, this does not include Module.Build.cs or .tps files, or if bGatherThirdPartySource is enabled
* Filter out source files that match */third_party/* by default
* Add option bIncludeDotNetPrograms to remove all ..NET projects, doing so saved me around 400MB. Note if you disable the .NET projects UnrealBuildTool won't build before compiling (but GenerateProjectFiles.bat will still build it)

Attempted changes that did not reduce memory usage by a significant amount:
* Filtering out project configurations
* Filtering out Engine/Shaders or Engine/Content/Editor/Templates
* Increasing the shared include size, even making it large enough to contain every include path didn't really make a difference. Also it can't really be increased anyway because that entire property is added to the process environment when starting a build and there's a max environment size of around 32k

Other Fixes:
* Fix vs2019 ToolVersionString 15.0 -> 16.0
* Add VCProjectVersion to Project globals
* Update UniqueIdentifier GUIDs in projects to be stable by hashing the directory path and using that hash as the GUID
* Don't write ProjectConfiguration if filtered out (Didn't affect memory usage)
* Add optional configuration filter for Debug, DebugGame, & Development

diff --git a/Engine/Source/Programs/UnrealBuildTool/ProjectFiles/ProjectFileGenerator.cs b/Engine/Source/Programs/UnrealBuildTool/ProjectFiles/ProjectFileGenerator.cs
--- a/Engine/Source/Programs/UnrealBuildTool/ProjectFiles/ProjectFileGenerator.cs
+++ b/Engine/Source/Programs/UnrealBuildTool/ProjectFiles/ProjectFileGenerator.cs
@@ -26,7 +26,7 @@ namespace UnrealBuildTool
 		/// </summary>
 		/// <param name="InitOwnerProjectFileGenerator">Project file generator that owns this object</param>
 		/// <param name="InitFolderName">Name for this folder</param>
-		public MasterProjectFolder( ProjectFileGenerator InitOwnerProjectFileGenerator, string InitFolderName )
+		public MasterProjectFolder(ProjectFileGenerator InitOwnerProjectFileGenerator, string InitFolderName)
 		{
 			OwnerProjectFileGenerator = InitOwnerProjectFileGenerator;
 			FolderName = InitFolderName;
@@ -45,7 +45,7 @@ namespace UnrealBuildTool
 		/// </summary>
 		/// <param name="SubFolderName">Name of the new folder</param>
 		/// <returns>The newly-added folder</returns>
-		public MasterProjectFolder AddSubFolder( string SubFolderName )
+		public MasterProjectFolder AddSubFolder(string SubFolderName)
 		{
 			MasterProjectFolder ResultFolder = null;
 
@@ -84,20 +84,20 @@ namespace UnrealBuildTool
 		/// </summary>
 		/// <param name="Project">The project file to look for</param>
 		/// <returns>The found folder that the project is in, or null</returns>
-		public MasterProjectFolder FindFolderForProject( ProjectFile Project )
+		public MasterProjectFolder FindFolderForProject(ProjectFile Project)
 		{
-			foreach( MasterProjectFolder CurFolder in SubFolders )
+			foreach (MasterProjectFolder CurFolder in SubFolders)
 			{
-				MasterProjectFolder FoundFolder = CurFolder.FindFolderForProject( Project );
-				if( FoundFolder != null )
+				MasterProjectFolder FoundFolder = CurFolder.FindFolderForProject(Project);
+				if (FoundFolder != null)
 				{
 					return FoundFolder;
 				}
 			}
 
-			foreach( ProjectFile ChildProject in ChildProjects )
+			foreach (ProjectFile ChildProject in ChildProjects)
 			{
-				if( ChildProject == Project )
+				if (ChildProject == Project)
 				{
 					return this;
 				}
@@ -117,7 +117,7 @@ namespace UnrealBuildTool
 
 		/// Files in this folder.  These are files that aren't part of any project, but display in the IDE under the project folder
 		/// and can be browsed/opened by the user easily in the user interface
-		public readonly List<string> Files = new List<string>();		
+		public readonly List<string> Files = new List<string>();
 	}
 
 	/// <summary>
@@ -142,9 +142,9 @@ namespace UnrealBuildTool
 		VisualStudioMac,
 		CLion,
 		Rider
-		#if __VPROJECT_AVAILABLE__
+#if __VPROJECT_AVAILABLE__
 					, VProject
-		#endif
+#endif
 	}
 
 	/// <summary>
@@ -165,10 +165,10 @@ namespace UnrealBuildTool
 		/// <returns>Sequence of project file formats</returns>
 		public static IEnumerable<ProjectFileFormat> ParseFormatList(string Formats)
 		{
-			foreach(string FormatName in Formats.Split('+').Select(x => x.Trim()))
+			foreach (string FormatName in Formats.Split('+').Select(x => x.Trim()))
 			{
 				ProjectFileFormat Format;
-				if(Enum.TryParse(FormatName, true, out Format))
+				if (Enum.TryParse(FormatName, true, out Format))
 				{
 					yield return Format;
 				}
@@ -222,6 +222,18 @@ namespace UnrealBuildTool
 		[XmlConfigFile]
 		public static bool bIncludeTestAndShippingConfigs = true;
 
+		/// <summary>
+		/// Whether we should include configurations for "Debug" and "DebugGame" in generated projects. Pass "-NoDebugConfigs" to disable this.
+		/// </summary>
+		[XmlConfigFile]
+		public static bool bIncludeDebugConfigs = true;
+
+		/// <summary>
+		/// Whether we should include configurations for "Development" in generated projects. Pass "-NoDevelopmentConfigs" to disable this.
+		/// </summary>
+		[XmlConfigFile]
+		public static bool bIncludeDevelopmentConfigs = true;
+
 		/// <summary>
 		/// True if intellisense data should be generated (takes a while longer).
 		/// </summary>
@@ -288,6 +300,12 @@ namespace UnrealBuildTool
 		[XmlConfigFile]
 		protected bool IncludeEnginePrograms = true;
 
+		/// <summary>
+		/// True if we should include csharp program projects in the generated solution. Pass "-DotNet" to enable this.
+		/// </summary>
+		[XmlConfigFile]
+		protected bool bIncludeDotNetPrograms = true;
+
 		/// <summary>
 		/// Whether to include temporary targets generated by UAT to support content only projects with non-default settings.
 		/// </summary>
@@ -295,10 +313,10 @@ namespace UnrealBuildTool
 		bool bIncludeTempTargets = false;
 
 		/// <summary>
-        /// True if we should reflect "Source" sub-directories on disk in the master project as master project directories.
-        /// This (arguably) adds some visual clutter to the master project but it is truer to the on-disk file organization.
-        /// </summary>
-        [XmlConfigFile]
+		/// True if we should reflect "Source" sub-directories on disk in the master project as master project directories.
+		/// This (arguably) adds some visual clutter to the master project but it is truer to the on-disk file organization.
+		/// </summary>
+		[XmlConfigFile]
 		bool bKeepSourceSubDirectories = true;
 
 		/// <summary>
@@ -328,7 +346,7 @@ namespace UnrealBuildTool
 		/// When ProjectsAreIntermediate is true, this is the directory to store generated project files
 		/// @todo projectfiles: Ideally, projects for game modules/targets would be created in the game's Intermediate folder!
 		/// </summary>
-		public static DirectoryReference IntermediateProjectFilesPath = DirectoryReference.Combine( UnrealBuildTool.EngineDirectory, "Intermediate", "ProjectFiles" );
+		public static DirectoryReference IntermediateProjectFilesPath = DirectoryReference.Combine(UnrealBuildTool.EngineDirectory, "Intermediate", "ProjectFiles");
 
 		/// <summary>
 		/// Path to timestamp file, recording when was the last time projects were created.
@@ -362,7 +380,7 @@ namespace UnrealBuildTool
 		/// <summary>
 		/// Maps all module names that were included in generated project files, to actual project file objects.
 		/// </summary>
-		protected Dictionary<string, ProjectFile> ModuleToProjectFileMap = new Dictionary<string, ProjectFile>( StringComparer.InvariantCultureIgnoreCase );
+		protected Dictionary<string, ProjectFile> ModuleToProjectFileMap = new Dictionary<string, ProjectFile>(StringComparer.InvariantCultureIgnoreCase);
 
 		/// <summary>
 		/// If generating project files for a single project, the path to its .uproject file.
@@ -421,7 +439,7 @@ namespace UnrealBuildTool
 			}
 
 			// Find all the automation modules .csproj files to add
-			List<FileReference> ModuleFiles = RulesCompiler.FindAllRulesSourceFiles(RulesCompiler.RulesFileType.AutomationModule, null, ForeignPlugins:null, AdditionalSearchPaths: BuildFolders );
+			List<FileReference> ModuleFiles = RulesCompiler.FindAllRulesSourceFiles(RulesCompiler.RulesFileType.AutomationModule, null, ForeignPlugins: null, AdditionalSearchPaths: BuildFolders);
 			foreach (FileReference ProjectFile in ModuleFiles)
 			{
 				if (FileReference.Exists(ProjectFile))
@@ -430,8 +448,8 @@ namespace UnrealBuildTool
 
 					Project.ShouldBuildForAllSolutionTargets = false;//true;
 					AddExistingProjectFile(Project, bForceDevelopmentConfiguration: true);
-					AutomationProjectFiles.Add( Project );
-					Folder.ChildProjects.Add( Project );
+					AutomationProjectFiles.Add(Project);
+					Folder.ChildProjects.Add(Project);
 
 					if (!ProjectFile.IsUnderDirectory(UnrealBuildTool.EngineDirectory))
 					{
@@ -447,7 +465,7 @@ namespace UnrealBuildTool
 		/// </summary>
 		void AddSharedDotNetModules(MasterProjectFolder ProgramsFolder)
 		{
-				DirectoryInfo DotNetDir = new DirectoryInfo(DirectoryReference.Combine(UnrealBuildTool.EngineSourceDirectory, "Programs", "Shared").FullName);
+			DirectoryInfo DotNetDir = new DirectoryInfo(DirectoryReference.Combine(UnrealBuildTool.EngineSourceDirectory, "Programs", "Shared").FullName);
 			if (DotNetDir.Exists)
 			{
 				List<FileInfo> ProjectFiles = new List<FileInfo>();
@@ -510,7 +528,7 @@ namespace UnrealBuildTool
 				DiscoverCSharpProgramProjectsRecursively(EngineDir, FoundProjects);
 			}
 
-			foreach(FileReference GameProjectFile in AllGameProjects)
+			foreach (FileReference GameProjectFile in AllGameProjects)
 			{
 				DirectoryReference GameRootFolder = GameProjectFile.Directory;
 				List<DirectoryReference> AllGameDirectories = UnrealBuildTool.GetExtensionDirs(GameRootFolder, "Source/Programs");
@@ -582,7 +600,7 @@ namespace UnrealBuildTool
 		public List<FileReference> FindGameProjects()
 		{
 			List<FileReference> ProjectFiles = new List<FileReference>();
-			if(OnlyGameProject != null)
+			if (OnlyGameProject != null)
 			{
 				ProjectFiles.Add(OnlyGameProject);
 			}
@@ -664,13 +682,13 @@ namespace UnrealBuildTool
 		/// </summary>
 		/// <param name="PlatformProjectGenerators">The registered platform project generators</param>
 		/// <param name="Arguments">Command-line arguments</param>
-		public virtual bool GenerateProjectFiles( PlatformProjectGeneratorCollection PlatformProjectGenerators, String[] Arguments )
+		public virtual bool GenerateProjectFiles(PlatformProjectGeneratorCollection PlatformProjectGenerators, String[] Arguments)
 		{
 			bool bSuccess = true;
-	
+
 			// Parse project generator options
 			bool IncludeAllPlatforms = true;
-			ConfigureProjectFileGeneration( Arguments, ref IncludeAllPlatforms);
+			ConfigureProjectFileGeneration(Arguments, ref IncludeAllPlatforms);
 
 			if (bGeneratingGameProjectFiles || UnrealBuildTool.IsEngineInstalled())
 			{
@@ -705,10 +723,10 @@ namespace UnrealBuildTool
 				{
 					MasterProjectName += "_" + Path.GetFileName(MasterProjectPath.ToString());
 				}
-				else if(bMasterProjectNameFromFolder)
+				else if (bMasterProjectNameFromFolder)
 				{
 					string NewMasterProjectName = MasterProjectPath.GetDirectoryName();
-					if(!String.IsNullOrEmpty(NewMasterProjectName))
+					if (!String.IsNullOrEmpty(NewMasterProjectName))
 					{
 						MasterProjectName = NewMasterProjectName;
 					}
@@ -752,11 +770,11 @@ namespace UnrealBuildTool
 
 			// Figure out which platforms we should generate project files for.
 			string SupportedPlatformNames;
-			SetupSupportedPlatformsAndConfigurations( IncludeAllPlatforms:IncludeAllPlatforms, SupportedPlatformNames:out SupportedPlatformNames );
+			SetupSupportedPlatformsAndConfigurations(IncludeAllPlatforms: IncludeAllPlatforms, SupportedPlatformNames: out SupportedPlatformNames);
 
-			Log.TraceVerbose( "Detected supported platforms: " + SupportedPlatformNames );
+			Log.TraceVerbose("Detected supported platforms: " + SupportedPlatformNames);
 
-			RootFolder = AllocateMasterProjectFolder( this, "<Root>" );
+			RootFolder = AllocateMasterProjectFolder(this, "<Root>");
 
 			// Build the list of games to generate projects for
 			List<FileReference> AllGameProjects = FindGameProjects();
@@ -781,7 +799,7 @@ namespace UnrealBuildTool
 			Dictionary<FileReference, ProjectFile> ProgramProjects;
 			{
 				// Setup buildable projects for all targets
-				AddProjectsForAllTargets( PlatformProjectGenerators, AllGameProjects, AllTargetFiles, Arguments, out EngineProject, out GameProjects, out ProgramProjects );
+				AddProjectsForAllTargets(PlatformProjectGenerators, AllGameProjects, AllTargetFiles, Arguments, out EngineProject, out GameProjects, out ProgramProjects);
 
 				// Add projects for mods
 				AddProjectsForMods(GameProjects, out ModProjects);
@@ -790,11 +808,11 @@ namespace UnrealBuildTool
 				AddAllGameProjects(GameProjects, SupportedPlatformNames, RootFolder);
 
 				// Set the game to be the default project
-				if(ModProjects.Count > 0)
+				if (ModProjects.Count > 0)
 				{
 					DefaultProject = ModProjects.First();
 				}
-				else if(bGeneratingGameProjectFiles && GameProjects.Count > 0)
+				else if (bGeneratingGameProjectFiles && GameProjects.Count > 0)
 				{
 					DefaultProject = GameProjects.First();
 				}
@@ -803,10 +821,10 @@ namespace UnrealBuildTool
 				List<Tuple<ProjectFile, string>> DebugProjectFiles = new List<Tuple<ProjectFile, string>>();
 
 				// Place projects into root level solution folders
-				if ( bIncludeEngineSource )
+				if (bIncludeEngineSource)
 				{
 					// If we're still missing an engine project because we don't have any targets for it, make one up.
-					if( EngineProject == null )
+					if (EngineProject == null)
 					{
 						FileReference ProjectFilePath = FileReference.Combine(IntermediateProjectFilesPath, EngineProjectFileNameBase + ProjectFileExtension);
 
@@ -818,15 +836,15 @@ namespace UnrealBuildTool
 						EngineProject.IsStubProject = true;
 					}
 
-					if( EngineProject != null )
+					if (EngineProject != null)
 					{
-						RootFolder.AddSubFolder( "Engine" ).ChildProjects.Add( EngineProject );
+						RootFolder.AddSubFolder("Engine").ChildProjects.Add(EngineProject);
 
 						// Engine config files
-						if( bIncludeConfigFiles )
+						if (bIncludeConfigFiles)
 						{
-							AddEngineConfigFiles( EngineProject );
-							if( IncludeEnginePrograms )
+							AddEngineConfigFiles(EngineProject);
+							if (IncludeEnginePrograms)
 							{
 								AddUnrealHeaderToolConfigFiles(EngineProject);
 								AddUBTConfigFilesToEngineProject(EngineProject);
@@ -840,35 +858,35 @@ namespace UnrealBuildTool
 						AddEngineExtensionFiles(EngineProject);
 
 						// Engine localization files
-						if( bIncludeLocalizationFiles )
+						if (bIncludeLocalizationFiles)
 						{
-							AddEngineLocalizationFiles( EngineProject );
+							AddEngineLocalizationFiles(EngineProject);
 						}
 
 						// Engine template files
 						if (bIncludeTemplateFiles)
 						{
-							AddEngineTemplateFiles( EngineProject );
+							AddEngineTemplateFiles(EngineProject);
 						}
 
-						if( bIncludeShaderSource )
+						if (bIncludeShaderSource)
 						{
-							Log.TraceVerbose( "Adding shader source code..." );
+							Log.TraceVerbose("Adding shader source code...");
 
 							// Find shader source files and generate stub project
-							AddEngineShaderSource( EngineProject );
+							AddEngineShaderSource(EngineProject);
 						}
 
-						if( bIncludeBuildSystemFiles )
+						if (bIncludeBuildSystemFiles)
 						{
-							Log.TraceVerbose( "Adding build system files..." );
+							Log.TraceVerbose("Adding build system files...");
 
-							AddEngineBuildFiles( EngineProject );
+							AddEngineBuildFiles(EngineProject);
 						}
 
-						if( bIncludeDocumentation )
+						if (bIncludeDocumentation)
 						{
-							AddEngineDocumentation( EngineProject );
+							AddEngineDocumentation(EngineProject);
 						}
 
 						List<Tuple<ProjectFile, string>> NewProjectFiles = EngineProject.WriteDebugProjectFiles(SupportedPlatforms, SupportedConfigurations, PlatformProjectGenerators);
@@ -879,7 +897,7 @@ namespace UnrealBuildTool
 						}
 					}
 
-					foreach( ProjectFile CurModProject in ModProjects )
+					foreach (ProjectFile CurModProject in ModProjects)
 					{
 						RootFolder.AddSubFolder("Mods").ChildProjects.Add(CurModProject);
 					}
@@ -887,7 +905,7 @@ namespace UnrealBuildTool
 					DirectoryReference TemplatesDirectory = DirectoryReference.Combine(UnrealBuildTool.RootDirectory, "Templates");
 					DirectoryReference SamplesDirectory = DirectoryReference.Combine(UnrealBuildTool.RootDirectory, "Samples");
 
-					foreach( ProjectFile CurGameProject in GameProjects )
+					foreach (ProjectFile CurGameProject in GameProjects)
 					{
 						// Templates go under a different solution folder than games
 						FileReference UnrealProjectFile = CurGameProject.ProjectTargets.First().UnrealProjectFilePath;
@@ -903,7 +921,7 @@ namespace UnrealBuildTool
 						}
 						else
 						{
-							RootFolder.AddSubFolder( "Games" ).ChildProjects.Add( CurGameProject );
+							RootFolder.AddSubFolder("Games").ChildProjects.Add(CurGameProject);
 						}
 
 						List<Tuple<ProjectFile, string>> NewProjectFiles = CurGameProject.WriteDebugProjectFiles(SupportedPlatforms, SupportedConfigurations, PlatformProjectGenerators);
@@ -917,30 +935,30 @@ namespace UnrealBuildTool
 
 					//Related Debug Project Files - Tuple has the related Debug Project, SolutionFolder
 					foreach (Tuple<ProjectFile, string> DebugProjectFile in DebugProjectFiles)
-					{						
+					{
 						AddExistingProjectFile(DebugProjectFile.Item1, bForceDevelopmentConfiguration: false);
 
 						//add it to the Android Debug Projects folder in the solution
-						RootFolder.AddSubFolder(DebugProjectFile.Item2).ChildProjects.Add(DebugProjectFile.Item1);						
+						RootFolder.AddSubFolder(DebugProjectFile.Item2).ChildProjects.Add(DebugProjectFile.Item1);
 					}
-					
 
-					foreach( KeyValuePair<FileReference, ProjectFile> CurProgramProject in ProgramProjects )
+
+					foreach (KeyValuePair<FileReference, ProjectFile> CurProgramProject in ProgramProjects)
 					{
-                        ProjectTarget Target = CurProgramProject.Value.ProjectTargets.FirstOrDefault(t => !String.IsNullOrEmpty(t.TargetRules.SolutionDirectory));
+						ProjectTarget Target = CurProgramProject.Value.ProjectTargets.FirstOrDefault(t => !String.IsNullOrEmpty(t.TargetRules.SolutionDirectory));
 
-                        if (Target != null)
-                        {
-                            RootFolder.AddSubFolder(Target.TargetRules.SolutionDirectory).ChildProjects.Add(CurProgramProject.Value);
-                        }
-                        else
-                        {
-							RootFolder.AddSubFolder( "Programs" ).ChildProjects.Add( CurProgramProject.Value );
-                        }
+						if (Target != null)
+						{
+							RootFolder.AddSubFolder(Target.TargetRules.SolutionDirectory).ChildProjects.Add(CurProgramProject.Value);
+						}
+						else
+						{
+							RootFolder.AddSubFolder("Programs").ChildProjects.Add(CurProgramProject.Value);
+						}
 					}
 
 					// Add all of the config files for generated program targets
-					AddEngineProgramConfigFiles( ProgramProjects );
+					AddEngineProgramConfigFiles(ProgramProjects);
 				}
 			}
 
@@ -948,12 +966,12 @@ namespace UnrealBuildTool
 			AddProjectsForAllModules(AllGameProjects, ProgramProjects, ModProjects, AllModuleFiles, bGatherThirdPartySource);
 
 			{
-				if( IncludeEnginePrograms )
+				if (bIncludeDotNetPrograms)
 				{
-					MasterProjectFolder ProgramsFolder = RootFolder.AddSubFolder( "Programs" );
+					MasterProjectFolder ProgramsFolder = RootFolder.AddSubFolder("Programs");
 
 					// Add UnrealBuildTool to the master project
-					AddUnrealBuildToolProject( ProgramsFolder );
+					AddUnrealBuildToolProject(ProgramsFolder);
 
 					// Add AutomationTool to the master project
 					VCSharpProjectFile AutomationToolProject = AddSimpleCSharpProject("AutomationTool", bShouldBuildForAllSolutionTargets: true, bForceDevelopmentConfiguration: true);
@@ -974,53 +992,53 @@ namespace UnrealBuildTool
 				// Eliminate all redundant master project folders.  E.g., folders which contain only one project and that project
 				// has the same name as the folder itself.  To the user, projects "feel like" folders already in the IDE, so we
 				// want to collapse them down where possible.
-				EliminateRedundantMasterProjectSubFolders( RootFolder, "" );
+				EliminateRedundantMasterProjectSubFolders(RootFolder, "");
 
 				// Figure out which targets we need about IntelliSense for.  We only need to worry about targets for projects
 				// that we're actually generating in this session.
 				List<Tuple<ProjectFile, ProjectTarget>> IntelliSenseTargetFiles = new List<Tuple<ProjectFile, ProjectTarget>>();
 				{
 					// Engine targets
-					if( EngineProject != null)
+					if (EngineProject != null)
 					{
-						foreach( ProjectTarget ProjectTarget in EngineProject.ProjectTargets )
+						foreach (ProjectTarget ProjectTarget in EngineProject.ProjectTargets)
 						{
-							if( ProjectTarget.TargetFilePath != null )
+							if (ProjectTarget.TargetFilePath != null)
 							{
 								// Only bother with the editor target.  We want to make sure that definitions are setup to be as inclusive as possible
 								// for good quality IntelliSense.  For example, we want WITH_EDITORONLY_DATA=1, so using the editor targets works well.
-								if( ProjectTarget.TargetRules.Type == TargetType.Editor )
-								{ 
-									IntelliSenseTargetFiles.Add( Tuple.Create(EngineProject, ProjectTarget) );
+								if (ProjectTarget.TargetRules.Type == TargetType.Editor)
+								{
+									IntelliSenseTargetFiles.Add(Tuple.Create(EngineProject, ProjectTarget));
 								}
 							}
 						}
 					}
 
 					// Program targets
-					foreach( ProjectFile ProgramProject in ProgramProjects.Values )
+					foreach (ProjectFile ProgramProject in ProgramProjects.Values)
 					{
-						foreach( ProjectTarget ProjectTarget in ProgramProject.ProjectTargets )
+						foreach (ProjectTarget ProjectTarget in ProgramProject.ProjectTargets)
 						{
-							if( ProjectTarget.TargetFilePath != null )
+							if (ProjectTarget.TargetFilePath != null)
 							{
-								IntelliSenseTargetFiles.Add( Tuple.Create( ProgramProject, ProjectTarget ) );
+								IntelliSenseTargetFiles.Add(Tuple.Create(ProgramProject, ProjectTarget));
 							}
 						}
 					}
 
 					// Game/template targets
-					foreach( ProjectFile GameProject in GameProjects )
+					foreach (ProjectFile GameProject in GameProjects)
 					{
-						foreach( ProjectTarget ProjectTarget in GameProject.ProjectTargets )
+						foreach (ProjectTarget ProjectTarget in GameProject.ProjectTargets)
 						{
-							if( ProjectTarget.TargetFilePath != null )
+							if (ProjectTarget.TargetFilePath != null)
 							{
 								// Only bother with the editor target.  We want to make sure that definitions are setup to be as inclusive as possible
 								// for good quality IntelliSense.  For example, we want WITH_EDITORONLY_DATA=1, so using the editor targets works well.
-								if( ProjectTarget.TargetRules.Type == TargetType.Editor )
-								{ 
-									IntelliSenseTargetFiles.Add( Tuple.Create( GameProject, ProjectTarget ) );
+								if (ProjectTarget.TargetRules.Type == TargetType.Editor)
+								{
+									IntelliSenseTargetFiles.Add(Tuple.Create(GameProject, ProjectTarget));
 								}
 							}
 						}
@@ -1033,13 +1051,13 @@ namespace UnrealBuildTool
 				// Generate IntelliSense data if we need to.  This involves having UBT simulate the action compilation of
 				// the targets so that we can extra the compiler defines, include paths, etc.
 				GenerateIntelliSenseData(Arguments, IntelliSenseTargetFiles);
-				
+
 				// Write the project files
 				WriteProjectFiles(PlatformProjectGenerators);
-				Log.TraceVerbose( "Project generation complete ({0} generated, {1} imported)", GeneratedProjectFiles.Count, OtherProjectFiles.Count );
+				Log.TraceVerbose("Project generation complete ({0} generated, {1} imported)", GeneratedProjectFiles.Count, OtherProjectFiles.Count);
 
 				// Generate all the target info files for the editor
-				foreach(FileReference ProjectFile in AllGameProjects)
+				foreach (FileReference ProjectFile in AllGameProjects)
 				{
 					RulesAssembly RulesAssembly = RulesCompiler.CreateProjectRulesAssembly(ProjectFile, false, false);
 					QueryTargetsMode.WriteTargetInfo(ProjectFile, RulesAssembly, QueryTargetsMode.GetDefaultOutputFile(ProjectFile), new CommandLineArguments(GetTargetArguments(Arguments)));
@@ -1062,7 +1080,7 @@ namespace UnrealBuildTool
 				));
 
 			List<XmlConfig.InputFile> InputFiles = XmlConfig.FindInputFiles();
-			foreach(XmlConfig.InputFile InputFile in InputFiles)
+			foreach (XmlConfig.InputFile InputFile in InputFiles)
 			{
 				EngineProject.AddAliasedFileToProject(
 						new AliasedFile(
@@ -1087,7 +1105,7 @@ namespace UnrealBuildTool
 		/// </summary>
 		/// <param name="Arguments">Arguments passed into the program</param>
 		/// <param name="IncludeAllPlatforms">True if all platforms should be included</param>
-		protected virtual void ConfigureProjectFileGeneration( String[] Arguments, ref bool IncludeAllPlatforms )
+		protected virtual void ConfigureProjectFileGeneration(String[] Arguments, ref bool IncludeAllPlatforms)
 		{
 			if (PlatformNames != null)
 			{
@@ -1102,11 +1120,11 @@ namespace UnrealBuildTool
 			}
 
 			bool bAlwaysIncludeEngineModules = false;
-			foreach( string CurArgument in Arguments )
+			foreach (string CurArgument in Arguments)
 			{
-				if( CurArgument.StartsWith( "-" ) )
+				if (CurArgument.StartsWith("-"))
 				{
-					if (CurArgument.StartsWith( "-Platforms=", StringComparison.InvariantCultureIgnoreCase ))
+					if (CurArgument.StartsWith("-Platforms=", StringComparison.InvariantCultureIgnoreCase))
 					{
 						// Parse the list... will be in Foo+Bar+New format
 						string PlatformList = CurArgument.Substring(11);
@@ -1143,62 +1161,86 @@ namespace UnrealBuildTool
 							bAppendPlatformSuffix = true;
 						}
 					}
-					else switch( CurArgument.ToUpperInvariant() )
-					{
-						case "-ALLPLATFORMS":
-							IncludeAllPlatforms = true;
-							break;
-
-						case "-CURRENTPLATFORM":
-							IncludeAllPlatforms = false;
-							break;
-
-						case "-THIRDPARTY":
-							bGatherThirdPartySource = true;
-							break;
-						
-						case "-GAME":
-							// Generates project files for a single game
-							bGeneratingGameProjectFiles = true;
-							break;
-
-						case "-ENGINE":
-							// Forces engine modules and targets to be included in game-specific project files
-							bAlwaysIncludeEngineModules = true;
-							break;
-
-						case "-NOINTELLISENSE":
-							bGenerateIntelliSenseData = false;
-							break;
-
-						case "-INTELLISENSE":
-							bGenerateIntelliSenseData = true;
-							break;
-
-						case "-SHIPPINGCONFIGS":
-							bIncludeTestAndShippingConfigs = true;
-							break;
-
-						case "-NOSHIPPINGCONFIGS":
-							bIncludeTestAndShippingConfigs = false;
-							break;
-
-						case "-ALLLANGUAGES":
-							bAllDocumentationLanguages = true;
-							break;
-
-						case "-USEPRECOMPILED":
-							bUsePrecompiled = true;
-							break;
-
-						case "-INCLUDETEMPTARGETS":
-							bIncludeTempTargets = true;
-							break;
-					}
+					else switch (CurArgument.ToUpperInvariant())
+						{
+							case "-ALLPLATFORMS":
+								IncludeAllPlatforms = true;
+								break;
+
+							case "-CURRENTPLATFORM":
+								IncludeAllPlatforms = false;
+								break;
+
+							case "-THIRDPARTY":
+								bGatherThirdPartySource = true;
+								break;
+
+							case "-GAME":
+								// Generates project files for a single game
+								bGeneratingGameProjectFiles = true;
+								break;
+
+							case "-ENGINE":
+								// Forces engine modules and targets to be included in game-specific project files
+								bAlwaysIncludeEngineModules = true;
+								break;
+
+							case "-NOINTELLISENSE":
+								bGenerateIntelliSenseData = false;
+								break;
+
+							case "-INTELLISENSE":
+								bGenerateIntelliSenseData = true;
+								break;
+
+							case "-SHIPPINGCONFIGS":
+								bIncludeTestAndShippingConfigs = true;
+								break;
+
+							case "-NOSHIPPINGCONFIGS":
+								bIncludeTestAndShippingConfigs = false;
+								break;
+
+							case "-DEBUGCONFIGS":
+								bIncludeDebugConfigs = true;
+								break;
+
+							case "-NODEBUGCONFIGS":
+								bIncludeDebugConfigs = false;
+								break;
+
+							case "-DEVELOPMENTCONFIGS":
+								bIncludeDevelopmentConfigs = true;
+								break;
+
+							case "-NODEVELOPMENTCONFIGS":
+								bIncludeDevelopmentConfigs = false;
+								break;
+
+							case "-DOTNET":
+								bIncludeDotNetPrograms = true;
+								break;
+
+							case "-NODOTNET":
+								bIncludeDotNetPrograms = false;
+								break;
+
+							case "-ALLLANGUAGES":
+								bAllDocumentationLanguages = true;
+								break;
+
+							case "-USEPRECOMPILED":
+								bUsePrecompiled = true;
+								break;
+
+							case "-INCLUDETEMPTARGETS":
+								bIncludeTempTargets = true;
+								break;
+						}
 				}
 			}
 
-			if( bGeneratingGameProjectFiles || UnrealBuildTool.IsEngineInstalled() )
+			if (bGeneratingGameProjectFiles || UnrealBuildTool.IsEngineInstalled())
 			{
 				if (OnlyGameProject == null)
 				{
@@ -1234,10 +1276,10 @@ namespace UnrealBuildTool
 		protected void AddAllGameProjects(List<ProjectFile> GameProjects, string SupportedPlatformNames, MasterProjectFolder ProjectsFolder)
 		{
 			HashSet<DirectoryReference> UniqueGameProjectDirectories = new HashSet<DirectoryReference>();
-			foreach( ProjectFile GameProject in GameProjects )
+			foreach (ProjectFile GameProject in GameProjects)
 			{
 				DirectoryReference GameProjectDirectory = GameProject.BaseDir;
-				if(UniqueGameProjectDirectories.Add(GameProjectDirectory))
+				if (UniqueGameProjectDirectories.Add(GameProjectDirectory))
 				{
 					// @todo projectfiles: We have engine localization files, but should we also add GAME localization files?
 
@@ -1247,13 +1289,13 @@ namespace UnrealBuildTool
 					}
 
 					// Game restricted source files, since they won't be added via a module FileReference
-					foreach (DirectoryReference GameRestrictedSourceDirectory in UnrealBuildTool.GetExtensionDirs(GameProjectDirectory, "Source", bIncludePlatformDirectories:false, bIncludeBaseDirectory:false))
+					foreach (DirectoryReference GameRestrictedSourceDirectory in UnrealBuildTool.GetExtensionDirs(GameProjectDirectory, "Source", bIncludePlatformDirectories: false, bIncludeBaseDirectory: false))
 					{
 						GameProject.AddFilesToProject(SourceFileSearch.FindFiles(GameRestrictedSourceDirectory), GameRestrictedSourceDirectory);
 					}
 
 					// Game config files
-					if ( bIncludeConfigFiles )
+					if (bIncludeConfigFiles)
 					{
 						foreach (DirectoryReference GameConfigDirectory in UnrealBuildTool.GetExtensionDirs(GameProjectDirectory, "Config"))
 						{
@@ -1262,7 +1304,7 @@ namespace UnrealBuildTool
 					}
 
 					// Game build files
-					if( bIncludeBuildSystemFiles )
+					if (bIncludeBuildSystemFiles)
 					{
 						foreach (DirectoryReference GameBuildDirectory in UnrealBuildTool.GetExtensionDirs(GameProjectDirectory, "Build"))
 						{
@@ -1273,7 +1315,7 @@ namespace UnrealBuildTool
 							SubdirectoryNamesToExclude.Add("PipelineCaches");
 							SubdirectoryNamesToExclude.Add("symbols");
 
-							GameProject.AddFilesToProject( SourceFileSearch.FindFiles( GameBuildDirectory, SubdirectoryNamesToExclude ), GameProjectDirectory );
+							GameProject.AddFilesToProject(SourceFileSearch.FindFiles(GameBuildDirectory, SubdirectoryNamesToExclude), GameProjectDirectory);
 						}
 					}
 
@@ -1282,39 +1324,39 @@ namespace UnrealBuildTool
 						GameProject.AddFilesToProject(SourceFileSearch.FindFiles(GameShaderDirectory), GameProjectDirectory);
 					}
 				}
-            }
+			}
 		}
 
 
 		/// Adds all engine localization text files to the specified project
-		private void AddEngineLocalizationFiles( ProjectFile EngineProject )
+		private void AddEngineLocalizationFiles(ProjectFile EngineProject)
 		{
-			DirectoryReference EngineLocalizationDirectory = DirectoryReference.Combine( UnrealBuildTool.EngineDirectory, "Content", "Localization" );
-			if( DirectoryReference.Exists(EngineLocalizationDirectory) )
+			DirectoryReference EngineLocalizationDirectory = DirectoryReference.Combine(UnrealBuildTool.EngineDirectory, "Content", "Localization");
+			if (DirectoryReference.Exists(EngineLocalizationDirectory))
 			{
-				EngineProject.AddFilesToProject( SourceFileSearch.FindFiles( EngineLocalizationDirectory ), UnrealBuildTool.EngineDirectory );
+				EngineProject.AddFilesToProject(SourceFileSearch.FindFiles(EngineLocalizationDirectory), UnrealBuildTool.EngineDirectory);
 			}
 		}
 
 
-        /// Adds all engine template text files to the specified project
-        private void AddEngineTemplateFiles( ProjectFile EngineProject )
-        {
-            DirectoryReference EngineTemplateDirectory = DirectoryReference.Combine(UnrealBuildTool.EngineDirectory, "Content", "Editor", "Templates");
-            if (DirectoryReference.Exists(EngineTemplateDirectory))
-            {
-				EngineProject.AddFilesToProject( SourceFileSearch.FindFiles( EngineTemplateDirectory ), UnrealBuildTool.EngineDirectory );
+		/// Adds all engine template text files to the specified project
+		private void AddEngineTemplateFiles(ProjectFile EngineProject)
+		{
+			DirectoryReference EngineTemplateDirectory = DirectoryReference.Combine(UnrealBuildTool.EngineDirectory, "Content", "Editor", "Templates");
+			if (DirectoryReference.Exists(EngineTemplateDirectory))
+			{
+				EngineProject.AddFilesToProject(SourceFileSearch.FindFiles(EngineTemplateDirectory), UnrealBuildTool.EngineDirectory);
 			}
-        }
+		}
 
 
 		/// Adds all engine config files to the specified project
-		private void AddEngineConfigFiles( ProjectFile EngineProject )
+		private void AddEngineConfigFiles(ProjectFile EngineProject)
 		{
-            DirectoryReference EngineConfigDirectory = DirectoryReference.Combine(UnrealBuildTool.EngineDirectory, "Config" );
-			if( DirectoryReference.Exists(EngineConfigDirectory) )
+			DirectoryReference EngineConfigDirectory = DirectoryReference.Combine(UnrealBuildTool.EngineDirectory, "Config");
+			if (DirectoryReference.Exists(EngineConfigDirectory))
 			{
-				EngineProject.AddFilesToProject( SourceFileSearch.FindFiles( EngineConfigDirectory ), UnrealBuildTool.EngineDirectory );
+				EngineProject.AddFilesToProject(SourceFileSearch.FindFiles(EngineConfigDirectory), UnrealBuildTool.EngineDirectory);
 			}
 		}
 
@@ -1324,11 +1366,11 @@ namespace UnrealBuildTool
 		}
 
 		/// Adds additional files from the platform extensions folder
-		protected virtual void AddEngineExtensionFiles( ProjectFile EngineProject )
+		protected virtual void AddEngineExtensionFiles(ProjectFile EngineProject)
 		{
 			// @todo: this will add the same files to the solution (like the UBT source files that also get added to UnrealBuildTool project).
 			// not sure of a good filtering method here
-			foreach(DirectoryReference ExtensionDir in UnrealBuildTool.GetExtensionDirs(UnrealBuildTool.EngineDirectory))
+			foreach (DirectoryReference ExtensionDir in UnrealBuildTool.GetExtensionDirs(UnrealBuildTool.EngineDirectory))
 			{
 				if (ExtensionDir != UnrealBuildTool.EngineDirectory)
 				{
@@ -1361,10 +1403,10 @@ namespace UnrealBuildTool
 			List<FileReference> AllModuleFiles = new List<FileReference>();
 
 			// Locate all modules (*.Build.cs files)
-			List<FileReference> FoundModuleFiles = RulesCompiler.FindAllRulesSourceFiles( RulesCompiler.RulesFileType.Module, GameFolders: AllGameProjects.Select(x => x.Directory).ToList(), ForeignPlugins:null, AdditionalSearchPaths:null );
-			foreach( FileReference BuildFileName in FoundModuleFiles )
+			List<FileReference> FoundModuleFiles = RulesCompiler.FindAllRulesSourceFiles(RulesCompiler.RulesFileType.Module, GameFolders: AllGameProjects.Select(x => x.Directory).ToList(), ForeignPlugins: null, AdditionalSearchPaths: null);
+			foreach (FileReference BuildFileName in FoundModuleFiles)
 			{
-				AllModuleFiles.Add( BuildFileName );
+				AllModuleFiles.Add(BuildFileName);
 			}
 			return AllModuleFiles;
 		}
@@ -1402,13 +1444,13 @@ namespace UnrealBuildTool
 			List<FileReference> AllTargetFiles = new List<FileReference>();
 
 			// Make a list of all platform name strings that we're *not* including in the project files
-			List<string> UnsupportedPlatformNameStrings = Utils.MakeListOfUnsupportedPlatforms(SupportedPlatforms, bIncludeUnbuildablePlatforms:true);
+			List<string> UnsupportedPlatformNameStrings = Utils.MakeListOfUnsupportedPlatforms(SupportedPlatforms, bIncludeUnbuildablePlatforms: true);
 
 			// Locate all targets (*.Target.cs files)
-			List<FileReference> FoundTargetFiles = RulesCompiler.FindAllRulesSourceFiles( RulesCompiler.RulesFileType.Target, AllGameProjects.Select(x => x.Directory).ToList(), ForeignPlugins:null, AdditionalSearchPaths:null, bIncludeTempTargets: bIncludeTempTargets );
-			foreach( FileReference CurTargetFile in FoundTargetFiles )
+			List<FileReference> FoundTargetFiles = RulesCompiler.FindAllRulesSourceFiles(RulesCompiler.RulesFileType.Target, AllGameProjects.Select(x => x.Directory).ToList(), ForeignPlugins: null, AdditionalSearchPaths: null, bIncludeTempTargets: bIncludeTempTargets);
+			foreach (FileReference CurTargetFile in FoundTargetFiles)
 			{
-				string CleanTargetFileName = Utils.CleanDirectorySeparators( CurTargetFile.FullName );
+				string CleanTargetFileName = Utils.CleanDirectorySeparators(CurTargetFile.FullName);
 
 				// remove the local root
 				string LocalRoot = UnrealBuildTool.RootDirectory.FullName;
@@ -1443,7 +1485,7 @@ namespace UnrealBuildTool
 
 				// Skip targets in unsupported platform directories
 				bool IncludeThisTarget = true;
-				foreach( string CurPlatformName in UnsupportedPlatformNameStrings )
+				foreach (string CurPlatformName in UnsupportedPlatformNameStrings)
 				{
 					if (Search.IndexOf(Path.DirectorySeparatorChar + CurPlatformName + Path.DirectorySeparatorChar, StringComparison.InvariantCultureIgnoreCase) != -1)
 					{
@@ -1452,9 +1494,9 @@ namespace UnrealBuildTool
 					}
 				}
 
-				if( IncludeThisTarget )
+				if (IncludeThisTarget)
 				{
-					AllTargetFiles.Add( CurTargetFile );
+					AllTargetFiles.Add(CurTargetFile);
 				}
 			}
 
@@ -1462,17 +1504,17 @@ namespace UnrealBuildTool
 		}
 
 
-	
+
 		/// <summary>
 		/// Recursively collapses all sub-folders that are redundant.  Should only be called after we're done adding
 		/// files and projects to the master project.
 		/// </summary>
 		/// <param name="Folder">The folder whose sub-folders we should potentially collapse into</param>
 		/// <param name="ParentMasterProjectFolderPath"></param>
-		void EliminateRedundantMasterProjectSubFolders( MasterProjectFolder Folder, string ParentMasterProjectFolderPath )
+		void EliminateRedundantMasterProjectSubFolders(MasterProjectFolder Folder, string ParentMasterProjectFolderPath)
 		{
 			// NOTE: This is for diagnostics output only
-			string MasterProjectFolderPath = String.IsNullOrEmpty( ParentMasterProjectFolderPath ) ? Folder.FolderName : ( ParentMasterProjectFolderPath + "/" + Folder.FolderName );
+			string MasterProjectFolderPath = String.IsNullOrEmpty(ParentMasterProjectFolderPath) ? Folder.FolderName : (ParentMasterProjectFolderPath + "/" + Folder.FolderName);
 
 			// We can eliminate folders that meet all of these requirements:
 			//		1) Have only a single project file in them
@@ -1482,26 +1524,26 @@ namespace UnrealBuildTool
 			// Additionally, if KeepSourceSubDirectories==false, we can eliminate directories called "Source".
 			//
 			// Also, we can kill folders that are completely empty.
-			
 
-			foreach( MasterProjectFolder SubFolder in Folder.SubFolders )
+
+			foreach (MasterProjectFolder SubFolder in Folder.SubFolders)
 			{
 				// Recurse
-				EliminateRedundantMasterProjectSubFolders( SubFolder, MasterProjectFolderPath );
+				EliminateRedundantMasterProjectSubFolders(SubFolder, MasterProjectFolderPath);
 			}
 
 			List<MasterProjectFolder> SubFoldersToAdd = new List<MasterProjectFolder>();
 			List<MasterProjectFolder> SubFoldersToRemove = new List<MasterProjectFolder>();
-			foreach( MasterProjectFolder SubFolder in Folder.SubFolders )
+			foreach (MasterProjectFolder SubFolder in Folder.SubFolders)
 			{
 				bool CanCollapseFolder = false;
 
 				// 1)
-				if( SubFolder.ChildProjects.Count == 1 )
+				if (SubFolder.ChildProjects.Count == 1)
 				{
 					// 2)
-					if( SubFolder.Files.Count == 0 &&
-						SubFolder.SubFolders.Count == 0 )
+					if (SubFolder.Files.Count == 0 &&
+						SubFolder.SubFolders.Count == 0)
 					{
 						// 3)
 						if (SubFolder.FolderName.Equals(SubFolder.ChildProjects[0].ProjectFilePath.GetFileNameWithoutExtension(), StringComparison.InvariantCultureIgnoreCase))
@@ -1511,47 +1553,47 @@ namespace UnrealBuildTool
 					}
 				}
 
-				if( !bKeepSourceSubDirectories )
+				if (!bKeepSourceSubDirectories)
 				{
-					if( SubFolder.FolderName.Equals( "Source", StringComparison.InvariantCultureIgnoreCase ) )
+					if (SubFolder.FolderName.Equals("Source", StringComparison.InvariantCultureIgnoreCase))
 					{
 						// Avoid collapsing the Engine's Source directory, since there are so many other solution folders in
 						// the parent directory.
-						if( !Folder.FolderName.Equals( "Engine", StringComparison.InvariantCultureIgnoreCase ) )
+						if (!Folder.FolderName.Equals("Engine", StringComparison.InvariantCultureIgnoreCase))
 						{
 							CanCollapseFolder = true;
 						}
 					}
 				}
 
-				if( SubFolder.ChildProjects.Count == 0 && SubFolder.Files.Count == 0 && SubFolder.SubFolders.Count == 0 )
+				if (SubFolder.ChildProjects.Count == 0 && SubFolder.Files.Count == 0 && SubFolder.SubFolders.Count == 0)
 				{
 					// Folder is totally empty
 					CanCollapseFolder = true;
 				}
 
-				if( CanCollapseFolder )
+				if (CanCollapseFolder)
 				{
 					// OK, this folder is redundant and can be collapsed away.
 
-					SubFoldersToAdd.AddRange( SubFolder.SubFolders );
+					SubFoldersToAdd.AddRange(SubFolder.SubFolders);
 					SubFolder.SubFolders.Clear();
 
-					Folder.ChildProjects.AddRange( SubFolder.ChildProjects );
+					Folder.ChildProjects.AddRange(SubFolder.ChildProjects);
 					SubFolder.ChildProjects.Clear();
 
-					Folder.Files.AddRange( SubFolder.Files );
+					Folder.Files.AddRange(SubFolder.Files);
 					SubFolder.Files.Clear();
 
-					SubFoldersToRemove.Add( SubFolder );
+					SubFoldersToRemove.Add(SubFolder);
 				}
 			}
 
-			foreach( MasterProjectFolder SubFolderToRemove in SubFoldersToRemove )
+			foreach (MasterProjectFolder SubFolderToRemove in SubFoldersToRemove)
 			{
-				Folder.SubFolders.Remove( SubFolderToRemove );
+				Folder.SubFolders.Remove(SubFolderToRemove);
 			}
-			Folder.SubFolders.AddRange( SubFoldersToAdd );
+			Folder.SubFolders.AddRange(SubFoldersToAdd);
 
 			// After everything has been collapsed, do a bit of data validation
 			Validate(Folder, ParentMasterProjectFolderPath);
@@ -1658,10 +1700,10 @@ namespace UnrealBuildTool
 		{
 			VCSharpProjectFile Project = null;
 
-			FileReference ProjectFileName = FileReference.Combine( UnrealBuildTool.EngineSourceDirectory, "Programs", ProjectName, Path.GetFileName( ProjectName ) + ".csproj" );
-			
-			FileInfo Info = new FileInfo( ProjectFileName.FullName );
-			if( Info.Exists )
+			FileReference ProjectFileName = FileReference.Combine(UnrealBuildTool.EngineSourceDirectory, "Programs", ProjectName, Path.GetFileName(ProjectName) + ".csproj");
+
+			FileInfo Info = new FileInfo(ProjectFileName.FullName);
+			if (Info.Exists)
 			{
 				Project = new VCSharpProjectFile(ProjectFileName);
 
@@ -1671,7 +1713,7 @@ namespace UnrealBuildTool
 			}
 			else
 			{
-				throw new BuildException( ProjectFileName.FullName + " doesn't exist!" );
+				throw new BuildException(ProjectFileName.FullName + " doesn't exist!");
 			}
 
 			return Project;
@@ -1683,11 +1725,11 @@ namespace UnrealBuildTool
 		/// <param name="RootKey"></param>
 		/// <param name="VisualStudioVersion"></param>
 		/// <returns></returns>
-		private bool CheckRegistryKey( RegistryKey RootKey, string VisualStudioVersion )
+		private bool CheckRegistryKey(RegistryKey RootKey, string VisualStudioVersion)
 		{
 			bool bInstalled = false;
-			RegistryKey VSSubKey = RootKey.OpenSubKey( "SOFTWARE\\Microsoft\\VisualStudio\\" + VisualStudioVersion + "\\Projects\\{E53F8FEA-EAE0-44A6-8774-FFD645390401}" );
-			if( VSSubKey != null )
+			RegistryKey VSSubKey = RootKey.OpenSubKey("SOFTWARE\\Microsoft\\VisualStudio\\" + VisualStudioVersion + "\\Projects\\{E53F8FEA-EAE0-44A6-8774-FFD645390401}");
+			if (VSSubKey != null)
 			{
 				bInstalled = true;
 				VSSubKey.Close();
@@ -1703,14 +1745,14 @@ namespace UnrealBuildTool
 		/// <param name="VisualStudioVersion"></param>
 		/// <param name="Extension"></param>
 		/// <returns></returns>
-		private bool CheckVisualStudioExtensionPackage( string VisualStudioFolder, string VisualStudioVersion, string Extension )
+		private bool CheckVisualStudioExtensionPackage(string VisualStudioFolder, string VisualStudioVersion, string Extension)
 		{
-			DirectoryInfo DirInfo = new DirectoryInfo( Path.Combine( VisualStudioFolder, VisualStudioVersion, "Extensions" ) );
-			if( DirInfo.Exists )
+			DirectoryInfo DirInfo = new DirectoryInfo(Path.Combine(VisualStudioFolder, VisualStudioVersion, "Extensions"));
+			if (DirInfo.Exists)
 			{
-				List<FileInfo> PackageDefs = DirInfo.GetFiles( "*.pkgdef", SearchOption.AllDirectories ).ToList();
-				List<string> PackageDefNames = PackageDefs.Select( x => x.Name ).ToList();
-				if( PackageDefNames.Contains( Extension ) )
+				List<FileInfo> PackageDefs = DirInfo.GetFiles("*.pkgdef", SearchOption.AllDirectories).ToList();
+				List<string> PackageDefNames = PackageDefs.Select(x => x.Name).ToList();
+				if (PackageDefNames.Contains(Extension))
 				{
 					return true;
 				}
@@ -1722,11 +1764,11 @@ namespace UnrealBuildTool
 		/// <summary>
 		/// Adds all of the config files for program targets to their project files
 		/// </summary>
-		private void AddEngineProgramConfigFiles( Dictionary<FileReference, ProjectFile> ProgramProjects )
+		private void AddEngineProgramConfigFiles(Dictionary<FileReference, ProjectFile> ProgramProjects)
 		{
-			if( bIncludeConfigFiles )
+			if (bIncludeConfigFiles)
 			{
-				foreach( KeyValuePair<FileReference, ProjectFile> FileAndProject in ProgramProjects )
+				foreach (KeyValuePair<FileReference, ProjectFile> FileAndProject in ProgramProjects)
 				{
 					string ProgramName = FileAndProject.Key.GetFileNameWithoutAnyExtensions();
 					ProjectFile ProgramProjectFile = FileAndProject.Value;
@@ -1735,12 +1777,12 @@ namespace UnrealBuildTool
 					// @todo projectfiles: Fragile assumption here about Programs always being under /Engine/Programs
 
 					DirectoryReference ProgramDirectory;
-					ProgramDirectory = DirectoryReference.Combine( UnrealBuildTool.EngineDirectory, "Programs", ProgramName );
+					ProgramDirectory = DirectoryReference.Combine(UnrealBuildTool.EngineDirectory, "Programs", ProgramName);
 
-					DirectoryReference ProgramConfigDirectory = DirectoryReference.Combine( ProgramDirectory, "Config" );
-					if( DirectoryReference.Exists(ProgramConfigDirectory) )
+					DirectoryReference ProgramConfigDirectory = DirectoryReference.Combine(ProgramDirectory, "Config");
+					if (DirectoryReference.Exists(ProgramConfigDirectory))
 					{
-						ProgramProjectFile.AddFilesToProject( SourceFileSearch.FindFiles( ProgramConfigDirectory ), ProgramDirectory );
+						ProgramProjectFile.AddFilesToProject(SourceFileSearch.FindFiles(ProgramConfigDirectory), ProgramDirectory);
 					}
 				}
 			}
@@ -1754,15 +1796,15 @@ namespace UnrealBuildTool
 		/// <return>Whether the process was successful or not</return>
 		private void GenerateIntelliSenseData(string[] Arguments, List<Tuple<ProjectFile, ProjectTarget>> Targets)
 		{
-			if(ShouldGenerateIntelliSenseData() && Targets.Count > 0)
+			if (ShouldGenerateIntelliSenseData() && Targets.Count > 0)
 			{
 				string ProgressInfoText = Utils.IsRunningOnWindows ? "Binding IntelliSense data..." : "Generating data for project indexing...";
-				using(ProgressWriter Progress = new ProgressWriter(ProgressInfoText, true))
+				using (ProgressWriter Progress = new ProgressWriter(ProgressInfoText, true))
 				{
-					for(int TargetIndex = 0; TargetIndex < Targets.Count; ++TargetIndex)
+					for (int TargetIndex = 0; TargetIndex < Targets.Count; ++TargetIndex)
 					{
-						ProjectFile TargetProjectFile = Targets[ TargetIndex ].Item1;
-						ProjectTarget CurTarget = Targets[ TargetIndex ].Item2;
+						ProjectFile TargetProjectFile = Targets[TargetIndex].Item1;
+						ProjectTarget CurTarget = Targets[TargetIndex].Item2;
 
 						// Ignore projects for platforms we can't build on this host
 						UnrealTargetPlatform IntellisensePlatform = BuildHostPlatform.Current.Platform;
@@ -1774,7 +1816,7 @@ namespace UnrealBuildTool
 						Log.TraceVerbose("Found target: " + CurTarget.Name);
 
 						List<string> NewArguments = new List<string>(Arguments.Length + 4);
-						if(CurTarget.TargetRules.Type != TargetType.Program)
+						if (CurTarget.TargetRules.Type != TargetType.Program)
 						{
 							NewArguments.Add("-precompile");
 						}
@@ -1794,10 +1836,10 @@ namespace UnrealBuildTool
 							AddTargetForIntellisense(Target);
 							// Generate a compile environment for each module in the binary
 							CppCompileEnvironment GlobalCompileEnvironment = Target.CreateCompileEnvironmentForProjectFiles();
-							foreach(UEBuildBinary Binary in Target.Binaries)
+							foreach (UEBuildBinary Binary in Target.Binaries)
 							{
 								CppCompileEnvironment BinaryCompileEnvironment = Binary.CreateBinaryCompileEnvironment(GlobalCompileEnvironment);
-								foreach(UEBuildModuleCPP Module in Binary.Modules.OfType<UEBuildModuleCPP>())
+								foreach (UEBuildModuleCPP Module in Binary.Modules.OfType<UEBuildModuleCPP>())
 								{
 									ProjectFile ProjectFileForIDE;
 									if (ModuleToProjectFileMap.TryGetValue(Module.Name, out ProjectFileForIDE) && ProjectFileForIDE == TargetProjectFile)
@@ -1817,7 +1859,7 @@ namespace UnrealBuildTool
 								FileReference.Delete(MakefileLocation);
 							}
 						}
-						catch(Exception Ex)
+						catch (Exception Ex)
 						{
 							Log.TraceWarning("Exception while generating include data for {0}: {1}", CurTarget.Name, Ex.ToString());
 						}
@@ -1831,7 +1873,7 @@ namespace UnrealBuildTool
 
 		protected virtual void AddTargetForIntellisense(UEBuildTarget Target)
 		{
-			
+
 		}
 
 
@@ -1857,7 +1899,7 @@ namespace UnrealBuildTool
 				{
 					// If there is a build platform present, add it to the SupportedPlatforms list
 					UEBuildPlatform BuildPlatform;
-					if(UEBuildPlatform.TryGetBuildPlatform(Platform, out BuildPlatform))
+					if (UEBuildPlatform.TryGetBuildPlatform(Platform, out BuildPlatform))
 					{
 						if (InstalledPlatformInfo.IsValidPlatform(Platform, EProjectType.Code))
 						{
@@ -1897,9 +1939,9 @@ namespace UnrealBuildTool
 			}
 
 			// Add all configurations
-			foreach( UnrealTargetConfiguration CurConfiguration in AllowedTargetConfigurations)
+			foreach (UnrealTargetConfiguration CurConfiguration in AllowedTargetConfigurations)
 			{
-				if( CurConfiguration != UnrealTargetConfiguration.Unknown )
+				if (CurConfiguration != UnrealTargetConfiguration.Unknown)
 				{
 					if (InstalledPlatformInfo.IsValidConfiguration(CurConfiguration, EProjectType.Code))
 					{
@@ -1959,16 +2001,16 @@ namespace UnrealBuildTool
 		/// <param name="ModProjects">All mod projects</param>
 		/// <param name="AllModuleFiles">List of *.Build.cs files for all engine programs and games</param>
 		/// <param name="bGatherThirdPartySource">True to gather source code from third party projects too</param>
-		protected void AddProjectsForAllModules( List<FileReference> AllGames, Dictionary<FileReference, ProjectFile> ProgramProjects, List<ProjectFile> ModProjects, List<FileReference> AllModuleFiles, bool bGatherThirdPartySource )
+		protected void AddProjectsForAllModules(List<FileReference> AllGames, Dictionary<FileReference, ProjectFile> ProgramProjects, List<ProjectFile> ModProjects, List<FileReference> AllModuleFiles, bool bGatherThirdPartySource)
 		{
 			HashSet<ProjectFile> ProjectsWithPlugins = new HashSet<ProjectFile>();
-			foreach( FileReference CurModuleFile in AllModuleFiles )
+			foreach (FileReference CurModuleFile in AllModuleFiles)
 			{
 				Log.TraceVerbose("AddProjectsForAllModules " + CurModuleFile);
 
 				// The module's "base directory" is simply the directory where its xxx.Build.cs file is stored.  We'll always
 				// harvest source files for this module in this base directory directory and all of its sub-directories.
-				string ModuleName = CurModuleFile.GetFileNameWithoutAnyExtensions();		// Remove both ".cs" and ".Build"
+				string ModuleName = CurModuleFile.GetFileNameWithoutAnyExtensions();        // Remove both ".cs" and ".Build"
 
 				bool WantProjectFileForModule = true;
 
@@ -1976,7 +2018,7 @@ namespace UnrealBuildTool
 				bool IsThirdPartyModule = CurModuleFile.ContainsName("ThirdParty", UnrealBuildTool.RootDirectory);
 
 				// check for engine, or platform extension engine folders
-				if( !bIncludeEngineSource )
+				if (!bIncludeEngineSource)
 				{
 					if (CurModuleFile.IsUnderDirectory(UnrealBuildTool.EngineDirectory))
 					{
@@ -1986,40 +2028,40 @@ namespace UnrealBuildTool
 					}
 				}
 
-				if( WantProjectFileForModule )
+				if (WantProjectFileForModule)
 				{
 					DirectoryReference BaseFolder;
 					ProjectFile ProjectFile = FindProjectForModule(CurModuleFile, AllGames, ProgramProjects, ModProjects, out BaseFolder);
 
 					// Update our module map
-					ModuleToProjectFileMap[ ModuleName ] = ProjectFile;
+					ModuleToProjectFileMap[ModuleName] = ProjectFile;
 					ProjectFile.IsGeneratedProject = true;
 
 					// Only search subdirectories for non-external modules.  We don't want to add all of the source and header files
 					// for every third-party module, unless we were configured to do so.
 					bool SearchSubdirectories = !IsThirdPartyModule || bGatherThirdPartySource;
 
-					if( bGatherThirdPartySource )
+					if (bGatherThirdPartySource)
 					{
-						Log.TraceInformation( "Searching for third-party source files..." );
+						Log.TraceInformation("Searching for third-party source files...");
 					}
 
 
 					// Find all of the source files (and other files) and add them to the project
-					List<FileReference> FoundFiles = SourceFileSearch.FindModuleSourceFiles( CurModuleFile, SearchSubdirectories:SearchSubdirectories );
+					List<FileReference> FoundFiles = SourceFileSearch.FindModuleSourceFiles(CurModuleFile, SearchSubdirectories: SearchSubdirectories);
 					// remove any target files, they are technically not part of the module and are explicitly added when the project is created
 					FoundFiles.RemoveAll(f => f.FullName.EndsWith(".Target.cs"));
-					ProjectFile.AddFilesToProject( FoundFiles, BaseFolder );
+					ProjectFile.AddFilesToProject(FoundFiles, BaseFolder);
 
 					// Check if there's a plugin directory here
-					if(!ProjectsWithPlugins.Contains(ProjectFile))
+					if (!ProjectsWithPlugins.Contains(ProjectFile))
 					{
 						foreach (DirectoryReference PluginFolder in UnrealBuildTool.GetExtensionDirs(BaseFolder, "Plugins"))
 						{
 							// Add all the plugin files for this project
-							foreach(FileReference PluginFileName in Plugins.EnumeratePlugins(PluginFolder))
+							foreach (FileReference PluginFileName in Plugins.EnumeratePlugins(PluginFolder))
 							{
-								if(!ModProjects.Any(x => x.BaseDir == PluginFileName.Directory))
+								if (!ModProjects.Any(x => x.BaseDir == PluginFileName.Directory))
 								{
 									AddPluginFilesToProject(PluginFileName, BaseFolder, ProjectFile);
 								}
@@ -2037,29 +2079,29 @@ namespace UnrealBuildTool
 			ProjectFile.AddFileToProject(PluginFileName, BaseFolder);
 
 			// Add plugin config files if we have any
-			if( bIncludeConfigFiles )
+			if (bIncludeConfigFiles)
 			{
 				DirectoryReference PluginConfigFolder = DirectoryReference.Combine(PluginFileName.Directory, "Config");
-				if(DirectoryReference.Exists(PluginConfigFolder))
+				if (DirectoryReference.Exists(PluginConfigFolder))
 				{
-					ProjectFile.AddFilesToProject(SourceFileSearch.FindFiles(PluginConfigFolder), BaseFolder );
+					ProjectFile.AddFilesToProject(SourceFileSearch.FindFiles(PluginConfigFolder), BaseFolder);
 				}
 			}
 
 			// Add plugin "resource" files if we have any
 			DirectoryReference PluginResourcesFolder = DirectoryReference.Combine(PluginFileName.Directory, "Resources");
-			if(DirectoryReference.Exists(PluginResourcesFolder))
+			if (DirectoryReference.Exists(PluginResourcesFolder))
 			{
-				ProjectFile.AddFilesToProject(SourceFileSearch.FindFiles(PluginResourcesFolder), BaseFolder );
-            }
+				ProjectFile.AddFilesToProject(SourceFileSearch.FindFiles(PluginResourcesFolder), BaseFolder);
+			}
 
-            // Add plugin shader files if we have any
-            DirectoryReference PluginShadersFolder = DirectoryReference.Combine(PluginFileName.Directory, "Shaders");
-            if (DirectoryReference.Exists(PluginShadersFolder))
-            {
-                ProjectFile.AddFilesToProject(SourceFileSearch.FindFiles(PluginShadersFolder), BaseFolder);
-            }
-        }
+			// Add plugin shader files if we have any
+			DirectoryReference PluginShadersFolder = DirectoryReference.Combine(PluginFileName.Directory, "Shaders");
+			if (DirectoryReference.Exists(PluginShadersFolder))
+			{
+				ProjectFile.AddFilesToProject(SourceFileSearch.FindFiles(PluginShadersFolder), BaseFolder);
+			}
+		}
 
 
 		private ProjectFile FindOrAddProjectHelper(string InProjectFileNameBase, DirectoryReference InBaseFolder)
@@ -2144,8 +2186,8 @@ namespace UnrealBuildTool
 			List<FileReference> AllGames,
 			List<FileReference> AllTargetFiles,
 			String[] Arguments,
-			out ProjectFile EngineProject, 
-			out List<ProjectFile> GameProjects, 
+			out ProjectFile EngineProject,
+			out List<ProjectFile> GameProjects,
 			out Dictionary<FileReference, ProjectFile> ProgramProjects)
 		{
 			// As we're creating project files, we'll also keep track of whether we created an "engine" project and return that if we have one
@@ -2154,17 +2196,17 @@ namespace UnrealBuildTool
 			ProgramProjects = new Dictionary<FileReference, ProjectFile>();
 
 			// Separate the .target.cs files that are platform extension specializations, per target name. These will be added alongside their base target.cs
-			Dictionary<string,List<FileReference>> AllSubTargetFilesPerTarget = new Dictionary<string, List<FileReference>>();
+			Dictionary<string, List<FileReference>> AllSubTargetFilesPerTarget = new Dictionary<string, List<FileReference>>();
 			HashSet<FileReference> AllSubTargetFiles = new HashSet<FileReference>();
-			foreach( FileReference TargetFilePath in AllTargetFiles )
+			foreach (FileReference TargetFilePath in AllTargetFiles)
 			{
-				string[] TargetPathSplit = TargetFilePath.GetFileNameWithoutAnyExtensions().Split(new char[]{'_'}, StringSplitOptions.RemoveEmptyEntries );
-				if (TargetPathSplit.Length > 1 && (UnrealTargetPlatform.IsValidName(TargetPathSplit.Last()) || UnrealPlatformGroup.IsValidName(TargetPathSplit.Last()) ) )
+				string[] TargetPathSplit = TargetFilePath.GetFileNameWithoutAnyExtensions().Split(new char[] { '_' }, StringSplitOptions.RemoveEmptyEntries);
+				if (TargetPathSplit.Length > 1 && (UnrealTargetPlatform.IsValidName(TargetPathSplit.Last()) || UnrealPlatformGroup.IsValidName(TargetPathSplit.Last())))
 				{
 					string TargetName = TargetPathSplit.First();
 					if (!AllSubTargetFilesPerTarget.ContainsKey(TargetName))
 					{
-						AllSubTargetFilesPerTarget.Add(TargetName, new List<FileReference>() );
+						AllSubTargetFilesPerTarget.Add(TargetName, new List<FileReference>());
 					}
 
 					AllSubTargetFilesPerTarget[TargetName].Add(TargetFilePath);
@@ -2178,34 +2220,34 @@ namespace UnrealBuildTool
 			// Keep a cache of the default editor target for each project so that we don't have to interrogate the configs for each target
 			Dictionary<string, string> DefaultProjectEditorTargetCache = new Dictionary<string, string>();
 
-			foreach( FileReference TargetFilePath in AllTargetFiles.Except(AllSubTargetFiles) )
+			foreach (FileReference TargetFilePath in AllTargetFiles.Except(AllSubTargetFiles))
 			{
-				string TargetName = TargetFilePath.GetFileNameWithoutAnyExtensions();		// Remove both ".cs" and ".Target"
+				string TargetName = TargetFilePath.GetFileNameWithoutAnyExtensions();       // Remove both ".cs" and ".Target"
 
 				// Check to see if this is an Engine target.  That is, the target is located under the "Engine" folder
 				bool IsEngineTarget = false;
 				bool WantProjectFileForTarget = true;
-				if(TargetFilePath.IsUnderDirectory(UnrealBuildTool.EngineDirectory))
+				if (TargetFilePath.IsUnderDirectory(UnrealBuildTool.EngineDirectory))
 				{
 					// This is an engine target
 					IsEngineTarget = true;
 
-					if(UnrealBuildTool.GetExtensionDirs(UnrealBuildTool.EngineDirectory, "Source/Programs").Any(x => TargetFilePath.IsUnderDirectory(x)))
+					if (UnrealBuildTool.GetExtensionDirs(UnrealBuildTool.EngineDirectory, "Source/Programs").Any(x => TargetFilePath.IsUnderDirectory(x)))
 					{
 						WantProjectFileForTarget = IncludeEnginePrograms;
 					}
-					else if(UnrealBuildTool.GetExtensionDirs(UnrealBuildTool.EngineDirectory, "Source").Any(x => TargetFilePath.IsUnderDirectory(x)))
+					else if (UnrealBuildTool.GetExtensionDirs(UnrealBuildTool.EngineDirectory, "Source").Any(x => TargetFilePath.IsUnderDirectory(x)))
 					{
 						WantProjectFileForTarget = bIncludeEngineSource;
 					}
 				}
-				
+
 				if (WantProjectFileForTarget)
 				{
 					RulesAssembly RulesAssembly;
 
 					FileReference CheckProjectFile = AllGames.FirstOrDefault(x => TargetFilePath.IsUnderDirectory(x.Directory));
-					if(CheckProjectFile == null)
+					if (CheckProjectFile == null)
 					{
 						RulesAssembly = RulesCompiler.CreateEngineRulesAssembly(false, false);
 					}
@@ -2272,7 +2314,7 @@ namespace UnrealBuildTool
 
 					// Get the suffix to use for this project file. If we have multiple targets of the same type, we'll have to split them out into separate IDE project files.
 					string GeneratedProjectName = TargetRulesObject.GeneratedProjectName;
-					if(GeneratedProjectName == null)
+					if (GeneratedProjectName == null)
 					{
 						ProjectFile ExistingProjectFile;
 						// We should create a separate project for targets which aren't the default for this target type.
@@ -2280,7 +2322,7 @@ namespace UnrealBuildTool
 						// if we aren't an editor, we'll just fall back to previous behavior and assume the first one we encounter
 						// should be added to the main project
 						bool bIsDefaultTargetForType = (DefaultEditorTarget == null) || (TargetRulesObject.Type != TargetType.Editor) || (TargetRulesObject.Name == DefaultEditorTarget);
-						if( !bIsDefaultTargetForType || (ProjectFileMap.TryGetValue(ProjectLocation, out ExistingProjectFile ) && ExistingProjectFile.ProjectTargets.Any(x => x.TargetRules.Type == TargetRulesObject.Type)))
+						if (!bIsDefaultTargetForType || (ProjectFileMap.TryGetValue(ProjectLocation, out ExistingProjectFile) && ExistingProjectFile.ProjectTargets.Any(x => x.TargetRules.Type == TargetRulesObject.Type)))
 						{
 							GeneratedProjectName = TargetRulesObject.Name;
 						}
@@ -2379,14 +2421,14 @@ namespace UnrealBuildTool
 					}
 
 					ProjectTarget ProjectTarget = new ProjectTarget()
-						{
-							TargetRules = TargetRulesObject,
-							TargetFilePath = TargetFilePath,
-							ProjectFilePath = ProjectFilePath,
-							UnrealProjectFilePath = CheckProjectFile,
-							SupportedPlatforms = TargetRulesObject.GetSupportedPlatforms().Where(x => UEBuildPlatform.TryGetBuildPlatform(x, out _)).ToArray(),
-							CreateRulesDelegate = (Platform, Configuration) => RulesAssembly.CreateTargetRules(TargetName, Platform, Configuration, "", CheckProjectFile, new CommandLineArguments(GetTargetArguments(Arguments)))
-                        };
+					{
+						TargetRules = TargetRulesObject,
+						TargetFilePath = TargetFilePath,
+						ProjectFilePath = ProjectFilePath,
+						UnrealProjectFilePath = CheckProjectFile,
+						SupportedPlatforms = TargetRulesObject.GetSupportedPlatforms().Where(x => UEBuildPlatform.TryGetBuildPlatform(x, out _)).ToArray(),
+						CreateRulesDelegate = (Platform, Configuration) => RulesAssembly.CreateTargetRules(TargetName, Platform, Configuration, "", CheckProjectFile, new CommandLineArguments(GetTargetArguments(Arguments)))
+					};
 
 					ProjectFile.ProjectTargets.Add(ProjectTarget);
 
@@ -2396,7 +2438,7 @@ namespace UnrealBuildTool
 					// Add all matching *_<platform>.Target.cs to the same folder
 					if (AllSubTargetFilesPerTarget.ContainsKey(TargetName))
 					{
-						ProjectFile.AddFilesToProject( AllSubTargetFilesPerTarget[TargetName], BaseFolder );
+						ProjectFile.AddFilesToProject(AllSubTargetFilesPerTarget[TargetName], BaseFolder);
 					}
 
 
@@ -2417,7 +2459,7 @@ namespace UnrealBuildTool
 			if (GameProjects.Count == 1)
 			{
 				ProjectFile GameProject = GameProjects.First();
-				foreach(PluginInfo PluginInfo in Plugins.ReadProjectPlugins(GameProject.BaseDir))
+				foreach (PluginInfo PluginInfo in Plugins.ReadProjectPlugins(GameProject.BaseDir))
 				{
 					if (PluginInfo.Descriptor.Modules != null && PluginInfo.Descriptor.Modules.Count > 0 && PluginInfo.Type == PluginType.Mod)
 					{
@@ -2448,93 +2490,93 @@ namespace UnrealBuildTool
 		{
 			return FileReference.Combine(IntermediateProjectFilesPath, BaseName + ProjectFileExtension);
 		}
-		
+
 		/// Adds shader source code to the specified project
-		protected void AddEngineShaderSource( ProjectFile EngineProject )
+		protected void AddEngineShaderSource(ProjectFile EngineProject)
 		{
 			// Setup a project file entry for this module's project.  Remember, some projects may host multiple modules!
-			DirectoryReference ShadersDirectory = DirectoryReference.Combine( UnrealBuildTool.EngineDirectory, "Shaders" );
-			if(DirectoryReference.Exists(ShadersDirectory))
+			DirectoryReference ShadersDirectory = DirectoryReference.Combine(UnrealBuildTool.EngineDirectory, "Shaders");
+			if (DirectoryReference.Exists(ShadersDirectory))
 			{
 				List<string> SubdirectoryNamesToExclude = new List<string>();
 				{
 					// Don't include binary shaders in the project file.
-					SubdirectoryNamesToExclude.Add( "Binaries" );
+					SubdirectoryNamesToExclude.Add("Binaries");
 					// We never want shader intermediate files in our project file
-					SubdirectoryNamesToExclude.Add( "PDBDump" );
-					SubdirectoryNamesToExclude.Add( "WorkingDirectory" );
+					SubdirectoryNamesToExclude.Add("PDBDump");
+					SubdirectoryNamesToExclude.Add("WorkingDirectory");
 				}
 
-				EngineProject.AddFilesToProject( SourceFileSearch.FindFiles( ShadersDirectory, SubdirectoryNamesToExclude ), UnrealBuildTool.EngineDirectory );
+				EngineProject.AddFilesToProject(SourceFileSearch.FindFiles(ShadersDirectory, SubdirectoryNamesToExclude), UnrealBuildTool.EngineDirectory);
 			}
 		}
 
 
 		/// Adds engine build infrastructure files to the specified project
-		protected void AddEngineBuildFiles( ProjectFile EngineProject )
+		protected void AddEngineBuildFiles(ProjectFile EngineProject)
 		{
-			DirectoryReference BuildDirectory = DirectoryReference.Combine( UnrealBuildTool.EngineDirectory, "Build" );
+			DirectoryReference BuildDirectory = DirectoryReference.Combine(UnrealBuildTool.EngineDirectory, "Build");
 
 			List<string> SubdirectoryNamesToExclude = new List<string>();
 			SubdirectoryNamesToExclude.Add("Receipts");
 
-			EngineProject.AddFilesToProject( SourceFileSearch.FindFiles( BuildDirectory, SubdirectoryNamesToExclude ), UnrealBuildTool.EngineDirectory );
+			EngineProject.AddFilesToProject(SourceFileSearch.FindFiles(BuildDirectory, SubdirectoryNamesToExclude), UnrealBuildTool.EngineDirectory);
 		}
 
 
 
 		/// Adds engine documentation to the specified project
-		protected void AddEngineDocumentation( ProjectFile EngineProject )
+		protected void AddEngineDocumentation(ProjectFile EngineProject)
 		{
 			// NOTE: The project folder added here will actually be collapsed away later if not needed
-			DirectoryReference DocumentationProjectDirectory = DirectoryReference.Combine( UnrealBuildTool.EngineDirectory, "Documentation" );
-			DirectoryReference DocumentationSourceDirectory = DirectoryReference.Combine( UnrealBuildTool.EngineDirectory, "Documentation", "Source" );
-			DirectoryInfo DirInfo = new DirectoryInfo( DocumentationProjectDirectory.FullName );
-			if( DirInfo.Exists && DirectoryReference.Exists(DocumentationSourceDirectory) )
+			DirectoryReference DocumentationProjectDirectory = DirectoryReference.Combine(UnrealBuildTool.EngineDirectory, "Documentation");
+			DirectoryReference DocumentationSourceDirectory = DirectoryReference.Combine(UnrealBuildTool.EngineDirectory, "Documentation", "Source");
+			DirectoryInfo DirInfo = new DirectoryInfo(DocumentationProjectDirectory.FullName);
+			if (DirInfo.Exists && DirectoryReference.Exists(DocumentationSourceDirectory))
 			{
-				Log.TraceVerbose( "Adding documentation files..." );
+				Log.TraceVerbose("Adding documentation files...");
 
 				List<string> SubdirectoryNamesToExclude = new List<string>();
 				{
 					// We never want any of the images or attachment files included in our generated project
-					SubdirectoryNamesToExclude.Add( "Images" );
-					SubdirectoryNamesToExclude.Add( "Attachments" );
+					SubdirectoryNamesToExclude.Add("Images");
+					SubdirectoryNamesToExclude.Add("Attachments");
 
 					// The API directory is huge, so don't include any of it
 					SubdirectoryNamesToExclude.Add("API");
 
 					// Omit Javascript source because it just confuses the Visual Studio IDE
-					SubdirectoryNamesToExclude.Add( "Javascript" );
+					SubdirectoryNamesToExclude.Add("Javascript");
 				}
 
-				List<FileReference> DocumentationFiles = SourceFileSearch.FindFiles( DocumentationSourceDirectory, SubdirectoryNamesToExclude );
+				List<FileReference> DocumentationFiles = SourceFileSearch.FindFiles(DocumentationSourceDirectory, SubdirectoryNamesToExclude);
 
 				// Filter out non-English documentation files if we were configured to do so
-				if( !bAllDocumentationLanguages )
+				if (!bAllDocumentationLanguages)
 				{
 					List<FileReference> FilteredDocumentationFiles = new List<FileReference>();
-					foreach( FileReference DocumentationFile in DocumentationFiles )
+					foreach (FileReference DocumentationFile in DocumentationFiles)
 					{
 						bool bPassesFilter = true;
-						if( DocumentationFile.FullName.EndsWith( ".udn", StringComparison.InvariantCultureIgnoreCase ) )
+						if (DocumentationFile.FullName.EndsWith(".udn", StringComparison.InvariantCultureIgnoreCase))
 						{
-							string LanguageSuffix = Path.GetExtension( Path.GetFileNameWithoutExtension( DocumentationFile.FullName ) );
-							if( !String.IsNullOrEmpty( LanguageSuffix ) &&
-								!LanguageSuffix.Equals( ".int", StringComparison.InvariantCultureIgnoreCase ) )
+							string LanguageSuffix = Path.GetExtension(Path.GetFileNameWithoutExtension(DocumentationFile.FullName));
+							if (!String.IsNullOrEmpty(LanguageSuffix) &&
+								!LanguageSuffix.Equals(".int", StringComparison.InvariantCultureIgnoreCase))
 							{
 								bPassesFilter = false;
 							}
 						}
 
-						if( bPassesFilter )
+						if (bPassesFilter)
 						{
-							FilteredDocumentationFiles.Add( DocumentationFile );
+							FilteredDocumentationFiles.Add(DocumentationFile);
 						}
 					}
 					DocumentationFiles = FilteredDocumentationFiles;
 				}
 
-				EngineProject.AddFilesToProject( DocumentationFiles, UnrealBuildTool.EngineDirectory );
+				EngineProject.AddFilesToProject(DocumentationFiles, UnrealBuildTool.EngineDirectory);
 			}
 			else
 			{
@@ -2553,29 +2595,29 @@ namespace UnrealBuildTool
 		/// <param name="IncludeInGeneratedProjects">True if this project should be included in the set of generated projects.  Only matters when actually generating project files.</param>
 		/// <param name="bAlreadyExisted">True if we already had this project file</param>
 		/// <returns>Object that represents this project file in Unreal Build Tool</returns>
-		public ProjectFile FindOrAddProject( FileReference FilePath, DirectoryReference BaseDir, bool IncludeInGeneratedProjects, out bool bAlreadyExisted )
+		public ProjectFile FindOrAddProject(FileReference FilePath, DirectoryReference BaseDir, bool IncludeInGeneratedProjects, out bool bAlreadyExisted)
 		{
-			if( FilePath == null )
+			if (FilePath == null)
 			{
-				throw new BuildException( "Not valid to call FindOrAddProject() with an empty file path!" );
+				throw new BuildException("Not valid to call FindOrAddProject() with an empty file path!");
 			}
 
 			// Do we already have this project?
 			ProjectFile ExistingProjectFile;
-			if( ProjectFileMap.TryGetValue( FilePath, out ExistingProjectFile ) )
+			if (ProjectFileMap.TryGetValue(FilePath, out ExistingProjectFile))
 			{
 				bAlreadyExisted = true;
 				return ExistingProjectFile;
 			}
 
 			// Add a new project file for the specified path
-			ProjectFile NewProjectFile = AllocateProjectFile( FilePath );
+			ProjectFile NewProjectFile = AllocateProjectFile(FilePath);
 			NewProjectFile.BaseDir = BaseDir;
-			ProjectFileMap[ FilePath ] = NewProjectFile;
+			ProjectFileMap[FilePath] = NewProjectFile;
 
-			if( IncludeInGeneratedProjects )
+			if (IncludeInGeneratedProjects)
 			{
-				GeneratedProjectFiles.Add( NewProjectFile );
+				GeneratedProjectFiles.Add(NewProjectFile);
 			}
 
 			bAlreadyExisted = false;
@@ -2588,7 +2630,7 @@ namespace UnrealBuildTool
 		/// </summary>
 		/// <param name="InitFilePath">Path to the project file</param>
 		/// <returns>The newly allocated project file object</returns>
-		protected abstract ProjectFile AllocateProjectFile( FileReference InitFilePath );
+		protected abstract ProjectFile AllocateProjectFile(FileReference InitFilePath);
 
 
 		/// <summary>
@@ -2597,7 +2639,7 @@ namespace UnrealBuildTool
 		/// <param name="OwnerProjectFileGenerator">Project file generator that owns this object</param>
 		/// <param name="FolderName">Name for this folder</param>
 		/// <returns>The newly allocated project folder object</returns>
-		public abstract MasterProjectFolder AllocateMasterProjectFolder(ProjectFileGenerator OwnerProjectFileGenerator, string FolderName );
+		public abstract MasterProjectFolder AllocateMasterProjectFolder(ProjectFileGenerator OwnerProjectFileGenerator, string FolderName);
 
 
 
@@ -2610,27 +2652,27 @@ namespace UnrealBuildTool
 			get
 			{
 				List<ProjectFile> CombinedList = new List<ProjectFile>();
-				CombinedList.AddRange( GeneratedProjectFiles );
-				CombinedList.AddRange( OtherProjectFiles );
+				CombinedList.AddRange(GeneratedProjectFiles);
+				CombinedList.AddRange(OtherProjectFiles);
 				return CombinedList;
 			}
 		}
 
-	
+
 		/// <summary>
 		/// Writes the project files to disk
 		/// </summary>
 		/// <returns>True if successful</returns>
 		protected virtual bool WriteProjectFiles(PlatformProjectGeneratorCollection PlatformProjectGenerators)
 		{
-            using(ProgressWriter Progress = new ProgressWriter("Writing project files...", true))
+			using (ProgressWriter Progress = new ProgressWriter("Writing project files...", true))
 			{
-				int TotalProjectFileCount = GeneratedProjectFiles.Count + 1;	// +1 for the master project file, which we'll save next
+				int TotalProjectFileCount = GeneratedProjectFiles.Count + 1;    // +1 for the master project file, which we'll save next
 
-				for(int ProjectFileIndex = 0 ; ProjectFileIndex < GeneratedProjectFiles.Count; ++ProjectFileIndex )
+				for (int ProjectFileIndex = 0; ProjectFileIndex < GeneratedProjectFiles.Count; ++ProjectFileIndex)
 				{
-					ProjectFile CurProject = GeneratedProjectFiles[ ProjectFileIndex ];
-					if( !CurProject.WriteProjectFile(SupportedPlatforms, SupportedConfigurations, PlatformProjectGenerators) )
+					ProjectFile CurProject = GeneratedProjectFiles[ProjectFileIndex];
+					if (!CurProject.WriteProjectFile(SupportedPlatforms, SupportedConfigurations, PlatformProjectGenerators))
 					{
 						return false;
 					}
@@ -2638,7 +2680,7 @@ namespace UnrealBuildTool
 					Progress.Write(ProjectFileIndex + 1, TotalProjectFileCount);
 				}
 
-				WriteMasterProjectFile( UBTProject, PlatformProjectGenerators );
+				WriteMasterProjectFile(UBTProject, PlatformProjectGenerators);
 				Progress.Write(TotalProjectFileCount, TotalProjectFileCount);
 			}
 
@@ -2651,7 +2693,7 @@ namespace UnrealBuildTool
 		/// <param name="UBTProject">The UnrealBuildTool project</param>
 		/// <param name="PlatformProjectGenerators">The platform project file generators</param>
 		/// <returns>True if successful</returns>
-		protected abstract bool WriteMasterProjectFile( ProjectFile UBTProject, PlatformProjectGeneratorCollection PlatformProjectGenerators );
+		protected abstract bool WriteMasterProjectFile(ProjectFile UBTProject, PlatformProjectGeneratorCollection PlatformProjectGenerators);
 
 
 		/// <summary>
@@ -2659,7 +2701,7 @@ namespace UnrealBuildTool
 		/// </summary>
 		/// <param name="PlatformProjectGenerators">The platform project generators</param>
 		/// <param name="IntermediateProjectFilesPath">Intermediate project files folder</param>
-		protected virtual void WriteDebugSolutionFiles( PlatformProjectGeneratorCollection PlatformProjectGenerators, DirectoryReference IntermediateProjectFilesPath )
+		protected virtual void WriteDebugSolutionFiles(PlatformProjectGeneratorCollection PlatformProjectGenerators, DirectoryReference IntermediateProjectFilesPath)
 		{
 		}
 
@@ -2671,27 +2713,27 @@ namespace UnrealBuildTool
 		/// <param name="NewFileContents">File content</param>
 		/// <param name="InEncoding"></param>
 		/// <returns>True if the file was saved, or if it didn't need to be overwritten because the content was unchanged</returns>
-		public static bool WriteFileIfChanged( string FileName, string NewFileContents, Encoding InEncoding = null )
+		public static bool WriteFileIfChanged(string FileName, string NewFileContents, Encoding InEncoding = null)
 		{
 			// Check to see if the file already exists, and if so, load it up
 			string LoadedFileContent = null;
-			bool FileAlreadyExists = File.Exists( FileName );
-			if( FileAlreadyExists )
+			bool FileAlreadyExists = File.Exists(FileName);
+			if (FileAlreadyExists)
 			{
 				try
 				{
-					LoadedFileContent = File.ReadAllText( FileName );
+					LoadedFileContent = File.ReadAllText(FileName);
 				}
-				catch( Exception )
+				catch (Exception)
 				{
-					Log.TraceInformation( "Error while trying to load existing file {0}.  Ignored.", FileName );
+					Log.TraceInformation("Error while trying to load existing file {0}.  Ignored.", FileName);
 				}
 			}
 
 
 			// Don't bother saving anything out if the new file content is the same as the old file's content
 			bool FileNeedsSave = true;
-			if( LoadedFileContent != null )
+			if (LoadedFileContent != null)
 			{
 				bool bIgnoreProjectFileWhitespaces = true;
 				if (ProjectFileComparer.CompareOrdinalIgnoreCase(LoadedFileContent, NewFileContents, bIgnoreProjectFileWhitespaces) == 0)
@@ -2700,32 +2742,32 @@ namespace UnrealBuildTool
 					FileNeedsSave = false;
 				}
 
-				if( !FileNeedsSave )
+				if (!FileNeedsSave)
 				{
-					Log.TraceVerbose( "Skipped saving {0} because contents haven't changed.", Path.GetFileName( FileName ) );
+					Log.TraceVerbose("Skipped saving {0} because contents haven't changed.", Path.GetFileName(FileName));
 				}
 			}
 
-			if( FileNeedsSave )
+			if (FileNeedsSave)
 			{
 				// Save the file
 				try
 				{
-					Directory.CreateDirectory( Path.GetDirectoryName( FileName ) );
-                    // When WriteAllText is passed Encoding.UTF8 it likes to write a BOM marker
-                    // at the start of the file (adding two bytes to the file length).  For most
-                    // files this is only mildly annoying but for Makefiles it can actually make
-                    // them un-useable.
-                    // TODO(sbc): See if we can just drop the Encoding.UTF8 argument on all
-                    // platforms.  In this case UTF8 encoding will still be used but without the
-                    // BOM, which is, AFAICT, desirable in almost all cases.
+					Directory.CreateDirectory(Path.GetDirectoryName(FileName));
+					// When WriteAllText is passed Encoding.UTF8 it likes to write a BOM marker
+					// at the start of the file (adding two bytes to the file length).  For most
+					// files this is only mildly annoying but for Makefiles it can actually make
+					// them un-useable.
+					// TODO(sbc): See if we can just drop the Encoding.UTF8 argument on all
+					// platforms.  In this case UTF8 encoding will still be used but without the
+					// BOM, which is, AFAICT, desirable in almost all cases.
 					if (BuildHostPlatform.Current.Platform == UnrealTargetPlatform.Linux || BuildHostPlatform.Current.Platform == UnrealTargetPlatform.Mac)
-                        File.WriteAllText(FileName, NewFileContents, new UTF8Encoding());
-                    else
-                        File.WriteAllText(FileName, NewFileContents, InEncoding != null ? InEncoding : Encoding.UTF8);
-                    Log.TraceVerbose("Saved {0}", Path.GetFileName(FileName));
+						File.WriteAllText(FileName, NewFileContents, new UTF8Encoding());
+					else
+						File.WriteAllText(FileName, NewFileContents, InEncoding != null ? InEncoding : Encoding.UTF8);
+					Log.TraceVerbose("Saved {0}", Path.GetFileName(FileName));
 				}
-				catch( Exception ex )
+				catch (Exception ex)
 				{
 					// Unable to write to the project file.
 					string Message = string.Format("Error while trying to write file {0}.  The file is probably read-only.", FileName);
@@ -2750,14 +2792,14 @@ namespace UnrealBuildTool
 		/// <returns>True if successful</returns>
 		public void AddExistingProjectFile(ProjectFile InProject, bool bNeedsAllPlatformAndConfigurations = false, bool bForceDevelopmentConfiguration = false, bool bProjectDeploys = false, List<UnrealTargetPlatform> InSupportedPlatforms = null, List<UnrealTargetConfiguration> InSupportedConfigurations = null)
 		{
-			if( InProject.ProjectTargets.Count != 0 )
+			if (InProject.ProjectTargets.Count != 0)
 			{
-				throw new BuildException( "Expecting existing project to not have any ProjectTargets defined yet." );
+				throw new BuildException("Expecting existing project to not have any ProjectTargets defined yet.");
 			}
-			
+
 			ProjectTarget ProjectTarget = new ProjectTarget();
 			ProjectTarget.SupportedPlatforms = new UnrealTargetPlatform[0];
-			if( bForceDevelopmentConfiguration )
+			if (bForceDevelopmentConfiguration)
 			{
 				ProjectTarget.ForceDevelopmentConfiguration = true;
 			}
@@ -2772,7 +2814,7 @@ namespace UnrealBuildTool
 				Array AllConfigurations = Enum.GetValues(typeof(UnrealTargetConfiguration));
 				foreach (UnrealTargetConfiguration CurConfiguration in AllConfigurations)
 				{
-					ProjectTarget.ExtraSupportedConfigurations.Add( CurConfiguration );
+					ProjectTarget.ExtraSupportedConfigurations.Add(CurConfiguration);
 				}
 			}
 			else if (InSupportedPlatforms != null || InSupportedConfigurations != null)
@@ -2811,76 +2853,76 @@ namespace UnrealBuildTool
 			}
 			else
 			{
-                bool bFoundDevelopmentConfig = false;
-                bool bFoundDebugConfig = false;
+				bool bFoundDevelopmentConfig = false;
+				bool bFoundDebugConfig = false;
 
-                try
-                {
+				try
+				{
 
-	                bool configsFound = false;
-                    // Parse the project and ensure both Development and Debug configurations are present
-                    foreach (string Config in XElement.Load(InProject.ProjectFilePath.FullName).Elements("{http://schemas.microsoft.com/developer/msbuild/2003}PropertyGroup")
-                                           .Where(node => node.Attribute("Condition") != null)
-                                           .Select(node => node.Attribute("Condition").ToString())
-                                           .ToList())
-                    {
-	                    configsFound = true;
+					bool configsFound = false;
+					// Parse the project and ensure both Development and Debug configurations are present
+					foreach (string Config in XElement.Load(InProject.ProjectFilePath.FullName).Elements("{http://schemas.microsoft.com/developer/msbuild/2003}PropertyGroup")
+										   .Where(node => node.Attribute("Condition") != null)
+										   .Select(node => node.Attribute("Condition").ToString())
+										   .ToList())
+					{
+						configsFound = true;
 
 						if (Config.Contains("Development|"))
-                        {
-                            bFoundDevelopmentConfig = true;
-                        }
-                        else if (Config.Contains("Debug|"))
-                        {
-                            bFoundDebugConfig = true;
-                        }
-                    }
-
-                    // for a net core project the PropertyGroup namespace changes, so lets see if we can find if its a project like that, net core also provides a hand semicolon separated list of configurations.
-                    // This is optional and defaults to Debug;Release if not present which still means Development is not present
+						{
+							bFoundDevelopmentConfig = true;
+						}
+						else if (Config.Contains("Debug|"))
+						{
+							bFoundDebugConfig = true;
+						}
+					}
+
+					// for a net core project the PropertyGroup namespace changes, so lets see if we can find if its a project like that, net core also provides a hand semicolon separated list of configurations.
+					// This is optional and defaults to Debug;Release if not present which still means Development is not present
 					// as such we we expect every proj to have configurations presents.
 					if (!configsFound)
 					{
-	                    foreach (string Config in XElement.Load(InProject.ProjectFilePath.FullName).Elements("PropertyGroup")
-		                    .Elements("Configurations")
-	                        .SelectMany(node => node.Value.Split(';'))
-	                        .ToList())
-	                    {
-		                    if (Config == "Development")
-		                    {
+						foreach (string Config in XElement.Load(InProject.ProjectFilePath.FullName).Elements("PropertyGroup")
+							.Elements("Configurations")
+							.SelectMany(node => node.Value.Split(';'))
+							.ToList())
+						{
+							if (Config == "Development")
+							{
 								bFoundDevelopmentConfig = true;
-		                    }
+							}
 							else if (Config == "Debug")
-		                    {
-			                    bFoundDebugConfig = true;
-		                    }
-	                    }
+							{
+								bFoundDebugConfig = true;
+							}
+						}
 
 					}
-                }
-                catch
-                {
-                    Trace.TraceError("Unable to parse existing project file {0}", InProject.ProjectFilePath.FullName);
-                }
+				}
+				catch
+				{
+					Trace.TraceError("Unable to parse existing project file {0}", InProject.ProjectFilePath.FullName);
+				}
 
-                if (!bFoundDebugConfig || !bFoundDevelopmentConfig)
-                {
-                    throw new BuildException("Existing C# project {0} must contain a {1} configuration", InProject.ProjectFilePath.FullName, bFoundDebugConfig ? "Development" : "Debug");
-                }
+				if (!bFoundDebugConfig || !bFoundDevelopmentConfig)
+				{
+					throw new BuildException("Existing C# project {0} must contain a {1} configuration", InProject.ProjectFilePath.FullName, bFoundDebugConfig ? "Development" : "Debug");
+				}
 
-                // For existing project files, just support the default desktop platforms and configurations
-                ProjectTarget.ExtraSupportedPlatforms.AddRange(Utils.GetPlatformsInClass(UnrealPlatformClass.Desktop));
+				// For existing project files, just support the default desktop platforms and configurations
+				ProjectTarget.ExtraSupportedPlatforms.AddRange(Utils.GetPlatformsInClass(UnrealPlatformClass.Desktop));
 				// Debug and Development only
 				ProjectTarget.ExtraSupportedConfigurations.Add(UnrealTargetConfiguration.Debug);
 				ProjectTarget.ExtraSupportedConfigurations.Add(UnrealTargetConfiguration.Development);
 			}
 
-			InProject.ProjectTargets.Add( ProjectTarget );
+			InProject.ProjectTargets.Add(ProjectTarget);
 
 			// Existing projects must always have a GUID.  This will throw an exception if one isn't found.
 			InProject.LoadGUIDFromExistingProject();
 
-			OtherProjectFiles.Add( InProject );
+			OtherProjectFiles.Add(InProject);
 		}
 
 		public virtual string[] GetTargetArguments(string[] Arguments)
@@ -2914,7 +2956,7 @@ namespace UnrealBuildTool
 		/// aren't generating them ourselves.  Note that these may *not* always be C++ project files (e.g. C#)
 		protected readonly List<ProjectFile> OtherProjectFiles = new List<ProjectFile>();
 
-        protected readonly List<ProjectFile> AutomationProjectFiles = new List<ProjectFile>();
+		protected readonly List<ProjectFile> AutomationProjectFiles = new List<ProjectFile>();
 
 		/// List of top-level folders in the master project file
 		protected MasterProjectFolder RootFolder;
diff --git a/Engine/Source/Programs/UnrealBuildTool/ProjectFiles/VisualStudio/VCProject.cs b/Engine/Source/Programs/UnrealBuildTool/ProjectFiles/VisualStudio/VCProject.cs
--- a/Engine/Source/Programs/UnrealBuildTool/ProjectFiles/VisualStudio/VCProject.cs
+++ b/Engine/Source/Programs/UnrealBuildTool/ProjectFiles/VisualStudio/VCProject.cs
@@ -176,12 +176,29 @@ namespace UnrealBuildTool
 		{
 			if (!ProjectFileGenerator.bIncludeTestAndShippingConfigs)
 			{
-				if(Configuration == UnrealTargetConfiguration.Test || Configuration == UnrealTargetConfiguration.Shipping)
+				if (Configuration == UnrealTargetConfiguration.Test || Configuration == UnrealTargetConfiguration.Shipping)
 				{
 					return false;
 				}
 			}
 
+			if (!ProjectFileGenerator.bIncludeDebugConfigs)
+			{
+				if (Configuration == UnrealTargetConfiguration.Debug || Configuration == UnrealTargetConfiguration.DebugGame)
+				{
+					return false;
+				}
+			}
+
+			if (!ProjectFileGenerator.bIncludeDevelopmentConfigs)
+			{
+				if (Configuration == UnrealTargetConfiguration.Development)
+				{
+					return false;
+				}
+			}
+
+
 			UEBuildPlatform BuildPlatform;
 			if (!UEBuildPlatform.TryGetBuildPlatform(Platform, out BuildPlatform))
 			{
@@ -330,7 +347,7 @@ namespace UnrealBuildTool
 				ProjectPlatformName = DefaultPlatformName;
 			}
 
-			if(TargetConfigurationName != TargetType.Game)
+			if (TargetConfigurationName != TargetType.Game)
 			{
 				ProjectConfigurationName += "_" + TargetConfigurationName.ToString();
 			}
@@ -347,7 +364,7 @@ namespace UnrealBuildTool
 		public override MSBuildProjectContext GetMatchingProjectContext(TargetType SolutionTarget, UnrealTargetConfiguration SolutionConfiguration, UnrealTargetPlatform SolutionPlatform, PlatformProjectGeneratorCollection PlatformProjectGenerators)
 		{
 			// Stub projects always build in the same configuration
-			if(IsStubProject)
+			if (IsStubProject)
 			{
 				return new MSBuildProjectContext(StubProjectConfigurationName, StubProjectPlatformName);
 			}
@@ -366,7 +383,7 @@ namespace UnrealBuildTool
 
 			// Programs are built in editor configurations (since the editor is like a desktop program too) and game configurations (since we omit the "game" qualification in the configuration name).
 			bool IsProgramProject = ProjectTargets[0].TargetRules != null && ProjectTargets[0].TargetRules.Type == TargetType.Program;
-			if(!IsProgramProject || SolutionTarget == TargetType.Game || SolutionTarget == TargetType.Editor)
+			if (!IsProgramProject || SolutionTarget == TargetType.Game || SolutionTarget == TargetType.Editor)
 			{
 				// Get the target type we expect to find for this project
 				TargetType TargetConfigurationName = SolutionTarget;
@@ -380,7 +397,7 @@ namespace UnrealBuildTool
 				List<ProjectTarget> MatchingProjectTargets = new List<ProjectTarget>();
 				foreach (ProjectTarget ProjectTarget in ProjectTargets)
 				{
-					if(VCProjectFile.IsValidProjectPlatformAndConfiguration(ProjectTarget, SolutionPlatform, SolutionConfiguration, PlatformProjectGenerators))
+					if (VCProjectFile.IsValidProjectPlatformAndConfiguration(ProjectTarget, SolutionPlatform, SolutionConfiguration, PlatformProjectGenerators))
 					{
 						if (ProjectTarget.TargetRules != null)
 						{
@@ -402,10 +419,10 @@ namespace UnrealBuildTool
 				// Always allow SCW and UnrealLighmass to build in editor configurations
 				if (MatchingProjectTargets.Count == 0 && SolutionTarget == TargetType.Editor && SolutionPlatform == UnrealTargetPlatform.Win64)
 				{
-					foreach(ProjectTarget ProjectTarget in ProjectTargets)
+					foreach (ProjectTarget ProjectTarget in ProjectTargets)
 					{
 						string TargetName = ProjectTargets[0].TargetRules.Name;
-						if(TargetName == "ShaderCompileWorker" || TargetName == "UnrealLightmass")
+						if (TargetName == "ShaderCompileWorker" || TargetName == "UnrealLightmass")
 						{
 							MatchingProjectTargets.Add(ProjectTarget);
 							break;
@@ -414,13 +431,13 @@ namespace UnrealBuildTool
 				}
 
 				// Make sure there's only one matching project target
-				if(MatchingProjectTargets.Count > 1)
+				if (MatchingProjectTargets.Count > 1)
 				{
 					throw new BuildException("Not expecting more than one target for project {0} to match solution configuration {1} {2} {3}", ProjectFilePath, SolutionTarget, SolutionConfiguration, SolutionPlatform);
 				}
 
 				// If we found a matching project, get matching configuration
-				if(MatchingProjectTargets.Count == 1)
+				if (MatchingProjectTargets.Count == 1)
 				{
 					// Get the matching target
 					ProjectTarget MatchingProjectTarget = MatchingProjectTargets[0];
@@ -473,7 +490,7 @@ namespace UnrealBuildTool
 				}
 			}
 
-			return new MSBuildProjectContext(ProjectConfigurationName, ProjectPlatformName){ bBuildByDefault = bBuildByDefault, bDeployByDefault = bDeployByDefault };
+			return new MSBuildProjectContext(ProjectConfigurationName, ProjectPlatformName) { bBuildByDefault = bBuildByDefault, bDeployByDefault = bDeployByDefault };
 		}
 
 		class ProjectConfigAndTargetCombination
@@ -537,7 +554,7 @@ namespace UnrealBuildTool
 
 		WindowsCompiler GetCompilerForIntellisense()
 		{
-			switch(ProjectFileFormat)
+			switch (ProjectFileFormat)
 			{
 				case VCProjectFileFormat.VisualStudio2019:
 					return WindowsCompiler.VisualStudio2019;
@@ -570,7 +587,7 @@ namespace UnrealBuildTool
 
 			return Version;
 		}
-		
+
 		/// <summary>
 		/// Gets compiler switch for specifying in AdditionalOptions in .vcxproj file for specific C++ version
 		/// </summary>
@@ -599,7 +616,7 @@ namespace UnrealBuildTool
 		private void BuildProjectConfigAndTargetCombinations(List<UnrealTargetPlatform> InPlatforms, List<UnrealTargetConfiguration> InConfigurations, PlatformProjectGeneratorCollection PlatformProjectGenerators)
 		{
 			//no need to do this more than once
-			if(ProjectConfigAndTargetCombinations == null)
+			if (ProjectConfigAndTargetCombinations == null)
 			{
 				// Build up a list of platforms and configurations this project will support.  In this list, Unknown simply
 				// means that we should use the default "stub" project platform and configuration name.
@@ -653,9 +670,9 @@ namespace UnrealBuildTool
 				}
 
 				// Create a list of platforms for the "invalid" configuration. We always require at least one of these.
-				if(ProjectConfigAndTargetCombinations.Count == 0)
+				if (ProjectConfigAndTargetCombinations.Count == 0)
 				{
-					InvalidConfigPlatformNames = new HashSet<string>{ DefaultPlatformName };
+					InvalidConfigPlatformNames = new HashSet<string> { DefaultPlatformName };
 				}
 				else
 				{
@@ -703,7 +720,7 @@ namespace UnrealBuildTool
 
 						//write out additional project files
 						Tuple<ProjectFile, string> DebugProjectInfo = ProjGenerator.WriteAdditionalProjFile(this);
-						if(DebugProjectInfo != null)
+						if (DebugProjectInfo != null)
 						{
 							ProjectFiles.Add(DebugProjectInfo);
 						}
@@ -717,7 +734,7 @@ namespace UnrealBuildTool
 		private string[] FilteredList = null;
 
 		bool IncludePathIsFilteredOut(DirectoryReference IncludePath)
-		{			
+		{
 			// Turn the filter string into an array, remove whitespace, and normalize any path statements the first time
 			// we are asked to check a path.
 			if (FilteredList == null)
@@ -740,7 +757,7 @@ namespace UnrealBuildTool
 			{
 				foreach (string Entry in FilteredList)
 				{
-					if (IncludePathWithSeparator.Contains(Entry, StringComparison.OrdinalIgnoreCase) 
+					if (IncludePathWithSeparator.Contains(Entry, StringComparison.OrdinalIgnoreCase)
 						|| IncludePathWithSeparator.EndsWith(Entry, StringComparison.OrdinalIgnoreCase))
 					{
 						return true;
@@ -765,7 +782,7 @@ namespace UnrealBuildTool
 				IncludePathsCollection Collection;
 				if (BaseDirToIncludePaths.TryGetValue(CurrentDir, out Collection))
 				{
-					foreach(DirectoryReference IncludePath in Collection.AbsolutePaths)
+					foreach (DirectoryReference IncludePath in Collection.AbsolutePaths)
 					{
 						if (!IgnorePaths.Contains(IncludePath) && !IncludePathIsFilteredOut(IncludePath))
 						{
@@ -795,7 +812,7 @@ namespace UnrealBuildTool
 				Dictionary<DirectoryReference, int> SourceDirToCount = new Dictionary<DirectoryReference, int>();
 				foreach (SourceFile SourceFile in SourceFiles)
 				{
-					if(SourceFile.Reference.HasExtension(".cpp"))
+					if (SourceFile.Reference.HasExtension(".cpp"))
 					{
 						DirectoryReference SourceDir = SourceFile.Reference.Directory;
 
@@ -827,7 +844,7 @@ namespace UnrealBuildTool
 
 				// Append the most common include paths to the search list.
 				if (Settings.MaxSharedIncludePaths > 0)
-				{	
+				{
 					foreach (DirectoryReference IncludePath in IncludePathToCount.OrderByDescending(x => x.Value).Select(x => x.Key))
 					{
 						string RelativePath = NormalizeProjectPath(IncludePath);
@@ -840,7 +857,7 @@ namespace UnrealBuildTool
 						{
 							SharedIncludeSearchPathsSet.Add(IncludePath);
 							SharedIncludeSearchPaths.AppendFormat("{0};", RelativePath);
-						}						
+						}
 					}
 				}
 
@@ -891,14 +908,18 @@ namespace UnrealBuildTool
 
 			// Make a list of the platforms and configs as project-format names
 			List<UnrealTargetPlatform> ProjectPlatforms = new List<UnrealTargetPlatform>();
-			List<Tuple<string, UnrealTargetPlatform>> ProjectPlatformNameAndPlatforms = new List<Tuple<string, UnrealTargetPlatform>>();	// ProjectPlatformName, Platform
-			List<Tuple<string, UnrealTargetConfiguration>> ProjectConfigurationNameAndConfigurations = new List<Tuple<string, UnrealTargetConfiguration>>();	// ProjectConfigurationName, Configuration
+			List<Tuple<string, UnrealTargetPlatform>> ProjectPlatformNameAndPlatforms = new List<Tuple<string, UnrealTargetPlatform>>();    // ProjectPlatformName, Platform
+			List<Tuple<string, UnrealTargetConfiguration>> ProjectConfigurationNameAndConfigurations = new List<Tuple<string, UnrealTargetConfiguration>>();    // ProjectConfigurationName, Configuration
 			foreach (ProjectConfigAndTargetCombination Combination in ProjectConfigAndTargetCombinations)
 			{
 				if (Combination.Platform == null)
 				{
 					continue;
 				}
+				if (!IsValidProjectPlatformAndConfiguration(Combination.ProjectTarget, (UnrealTargetPlatform)Combination.Platform, Combination.Configuration, PlatformProjectGenerators))
+				{
+					continue;
+				}
 				if (!ProjectPlatforms.Contains(Combination.Platform.Value))
 				{
 					ProjectPlatforms.Add(Combination.Platform.Value);
@@ -914,7 +935,7 @@ namespace UnrealBuildTool
 			}
 
 			// Add the "invalid" configuration for each platform. We use this when the solution configuration does not match any project configuration.
-			foreach(string InvalidConfigPlatformName in InvalidConfigPlatformNames)
+			foreach (string InvalidConfigPlatformName in InvalidConfigPlatformNames)
 			{
 				VCProjectFileContent.AppendLine("    <ProjectConfiguration Include=\"Invalid|{0}\">", InvalidConfigPlatformName);
 				VCProjectFileContent.AppendLine("      <Configuration>Invalid</Configuration>");
@@ -962,6 +983,7 @@ namespace UnrealBuildTool
 				VCProjectFileContent.AppendLine("    <RootNamespace>{0}</RootNamespace>", ProjectName);
 				VCProjectFileGenerator.AppendPlatformToolsetProperty(VCProjectFileContent, ProjectFileFormat);
 				VCProjectFileContent.AppendLine("    <MinimumVisualStudioVersion>{0}</MinimumVisualStudioVersion>", VCProjectFileGenerator.GetProjectFileToolVersionString(ProjectFileFormat));
+				VCProjectFileContent.AppendLine("    <VCProjectVersion>{0}</VCProjectVersion>", VCProjectFileGenerator.GetProjectFileToolVersionString(ProjectFileFormat));
 				VCProjectFileContent.AppendLine("    <NMakeUseOemCodePage>true</NMakeUseOemCodePage>"); // Fixes mojibake with non-Latin character sets (UE-102825)
 				VCProjectFileContent.AppendLine("    <TargetRuntime>Native</TargetRuntime>");
 				VCProjectFileContent.AppendLine("  </PropertyGroup>");
@@ -1006,7 +1028,7 @@ namespace UnrealBuildTool
 			VCProjectFileContent.AppendLine("  <Import Project=\"$(VCTargetsPath)\\Microsoft.Cpp.Default.props\" />");
 
 			// Write the invalid configuration data
-			foreach(string InvalidConfigPlatformName in InvalidConfigPlatformNames)
+			foreach (string InvalidConfigPlatformName in InvalidConfigPlatformNames)
 			{
 				VCProjectFileContent.AppendLine("  <PropertyGroup Condition=\"'$(Configuration)|$(Platform)'=='Invalid|{0}'\" Label=\"Configuration\">", InvalidConfigPlatformName);
 				VCProjectFileContent.AppendLine("    <ConfigurationType>Makefile</ConfigurationType>");
@@ -1031,7 +1053,7 @@ namespace UnrealBuildTool
 			VCProjectFileContent.AppendLine("  <PropertyGroup Label=\"UserMacros\" />");
 
 			// Write the invalid configuration
-			foreach(string InvalidConfigPlatformName in InvalidConfigPlatformNames)
+			foreach (string InvalidConfigPlatformName in InvalidConfigPlatformNames)
 			{
 				const string InvalidMessage = "echo The selected platform/configuration is not valid for this target.";
 
@@ -1164,7 +1186,7 @@ namespace UnrealBuildTool
 								ForceIncludePaths = string.Join(";", PathList.Where(P => !IncludePathIsFilteredOut(new DirectoryReference(P))));
 							}
 
-							DirectoryToForceIncludePaths[Directory] = ForceIncludePaths;							
+							DirectoryToForceIncludePaths[Directory] = ForceIncludePaths;
 						}
 
 
@@ -1379,10 +1401,10 @@ namespace UnrealBuildTool
 					{
 						FilterDirectories.Add(LeadingDirectory);
 
-						// Generate a unique GUID for this folder
+						// Generate a unique stable GUID for this folder by hashing the LeadingDirectory string
 						// NOTE: When saving generated project files, we ignore differences in GUIDs if every other part of the file
 						//       matches identically with the pre-existing file
-						string FilterGUID = Guid.NewGuid().ToString("B").ToUpperInvariant();
+						string FilterGUID = VCProjectFileGenerator.MakeMd5Guid(LeadingDirectory).ToString("B").ToUpperInvariant();
 
 						VCFiltersFileContent.AppendLine("    <Filter Include=\"{0}\">", EscapeFileName(LeadingDirectory));
 						VCFiltersFileContent.AppendLine("      <UniqueIdentifier>{0}</UniqueIdentifier>", FilterGUID);
@@ -1439,7 +1461,7 @@ namespace UnrealBuildTool
 			string ProjectConfigurationAndPlatformName = ProjectConfigurationName + "|" + ProjectPlatformName;
 			string ConditionString = "Condition=\"'$(Configuration)|$(Platform)'=='" + ProjectConfigurationAndPlatformName + "'\"";
 
-			if(ProjGenerator != null)
+			if (ProjGenerator != null)
 			{
 				StringBuilder PlatformToolsetString = new StringBuilder();
 				ProjGenerator.GetVisualStudioPreDefaultString(TargetPlatform, TargetConfiguration, PlatformToolsetString);
@@ -1501,7 +1523,7 @@ namespace UnrealBuildTool
 			{
 				VCProjectFileContent.AppendLine("  <ImportGroup {0} Label=\"PropertySheets\">", ConditionString);
 				VCProjectFileContent.AppendLine("    <Import Project=\"$(UserRootDir)\\Microsoft.Cpp.$(Platform).user.props\" Condition=\"exists('$(UserRootDir)\\Microsoft.Cpp.$(Platform).user.props')\" Label=\"LocalAppDataPlatform\" />");
-				if(ProjGenerator != null)
+				if (ProjGenerator != null)
 				{
 					ProjGenerator.GetVisualStudioImportGroupProperties(Combination.Platform.Value, VCProjectFileContent);
 				}
@@ -1522,7 +1544,7 @@ namespace UnrealBuildTool
 					VCProjectFileContent.AppendLine("    <NMakeOutput/>");
 					VCProjectFileContent.AppendLine("  </PropertyGroup>");
 				}
-				else if (UnrealBuildTool.IsEngineInstalled() && Combination.ProjectTarget != null && Combination.ProjectTarget.TargetRules != null && 
+				else if (UnrealBuildTool.IsEngineInstalled() && Combination.ProjectTarget != null && Combination.ProjectTarget.TargetRules != null &&
 					(Combination.Platform == null || !Combination.ProjectTarget.SupportedPlatforms.Contains(Combination.Platform.Value)))
 				{
 					string ProjectRelativeUnusedDirectory = NormalizeProjectPath(DirectoryReference.Combine(UnrealBuildTool.EngineDirectory, "Intermediate", "Build", "Unused"));
@@ -1553,7 +1575,7 @@ namespace UnrealBuildTool
 
 					// Figure out if this is a monolithic build
 					bool bShouldCompileMonolithic = BuildPlatform.ShouldCompileMonolithicBinary(Platform);
-					if(!bShouldCompileMonolithic)
+					if (!bShouldCompileMonolithic)
 					{
 						bShouldCompileMonolithic = (Combination.ProjectTarget.CreateRulesDelegate(Platform, Configuration).LinkType == TargetLinkType.Monolithic);
 					}
@@ -1599,7 +1621,7 @@ namespace UnrealBuildTool
 					VCProjectFileContent.AppendLine("  <PropertyGroup {0}>", ConditionString);
 
 					StringBuilder PathsStringBuilder = new StringBuilder();
-					if(ProjGenerator != null)
+					if (ProjGenerator != null)
 					{
 						ProjGenerator.GetVisualStudioPathsEntries(Platform, Configuration, TargetRulesObject.Type, TargetFilePath, ProjectFilePath, NMakePath, ProjectFileFormat, PathsStringBuilder);
 					}
@@ -1639,16 +1661,16 @@ namespace UnrealBuildTool
 						}
 						if (TargetRulesObject.bWithLiveCoding && Settings.bBuildLiveCodingConsole && !UnrealBuildTool.IsEngineInstalled() && TargetRulesObject.Name != "LiveCodingConsole")
 						{
-							ExtraTargets.Add(TargetRulesObject.bUseDebugLiveCodingConsole? "LiveCodingConsole Win64 Debug" : "LiveCodingConsole Win64 Development");
+							ExtraTargets.Add(TargetRulesObject.bUseDebugLiveCodingConsole ? "LiveCodingConsole Win64 Debug" : "LiveCodingConsole Win64 Development");
 						}
 					}
 
-					if(ExtraTargets.Count > 0)
+					if (ExtraTargets.Count > 0)
 					{
 						BuildArguments.Replace("\"", "\\\"");
 						BuildArguments.Insert(0, "-Target=\"");
 						BuildArguments.Append("\"");
-						foreach(string ExtraTarget in ExtraTargets)
+						foreach (string ExtraTarget in ExtraTargets)
 						{
 							BuildArguments.AppendFormat(" -Target=\"{0} -Quiet\"", ExtraTarget);
 						}
@@ -1673,7 +1695,7 @@ namespace UnrealBuildTool
 
 					DirectoryReference BatchFilesDirectory = DirectoryReference.Combine(UnrealBuildTool.EngineDirectory, "Build", "BatchFiles");
 
-					if(BuildToolOverride != null)
+					if (BuildToolOverride != null)
 					{
 						BuildArguments.AppendFormat(" {0}", BuildToolOverride);
 					}
@@ -1698,7 +1720,7 @@ namespace UnrealBuildTool
 
 					VCProjectFileContent.AppendLine("  </PropertyGroup>");
 
-					if(ProjGenerator != null)
+					if (ProjGenerator != null)
 					{
 						VCProjectFileContent.Append(ProjGenerator.GetVisualStudioLayoutDirSection(Platform, Configuration, ConditionString, Combination.ProjectTarget.TargetRules.Type, Combination.ProjectTarget.TargetFilePath, ProjectFilePath, NMakePath, ProjectFileFormat));
 					}
@@ -1731,7 +1753,7 @@ namespace UnrealBuildTool
 						VCUserFileContent.AppendLine("  </PropertyGroup>");
 					}
 
-					if(ProjGenerator != null)
+					if (ProjGenerator != null)
 					{
 						VCUserFileContent.Append(ProjGenerator.GetVisualStudioUserFileStrings(Combination.Platform.Value, Configuration, ConditionString, TargetRulesObject, Combination.ProjectTarget.TargetFilePath, ProjectFilePath));
 					}
@@ -1788,13 +1810,13 @@ namespace UnrealBuildTool
 				{
 					foreach (XmlElement Element in Document.DocumentElement.ChildNodes.OfType<XmlElement>())
 					{
-						if(Element.Name == "PropertyGroup")
+						if (Element.Name == "PropertyGroup")
 						{
 							string Condition = Element.GetAttribute("Condition");
-							if(!String.IsNullOrEmpty(Condition))
+							if (!String.IsNullOrEmpty(Condition))
 							{
 								Match Match = Regex.Match(Condition, "^\\s*'\\$\\(Configuration\\)\\|\\$\\(Platform\\)'\\s*==\\s*'(.+)\\|(.+)'\\s*$");
-								if(Match.Success && Match.Groups.Count == 3)
+								if (Match.Success && Match.Groups.Count == 3)
 								{
 									Configurations.Add(Match.Groups[1].Value);
 									Platforms.Add(Match.Groups[2].Value);
@@ -1817,7 +1839,7 @@ namespace UnrealBuildTool
 						// if this property group does not set configurations we do not care about it
 						if (ConfigNodeList.Count == 0)
 							continue;
-						
+
 						if (PropertyGroup.HasAttribute("Condition"))
 						{
 							string Condition = PropertyGroup.GetAttribute("Condition");
@@ -1847,7 +1869,7 @@ namespace UnrealBuildTool
 					}
 				}
 			}
-			catch(Exception Ex)
+			catch (Exception Ex)
 			{
 				Log.TraceWarning("Unable to parse {0}: {1}", InitFilePath, Ex.ToString());
 			}
@@ -1897,7 +1919,7 @@ namespace UnrealBuildTool
 		{
 			// Find the matching platform name
 			string ProjectPlatformName;
-			if(Platforms.Contains("x64"))
+			if (Platforms.Contains("x64"))
 			{
 				ProjectPlatformName = "x64";
 			}
@@ -1908,11 +1930,11 @@ namespace UnrealBuildTool
 
 			// Find the matching configuration
 			string ProjectConfigurationName;
-			if(Configurations.Contains(SolutionConfiguration.ToString()))
+			if (Configurations.Contains(SolutionConfiguration.ToString()))
 			{
 				ProjectConfigurationName = SolutionConfiguration.ToString();
 			}
-			else if(Configurations.Contains("Development"))
+			else if (Configurations.Contains("Development"))
 			{
 				ProjectConfigurationName = "Development";
 			}
@@ -1923,13 +1945,13 @@ namespace UnrealBuildTool
 
 			// Figure out whether to build it by default
 			bool bBuildByDefault = ShouldBuildByDefaultForSolutionTargets;
-			if(SolutionTarget == TargetType.Game || SolutionTarget == TargetType.Editor)
+			if (SolutionTarget == TargetType.Game || SolutionTarget == TargetType.Editor)
 			{
 				bBuildByDefault = true;
 			}
 
 			// Create the context
-			return new MSBuildProjectContext(ProjectConfigurationName, ProjectPlatformName){ bBuildByDefault = bBuildByDefault };
+			return new MSBuildProjectContext(ProjectConfigurationName, ProjectPlatformName) { bBuildByDefault = bBuildByDefault };
 		}
 
 		/// <summary>
diff --git a/Engine/Source/Programs/UnrealBuildTool/ProjectFiles/VisualStudio/VCProjectFileGenerator.cs b/Engine/Source/Programs/UnrealBuildTool/ProjectFiles/VisualStudio/VCProjectFileGenerator.cs
--- a/Engine/Source/Programs/UnrealBuildTool/ProjectFiles/VisualStudio/VCProjectFileGenerator.cs
+++ b/Engine/Source/Programs/UnrealBuildTool/ProjectFiles/VisualStudio/VCProjectFileGenerator.cs
@@ -124,20 +124,20 @@ namespace UnrealBuildTool
 		{
 			XmlConfig.ApplyTo(Settings);
 
-			if(InProjectFileFormat != VCProjectFileFormat.Default)
+			if (InProjectFileFormat != VCProjectFileFormat.Default)
 			{
 				Settings.ProjectFileFormat = InProjectFileFormat;
 			}
 
-			if(InArguments.HasOption("-2015"))
+			if (InArguments.HasOption("-2015"))
 			{
 				BuildToolOverride = "-2015";
 			}
-			else if(InArguments.HasOption("-2017"))
+			else if (InArguments.HasOption("-2017"))
 			{
 				BuildToolOverride = "-2017";
 			}
-			else if(InArguments.HasOption("-2019"))
+			else if (InArguments.HasOption("-2019"))
 			{
 				BuildToolOverride = "-2019";
 			}
@@ -224,9 +224,9 @@ namespace UnrealBuildTool
 		static public string GetProjectFileToolVersionString(VCProjectFileFormat ProjectFileFormat)
 		{
 			switch (ProjectFileFormat)
-            {
-                case VCProjectFileFormat.VisualStudio2012:
-                    return "4.0";
+			{
+				case VCProjectFileFormat.VisualStudio2012:
+					return "4.0";
 				case VCProjectFileFormat.VisualStudio2013:
 					return "12.0";
 				case VCProjectFileFormat.VisualStudio2015:
@@ -234,7 +234,7 @@ namespace UnrealBuildTool
 				case VCProjectFileFormat.VisualStudio2017:
 					return "15.0";
 				case VCProjectFileFormat.VisualStudio2019:
-					return "15.0"; // Correct as of VS2019 Preview 1
+					return "16.0";
 			}
 			return string.Empty;
 		}
@@ -242,31 +242,28 @@ namespace UnrealBuildTool
 		/// for instance: <PlatformToolset>v110</PlatformToolset>
 		static public string GetProjectFilePlatformToolsetVersionString(VCProjectFileFormat ProjectFileFormat)
 		{
-            switch (ProjectFileFormat)
-            {
-                case VCProjectFileFormat.VisualStudio2012:
-                    return "v110";
-                case VCProjectFileFormat.VisualStudio2013:
-                    return "v120";
-                case VCProjectFileFormat.VisualStudio2015:
+			switch (ProjectFileFormat)
+			{
+				case VCProjectFileFormat.VisualStudio2012:
+					return "v110";
+				case VCProjectFileFormat.VisualStudio2013:
+					return "v120";
+				case VCProjectFileFormat.VisualStudio2015:
 					return "v140";
 				case VCProjectFileFormat.VisualStudio2017:
-                    return "v141";
+					return "v141";
 				case VCProjectFileFormat.VisualStudio2019:
-					return "v142"; // Correct as of VS2019 Preview 2
+					return "v142";
 
-            }
+			}
 			return string.Empty;
 		}
 
 		static public void AppendPlatformToolsetProperty(StringBuilder VCProjectFileContent, VCProjectFileFormat ProjectFileFormat)
 		{
-			VCProjectFileContent.AppendLine("    <PlatformToolset>{0}</PlatformToolset>", GetProjectFilePlatformToolsetVersionString(ProjectFileFormat));
-
-			if(ProjectFileFormat == VCProjectFileFormat.VisualStudio2019)
-			{
-				VCProjectFileContent.AppendLine("    <PlatformToolset Condition=\"'$(VisualStudioVersion)' == '15.0'\">v141</PlatformToolset>");
-			}
+			string ToolVersionString = GetProjectFileToolVersionString(ProjectFileFormat);
+			string PlatformToolsetVersionString = GetProjectFilePlatformToolsetVersionString(ProjectFileFormat);
+			VCProjectFileContent.AppendLine("    <PlatformToolset>{0}</PlatformToolset>", PlatformToolsetVersionString);
 		}
 
 		/// <summary>
@@ -291,7 +288,7 @@ namespace UnrealBuildTool
 			base.SetupSupportedPlatformsAndConfigurations(IncludeAllPlatforms, out SupportedPlatformNames);
 
 			// If we have a non-default setting for visual studio, check the compiler exists. If not, revert to the default.
-			if(Settings.ProjectFileFormat == VCProjectFileFormat.VisualStudio2015)
+			if (Settings.ProjectFileFormat == VCProjectFileFormat.VisualStudio2015)
 			{
 				if (!WindowsPlatform.HasCompiler(WindowsCompiler.VisualStudio2015_DEPRECATED))
 				{
@@ -299,7 +296,7 @@ namespace UnrealBuildTool
 					Settings.ProjectFileFormat = VCProjectFileFormat.Default;
 				}
 			}
-			else if(Settings.ProjectFileFormat == VCProjectFileFormat.VisualStudio2017)
+			else if (Settings.ProjectFileFormat == VCProjectFileFormat.VisualStudio2017)
 			{
 				if (!WindowsPlatform.HasCompiler(WindowsCompiler.VisualStudio2017))
 				{
@@ -307,7 +304,7 @@ namespace UnrealBuildTool
 					Settings.ProjectFileFormat = VCProjectFileFormat.Default;
 				}
 			}
-			else if(Settings.ProjectFileFormat == VCProjectFileFormat.VisualStudio2019)
+			else if (Settings.ProjectFileFormat == VCProjectFileFormat.VisualStudio2019)
 			{
 				if (!WindowsPlatform.HasCompiler(WindowsCompiler.VisualStudio2019))
 				{
@@ -458,13 +455,20 @@ namespace UnrealBuildTool
 			return MakeMd5Guid(Input);
 		}
 
+		public static Guid MakeMd5Guid(string Text)
+		{
+			byte[] Input = Encoding.UTF8.GetBytes(Text);
+
+			return MakeMd5Guid(Input);
+		}
+
 		/// <summary>
 		/// Writes the project files to disk
 		/// </summary>
 		/// <returns>True if successful</returns>
 		protected override bool WriteProjectFiles(PlatformProjectGeneratorCollection PlatformProjectGenerators)
 		{
-			if(!base.WriteProjectFiles(PlatformProjectGenerators))
+			if (!base.WriteProjectFiles(PlatformProjectGenerators))
 			{
 				return false;
 			}
@@ -489,7 +493,7 @@ namespace UnrealBuildTool
 					)
 				).Save(FileReference.Combine(AutomationToolDir, "AutomationTool.csproj.References").FullName);
 			}
-	
+
 			return true;
 		}
 
@@ -533,13 +537,13 @@ namespace UnrealBuildTool
 				VCSolutionFileContent.AppendLine();
 				VCSolutionFileContent.AppendLine("Microsoft Visual Studio Solution File, Format Version 12.00");
 				VCSolutionFileContent.AppendLine("# Visual Studio 2013");
-            }
-            else if (Settings.ProjectFileFormat == VCProjectFileFormat.VisualStudio2012)
-            {
+			}
+			else if (Settings.ProjectFileFormat == VCProjectFileFormat.VisualStudio2012)
+			{
 				VCSolutionFileContent.AppendLine();
-                VCSolutionFileContent.AppendLine("Microsoft Visual Studio Solution File, Format Version 12.00");
+				VCSolutionFileContent.AppendLine("Microsoft Visual Studio Solution File, Format Version 12.00");
 				VCSolutionFileContent.AppendLine("# Visual Studio 2012");
-            }
+			}
 			else
 			{
 				throw new BuildException("Unexpected ProjectFileFormat");
@@ -769,7 +773,7 @@ namespace UnrealBuildTool
 								if (ProjectContext.bBuildByDefault)
 								{
 									VCSolutionFileContent.AppendLine("		{0}.{1}.Build.0 = {2}", CurProjectGUID, SolutionConfigCombination.VCSolutionConfigAndPlatformName, ProjectContext.Name);
-									if(ProjectContext.bDeployByDefault)
+									if (ProjectContext.bDeployByDefault)
 									{
 										VCSolutionFileContent.AppendLine("		{0}.{1}.Deploy.0 = {2}", CurProjectGUID, SolutionConfigCombination.VCSolutionConfigAndPlatformName, ProjectContext.Name);
 									}
@@ -847,10 +851,10 @@ namespace UnrealBuildTool
 				// Figure out the filename for the SUO file. VS will automatically import the options from earlier versions if necessary.
 				FileReference SolutionOptionsFileName;
 				switch (Settings.ProjectFileFormat)
-                {
-                    case VCProjectFileFormat.VisualStudio2012:
+				{
+					case VCProjectFileFormat.VisualStudio2012:
 						SolutionOptionsFileName = FileReference.Combine(MasterProjectPath, Path.ChangeExtension(SolutionFileName, "v11.suo"));
-                        break;
+						break;
 					case VCProjectFileFormat.VisualStudio2013:
 						SolutionOptionsFileName = FileReference.Combine(MasterProjectPath, Path.ChangeExtension(SolutionFileName, "v12.suo"));
 						break;
@@ -861,7 +865,7 @@ namespace UnrealBuildTool
 						SolutionOptionsFileName = FileReference.Combine(MasterProjectPath, ".vs", Path.GetFileNameWithoutExtension(SolutionFileName), "v15", ".suo");
 						break;
 					case VCProjectFileFormat.VisualStudio2019:
-						SolutionOptionsFileName = FileReference.Combine(MasterProjectPath, ".vs", Path.GetFileNameWithoutExtension(SolutionFileName), "v15", ".suo"); // Still uses v15
+						SolutionOptionsFileName = FileReference.Combine(MasterProjectPath, ".vs", Path.GetFileNameWithoutExtension(SolutionFileName), "v16", ".suo");
 						break;
 					default:
 						throw new BuildException("Unsupported Visual Studio version");
@@ -889,7 +893,7 @@ namespace UnrealBuildTool
 
 					// Mark all the projects as closed by default, apart from the startup project
 					VCSolutionExplorerState ExplorerState = new VCSolutionExplorerState();
-					if(Settings.ProjectFileFormat >= VCProjectFileFormat.VisualStudio2017)
+					if (Settings.ProjectFileFormat >= VCProjectFileFormat.VisualStudio2017)
 					{
 						BuildSolutionExplorerState_VS2017(RootFolder, "", ExplorerState, DefaultProject);
 					}
@@ -910,7 +914,7 @@ namespace UnrealBuildTool
 			return bSuccess;
 		}
 
-		protected override void WriteDebugSolutionFiles( PlatformProjectGeneratorCollection PlatformProjectGenerators, DirectoryReference IntermediateProjectFilesPath )
+		protected override void WriteDebugSolutionFiles(PlatformProjectGeneratorCollection PlatformProjectGenerators, DirectoryReference IntermediateProjectFilesPath)
 		{
 			//build and collect UnrealVS configuration
 			StringBuilder UnrealVSContent = new StringBuilder();
@@ -922,19 +926,20 @@ namespace UnrealBuildTool
 					ProjGenerator.GetUnrealVSConfigurationEntries(UnrealVSContent);
 				}
 			}
-			if (UnrealVSContent.Length > 0 )
+			if (UnrealVSContent.Length > 0)
 			{
 				UnrealVSContent.Insert(0, "<UnrealVS>" + ProjectFileGenerator.NewLine);
-				UnrealVSContent.Append("</UnrealVS>" + ProjectFileGenerator.NewLine );
+				UnrealVSContent.Append("</UnrealVS>" + ProjectFileGenerator.NewLine);
 
 				string ConfigFilePath = FileReference.Combine(IntermediateProjectFilesPath, "UnrealVS.xml").FullName;
-				/* bool bSuccess = */ ProjectFileGenerator.WriteFileIfChanged(ConfigFilePath, UnrealVSContent.ToString());
+				/* bool bSuccess = */
+				ProjectFileGenerator.WriteFileIfChanged(ConfigFilePath, UnrealVSContent.ToString());
 			}
 		}
 
 		static void BuildSolutionExplorerState_VS2017(MasterProjectFolder Folder, string Suffix, VCSolutionExplorerState ExplorerState, ProjectFile DefaultProject)
 		{
-			foreach(ProjectFile Project in Folder.ChildProjects)
+			foreach (ProjectFile Project in Folder.ChildProjects)
 			{
 				string ProjectIdentifier = String.Format("{0}{1}", Project.ProjectFilePath.GetFileNameWithoutExtension(), Suffix);
 				if (Project == DefaultProject)
@@ -947,10 +952,10 @@ namespace UnrealBuildTool
 				}
 			}
 
-			foreach(MasterProjectFolder SubFolder in Folder.SubFolders)
+			foreach (MasterProjectFolder SubFolder in Folder.SubFolders)
 			{
 				string SubFolderName = SubFolder.FolderName + Suffix;
-				if(SubFolderName == "Automation;Programs")
+				if (SubFolderName == "Automation;Programs")
 				{
 					ExplorerState.OpenProjects.Add(new Tuple<string, string[]>(SubFolderName, new string[] { }));
 				}
diff --git a/Engine/Source/Programs/UnrealBuildTool/System/SourceFileSearch.cs b/Engine/Source/Programs/UnrealBuildTool/System/SourceFileSearch.cs
--- a/Engine/Source/Programs/UnrealBuildTool/System/SourceFileSearch.cs
+++ b/Engine/Source/Programs/UnrealBuildTool/System/SourceFileSearch.cs
@@ -34,7 +34,9 @@ namespace UnrealBuildTool
 		// Default directory names to exclude. Must be lowercase.
 		static readonly string[] DefaultExcludedDirectorySuffixes = new string[] 
 		{
-			Path.DirectorySeparatorChar + "intermediate"
+			Path.DirectorySeparatorChar + "intermediate",
+			Path.DirectorySeparatorChar + "source" + Path.DirectorySeparatorChar + "thirdparty",
+			Path.DirectorySeparatorChar + "third_party",
 		};
 
 		/// Finds mouse source files
