Fixed -skipdeploy UBT command. This was resulting in BuildCookRun making and APK twice

diff --git a/Engine/Source/Programs/UnrealBuildTool/Configuration/TargetRules.cs b/Engine/Source/Programs/UnrealBuildTool/Configuration/TargetRules.cs
--- a/Engine/Source/Programs/UnrealBuildTool/Configuration/TargetRules.cs
+++ b/Engine/Source/Programs/UnrealBuildTool/Configuration/TargetRules.cs
@@ -1228,9 +1228,14 @@ namespace UnrealBuildTool
 		/// Whether to deploy the executable after compilation on platforms that require deployment.
 		/// </summary>
 		[CommandLine("-Deploy")]
-		[CommandLine("-SkipDeploy", Value = "false")]
 		public bool bDeployAfterCompile = false;
 
+		/// <summary>
+		/// Whether to force skipping deployment for platforms that require deployment by default.
+		/// </summary>
+		[CommandLine("-SkipDeploy")]
+		private bool bForceSkipDeploy = false; 
+
 		/// <summary>
 		/// When enabled, allows XGE to compile pre-compiled header files on remote machines.  Otherwise, PCHs are always generated locally.
 		/// </summary>
@@ -1616,6 +1621,7 @@ namespace UnrealBuildTool
 
 			// Allow the build platform to set defaults for this target
 			UEBuildPlatform.GetBuildPlatform(Platform).ResetTarget(this);
+			bDeployAfterCompile = bForceSkipDeploy ? false : bDeployAfterCompile;
 
 			// Set the default build version
 			if(String.IsNullOrEmpty(BuildVersion))
