UnrealBuildTool Utils.cs, Mac support:

- Account for caching of disk-backed files when computing free memory
- Determine actual number of physical CPU cores

diff --git a/Engine/Source/Programs/UnrealBuildTool/System/Utils.cs b/Engine/Source/Programs/UnrealBuildTool/System/Utils.cs
--- a/Engine/Source/Programs/UnrealBuildTool/System/Utils.cs
+++ b/Engine/Source/Programs/UnrealBuildTool/System/Utils.cs
@@ -1085,48 +1085,63 @@ namespace UnrealBuildTool
 			return Environment.ProcessorCount;
 		}
 
+		
+		// int sysctlbyname(const char *name, void *oldp, size_t *oldlenp, void *newp, size_t newlen); // from man page
+		[DllImport("libc")]
+		extern static int sysctlbyname(string name, out int oldp, ref UInt64 oldlenp, IntPtr newp, UInt64 newlen);
+		
 		/// <summary>
 		/// Gets the number of physical cores, excluding hyper threading.
 		/// </summary>
 		/// <returns>The number of physical cores, or -1 if it could not be obtained</returns>
 		public static int GetPhysicalProcessorCount()
 		{
-			// This function uses Windows P/Invoke calls; if we're not running on Windows, just fail.
-			if (!Utils.IsRunningOnWindows)
+			if (Utils.IsRunningOnWindows)
 			{
-				return -1;
-			}
-
-			const int ERROR_INSUFFICIENT_BUFFER = 122;
+				const int ERROR_INSUFFICIENT_BUFFER = 122;
 
-			// Determine the required buffer size to store the processor information
-			uint ReturnLength = 0;
-			if(!GetLogicalProcessorInformationEx(LOGICAL_PROCESSOR_RELATIONSHIP.RelationProcessorCore, IntPtr.Zero, ref ReturnLength) && Marshal.GetLastWin32Error() == ERROR_INSUFFICIENT_BUFFER)
-			{
-				// Allocate a buffer for it
-				IntPtr Ptr = Marshal.AllocHGlobal((int)ReturnLength);
-				try
+				// Determine the required buffer size to store the processor information
+				uint ReturnLength = 0;
+				if (!GetLogicalProcessorInformationEx(LOGICAL_PROCESSOR_RELATIONSHIP.RelationProcessorCore, IntPtr.Zero,
+					ref ReturnLength) && Marshal.GetLastWin32Error() == ERROR_INSUFFICIENT_BUFFER)
 				{
-					if (GetLogicalProcessorInformationEx(LOGICAL_PROCESSOR_RELATIONSHIP.RelationProcessorCore, Ptr, ref ReturnLength))
+					// Allocate a buffer for it
+					IntPtr Ptr = Marshal.AllocHGlobal((int) ReturnLength);
+					try
 					{
-						// As per-MSDN, this will return one structure per physical processor. Each SYSTEM_LOGICAL_PROCESSOR_INFORMATION_EX structure is of a variable size, so just skip 
-						// through the list and count the number of entries.
-						int Count = 0;
-						for(int Pos = 0; Pos < ReturnLength; )
+						if (GetLogicalProcessorInformationEx(LOGICAL_PROCESSOR_RELATIONSHIP.RelationProcessorCore, Ptr,
+							ref ReturnLength))
 						{
-							LOGICAL_PROCESSOR_RELATIONSHIP Type = (LOGICAL_PROCESSOR_RELATIONSHIP)Marshal.ReadInt16(Ptr, Pos);
-							if(Type == LOGICAL_PROCESSOR_RELATIONSHIP.RelationProcessorCore)
+							// As per-MSDN, this will return one structure per physical processor. Each SYSTEM_LOGICAL_PROCESSOR_INFORMATION_EX structure is of a variable size, so just skip 
+							// through the list and count the number of entries.
+							int Count = 0;
+							for (int Pos = 0; Pos < ReturnLength;)
 							{
-								Count++;
+								LOGICAL_PROCESSOR_RELATIONSHIP Type =
+									(LOGICAL_PROCESSOR_RELATIONSHIP) Marshal.ReadInt16(Ptr, Pos);
+								if (Type == LOGICAL_PROCESSOR_RELATIONSHIP.RelationProcessorCore)
+								{
+									Count++;
+								}
+
+								Pos += Marshal.ReadInt32(Ptr, Pos + 4);
 							}
-							Pos += Marshal.ReadInt32(Ptr, Pos + 4);
+
+							return Count;
 						}
-						return Count;
+					}
+					finally
+					{
+						Marshal.FreeHGlobal(Ptr);
 					}
 				}
-				finally
+			}
+			else if(Utils.IsRunningOnMac)
+			{
+				UInt64 Size = 4;
+				if (0 == sysctlbyname("hw.physicalcpu", out int Value, ref Size, IntPtr.Zero, 0))
 				{
-					Marshal.FreeHGlobal(Ptr);		
+					return Value;
 				}
 			}
 
@@ -1142,9 +1157,55 @@ namespace UnrealBuildTool
 		{
 			GCMemoryInfo MemoryInfo = GC.GetGCMemoryInfo();
 			// TotalAvailableMemoryBytes will be 0 if garbage collection has not run yet
-			return MemoryInfo.TotalAvailableMemoryBytes != 0 ? MemoryInfo.TotalAvailableMemoryBytes :- 1;
+			return MemoryInfo.TotalAvailableMemoryBytes != 0 ? MemoryInfo.TotalAvailableMemoryBytes : -1;
 		}
 
+		
+		// vm_statistics64, based on the definition in <mach/vm_statistics.h>
+		[StructLayout(LayoutKind.Sequential)]
+        struct vm_statistics64 {
+          	/*natural_t*/ public int	free_count;				/* # of pages free */
+          	/*natural_t*/ public int	active_count;			/* # of pages active */
+          	/*natural_t*/ public int	inactive_count;			/* # of pages inactive */
+          	/*natural_t*/ public int	wire_count;				/* # of pages wired down */
+          	/*uint64_t */ public UInt64	zero_fill_count;		/* # of zero fill pages */
+          	/*uint64_t */ public UInt64	reactivations;			/* # of pages reactivated */
+		  	/*uint64_t */ public UInt64	pageins;				/* # of pageins */
+          	/*uint64_t */ public UInt64	pageouts;				/* # of pageouts */
+          	/*uint64_t */ public UInt64	faults;					/* # of faults */
+          	/*uint64_t */ public UInt64	cow_faults;				/* # of copy-on-writes */
+          	/*uint64_t */ public UInt64	lookups;				/* object cache lookups */
+          	/*uint64_t */ public UInt64	hits;					/* object cache hits */
+          	/*uint64_t */ public UInt64	purges;					/* # of pages purged */
+          	/*natural_t*/ public int	purgeable_count;		/* # of pages purgeable */
+          	/*
+          	 * NB: speculative pages are already accounted for in "free_count",
+          	 * so "speculative_count" is the number of "free" pages that are
+          	 * used to hold data that was read speculatively from disk but
+          	 * haven't actually been used by anyone so far.
+          	 */
+          	/*natural_t*/ public int	speculative_count;		/* # of pages speculative */
+          
+          	/* added for rev1 */
+          	/*uint64_t */ public UInt64	decompressions;			/* # of pages decompressed */
+          	/*uint64_t */ public UInt64	compressions;			/* # of pages compressed */
+          	/*uint64_t */ public UInt64	swapins;				/* # of pages swapped in (via compression segments) */
+          	/*uint64_t */ public UInt64	swapouts;				/* # of pages swapped out (via compression segments) */
+          	/*natural_t*/ public int	compressor_page_count;	/* # of pages used by the compressed pager to hold all the compressed data */
+          	/*natural_t*/ public int	throttled_count;		/* # of pages throttled */
+          	/*natural_t*/ public int	external_page_count;	/* # of pages that are file-backed (non-swap) */
+          	/*natural_t*/ public int	internal_page_count;	/* # of pages that are anonymous */
+          	/*uint64_t */ public UInt64	total_uncompressed_pages_in_compressor; /* # of pages (uncompressed) held within the compressor. */
+        } // __attribute__((aligned(8))); 
+		
+		// kern_return_t host_statistics64(host_t host_priv, host_flavor_t flavor, host_info64_t host_info64_out, mach_msg_type_number_t *host_info64_outCnt); // from <mach/mach_host.h>
+		[DllImport("libc")]
+		extern static int host_statistics64(IntPtr host_priv, int flavor, out vm_statistics64 host_info64_out, ref uint host_info_count);
+
+		// mach_port_t mach_host_self() // from <mach/mach_init.h>
+		[DllImport("libc")]
+		extern static IntPtr mach_host_self();
+		
 		/// <summary>
 		/// Gets the total memory bytes free, based on what is known to the garbage collector.
 		/// </summary>
@@ -1152,9 +1213,40 @@ namespace UnrealBuildTool
 		/// <returns>The total memory free, in bytes.</returns>
 		public static long GetFreeMemoryBytes()
 		{
+			long FreeMemoryBytes = -1;
+			
 			GCMemoryInfo MemoryInfo = GC.GetGCMemoryInfo();
 			// TotalAvailableMemoryBytes will be 0 if garbage collection has not run yet
-			return MemoryInfo.TotalAvailableMemoryBytes != 0 ? MemoryInfo.TotalAvailableMemoryBytes - MemoryInfo.MemoryLoadBytes : -1;
+			if (MemoryInfo.TotalAvailableMemoryBytes != 0)
+			{
+				FreeMemoryBytes = MemoryInfo.TotalAvailableMemoryBytes - MemoryInfo.MemoryLoadBytes;
+			}
+
+			// On Mac, MemoryInfo.MemoryLoadBytes includes memory used to cache disk-backed files ("Cached Files" in
+			// Activity Monitor), which can result in a significant over-estimate of memory pressure.
+			// We treat memory used for caching of disk-backed files as free for use in compilation tasks.
+			if (Utils.IsRunningOnMac)	
+			{
+				// host_statistics64() flavor, from <mach/host_info.h>
+				int HOST_VM_INFO64 = 4;
+				// host_statistics64() count of 32bit values in output struct, from <mach/host_info.h>
+				int HOST_VM_INFO64_COUNT = Marshal.SizeOf(typeof(vm_statistics64)) / 4;
+		
+				vm_statistics64 VMStats;
+				uint StructSize = (uint)HOST_VM_INFO64_COUNT;
+				IntPtr Host = mach_host_self();
+				host_statistics64(Host, HOST_VM_INFO64, out VMStats, ref StructSize);
+
+				int PageSize = 0;
+				UInt64 OutSize = 4;
+				if (0 != sysctlbyname("hw.pagesize", out PageSize, ref OutSize, IntPtr.Zero, 0))
+				{
+					PageSize = 4096; // likely result
+				}
+				
+				FreeMemoryBytes += (long)PageSize * (long)VMStats.external_page_count;
+			}
+			return FreeMemoryBytes;
 		}
 
 		/// <summary>
