Add launch options to VSCodeProjectFileGenerator that allow attaching to a running process or debugging a core file

diff --git a/Engine/Source/Programs/UnrealBuildTool/ProjectFiles/VisualStudioCode/VSCodeProjectFileGenerator.cs b/Engine/Source/Programs/UnrealBuildTool/ProjectFiles/VisualStudioCode/VSCodeProjectFileGenerator.cs
--- a/Engine/Source/Programs/UnrealBuildTool/ProjectFiles/VisualStudioCode/VSCodeProjectFileGenerator.cs
+++ b/Engine/Source/Programs/UnrealBuildTool/ProjectFiles/VisualStudioCode/VSCodeProjectFileGenerator.cs
@@ -50,6 +50,18 @@ namespace UnrealBuildTool
 		[XmlConfigFile(Name = "IncludeAllFiles")]
 		private bool IncludeAllFiles = false;
 
+		/// <summary>
+		/// Whether VS Code project generation should include debug configurations to allow attaching to already running processes
+		/// </summary>
+		[XmlConfigFile(Name = "AddDebugAttachConfig")]
+		private bool bAddDebugAttachConfig = false;
+
+		/// <summary>
+		/// Whether VS Code project generation should include debug configurations to allow core dump debugging
+		/// </summary>
+		[XmlConfigFile(Name = "AddDebugCoreConfig")]
+		private bool bAddDebugCoreConfig = false;
+
 		private enum EPathType
 		{
 			Absolute,
@@ -1187,70 +1199,106 @@ namespace UnrealBuildTool
 			}
 		}
 
-		private void WriteNativeLaunchConfig(ProjectData.Project InProject, JsonFile OutFile)
+		private void WriteNativeLaunchConfig(ProjectData.Project InProject, JsonFile OutFile, ProjectData.Target Target, ProjectData.BuildProduct BuildProduct)
 		{
-			foreach (ProjectData.Target Target in InProject.Targets)
+			bool bIsLinux = BuildProduct.Platform == UnrealTargetPlatform.Linux;
+			List<string> Types = new List<string>();
+			Types.Add("Launch");
+
+			if(bAddDebugAttachConfig && bIsLinux)
 			{
-				foreach (ProjectData.BuildProduct BuildProduct in Target.BuildProducts)
+				Types.Add("Attach");
+			}
+
+			if (bAddDebugCoreConfig && bIsLinux)
+			{
+				Types.Add("Debug Core");
+			}
+
+			string LaunchTaskName = String.Format("{0} {1} {2} Build", Target.Name, BuildProduct.Platform, BuildProduct.Config);
+
+			foreach(string Type in Types)
+			{
+				OutFile.BeginObject();
 				{
-					if (BuildProduct.Platform == HostPlatform)
+					OutFile.AddField("name", Type + " " + Target.Name + " (" + BuildProduct.Config.ToString() + ")");
+					OutFile.AddField("request", (Type == "Attach")?"attach":"launch");
+					OutFile.AddField("program", MakeUnquotedPathString(BuildProduct.OutputFile, EPathType.Absolute));
+					switch (Type)
 					{
-						string LaunchTaskName = String.Format("{0} {1} {2} Build", Target.Name, BuildProduct.Platform, BuildProduct.Config);
+						case "Launch":
+							OutFile.AddField("preLaunchTask", LaunchTaskName);
+							break;
+						case "Debug Core":
+							OutFile.AddField("coreDumpPath", "${input:coreFileName}");
+							break;
+						case "Attach": 
+							OutFile.AddField("processId", "${command:pickProcess}");
+							break;
+					}
 
-						OutFile.BeginObject();
+					if(Type != "Attach")
+					{
+						OutFile.BeginArray("args");
 						{
-							OutFile.AddField("name", Target.Name + " (" + BuildProduct.Config.ToString() + ")");
-							OutFile.AddField("request", "launch");
-							OutFile.AddField("preLaunchTask", LaunchTaskName);
-							OutFile.AddField("program", MakeUnquotedPathString(BuildProduct.OutputFile, EPathType.Absolute));								
-							
-							OutFile.BeginArray("args");
+							if (Target.Type == TargetRules.TargetType.Editor)
 							{
-								if (Target.Type == TargetRules.TargetType.Editor)
+								if (InProject.Name != "UE5")
 								{
-									if (InProject.Name != "UE5")
+									if (bForeignProject)
 									{
-										if (bForeignProject)
-										{
-											OutFile.AddUnnamedField(MakePathString(BuildProduct.UProjectFile, false, true));
-										}
-										else
-										{
-											OutFile.AddUnnamedField(InProject.Name);
-										}
+										OutFile.AddUnnamedField(MakePathString(BuildProduct.UProjectFile, false, true));
+									}
+									else
+									{
+										OutFile.AddUnnamedField(InProject.Name);
 									}
 								}
-
 							}
-							OutFile.EndArray();
 
-/*
-							DirectoryReference CWD = BuildProduct.OutputFile.Directory;
-							while (HostPlatform == UnrealTargetPlatform.Mac && CWD != null && CWD.ToString().Contains(".app"))
-							{
-								CWD = CWD.ParentDirectory;
-							}
-							if (CWD != null)
-							{
-								OutFile.AddField("cwd", MakePathString(CWD, true, true));
-							}
- */
-							OutFile.AddField("cwd", MakeUnquotedPathString(UE4ProjectRoot, EPathType.Absolute));
+						}
+						OutFile.EndArray();
+					}
 
-							if (HostPlatform == UnrealTargetPlatform.Win64)
-							{
-								OutFile.AddField("stopAtEntry", false);
-								OutFile.AddField("externalConsole", true);
+					/*
+									DirectoryReference CWD = BuildProduct.OutputFile.Directory;
+									while (HostPlatform == UnrealTargetPlatform.Mac && CWD != null && CWD.ToString().Contains(".app"))
+									{
+										CWD = CWD.ParentDirectory;
+									}
+									if (CWD != null)
+									{
+										OutFile.AddField("cwd", MakePathString(CWD, true, true));
+									}
+					*/
+					OutFile.AddField("cwd", MakeUnquotedPathString(UE4ProjectRoot, EPathType.Absolute));
 
-								OutFile.AddField("type", "cppvsdbg");
-								OutFile.AddField("visualizerFile", MakeUnquotedPathString(FileReference.Combine(UE4ProjectRoot, "Engine", "Extras", "VisualStudioDebugging", "Unreal.natvis"), EPathType.Absolute));
-							}
-							else
-							{
-								OutFile.AddField("type", "lldb");
-							}
-						}
-						OutFile.EndObject();
+					if (HostPlatform == UnrealTargetPlatform.Win64)
+					{
+						OutFile.AddField("stopAtEntry", false);
+						OutFile.AddField("externalConsole", true);
+
+						OutFile.AddField("type", "cppvsdbg");
+						OutFile.AddField("visualizerFile", MakeUnquotedPathString(FileReference.Combine(UE4ProjectRoot, "Engine", "Extras", "VisualStudioDebugging", "Unreal.natvis"), EPathType.Absolute));
+					}
+					else
+					{
+						OutFile.AddField("type", "cppdbg");
+					}
+				}
+				OutFile.EndObject();
+			}
+		}
+
+		private void WriteNativeLaunchConfig(ProjectData.Project InProject, JsonFile OutFile)
+		{
+			foreach (ProjectData.Target Target in InProject.Targets)
+			{
+				foreach (ProjectData.BuildProduct BuildProduct in Target.BuildProducts)
+				{
+					if (BuildProduct.Platform == HostPlatform)
+					{
+						WriteNativeLaunchConfig(InProject, OutFile, Target, BuildProduct);
 					}
 					else if (BuildProduct.Platform == UnrealTargetPlatform.Android)
 					{
@@ -1330,6 +1378,25 @@ namespace UnrealBuildTool
 			OutFile.BeginRootObject();
 			{
 				OutFile.AddField("version", "0.2.0");
+				if(bAddDebugCoreConfig)
+				{
+					OutFile.BeginArray("inputs");
+					OutFile.BeginObject();
+					OutFile.AddField("id", "coreFileName");
+					OutFile.AddField("type", "command");
+					OutFile.AddField("command", "filePicker.pick");
+					OutFile.BeginObject("args");
+					OutFile.AddField("masks", "core*");
+					OutFile.BeginObject("display");
+					OutFile.AddField("type", "fileRelativePath");
+					OutFile.AddField("detail", "filePath");
+					OutFile.EndObject();
+					OutFile.AddField("output", "filePath");
+					OutFile.EndObject();
+					OutFile.EndObject();
+					OutFile.EndArray();
+				}
+
 				OutFile.BeginArray("configurations");
 				{
 					foreach (ProjectData.Project Project in ProjectData.NativeProjects)
@@ -1499,6 +1566,7 @@ namespace UnrealBuildTool
 					{
 						WorkspaceFile.AddUnnamedField("vadimcn.vscode-lldb");
 						WorkspaceFile.AddUnnamedField("ms-vscode.mono-debug");
+						WorkspaceFile.AddUnnamedField("dfarley1.file-picker");
 					}
 				}
 				WorkspaceFile.EndArray();
