Update Mac/iOS/tvOS compiler options

- Remove "-fobjc-abi-version=2"
- Remove "-fobjc-legacy-dispatch"
- Remove "-gdwarf-2" and replace with "-g"

#rnx

diff --git a/Engine/Source/Programs/UnrealBuildTool/Platform/IOS/IOSToolChain.cs b/Engine/Source/Programs/UnrealBuildTool/Platform/IOS/IOSToolChain.cs
--- a/Engine/Source/Programs/UnrealBuildTool/Platform/IOS/IOSToolChain.cs
+++ b/Engine/Source/Programs/UnrealBuildTool/Platform/IOS/IOSToolChain.cs
@@ -347,7 +347,7 @@ namespace UnrealBuildTool
 			// Create DWARF format debug info if wanted,
 			if (CompileEnvironment.bCreateDebugInfo)
 			{
-				Result += " -gdwarf-2";
+				Result += " -g";
 			}
 
 			// Add additional frameworks so that their headers can be found
@@ -373,8 +373,6 @@ namespace UnrealBuildTool
 		{
 			string Result = "";
 			Result += " -x objective-c++";
-			Result += " -fobjc-abi-version=2";
-			Result += " -fobjc-legacy-dispatch";
 			Result += GetCppStandardCompileArgument(CompileEnvironment.CppStandard);
 			Result += " -stdlib=libc++";
 			return Result;
@@ -384,8 +382,6 @@ namespace UnrealBuildTool
 		{
 			string Result = "";
 			Result += " -x objective-c++";
-			Result += " -fobjc-abi-version=2";
-			Result += " -fobjc-legacy-dispatch";
 			Result += GetCppStandardCompileArgument(CompileEnvironment.CppStandard);
 			Result += " -stdlib=libc++";
 			return Result;
@@ -395,8 +391,6 @@ namespace UnrealBuildTool
 		{
 			string Result = "";
 			Result += " -x objective-c";
-			Result += " -fobjc-abi-version=2";
-			Result += " -fobjc-legacy-dispatch";
 			Result += GetCppStandardCompileArgument(CompileEnvironment.CppStandard);
 			Result += " -stdlib=libc++";
 			return Result;
diff --git a/Engine/Source/Programs/UnrealBuildTool/Platform/Mac/MacToolChain.cs b/Engine/Source/Programs/UnrealBuildTool/Platform/Mac/MacToolChain.cs
--- a/Engine/Source/Programs/UnrealBuildTool/Platform/Mac/MacToolChain.cs
+++ b/Engine/Source/Programs/UnrealBuildTool/Platform/Mac/MacToolChain.cs
@@ -287,7 +287,7 @@ namespace UnrealBuildTool
 			// Create DWARF format debug info if wanted,
 			if (CompileEnvironment.bCreateDebugInfo)
 			{
-				Result += " -gdwarf-2";
+				Result += " -g";
 			}
 
 			return Result;
@@ -297,8 +297,6 @@ namespace UnrealBuildTool
 		{
 			string Result = "";
 			Result += " -x objective-c++";
-			Result += " -fobjc-abi-version=2";
-			Result += " -fobjc-legacy-dispatch";
 			Result += GetCppStandardCompileArgument(CompileEnvironment.CppStandard);
 			Result += " -stdlib=libc++";
 			return Result;
@@ -308,8 +306,6 @@ namespace UnrealBuildTool
 		{
 			string Result = "";
 			Result += " -x objective-c++";
-			Result += " -fobjc-abi-version=2";
-			Result += " -fobjc-legacy-dispatch";
 			Result += GetCppStandardCompileArgument(CompileEnvironment.CppStandard);
 			Result += " -stdlib=libc++";
 			return Result;
@@ -319,8 +315,6 @@ namespace UnrealBuildTool
 		{
 			string Result = "";
 			Result += " -x objective-c";
-			Result += " -fobjc-abi-version=2";
-			Result += " -fobjc-legacy-dispatch";
 			Result += " -stdlib=libc++";
 			return Result;
 		}
@@ -336,8 +330,6 @@ namespace UnrealBuildTool
 		{
 			string Result = "";
 			Result += " -x objective-c++-header";
-			Result += " -fobjc-abi-version=2";
-			Result += " -fobjc-legacy-dispatch";
 			Result += GetCppStandardCompileArgument(CompileEnvironment.CppStandard);
 			Result += " -stdlib=libc++";
 			return Result;
