Re-enabling type truncation warnings on Win64 VS compiles (this is expected to cause some fallout with LWC enabled)

diff --git a/Engine/Source/Programs/UnrealBuildTool/Platform/Windows/UEBuildWindows.cs b/Engine/Source/Programs/UnrealBuildTool/Platform/Windows/UEBuildWindows.cs
--- a/Engine/Source/Programs/UnrealBuildTool/Platform/Windows/UEBuildWindows.cs
+++ b/Engine/Source/Programs/UnrealBuildTool/Platform/Windows/UEBuildWindows.cs
@@ -1903,9 +1903,6 @@ namespace UnrealBuildTool
 		/// <param name="LinkEnvironment">The link environment for this target</param>
 		public override void SetUpEnvironment(ReadOnlyTargetRules Target, CppCompileEnvironment CompileEnvironment, LinkEnvironment LinkEnvironment)
 		{
-			// LWC_TODO: Disable "possible loss of data" conversion warnings (double -> float)
-			CompileEnvironment.AdditionalArguments += "/wd4244";
-
 			// @todo Remove this hack to work around broken includes
 			CompileEnvironment.Definitions.Add("NDIS_MINIPORT_MAJOR_VERSION=0");
 
