Bump UE5 to new LLVM 11.0.1 v18 Linux

diff --git a/Engine/Source/Programs/UnrealBuildTool/Platform/Linux/LinuxPlatformSDK.cs b/Engine/Source/Programs/UnrealBuildTool/Platform/Linux/LinuxPlatformSDK.cs
--- a/Engine/Source/Programs/UnrealBuildTool/Platform/Linux/LinuxPlatformSDK.cs
+++ b/Engine/Source/Programs/UnrealBuildTool/Platform/Linux/LinuxPlatformSDK.cs
@@ -14,14 +14,14 @@ namespace UnrealBuildTool
 	{
 		public override string GetMainVersion()
 		{
-			return "v17_clang-10.0.1-centos7";
+			return "v18_clang-11.0.1-centos7";
 		}
 
 		public override void GetValidVersionRange(out string MinVersion, out string MaxVersion)
 		{
 			// all that matters is the number after the v, according to TryConvertVersionToInt()
 			MinVersion = "v10_clang-5.0.0-centos7";
-			MaxVersion = "v17_clang-10.0.1-centos7";
+			MaxVersion = "v18_clang-11.0.1-centos7";
 		}
 
 		public override void GetValidSoftwareVersionRange(out string MinVersion, out string MaxVersion)
diff --git a/Engine/Source/Programs/UnrealBuildTool/Platform/Linux/LinuxToolChain.cs b/Engine/Source/Programs/UnrealBuildTool/Platform/Linux/LinuxToolChain.cs
--- a/Engine/Source/Programs/UnrealBuildTool/Platform/Linux/LinuxToolChain.cs
+++ b/Engine/Source/Programs/UnrealBuildTool/Platform/Linux/LinuxToolChain.cs
@@ -188,10 +188,10 @@ namespace UnrealBuildTool
 				throw new BuildException("Unable to build: no compatible clang version found. Please run Setup.sh");
 			}
 			// prevent unknown clangs since the build is likely to fail on too old or too new compilers
-			else if ((CompilerVersionMajor * 10 + CompilerVersionMinor) > 100 || (CompilerVersionMajor * 10 + CompilerVersionMinor) < 60)
+			else if ((CompilerVersionMajor * 10 + CompilerVersionMinor) > 110 || (CompilerVersionMajor * 10 + CompilerVersionMinor) < 60)
 			{
 				throw new BuildException(
-					string.Format("This version of the Unreal Engine can only be compiled with clang 10.0, 9.0, 8.0, 7.0 and 6.0. clang {0} may not build it - please use a different version.",
+					string.Format("This version of the Unreal Engine can only be compiled with clang 11.0, 10.0, 9.0, 8.0, 7.0 and 6.0. clang {0} may not build it - please use a different version.",
 						CompilerVersionString)
 					);
 			}
