UnrealBuildTool: Support VS2022 Preview

Default is still VS2019, running `GenerateProjectFiles.bat -2022` will create a solution and projects that pass -2022 to UnrealBuild tool when compiling from Visual Studio to force it to use the 2022 toolchain.
Please note, as the defaults are unchanged building from UnrealGameSync will still compile with VS2019 so I would disable that build if testing VS2022.
None of this is necessary to use VS2022, it can open VS2019 solutions and will use the Vs2019 toolchain to build.

#pf 60aebccd7d4b9f0001197729

diff --git a/Engine/Source/Programs/UnrealBuildTool/Modes/GenerateProjectFilesMode.cs b/Engine/Source/Programs/UnrealBuildTool/Modes/GenerateProjectFilesMode.cs
--- a/Engine/Source/Programs/UnrealBuildTool/Modes/GenerateProjectFilesMode.cs
+++ b/Engine/Source/Programs/UnrealBuildTool/Modes/GenerateProjectFilesMode.cs
@@ -26,6 +26,7 @@ namespace UnrealBuildTool
 		[CommandLine("-ProjectFileFormat")]
 		[CommandLine("-2017", Value = nameof(ProjectFileFormat.VisualStudio2017))] // + override compiler
 		[CommandLine("-2019", Value = nameof(ProjectFileFormat.VisualStudio2019))] // + override compiler
+		[CommandLine("-2022", Value = nameof(ProjectFileFormat.VisualStudio2022))] // + override compiler
 		[CommandLine("-Makefile", Value = nameof(ProjectFileFormat.Make))]
 		[CommandLine("-CMakefile", Value = nameof(ProjectFileFormat.CMake))]
 		[CommandLine("-QMakefile", Value = nameof(ProjectFileFormat.QMake))]
@@ -197,6 +198,9 @@ namespace UnrealBuildTool
 					case ProjectFileFormat.VisualStudio2019:
 						Generator = new VCProjectFileGenerator(ProjectFile, VCProjectFileFormat.VisualStudio2019, Arguments);
 						break;
+					case ProjectFileFormat.VisualStudio2022:
+						Generator = new VCProjectFileGenerator(ProjectFile, VCProjectFileFormat.VisualStudio2022, Arguments);
+						break;
 					case ProjectFileFormat.XCode:
 						Generator = new XcodeProjectFileGenerator(ProjectFile, Arguments);
 						break;
diff --git a/Engine/Source/Programs/UnrealBuildTool/Platform/HoloLens/UEBuildHoloLens.cs b/Engine/Source/Programs/UnrealBuildTool/Platform/HoloLens/UEBuildHoloLens.cs
--- a/Engine/Source/Programs/UnrealBuildTool/Platform/HoloLens/UEBuildHoloLens.cs
+++ b/Engine/Source/Programs/UnrealBuildTool/Platform/HoloLens/UEBuildHoloLens.cs
@@ -26,6 +26,7 @@ namespace UnrealBuildTool
 		[XmlConfigFile(Category = "HoloLensPlatform")]
 		[CommandLine("-2017", Value = nameof(WindowsCompiler.VisualStudio2017))]
 		[CommandLine("-2019", Value = nameof(WindowsCompiler.VisualStudio2019))]
+		[CommandLine("-2022", Value = nameof(WindowsCompiler.VisualStudio2022))]
 		public WindowsCompiler Compiler = WindowsCompiler.Default;
 
 		/// <summary>
diff --git a/Engine/Source/Programs/UnrealBuildTool/Platform/Lumin/LuminProjectGenerator.cs b/Engine/Source/Programs/UnrealBuildTool/Platform/Lumin/LuminProjectGenerator.cs
--- a/Engine/Source/Programs/UnrealBuildTool/Platform/Lumin/LuminProjectGenerator.cs
+++ b/Engine/Source/Programs/UnrealBuildTool/Platform/Lumin/LuminProjectGenerator.cs
@@ -39,7 +39,8 @@ namespace UnrealBuildTool
 			{
 				// TODO: add a registry check or file exists check to confirm if MLExtension is installed on the given ProjectFileFormat.
 				VSDebuggingEnabled = (ProjectFileFormat == VCProjectFileFormat.VisualStudio2017 
-									|| ProjectFileFormat == VCProjectFileFormat.VisualStudio2019);
+									|| ProjectFileFormat == VCProjectFileFormat.VisualStudio2019
+									|| ProjectFileFormat == VCProjectFileFormat.VisualStudio2022);
 				VSSupportChecked = true;
 			}
             return VSDebuggingEnabled;
diff --git a/Engine/Source/Programs/UnrealBuildTool/Platform/Windows/UEBuildWindows.cs b/Engine/Source/Programs/UnrealBuildTool/Platform/Windows/UEBuildWindows.cs
--- a/Engine/Source/Programs/UnrealBuildTool/Platform/Windows/UEBuildWindows.cs
+++ b/Engine/Source/Programs/UnrealBuildTool/Platform/Windows/UEBuildWindows.cs
@@ -44,6 +44,11 @@ namespace UnrealBuildTool
 		/// Visual Studio 2019 (Visual C++ 16.0)
 		/// </summary>
 		VisualStudio2019,
+
+		/// <summary>
+		/// Visual Studio 2022 (Visual C++ 17.0)
+		/// </summary>
+		VisualStudio2022,
 	}
 
 	/// <summary>
@@ -106,6 +111,7 @@ namespace UnrealBuildTool
 		[XmlConfigFile(Category = "WindowsPlatform")]
 		[CommandLine("-2017", Value = nameof(WindowsCompiler.VisualStudio2017))]
 		[CommandLine("-2019", Value = nameof(WindowsCompiler.VisualStudio2019))]
+		[CommandLine("-2022", Value = nameof(WindowsCompiler.VisualStudio2022))]
 		[CommandLine("-Compiler=")]
 		public WindowsCompiler Compiler = WindowsCompiler.Default;
 
@@ -220,7 +226,7 @@ namespace UnrealBuildTool
 		/// Microsoft provides legacy_stdio_definitions library to enable building with VS2015 until they fix everything up.
 		public bool bNeedsLegacyStdioDefinitionsLib
 		{
-			get { return Compiler == WindowsCompiler.VisualStudio2017 || Compiler == WindowsCompiler.VisualStudio2019 || Compiler == WindowsCompiler.Clang; }
+			get { return Compiler == WindowsCompiler.VisualStudio2017 || Compiler == WindowsCompiler.VisualStudio2019 || Compiler == WindowsCompiler.VisualStudio2022 || Compiler == WindowsCompiler.Clang; }
 		}
 
 		/// <summary>
@@ -381,6 +387,7 @@ namespace UnrealBuildTool
 				case WindowsCompiler.Intel:
 				case WindowsCompiler.VisualStudio2017:
 				case WindowsCompiler.VisualStudio2019:
+				case WindowsCompiler.VisualStudio2022:
 					return "2015"; // VS2017 is backwards compatible with VS2015 compiler
 
 				default:
@@ -992,6 +999,8 @@ namespace UnrealBuildTool
 					return "Visual Studio 2017";
 				case WindowsCompiler.VisualStudio2019:
 					return "Visual Studio 2019";
+				case WindowsCompiler.VisualStudio2022:
+					return "Visual Studio 2022";
 				default:
 					return Compiler.ToString();
 			}
@@ -1039,7 +1048,7 @@ namespace UnrealBuildTool
 				Installations = new List<VisualStudioInstallation>();
 			    if (BuildHostPlatform.Current.Platform == UnrealTargetPlatform.Win64)
 			    {
-				    if(Compiler == WindowsCompiler.VisualStudio2017 || Compiler == WindowsCompiler.VisualStudio2019)
+				    if(Compiler == WindowsCompiler.VisualStudio2017 || Compiler == WindowsCompiler.VisualStudio2019 || Compiler == WindowsCompiler.VisualStudio2022)
 				    {
 						try
 						{
@@ -1066,7 +1075,12 @@ namespace UnrealBuildTool
 									if (VersionNumber.TryParse(VersionString, out Version))
 									{
 										VersionNumber Version2019 = new VersionNumber(16);
-										if(Compiler == WindowsCompiler.VisualStudio2019 && Version < Version2019)
+										VersionNumber Version2022 = new VersionNumber(17);
+										if (Compiler == WindowsCompiler.VisualStudio2022 && Version < Version2022)
+										{
+											continue;
+										}
+										else if (Compiler == WindowsCompiler.VisualStudio2019 && Version < Version2019)
 										{
 											continue;
 										}
@@ -1170,7 +1184,7 @@ namespace UnrealBuildTool
 							}
 						}
 					}
-				    else if(Compiler == WindowsCompiler.VisualStudio2017 || Compiler == WindowsCompiler.VisualStudio2019)
+				    else if(Compiler == WindowsCompiler.VisualStudio2017 || Compiler == WindowsCompiler.VisualStudio2019 || Compiler == WindowsCompiler.VisualStudio2022)
 				    {
 						// Enumerate all the manually installed toolchains
 						List<VisualStudioInstallation> Installations = FindVisualStudioInstallations(Compiler);
@@ -1184,11 +1198,22 @@ namespace UnrealBuildTool
 						DirectoryReference PlatformDir;
 						if (UEBuildPlatformSDK.TryGetHostPlatformAutoSDKDir(out PlatformDir))
 						{
-							DirectoryReference ReleaseBaseDir = DirectoryReference.Combine(PlatformDir, "Win64", (Compiler == WindowsCompiler.VisualStudio2019) ? "VS2019" : "VS2017");
-							FindVisualStudioToolChains(ReleaseBaseDir, false, ToolChains);
+							string VSDir = string.Empty;
+							switch (Compiler)
+							{
+								case WindowsCompiler.VisualStudio2017: VSDir = "VS2017"; break;
+								case WindowsCompiler.VisualStudio2019: VSDir = "VS2019"; break;
+								case WindowsCompiler.VisualStudio2022: VSDir = "VS2022"; break;
+							}
+
+							if (!string.IsNullOrEmpty(VSDir))
+							{
+								DirectoryReference ReleaseBaseDir = DirectoryReference.Combine(PlatformDir, "Win64", VSDir);
+								FindVisualStudioToolChains(ReleaseBaseDir, false, ToolChains);
 
-							DirectoryReference PreviewBaseDir = DirectoryReference.Combine(PlatformDir, "Win64", (Compiler == WindowsCompiler.VisualStudio2019) ? "VS2019-Preview" : "VS2017-Preview");
-							FindVisualStudioToolChains(PreviewBaseDir, true, ToolChains);
+								DirectoryReference PreviewBaseDir = DirectoryReference.Combine(PlatformDir, "Win64", $"{VSDir}-Preview");
+								FindVisualStudioToolChains(PreviewBaseDir, true, ToolChains);
+							}
 						}
 					}
 					else
@@ -1323,16 +1348,6 @@ namespace UnrealBuildTool
 			return FileReference.Exists(FileReference.Combine(ToolChainDir, "bin", "clang-cl.exe"));
 		}
 
-		/// <summary>
-		/// Checks if the given directory contains a valid Visual Studio 2015 toolchain
-		/// </summary>
-		/// <param name="ToolChainDir">Directory to check</param>
-		/// <returns>True if the given directory is valid</returns>
-		static bool IsValidToolChainDir2015(DirectoryReference ToolChainDir)
-		{
-			return FileReference.Exists(FileReference.Combine(ToolChainDir, "bin", "amd64", "cl.exe")) || FileReference.Exists(FileReference.Combine(ToolChainDir, "bin", "x86_amd64", "cl.exe"));
-		}
-
 		/// <summary>
 		/// Determines if the given path is a valid Visual C++ version number
 		/// </summary>
@@ -1642,10 +1657,21 @@ namespace UnrealBuildTool
 				DirectoryReference PlatformDir;
 				if (UEBuildPlatformSDK.TryGetHostPlatformAutoSDKDir(out PlatformDir))
 				{
-					DirectoryReference DiaSdkDir = DirectoryReference.Combine(PlatformDir, "Win64", "DIA SDK", (Compiler == WindowsCompiler.VisualStudio2019) ? "VS2019" : "VS2017");
-					if (IsValidDiaSdkDir(DiaSdkDir))
+					string VSDir = string.Empty;
+					switch (Compiler)
 					{
-						DiaSdkDirs.Add(DiaSdkDir);
+						case WindowsCompiler.VisualStudio2017: VSDir = "VS2017"; break;
+						case WindowsCompiler.VisualStudio2019: VSDir = "VS2019"; break;
+						case WindowsCompiler.VisualStudio2022: VSDir = "VS2022"; break;
+					}
+
+					if (!string.IsNullOrEmpty(VSDir))
+					{
+						DirectoryReference DiaSdkDir = DirectoryReference.Combine(PlatformDir, "Win64", "DIA SDK", VSDir);
+						if (IsValidDiaSdkDir(DiaSdkDir))
+						{
+							DiaSdkDirs.Add(DiaSdkDir);
+						}
 					}
 				}
 
diff --git a/Engine/Source/Programs/UnrealBuildTool/Platform/Windows/VCEnvironment.cs b/Engine/Source/Programs/UnrealBuildTool/Platform/Windows/VCEnvironment.cs
--- a/Engine/Source/Programs/UnrealBuildTool/Platform/Windows/VCEnvironment.cs
+++ b/Engine/Source/Programs/UnrealBuildTool/Platform/Windows/VCEnvironment.cs
@@ -448,9 +448,13 @@ namespace UnrealBuildTool
 				{
 					ToolChain = WindowsCompiler.VisualStudio2017;
 				}
+				else if (WindowsPlatform.TryGetToolChainDir(WindowsCompiler.VisualStudio2022, null, out SelectedToolChainVersion, out SelectedToolChainDir))
+				{
+					ToolChain = WindowsCompiler.VisualStudio2022;
+				}
 				else
 				{
-					throw new BuildException("{0} or {1} must be installed in order to build this target.", WindowsPlatform.GetCompilerName(WindowsCompiler.VisualStudio2019), WindowsPlatform.GetCompilerName(WindowsCompiler.VisualStudio2017));
+					throw new BuildException("{0}, {1}, or {2} must be installed in order to build this target.", WindowsPlatform.GetCompilerName(WindowsCompiler.VisualStudio2019), WindowsPlatform.GetCompilerName(WindowsCompiler.VisualStudio2017), WindowsPlatform.GetCompilerName(WindowsCompiler.VisualStudio2022));
 				}
 			}
 			else
diff --git a/Engine/Source/Programs/UnrealBuildTool/Platform/Windows/VCToolChain.cs b/Engine/Source/Programs/UnrealBuildTool/Platform/Windows/VCToolChain.cs
--- a/Engine/Source/Programs/UnrealBuildTool/Platform/Windows/VCToolChain.cs
+++ b/Engine/Source/Programs/UnrealBuildTool/Platform/Windows/VCToolChain.cs
@@ -851,6 +851,7 @@ namespace UnrealBuildTool
 							break;
 						case WindowsCompiler.VisualStudio2017:
 						case WindowsCompiler.VisualStudio2019:
+						case WindowsCompiler.VisualStudio2022:
 							Arguments[Arguments.Count - 1] += ":FASTLINK";
 							break;
 					}
diff --git a/Engine/Source/Programs/UnrealBuildTool/ProjectFiles/ProjectFileGenerator.cs b/Engine/Source/Programs/UnrealBuildTool/ProjectFiles/ProjectFileGenerator.cs
--- a/Engine/Source/Programs/UnrealBuildTool/ProjectFiles/ProjectFileGenerator.cs
+++ b/Engine/Source/Programs/UnrealBuildTool/ProjectFiles/ProjectFileGenerator.cs
@@ -133,6 +133,7 @@ namespace UnrealBuildTool
 		VisualStudio,
 		VisualStudio2017,
 		VisualStudio2019,
+		VisualStudio2022,
 		XCode,
 		Eddie,
 		VisualStudioCode,
@@ -662,8 +663,14 @@ namespace UnrealBuildTool
 					Format = ProjectFileFormat.VisualStudio2019;
 					return true;
 				}
+				else if (PreferredAccessor == "visualstudio2022")
+				{
+					Format = ProjectFileFormat.VisualStudio2022;
+					return true;
+				}
 			}
 
+
 			Format = ProjectFileFormat.VisualStudio;
 			return false;
 		}
diff --git a/Engine/Source/Programs/UnrealBuildTool/ProjectFiles/VisualStudio/VCProject.cs b/Engine/Source/Programs/UnrealBuildTool/ProjectFiles/VisualStudio/VCProject.cs
--- a/Engine/Source/Programs/UnrealBuildTool/ProjectFiles/VisualStudio/VCProject.cs
+++ b/Engine/Source/Programs/UnrealBuildTool/ProjectFiles/VisualStudio/VCProject.cs
@@ -556,6 +556,8 @@ namespace UnrealBuildTool
 		{
 			switch (ProjectFileFormat)
 			{
+				case VCProjectFileFormat.VisualStudio2022:
+					return WindowsCompiler.VisualStudio2022;
 				case VCProjectFileFormat.VisualStudio2019:
 					return WindowsCompiler.VisualStudio2019;
 				case VCProjectFileFormat.VisualStudio2017:
@@ -1705,6 +1707,11 @@ namespace UnrealBuildTool
 						BuildArguments.Append(ProjGenerator.GetExtraBuildArguments(Platform, Configuration));
 					}
 
+					if (ProjectFileFormat == VCProjectFileFormat.VisualStudio2022)
+					{
+						BuildArguments.Append(" -2022");
+					}
+
 					// NMake Build command line
 					VCProjectFileContent.AppendLine("    <NMakeBuildCommandLine>{0} {1}</NMakeBuildCommandLine>", EscapePath(NormalizeProjectPath(FileReference.Combine(BatchFilesDirectory, "Build.bat"))), BuildArguments.ToString());
 					VCProjectFileContent.AppendLine("    <NMakeReBuildCommandLine>{0} {1}</NMakeReBuildCommandLine>", EscapePath(NormalizeProjectPath(FileReference.Combine(BatchFilesDirectory, "Rebuild.bat"))), BuildArguments.ToString());
diff --git a/Engine/Source/Programs/UnrealBuildTool/ProjectFiles/VisualStudio/VCProjectFileGenerator.cs b/Engine/Source/Programs/UnrealBuildTool/ProjectFiles/VisualStudio/VCProjectFileGenerator.cs
--- a/Engine/Source/Programs/UnrealBuildTool/ProjectFiles/VisualStudio/VCProjectFileGenerator.cs
+++ b/Engine/Source/Programs/UnrealBuildTool/ProjectFiles/VisualStudio/VCProjectFileGenerator.cs
@@ -32,6 +32,7 @@ namespace UnrealBuildTool
 		Default,          // Default to the best installed version, but allow SDKs to override
 		VisualStudio2017,
 		VisualStudio2019,
+		VisualStudio2022,
 	}
 
 	class VCProjectFileSettings
@@ -222,6 +223,8 @@ namespace UnrealBuildTool
 					return "15.0";
 				case VCProjectFileFormat.VisualStudio2019:
 					return "16.0";
+				case VCProjectFileFormat.VisualStudio2022:
+					return "17.0";
 			}
 			return string.Empty;
 		}
@@ -234,6 +237,7 @@ namespace UnrealBuildTool
 				case VCProjectFileFormat.VisualStudio2017:
 					return "v141";
 				case VCProjectFileFormat.VisualStudio2019:
+				case VCProjectFileFormat.VisualStudio2022:
 					return "v142";
 
 			}
@@ -285,6 +289,14 @@ namespace UnrealBuildTool
 					Settings.ProjectFileFormat = VCProjectFileFormat.Default;
 				}
 			}
+			else if (Settings.ProjectFileFormat == VCProjectFileFormat.VisualStudio2022)
+			{
+				if (!WindowsPlatform.HasCompiler(WindowsCompiler.VisualStudio2022))
+				{
+					Log.TraceWarning("Visual Studio C++ 2022 installation not found - ignoring preferred project file format.");
+					Settings.ProjectFileFormat = VCProjectFileFormat.Default;
+				}
+			}
 
 			// Certain platforms override the project file format because their debugger add-ins may not yet support the latest
 			// version of Visual Studio.  This is their chance to override that.
@@ -300,6 +312,10 @@ namespace UnrealBuildTool
 				{
 					Settings.ProjectFileFormat = VCProjectFileFormat.VisualStudio2017;
 				}
+				else if (WindowsPlatform.HasCompiler(WindowsCompiler.VisualStudio2022) && WindowsPlatform.HasIDE(WindowsCompiler.VisualStudio2022))
+				{
+					Settings.ProjectFileFormat = VCProjectFileFormat.VisualStudio2022;
+				}
 
 				// Allow the SDKs to override
 				foreach (UnrealTargetPlatform SupportedPlatform in SupportedPlatforms)
@@ -477,7 +493,15 @@ namespace UnrealBuildTool
 			StringBuilder VCSolutionFileContent = new StringBuilder();
 
 			// Solution file header. Note that a leading newline is required for file type detection to work correclty in the shell.
-			if (Settings.ProjectFileFormat == VCProjectFileFormat.VisualStudio2019)
+			if (Settings.ProjectFileFormat == VCProjectFileFormat.VisualStudio2022)
+			{
+				VCSolutionFileContent.AppendLine();
+				VCSolutionFileContent.AppendLine("Microsoft Visual Studio Solution File, Format Version 12.00");
+				VCSolutionFileContent.AppendLine("# Visual Studio Version 17");
+				VCSolutionFileContent.AppendLine("VisualStudioVersion = 17.0.31314.256");
+				VCSolutionFileContent.AppendLine("MinimumVisualStudioVersion = 10.0.40219.1");
+			}
+			else if (Settings.ProjectFileFormat == VCProjectFileFormat.VisualStudio2019)
 			{
 				VCSolutionFileContent.AppendLine();
 				VCSolutionFileContent.AppendLine("Microsoft Visual Studio Solution File, Format Version 12.00");
@@ -807,6 +831,9 @@ namespace UnrealBuildTool
 					case VCProjectFileFormat.VisualStudio2019:
 						SolutionOptionsFileName = FileReference.Combine(MasterProjectPath, ".vs", Path.GetFileNameWithoutExtension(SolutionFileName), "v16", ".suo");
 						break;
+					case VCProjectFileFormat.VisualStudio2022:
+						SolutionOptionsFileName = FileReference.Combine(MasterProjectPath, ".vs", Path.GetFileNameWithoutExtension(SolutionFileName), "v17", ".suo");
+						break;
 					default:
 						throw new BuildException("Unsupported Visual Studio version");
 				}
