Add HasExplicitPlatforms property to uplugin definitions, allowing the SupportedTargetPlatforms and WhitelistPlatforms fields to be entirely defined in child plugins. The existing functionality means that an empty platform list implies "all platforms".

diff --git a/Engine/Source/Editor/UnrealEd/Private/CookOnTheFlyServer.cpp b/Engine/Source/Editor/UnrealEd/Private/CookOnTheFlyServer.cpp
--- a/Engine/Source/Editor/UnrealEd/Private/CookOnTheFlyServer.cpp
+++ b/Engine/Source/Editor/UnrealEd/Private/CookOnTheFlyServer.cpp
@@ -7454,7 +7454,7 @@ void UCookOnTheFlyServer::DiscoverPlatformSpecificNeverCookPackages(
 		const FPluginDescriptor& Descriptor = Plugin->GetDescriptor();
 
 		// we are only interested in plugins that does not support all platforms
-		if (Descriptor.SupportedTargetPlatforms.Num() == 0)
+		if (Descriptor.SupportedTargetPlatforms.Num() == 0 && !Descriptor.bHasExplicitPlatforms)
 		{
 			continue;
 		}
diff --git a/Engine/Source/Programs/UnrealBuildTool/Configuration/UEBuildTarget.cs b/Engine/Source/Programs/UnrealBuildTool/Configuration/UEBuildTarget.cs
--- a/Engine/Source/Programs/UnrealBuildTool/Configuration/UEBuildTarget.cs
+++ b/Engine/Source/Programs/UnrealBuildTool/Configuration/UEBuildTarget.cs
@@ -3102,6 +3102,7 @@ namespace UnrealBuildTool
 
 						PluginReferenceDescriptor PluginReference = new PluginReferenceDescriptor(Plugin.Name, null, true);
 						PluginReference.SupportedTargetPlatforms = Plugin.Descriptor.GetSupportedTargetPlatformNames();
+						PluginReference.bHasExplicitPlatforms = Plugin.Descriptor.bHasExplicitPlatforms;
 						PluginReference.bOptional = true;
 
 						AddPlugin(PluginReference, "default plugins", ExcludeFolders, NameToInstance, NameToInfo);
diff --git a/Engine/Source/Programs/UnrealBuildTool/System/ModuleDescriptor.cs b/Engine/Source/Programs/UnrealBuildTool/System/ModuleDescriptor.cs
--- a/Engine/Source/Programs/UnrealBuildTool/System/ModuleDescriptor.cs
+++ b/Engine/Source/Programs/UnrealBuildTool/System/ModuleDescriptor.cs
@@ -212,6 +212,11 @@ namespace UnrealBuildTool
 		/// </summary>
 		public string[]? AdditionalDependencies;
 
+		/// <summary>
+		/// When true, empty WhitelistPlatforms are interpeted as 'no platforms' with the expectation that explict platforms will be added in plugin extensions */
+		/// </summary>
+		public bool bHasExplicitPlatforms;
+
 		/// <summary>
 		/// Constructor
 		/// </summary>
@@ -327,6 +332,12 @@ namespace UnrealBuildTool
 				Module.AdditionalDependencies = AdditionalDependencies;
 			}
 
+			bool bHasExplicitPlatforms;
+			if (InObject.TryGetBoolField("HasExplicitPlatforms", out bHasExplicitPlatforms))
+			{
+				Module.bHasExplicitPlatforms = bHasExplicitPlatforms;
+			}
+
 			return Module;
 		}
 
@@ -414,6 +425,10 @@ namespace UnrealBuildTool
 				}
 				Writer.WriteArrayEnd();
 			}
+			if (bHasExplicitPlatforms)
+			{
+				Writer.WriteValue("HasExplicitPlatforms", bHasExplicitPlatforms);
+			}
 			Writer.WriteObjectEnd();
 		}
 
@@ -463,7 +478,8 @@ namespace UnrealBuildTool
 			// Check the platform is whitelisted
 			// important note: we don't check the length of the whitelist platforms, because if an unknown platform was read in, but was not valid, the 
 			// list will exist but be empty. In this case, we need to disallow all platforms from building, otherwise, build errors will occur when
-			// it starts compiling for _all_ platforms
+			// it starts compiling for _all_ platforms. This means we don't need to check bHasExplicitPlatforms either
+			
 			if (WhitelistPlatforms != null && !WhitelistPlatforms.Contains(Platform))
 			{
 				return false;
diff --git a/Engine/Source/Programs/UnrealBuildTool/System/PluginDescriptor.cs b/Engine/Source/Programs/UnrealBuildTool/System/PluginDescriptor.cs
--- a/Engine/Source/Programs/UnrealBuildTool/System/PluginDescriptor.cs
+++ b/Engine/Source/Programs/UnrealBuildTool/System/PluginDescriptor.cs
@@ -174,6 +174,11 @@ namespace UnrealBuildTool
 		/// </summary>
 		public bool bExplicitlyLoaded;
 
+		/// <summary>
+		/// When true, an empty SupportedTargetPlatforms is interpeted as 'no platforms' with the expectation that explict platforms will be added in plugin platform extensions
+		/// </summary>
+		public bool bHasExplicitPlatforms;
+
 		/// <summary>
 		/// Set of pre-build steps to execute, keyed by host platform name.
 		/// </summary>
@@ -297,6 +302,7 @@ namespace UnrealBuildTool
 
 			RawObject.TryGetBoolField("RequiresBuildPlatform", out bRequiresBuildPlatform);
 			RawObject.TryGetBoolField("ExplicitlyLoaded", out bExplicitlyLoaded);
+			RawObject.TryGetBoolField("HasExplicitPlatforms", out bHasExplicitPlatforms);
 
 			CustomBuildSteps.TryRead(RawObject, "PreBuildSteps", out PreBuildSteps);
 			CustomBuildSteps.TryRead(RawObject, "PostBuildSteps", out PostBuildSteps);
@@ -397,6 +403,11 @@ namespace UnrealBuildTool
 				Writer.WriteValue("ExplicitlyLoaded", bExplicitlyLoaded);
 			}
 
+			if (bHasExplicitPlatforms)
+			{
+				Writer.WriteValue("HasExplicitPlatforms", bHasExplicitPlatforms);
+			}
+
 			if(SupportedTargetPlatforms != null && SupportedTargetPlatforms.Count > 0)
 			{
 				Writer.WriteStringArrayField("SupportedTargetPlatforms", SupportedTargetPlatforms.Select<UnrealTargetPlatform, string>(x => x.ToString()).ToArray());
@@ -441,7 +452,14 @@ namespace UnrealBuildTool
 		/// <returns>True if the plugin should be enabled</returns>
 		public bool SupportsTargetPlatform(UnrealTargetPlatform Platform)
 		{
-			return SupportedTargetPlatforms == null || SupportedTargetPlatforms.Count == 0 || SupportedTargetPlatforms.Contains(Platform);
+			if (bHasExplicitPlatforms)
+			{
+				return SupportedTargetPlatforms != null && SupportedTargetPlatforms.Contains(Platform);
+			}
+			else
+			{
+				return SupportedTargetPlatforms == null || SupportedTargetPlatforms.Count == 0 || SupportedTargetPlatforms.Contains(Platform);
+			}
 		}
 
 		/// <summary>
diff --git a/Engine/Source/Programs/UnrealBuildTool/System/PluginReferenceDescriptor.cs b/Engine/Source/Programs/UnrealBuildTool/System/PluginReferenceDescriptor.cs
--- a/Engine/Source/Programs/UnrealBuildTool/System/PluginReferenceDescriptor.cs
+++ b/Engine/Source/Programs/UnrealBuildTool/System/PluginReferenceDescriptor.cs
@@ -76,6 +76,12 @@ namespace UnrealBuildTool
 		/// </summary>
 		public string[]? SupportedTargetPlatforms;
 
+		/// <summary>
+		/// When true, empty SupportedTargetPlatforms and WhitelistPlatforms are interpeted as 'no platforms' with the expectation that explict platforms will be added in plugin platform extensions
+		/// </summary>
+		public bool bHasExplicitPlatforms;
+
+
 		/// <summary>
 		/// Constructor
 		/// </summary>
@@ -138,6 +144,10 @@ namespace UnrealBuildTool
 			{
 				Writer.WriteStringArrayField("SupportedTargetPlatforms", SupportedTargetPlatforms.Select(x => x.ToString()).ToArray());
 			}
+			if(bHasExplicitPlatforms)
+			{
+				Writer.WriteValue("HasExplicitPlatforms",bHasExplicitPlatforms);
+			}
 			Writer.WriteObjectEnd();
 		}
 
@@ -201,6 +211,7 @@ namespace UnrealBuildTool
 				RawObject.TryGetEnumArrayField<TargetType>("WhitelistTargets", out Descriptor.WhitelistTargets);
 				RawObject.TryGetEnumArrayField<TargetType>("BlacklistTargets", out Descriptor.BlacklistTargets);
 				RawObject.TryGetStringArrayField("SupportedTargetPlatforms", out Descriptor.SupportedTargetPlatforms);
+				RawObject.TryGetBoolField("HasExplicitPlatforms", out Descriptor.bHasExplicitPlatforms);
 
 				Descriptor.VerifyPlatformNames(Descriptor.WhitelistPlatforms);
 				Descriptor.VerifyPlatformNames(Descriptor.BlacklistPlatforms);
@@ -221,7 +232,14 @@ namespace UnrealBuildTool
 			{
 				return false;
 			}
-			if (WhitelistPlatforms != null && WhitelistPlatforms.Length > 0 && !WhitelistPlatforms.Contains(Platform.ToString()))
+			if (bHasExplicitPlatforms)
+			{
+				if (WhitelistPlatforms == null || !WhitelistPlatforms.Contains(Platform.ToString()))
+				{
+					return false;
+				}
+			}
+			else if (WhitelistPlatforms != null && WhitelistPlatforms.Length > 0 && !WhitelistPlatforms.Contains(Platform.ToString()))
 			{
 				return false;
 			}
@@ -283,7 +301,14 @@ namespace UnrealBuildTool
 		/// <returns>True if the plugin for this target platform</returns>
 		public bool IsSupportedTargetPlatform(UnrealTargetPlatform Platform)
 		{
-			return SupportedTargetPlatforms == null || SupportedTargetPlatforms.Length == 0 || SupportedTargetPlatforms.Contains(Platform.ToString());
+			if (bHasExplicitPlatforms)
+			{
+				return SupportedTargetPlatforms != null && SupportedTargetPlatforms.Contains(Platform.ToString());
+			}
+			else
+			{
+				return SupportedTargetPlatforms == null || SupportedTargetPlatforms.Length == 0 || SupportedTargetPlatforms.Contains(Platform.ToString());
+			}
 		}
 	}
 }
diff --git a/Engine/Source/Runtime/Projects/Private/ModuleDescriptor.cpp b/Engine/Source/Runtime/Projects/Private/ModuleDescriptor.cpp
--- a/Engine/Source/Runtime/Projects/Private/ModuleDescriptor.cpp
+++ b/Engine/Source/Runtime/Projects/Private/ModuleDescriptor.cpp
@@ -152,6 +152,7 @@ FModuleDescriptor::FModuleDescriptor(const FName InName, EHostType::Type InType,
 	: Name(InName)
 	, Type(InType)
 	, LoadingPhase(InLoadingPhase)
+	, bHasExplicitPlatforms(false)
 {
 }
 
@@ -195,6 +196,7 @@ bool FModuleDescriptor::Read(const FJsonObject& Object, FText& OutFailReason)
 	// Read the whitelisted and blacklisted platforms
 	Object.TryGetStringArrayField(TEXT("WhitelistPlatforms"), WhitelistPlatforms);
 	Object.TryGetStringArrayField(TEXT("BlacklistPlatforms"), BlacklistPlatforms);
+	Object.TryGetBoolField(TEXT("HasExplicitPlatforms"), bHasExplicitPlatforms);
 
 	// Read the whitelisted and blacklisted targets
 	Object.TryGetEnumArrayField(TEXT("WhitelistTargets"), WhitelistTargets);
@@ -387,6 +389,16 @@ void FModuleDescriptor::UpdateJson(FJsonObject& JsonObject) const
 	{
 		JsonObject.RemoveField(TEXT("AdditionalDependencies"));
 	}
+
+	if (bHasExplicitPlatforms)
+	{
+		JsonObject.SetBoolField(TEXT("HasExplicitPlatforms"), bHasExplicitPlatforms);
+	}
+	else
+	{
+		JsonObject.RemoveField(TEXT("HasExplicitPlatforms"));
+	}
+
 }
 
 void FModuleDescriptor::WriteArray(TJsonWriter<>& Writer, const TCHAR* ArrayName, const TArray<FModuleDescriptor>& Modules)
@@ -416,7 +428,7 @@ void FModuleDescriptor::UpdateArray(FJsonObject& JsonObject, const TCHAR* ArrayN
 bool FModuleDescriptor::IsCompiledInConfiguration(const FString& Platform, EBuildConfiguration Configuration, const FString& TargetName, EBuildTargetType TargetType, bool bBuildDeveloperTools, bool bBuildRequiresCookedData) const
 {
 	// Check the platform is whitelisted
-	if (WhitelistPlatforms.Num() > 0 && !WhitelistPlatforms.Contains(Platform))
+	if ((bHasExplicitPlatforms || WhitelistPlatforms.Num() > 0) && !WhitelistPlatforms.Contains(Platform))
 	{
 		return false;
 	}
diff --git a/Engine/Source/Runtime/Projects/Private/PluginDescriptor.cpp b/Engine/Source/Runtime/Projects/Private/PluginDescriptor.cpp
--- a/Engine/Source/Runtime/Projects/Private/PluginDescriptor.cpp
+++ b/Engine/Source/Runtime/Projects/Private/PluginDescriptor.cpp
@@ -76,6 +76,7 @@ FPluginDescriptor::FPluginDescriptor()
 	, bRequiresBuildPlatform(false)
 	, bIsHidden(false)
 	, bExplicitlyLoaded(false)
+	, bHasExplicitPlatforms(false)
 	, bIsPluginExtension(false)
 {
 }
@@ -198,6 +199,7 @@ bool FPluginDescriptor::Read(const FJsonObject& Object, FText& OutFailReason)
 	Object.TryGetBoolField(TEXT("RequiresBuildPlatform"), bRequiresBuildPlatform);
 	Object.TryGetBoolField(TEXT("Hidden"), bIsHidden);
 	Object.TryGetBoolField(TEXT("ExplicitlyLoaded"), bExplicitlyLoaded);
+	Object.TryGetBoolField(TEXT("HasExplicitPlatforms"), bHasExplicitPlatforms);
 
 	bool bCanBeUsedWithUnrealHeaderTool;
 	if(Object.TryGetBoolField("CanBeUsedWithUnrealHeaderTool", bCanBeUsedWithUnrealHeaderTool) && bCanBeUsedWithUnrealHeaderTool)
@@ -365,6 +367,15 @@ void FPluginDescriptor::UpdateJson(FJsonObject& JsonObject) const
 		JsonObject.RemoveField(TEXT("ExplicitlyLoaded"));
 	}
 
+	if (bHasExplicitPlatforms)
+	{
+		JsonObject.SetBoolField(TEXT("HasExplicitPlatforms"), bHasExplicitPlatforms);
+	}
+	else
+	{
+		JsonObject.RemoveField("HasExplicitPlatforms");
+	}
+
 	PreBuildSteps.UpdateJson(JsonObject, TEXT("PreBuildSteps"));
 	PostBuildSteps.UpdateJson(JsonObject, TEXT("PostBuildSteps"));
 
@@ -421,7 +432,14 @@ bool FPluginDescriptor::UpdatePluginFile(const FString& FileName, FText& OutFail
 
 bool FPluginDescriptor::SupportsTargetPlatform(const FString& Platform) const
 {
-	return SupportedTargetPlatforms.Num() == 0 || SupportedTargetPlatforms.Contains(Platform);
+	if (bHasExplicitPlatforms)
+	{
+		return SupportedTargetPlatforms.Contains(Platform);
+	}
+	else
+	{
+		return SupportedTargetPlatforms.Num() == 0 || SupportedTargetPlatforms.Contains(Platform);
+	}
 }
 
 #undef LOCTEXT_NAMESPACE
diff --git a/Engine/Source/Runtime/Projects/Private/PluginManager.cpp b/Engine/Source/Runtime/Projects/Private/PluginManager.cpp
--- a/Engine/Source/Runtime/Projects/Private/PluginManager.cpp
+++ b/Engine/Source/Runtime/Projects/Private/PluginManager.cpp
@@ -422,7 +422,7 @@ void FPluginManager::ReadAllPlugins(TMap<FString, TSharedRef<FPlugin>>& Plugins,
 					if (ParentModule.Name == ChildModule.Name && ParentModule.Type == ChildModule.Type)
 					{
 						// we only need to whitelist the platform if the parent had a whitelist (otherwise, we could mistakenly remove all other platforms)
-						if (ParentModule.WhitelistPlatforms.Num() > 0)
+						if (ParentModule.bHasExplicitPlatforms || ParentModule.WhitelistPlatforms.Num() > 0)
 						{
 							ParentModule.WhitelistPlatforms.Append(ChildModule.WhitelistPlatforms);
 						}
@@ -433,7 +433,7 @@ void FPluginManager::ReadAllPlugins(TMap<FString, TSharedRef<FPlugin>>& Plugins,
 				}
 			}
 
-			if (Parent->GetDescriptor().SupportedTargetPlatforms.Num() != 0)
+			if (Parent->GetDescriptor().bHasExplicitPlatforms || Parent->GetDescriptor().SupportedTargetPlatforms.Num() != 0)
 			{
 				for (const FString& SupportedTargetPlatform : Child->GetDescriptor().SupportedTargetPlatforms)
 				{
diff --git a/Engine/Source/Runtime/Projects/Private/PluginReferenceDescriptor.cpp b/Engine/Source/Runtime/Projects/Private/PluginReferenceDescriptor.cpp
--- a/Engine/Source/Runtime/Projects/Private/PluginReferenceDescriptor.cpp
+++ b/Engine/Source/Runtime/Projects/Private/PluginReferenceDescriptor.cpp
@@ -31,6 +31,7 @@ FPluginReferenceDescriptor::FPluginReferenceDescriptor( const FString& InName, b
 	: Name(InName)
 	, bEnabled(bInEnabled)
 	, bOptional(false)
+	, bHasExplicitPlatforms(false)
 { }
 
 
@@ -43,7 +44,7 @@ bool FPluginReferenceDescriptor::IsEnabledForPlatform( const FString& Platform )
 	}
 
 	// If there is a list of whitelisted platforms, and this isn't one of them, return false
-	if(WhitelistPlatforms.Num() > 0 && !WhitelistPlatforms.Contains(Platform))
+	if( (bHasExplicitPlatforms || WhitelistPlatforms.Num() > 0) && !WhitelistPlatforms.Contains(Platform))
 	{
 		return false;
 	}
@@ -105,7 +106,14 @@ bool FPluginReferenceDescriptor::IsEnabledForTargetConfiguration(EBuildConfigura
 
 bool FPluginReferenceDescriptor::IsSupportedTargetPlatform(const FString& Platform) const
 {
-	return SupportedTargetPlatforms.Num() == 0 || SupportedTargetPlatforms.Contains(Platform);
+	if (bHasExplicitPlatforms)
+	{
+		return SupportedTargetPlatforms.Contains(Platform);
+	}
+	else
+	{
+		return SupportedTargetPlatforms.Num() == 0 || SupportedTargetPlatforms.Contains(Platform);
+	}
 }
 
 bool FPluginReferenceDescriptor::Read( const FJsonObject& Object, FText& OutFailReason )
@@ -145,6 +153,7 @@ bool FPluginReferenceDescriptor::Read( const FJsonObject& Object, FText& OutFail
 
 	// Get the supported platform list
 	Object.TryGetStringArrayField(TEXT("SupportedTargetPlatforms"), SupportedTargetPlatforms);
+	Object.TryGetBoolField(TEXT("HasExplicitPlatforms"), bHasExplicitPlatforms);
 
 	return true;
 }
@@ -315,6 +324,15 @@ void FPluginReferenceDescriptor::UpdateJson(FJsonObject& JsonObject) const
 	{
 		JsonObject.RemoveField(TEXT("SupportedTargetPlatforms"));
 	}
+
+	if (bHasExplicitPlatforms)
+	{
+		JsonObject.SetBoolField(TEXT("HasExplicitPlatforms"), bHasExplicitPlatforms);
+	}
+	else
+	{
+		JsonObject.RemoveField(TEXT("HasExplicitPlatforms"));
+	}
 }
 
 void FPluginReferenceDescriptor::WriteArray(TJsonWriter<>& Writer, const TCHAR* ArrayName, const TArray<FPluginReferenceDescriptor>& Plugins)
diff --git a/Engine/Source/Runtime/Projects/Public/ModuleDescriptor.h b/Engine/Source/Runtime/Projects/Public/ModuleDescriptor.h
--- a/Engine/Source/Runtime/Projects/Public/ModuleDescriptor.h
+++ b/Engine/Source/Runtime/Projects/Public/ModuleDescriptor.h
@@ -179,6 +179,10 @@ struct PROJECTS_API FModuleDescriptor
 	/** List of additional dependencies for building this module. */
 	TArray<FString> AdditionalDependencies;
 
+	/** When true, empty WhitelistPlatforms are interpeted as 'no platforms' with the expectation that explict platforms will be added in plugin extensions */
+	bool bHasExplicitPlatforms;
+
+
 	/** Normal constructor */
 	FModuleDescriptor(const FName InName = NAME_None, EHostType::Type InType = EHostType::Runtime, ELoadingPhase::Type InLoadingPhase = ELoadingPhase::Default);
 
diff --git a/Engine/Source/Runtime/Projects/Public/PluginDescriptor.h b/Engine/Source/Runtime/Projects/Public/PluginDescriptor.h
--- a/Engine/Source/Runtime/Projects/Public/PluginDescriptor.h
+++ b/Engine/Source/Runtime/Projects/Public/PluginDescriptor.h
@@ -105,6 +105,9 @@ struct PROJECTS_API FPluginDescriptor
 	/** When true, this plugin's modules will not be loaded automatically nor will it's content be mounted automatically. It will load/mount when explicitly requested and LoadingPhases will be ignored */
 	bool bExplicitlyLoaded;
 
+	/** When true, an empty SupportedTargetPlatforms is interpeted as 'no platforms' with the expectation that explict platforms will be added in plugin platform extensions */
+	bool bHasExplicitPlatforms;
+
 	/** If true, this plugin from a platform extension extending another plugin */
 	bool bIsPluginExtension;
 
diff --git a/Engine/Source/Runtime/Projects/Public/PluginReferenceDescriptor.h b/Engine/Source/Runtime/Projects/Public/PluginReferenceDescriptor.h
--- a/Engine/Source/Runtime/Projects/Public/PluginReferenceDescriptor.h
+++ b/Engine/Source/Runtime/Projects/Public/PluginReferenceDescriptor.h
@@ -48,6 +48,10 @@ struct PROJECTS_API FPluginReferenceDescriptor
 	/** The list of supported target platforms for this plugin. This field is copied from the plugin descriptor, and supplements the user's whitelisted and blacklisted platforms. */
 	TArray<FString> SupportedTargetPlatforms;
 
+	/** When true, empty SupportedTargetPlatforms and WhitelistPlatforms are interpeted as 'no platforms' with the expectation that explict platforms will be added in plugin platform extensions */
+	bool bHasExplicitPlatforms;
+
+
 	/** Constructor */
 	FPluginReferenceDescriptor(const FString& InName = TEXT(""), bool bInEnabled = false);
 
