Remove iOS 11 and 12 references

diff --git a/Engine/Plugins/Online/IOS/OnlineSubsystemIOS/Source/Private/OnlineFriendsInterfaceIOS.cpp b/Engine/Plugins/Online/IOS/OnlineSubsystemIOS/Source/Private/OnlineFriendsInterfaceIOS.cpp
--- a/Engine/Plugins/Online/IOS/OnlineSubsystemIOS/Source/Private/OnlineFriendsInterfaceIOS.cpp
+++ b/Engine/Plugins/Online/IOS/OnlineSubsystemIOS/Source/Private/OnlineFriendsInterfaceIOS.cpp
@@ -53,94 +53,8 @@ FOnlineFriendsIOS::FOnlineFriendsIOS(FOnlineSubsystemIOS* InSubsystem):IOSSubsys
 bool FOnlineFriendsIOS::ReadFriendsList(int32 LocalUserNum, const FString& ListName, const FOnReadFriendsListComplete& Delegate /*= FOnReadFriendsListComplete()*/)
 {
 	UE_LOG_ONLINE_FRIEND(Verbose, TEXT("FOnlineFriendsIOS::ReadFriendsList()"));
-#if __IPHONE_OS_VERSION_MIN_REQUIRED < __IPHONE_10_0
-	bool bSuccessfullyBeganReadFriends = false;
-
-	if( IdentityInterface->GetLocalGameCenterUser() != NULL && IdentityInterface->GetLocalGameCenterUser().isAuthenticated )
-	{
-		bSuccessfullyBeganReadFriends = true;
-		dispatch_async(dispatch_get_main_queue(), ^
-		{
-		// Get the friends list for the local player from the server
-            if ([[GKLocalPlayer localPlayer] respondsToSelector:@selector(loadFriendPlayersWithCompletionHandler)] == YES)
-            {
-				PRAGMA_DISABLE_DEPRECATION_WARNINGS
-                [[GKLocalPlayer localPlayer] loadFriendPlayersWithCompletionHandler:
-				PRAGMA_ENABLE_DEPRECATION_WARNINGS
-                 ^(NSArray* Friends, NSError* Error)
-                 {
-                    if( Error )
-                    {
-                        FString ErrorStr(FString::Printf(TEXT("FOnlineFriendsIOS::ReadFriendsList() - Failed to read friends list with error: [%i]"), [Error code]));
-                        UE_LOG_ONLINE_FRIEND(Verbose, TEXT("%s"), *ErrorStr);
-                 
-                        // Report back to the game thread whether this succeeded.
-                        [FIOSAsyncTask CreateTaskWithBlock : ^ bool(void)
-                         {
-                            Delegate.ExecuteIfBound(0, false, ListName, ErrorStr);
-                            return true;
-                         }];
-                    }
-                    else
-                    {
-                        [GKPlayer loadPlayersForIdentifiers:Friends withCompletionHandler:
-                         ^(NSArray* Players, NSError* Error2)
-                         {
-                            FString ErrorStr;
-                            bool bWasSuccessful = false;
-                            if( Error2 )
-                            {
-                                ErrorStr = FString::Printf(TEXT("FOnlineFriendsIOS::ReadFriendsList() - Failed to loadPlayersForIdentifiers with error: [%i]"), [Error2 code]);
-                            }
-                            else
-                            {
-                                // Clear our previosly cached friends before we repopulate the cache.
-                                CachedFriends.Empty();
-                                for( int32 FriendIdx = 0; FriendIdx < [Players count]; FriendIdx++ )
-                                {
-                                    GKPlayer* Friend = Players[ FriendIdx ];
-                  
-                                    // Add new friend entry to list
-                                    TSharedRef<FOnlineFriendIOS> FriendEntry(
-										new FOnlineFriendIOS(*FString(FOnlineSubsystemIOS::GetPlayerId(Friend)))
-                                    );
-                                    FriendEntry->AccountData.Add(
-                                        TEXT("nickname"), *FString(Friend.alias)
-                                    );
-                                    CachedFriends.Add(FriendEntry);
-                  
-                                    UE_LOG_ONLINE_FRIEND(Verbose, TEXT("GCFriend - Id:%s Alias:%s"),
-										*FString(FOnlineSubsystemIOS::GetPlayerId(Friend)),
-                                        *FriendEntry->GetDisplayName()
-									);
-                                }
-                  
-                                bWasSuccessful = true;
-                            }
-                  
-                            // Report back to the game thread whether this succeeded.
-                            [FIOSAsyncTask CreateTaskWithBlock : ^ bool(void)
-                             {
-                                Delegate.ExecuteIfBound(0, bWasSuccessful, ListName, ErrorStr);
-                                return true;
-                             }];
-                         }];
-                    }
-                 }];
-            }
-		});
-	}
-	else
-	{
-		// no gamecenter login means we cannot read the friends.
-		Delegate.ExecuteIfBound( 0 , false, ListName, TEXT("not logged in") );
-	}
-
-	return bSuccessfullyBeganReadFriends;
-#else
 	Delegate.ExecuteIfBound(LocalUserNum, false, ListName, FString(TEXT("ReadFriendsList() is not supported as of UE 4.21")));
 	return false;
-#endif
 }
 
 bool FOnlineFriendsIOS::DeleteFriendsList(int32 LocalUserNum, const FString& ListName, const FOnDeleteFriendsListComplete& Delegate /*= FOnDeleteFriendsListComplete()*/)
diff --git a/Engine/Plugins/Online/IOS/OnlineSubsystemIOS/Source/Private/OnlineIdentityInterfaceIOS.cpp b/Engine/Plugins/Online/IOS/OnlineSubsystemIOS/Source/Private/OnlineIdentityInterfaceIOS.cpp
--- a/Engine/Plugins/Online/IOS/OnlineSubsystemIOS/Source/Private/OnlineIdentityInterfaceIOS.cpp
+++ b/Engine/Plugins/Online/IOS/OnlineSubsystemIOS/Source/Private/OnlineIdentityInterfaceIOS.cpp
@@ -7,12 +7,6 @@
 #include "IOS/IOSAppDelegate.h"
 #import "OnlineAppStoreUtils.h"
 
-#if (defined(__IPHONE_13_0) && __IPHONE_OS_VERSION_MAX_ALLOWED >= __IPHONE_13_0) || (defined(__TVOS_13_0) && __TV_OS_VERSION_MAX_ALLOWED >= __TVOS_13_0)
-#define SUPPORTS_PERSISTENT_SCOPEIDS 1
-#else
-#define SUPPORTS_PERSISTENT_SCOPEIDS 0
-#endif
-
 FOnlineIdentityIOS::FOnlineIdentityIOS()
 	: UniqueNetId(nullptr)
 {
@@ -64,8 +58,6 @@ bool FOnlineIdentityIOS::Login(int32 LocalUserNum, const FOnlineAccountCredentia
 	{
 		// Now logged in
 		bStartedLogin = true;
-		
-#if SUPPORTS_PERSISTENT_SCOPEIDS
 		if ([GKPlayer respondsToSelector:@selector(scopedIDsArePersistent)] == YES)
 		{
 			if ([GKLocalUser scopedIDsArePersistent])
@@ -87,7 +79,6 @@ bool FOnlineIdentityIOS::Login(int32 LocalUserNum, const FOnlineAccountCredentia
 			}
 		}
 		else
-#endif
 		{
 			const FString PlayerId(FString(FOnlineSubsystemIOS::GetPlayerId(GKLocalUser)));
 
@@ -126,7 +117,6 @@ bool FOnlineIdentityIOS::Login(int32 LocalUserNum, const FOnlineAccountCredentia
 						GKLocalPlayer* GKLocalUserAuth = [GKLocalPlayer localPlayer];
 
 						/* Perform additional tasks for the authenticated player here */
-#if SUPPORTS_PERSISTENT_SCOPEIDS
 						if ([GKPlayer respondsToSelector:@selector(scopedIDsArePersistent)] == YES)
 						{
 							if ([GKLocalUserAuth scopedIDsArePersistent])
@@ -144,7 +134,6 @@ bool FOnlineIdentityIOS::Login(int32 LocalUserNum, const FOnlineAccountCredentia
 							}
 						}
 						else
-#endif
 						{
 							PlayerId = FString(FOnlineSubsystemIOS::GetPlayerId(GKLocalUserAuth));
 
diff --git a/Engine/Plugins/Online/OnlineSubsystemApple/Source/Private/OnlineExternalUIInterfaceApple.cpp b/Engine/Plugins/Online/OnlineSubsystemApple/Source/Private/OnlineExternalUIInterfaceApple.cpp
--- a/Engine/Plugins/Online/OnlineSubsystemApple/Source/Private/OnlineExternalUIInterfaceApple.cpp
+++ b/Engine/Plugins/Online/OnlineSubsystemApple/Source/Private/OnlineExternalUIInterfaceApple.cpp
@@ -14,11 +14,7 @@
 #import <Foundation/Foundation.h>
 
 #if ONLINESUBSYSTEMAPPLE_IDENTITY_ENABLE_SIWA
-	#if defined(__MAC_10_15) || (defined(__IPHONE_13_0) && __IPHONE_OS_VERSION_MAX_ALLOWED >= __IPHONE_13_0) || (defined(__TVOS_13_0) && __TV_OS_VERSION_MAX_ALLOWED >= __TVOS_13_0)
-		#define SIWA_SUPPORTED 1
-	#else
-		#define SIWA_SUPPORTED 0
-	#endif
+	#define SIWA_SUPPORTED 1
 #else
 	#define SIWA_SUPPORTED 0
 #endif
diff --git a/Engine/Plugins/Online/OnlineSubsystemApple/Source/Private/OnlineIdentityInterfaceApple.cpp b/Engine/Plugins/Online/OnlineSubsystemApple/Source/Private/OnlineIdentityInterfaceApple.cpp
--- a/Engine/Plugins/Online/OnlineSubsystemApple/Source/Private/OnlineIdentityInterfaceApple.cpp
+++ b/Engine/Plugins/Online/OnlineSubsystemApple/Source/Private/OnlineIdentityInterfaceApple.cpp
@@ -6,11 +6,7 @@
 #include "OnlineError.h"
 
 #if ONLINESUBSYSTEMAPPLE_IDENTITY_ENABLE_SIWA
-	#if defined(__MAC_10_15) || (defined(__IPHONE_13_0) && __IPHONE_OS_VERSION_MAX_ALLOWED >= __IPHONE_13_0) || (defined(__TVOS_13_0) && __TV_OS_VERSION_MAX_ALLOWED >= __TVOS_13_0)
-		#define SIWA_SUPPORTED 1
-	#else
-		#define SIWA_SUPPORTED 0
-	#endif
+	#define SIWA_SUPPORTED 1
 #else
 	#define SIWA_SUPPORTED 0
 #endif
diff --git a/Engine/Source/Programs/UnrealBuildTool/Platform/IOS/IOSToolChain.cs b/Engine/Source/Programs/UnrealBuildTool/Platform/IOS/IOSToolChain.cs
--- a/Engine/Source/Programs/UnrealBuildTool/Platform/IOS/IOSToolChain.cs
+++ b/Engine/Source/Programs/UnrealBuildTool/Platform/IOS/IOSToolChain.cs
@@ -184,12 +184,14 @@ namespace UnrealBuildTool
 					FileReference AssetFile = FileReference.Combine(Binary.OutputFilePath.Directory, "AssetCatalog", "Assets.car");
 					BuildProducts.Add(AssetFile, BuildProductType.RequiredResource);
 				}
-				else if (Target.Platform == UnrealTargetPlatform.IOS && Settings.Value.IOSSDKVersionFloat >= 13.0f)
+				else if (Target.Platform == UnrealTargetPlatform.IOS)
 				{
 					int Index = Binary.OutputFilePath.GetFileNameWithoutExtension().IndexOf("-");
 					string OutputFile = Binary.OutputFilePath.GetFileNameWithoutExtension().Substring(0, Index > 0 ? Index : Binary.OutputFilePath.GetFileNameWithoutExtension().Length);
 					FileReference AssetFile = FileReference.Combine(Binary.OutputFilePath.Directory, "Payload", OutputFile + ".app", "Assets.car");
 					BuildProducts.Add(AssetFile, BuildProductType.RequiredResource);
+
+					// Default AppIcons need to be copied as is in the Payload's root, even when using an asset catalog.
 					AssetFile = FileReference.Combine(Binary.OutputFilePath.Directory, "Payload", OutputFile + ".app", "[REDACTED_EMAIL]");
 					BuildProducts.Add(AssetFile, BuildProductType.RequiredResource);
 					AssetFile = FileReference.Combine(Binary.OutputFilePath.Directory, "Payload", OutputFile + ".app", "AppIcon76x76@2x~ipad.png");
@@ -1000,7 +1002,7 @@ namespace UnrealBuildTool
 				Arguments.Append(" --launch-image \"Launch Image\"");
 				Arguments.Append(" --filter-for-device-model AppleTV5,3");
 				Arguments.Append(" --target-device tv");
-				Arguments.Append(" --minimum-deployment-target 12.0");
+				Arguments.Append(" --minimum-deployment-target 13.0");
 				Arguments.Append(" --platform appletvos");
 			}
 			else
@@ -1009,7 +1011,7 @@ namespace UnrealBuildTool
 				Arguments.Append(" --product-type com.apple.product-type.application");
 				Arguments.Append(" --target-device iphone");
 				Arguments.Append(" --target-device ipad");
-				Arguments.Append(" --minimum-deployment-target 12.0");
+				Arguments.Append(" --minimum-deployment-target 13.0");
 				Arguments.Append(" --platform iphoneos");
 			}
 			Arguments.Append(" --enable-on-demand-resources YES");
@@ -1367,23 +1369,19 @@ namespace UnrealBuildTool
 				}
 				// copy the icons from the game directory if it has any
 				string[][] Images = {
-					new string []{ "[REDACTED_EMAIL]", "[REDACTED_EMAIL]", "Icon40.png" },
-					new string []{ "[REDACTED_EMAIL]", "[REDACTED_EMAIL]", "Icon60.png" },
-					new string []{ "[REDACTED_EMAIL]", "[REDACTED_EMAIL]", "Icon58.png" },
-					new string []{ "[REDACTED_EMAIL]", "[REDACTED_EMAIL]", "Icon87.png" },
-					new string []{ "[REDACTED_EMAIL]", "[REDACTED_EMAIL]", "Icon80.png" },
-					new string []{ "[REDACTED_EMAIL]", "[REDACTED_EMAIL]", "[REDACTED_EMAIL]", "Icon120.png" },
-					new string []{ "[REDACTED_EMAIL]", "[REDACTED_EMAIL]", "[REDACTED_EMAIL]", "Icon120.png" },
-					new string []{ "[REDACTED_EMAIL]", "[REDACTED_EMAIL]", "Icon180.png" },
-					new string []{ "IPadIcon20.png", "Icon20.png" },
-					new string []{ "[REDACTED_EMAIL]", "[REDACTED_EMAIL]", "Icon40.png" },
-					new string []{ "IPadIcon29.png", "Icon29.png" },
-					new string []{ "[REDACTED_EMAIL]", "[REDACTED_EMAIL]", "Icon58.png" },
-					new string []{ "IPadIcon40.png", "Icon40.png", "[REDACTED_EMAIL]" },
-					new string []{ "[REDACTED_EMAIL]", "Icon80.png", "[REDACTED_EMAIL]" },
-					new string []{ "IPadIcon76.png", "Icon76.png" },
-					new string []{ "[REDACTED_EMAIL]", "[REDACTED_EMAIL]", "Icon152.png" },
-					new string []{ "[REDACTED_EMAIL]", "[REDACTED_EMAIL]", "Icon167.png" },
+					new string []{ "[REDACTED_EMAIL]", "[REDACTED_EMAIL]" },
+					new string []{ "[REDACTED_EMAIL]", "[REDACTED_EMAIL]" },
+					new string []{ "[REDACTED_EMAIL]", "[REDACTED_EMAIL]" },
+					new string []{ "[REDACTED_EMAIL]", "[REDACTED_EMAIL]" },
+					new string []{ "[REDACTED_EMAIL]", "[REDACTED_EMAIL]" },
+					new string []{ "[REDACTED_EMAIL]", "[REDACTED_EMAIL]" },
+					new string []{ "[REDACTED_EMAIL]", "[REDACTED_EMAIL]" },
+					new string []{ "[REDACTED_EMAIL]", "[REDACTED_EMAIL]" },
+					new string []{ "[REDACTED_EMAIL]", "[REDACTED_EMAIL]"},
+					new string []{ "[REDACTED_EMAIL]", "[REDACTED_EMAIL]"},
+					new string []{ "[REDACTED_EMAIL]", "[REDACTED_EMAIL]" },
+					new string []{ "[REDACTED_EMAIL]", "[REDACTED_EMAIL]"},
+					new string []{ "[REDACTED_EMAIL]", "[REDACTED_EMAIL]"},
 					new string []{ "Icon1024.png", "Icon1024.png" },
 				};
 				Dir = Path.Combine(IntermediateDir, "Resources", "Assets.xcassets", "AppIcon.appiconset");
@@ -1392,14 +1390,6 @@ namespace UnrealBuildTool
 				for (int Index = 0; Index < Images.Length; ++Index)
 				{
 					string Image = Path.Combine((Directory.Exists(Path.Combine(BuildDir, "Resources", "Graphics")) ? (BuildDir) : (Path.Combine(EngineDir, "Build", "IOS"))), "Resources", "Graphics", Images[Index][1]);
-					if (!File.Exists(Image) && Images[Index].Length > 2)
-					{
-						Image = Path.Combine((Directory.Exists(Path.Combine(BuildDir, "Resources", "Graphics")) ? (BuildDir) : (Path.Combine(EngineDir, "Build", "IOS"))), "Resources", "Graphics", Images[Index][2]);
-					}
-					if (!File.Exists(Image) && Images[Index].Length > 3)
-					{
-						Image = Path.Combine((Directory.Exists(Path.Combine(BuildDir, "Resources", "Graphics")) ? (BuildDir) : (Path.Combine(EngineDir, "Build", "IOS"))), "Resources", "Graphics", Images[Index][3]);
-					}
 					if (File.Exists(Image))
 					{
 						bUserImagesExist |= Image.StartsWith(BuildResourcesGraphicsDir);
diff --git a/Engine/Source/Runtime/Apple/MetalRHI/Private/MetalLLM.cpp b/Engine/Source/Runtime/Apple/MetalRHI/Private/MetalLLM.cpp
--- a/Engine/Source/Runtime/Apple/MetalRHI/Private/MetalLLM.cpp
+++ b/Engine/Source/Runtime/Apple/MetalRHI/Private/MetalLLM.cpp
@@ -166,12 +166,10 @@ static mtlpp::SizeAndAlign TextureSizeAndAlign(mtlpp::TextureType TextureType, u
 			SizeAlign.Size = RHICalcTextureCubePlatformSize(Width, MetalToRHIPixelFormat(Format), MipCount, TexCreate_None, CreateInfo, Align);
 			SizeAlign.Align = Align;
 			break;
-#if PLATFORM_MAC || __IPHONE_OS_VERSION_MAX_ALLOWED >= 110000
 		case mtlpp::TextureType::TextureCubeArray:
 			SizeAlign.Size = RHICalcTextureCubePlatformSize(Width, MetalToRHIPixelFormat(Format), MipCount, TexCreate_None, CreateInfo, Align) * ArrayCount;
 			SizeAlign.Align = Align;
 			break;
-#endif
 		case mtlpp::TextureType::Texture3D:
 			SizeAlign.Size = RHICalcTexture3DPlatformSize(Width, Height, Depth, MetalToRHIPixelFormat(Format), MipCount, TexCreate_None, CreateInfo, Align);
 			SizeAlign.Align = Align;
@@ -189,16 +187,10 @@ static mtlpp::SizeAndAlign TextureSizeAndAlign(mtlpp::TextureType TextureType, u
 void MetalLLM::LogAllocTexture(mtlpp::Device& Device, mtlpp::TextureDescriptor const& Desc, mtlpp::Texture const& Texture)
 {
 	mtlpp::SizeAndAlign SizeAlign;
-#if PLATFORM_MAC || __IPHONE_OS_VERSION_MAX_ALLOWED >= 110000
 	if (FMetalCommandQueue::SupportsFeature(EMetalFeaturesGPUCaptureManager))
 	{
 		SizeAlign = Device.HeapTextureSizeAndAlign(Desc);
 	}
-	else
-#endif
-	{
-		SizeAlign = TextureSizeAndAlign(Desc.GetTextureType(), Desc.GetWidth(), Desc.GetHeight(), Desc.GetDepth(), Desc.GetPixelFormat(), Desc.GetMipmapLevelCount(), Desc.GetSampleCount(), Desc.GetArrayLength());
-	}
 	
 	void* Ptr = (void*)Texture.GetPtr();
 	uint64 Size = SizeAlign.Size;
diff --git a/Engine/Source/Runtime/Apple/MetalRHI/Private/MetalRHIPrivate.h b/Engine/Source/Runtime/Apple/MetalRHI/Private/MetalRHIPrivate.h
--- a/Engine/Source/Runtime/Apple/MetalRHI/Private/MetalRHIPrivate.h
+++ b/Engine/Source/Runtime/Apple/MetalRHI/Private/MetalRHIPrivate.h
@@ -71,9 +71,9 @@ enum EMTLTextureType
 };
 
 // This is the right VERSION check, see Availability.h in the SDK
-#define METAL_SUPPORTS_INDIRECT_ARGUMENT_BUFFERS ((PLATFORM_MAC && __MAC_OS_X_VERSION_MAX_ALLOWED >= 101300) || (!PLATFORM_MAC && __IPHONE_OS_VERSION_MAX_ALLOWED >= 110000)) && (__clang_major__ >= 9)
-#define METAL_SUPPORTS_CAPTURE_MANAGER (PLATFORM_MAC && __MAC_OS_X_VERSION_MAX_ALLOWED >= 101300) || (!PLATFORM_MAC && __IPHONE_OS_VERSION_MAX_ALLOWED >= 110000) && (__clang_major__ >= 9)
-#define METAL_SUPPORTS_TILE_SHADERS (!PLATFORM_MAC && !PLATFORM_TVOS && __IPHONE_OS_VERSION_MAX_ALLOWED >= 110000) && (__clang_major__ >= 9)
+#define METAL_SUPPORTS_INDIRECT_ARGUMENT_BUFFERS 1
+#define METAL_SUPPORTS_CAPTURE_MANAGER 1
+#define METAL_SUPPORTS_TILE_SHADERS 1
 // In addition to compile-time SDK checks we also need a way to check if these are available on runtime
 extern bool GMetalSupportsCaptureManager;
 
