Allow APK generation for UnrealGame when using AGDE, to debug Blueprint projects

diff --git a/Engine/Source/Programs/UnrealBuildTool/Platform/Android/UEDeployAndroid.cs b/Engine/Source/Programs/UnrealBuildTool/Platform/Android/UEDeployAndroid.cs
--- a/Engine/Source/Programs/UnrealBuildTool/Platform/Android/UEDeployAndroid.cs
+++ b/Engine/Source/Programs/UnrealBuildTool/Platform/Android/UEDeployAndroid.cs
@@ -3581,7 +3581,7 @@ namespace UnrealBuildTool
 			List<string> GPUArchitectures = ToolChain.GetAllGPUArchitectures();
 
 			// we do not need to really build an engine UnrealGame.apk so short-circuit it
-			if (ProjectName == "UnrealGame" && OutputPath.Replace("\\", "/").Contains("/Engine/Binaries/Android/") && Path.GetFileNameWithoutExtension(OutputPath).StartsWith("UnrealGame"))
+			if (!ForceAPKGeneration && ProjectName == "UnrealGame" && OutputPath.Replace("\\", "/").Contains("/Engine/Binaries/Android/") && Path.GetFileNameWithoutExtension(OutputPath).StartsWith("UnrealGame"))
 			{
 				if (!bSkipGradleBuild)
 				{
