PR #8076: Fix 4668 error when using Fastbuild (Contributed by [REDACTED_PERSON])

#rnx

diff --git a/Engine/Source/Programs/UnrealBuildTool/Executors/Experimental/FASTBuild.cs b/Engine/Source/Programs/UnrealBuildTool/Executors/Experimental/FASTBuild.cs
--- a/Engine/Source/Programs/UnrealBuildTool/Executors/Experimental/FASTBuild.cs
+++ b/Engine/Source/Programs/UnrealBuildTool/Executors/Experimental/FASTBuild.cs
@@ -738,7 +738,7 @@ namespace UnrealBuildTool
 				{
 					ParsedCompilerOptions["InputFile"] = Token;
 					ProcessedTokens.RemoveAt(i);
-					break;
+					i--;
 				}
 			}
 
