Fix for single file compile via UnrealVS

diff --git a/Engine/Source/Programs/UnrealBuildTool/Modes/BuildMode.cs b/Engine/Source/Programs/UnrealBuildTool/Modes/BuildMode.cs
--- a/Engine/Source/Programs/UnrealBuildTool/Modes/BuildMode.cs
+++ b/Engine/Source/Programs/UnrealBuildTool/Modes/BuildMode.cs
@@ -167,7 +167,7 @@ namespace UnrealBuildTool
 				// Hack for specific files compile; don't build the ShaderCompileWorker target that's added to the command line for generated project files
 				if(TargetDescriptors.Count >= 2)
 				{
-					TargetDescriptors.RemoveAll(x => (x.Name == "ShaderCompileWorker" || x.Name == "LiveCodingConsole") && x.SpecificFilesToCompile.Count > 0);
+					TargetDescriptors.RemoveAll(x => (x.Name == "ShaderCompileWorker" || x.Name == "LiveCodingConsole" || x.Name == "InterchangeWorker") && x.SpecificFilesToCompile.Count > 0);
 				}
 
 				// Handle remote builds
