[Workspace] - add machinename to workspace default name, and lowercase the whole name to fix UE-115727
- this allows 2 machines with the same user & project to share hw

diff --git a/Engine/Source/Programs/UnrealBuildTool/System/DeployExports.cs b/Engine/Source/Programs/UnrealBuildTool/System/DeployExports.cs
--- a/Engine/Source/Programs/UnrealBuildTool/System/DeployExports.cs
+++ b/Engine/Source/Programs/UnrealBuildTool/System/DeployExports.cs
@@ -17,20 +17,21 @@ namespace UnrealBuildTool
 		/// <summary>
 		/// The max length of the deploy folder name.
 		/// </summary>
-		public const int DeployFolderMaxLength = 50; 
+		public const int DeployFolderMaxLength = 50;
 
 		/// <summary>
-		/// Gets the default deploy folder name on the format ShortProjectName-RootDirectoryName-UserName,
+		/// Gets the default deploy folder name in the format RootDirectoryName-ShortProjectName-UserName-MachineNameHash,
 		/// truncated to DeployFolderMaxLength characters, with dots '.' and whitespace ' ' stripped.
 		/// </summary>
 		/// <param name="ShortProjectName">Project to deploy</param>
 		/// <returns>The default deploy folder name.</returns>
 		public static string GetDefaultDeployFolder(string ShortProjectName)
 		{
-			string DeployFolder = string.Format("{0}-{1}-{2}",
-				ShortProjectName,
+			string DeployFolder = string.Format("{0}-{1}-{2}-{3}",
 				UnrealBuildTool.RootDirectory.GetDirectoryName(),
-				Environment.UserName);
+				ShortProjectName,
+				Environment.UserName,
+				Environment.MachineName.GetHashCode().ToString()).ToLower();
 
 			DeployFolder = DeployFolder.Replace(" ", "");
 			DeployFolder = DeployFolder.Replace(".", "");
