[DEPLOY] - optionally append machine name onto the deployed PackageName to allow multiple host machines to share hw.

diff --git a/Engine/Source/Programs/UnrealBuildTool/Platform/Windows/VCManifestGenerator.cs b/Engine/Source/Programs/UnrealBuildTool/Platform/Windows/VCManifestGenerator.cs
--- a/Engine/Source/Programs/UnrealBuildTool/Platform/Windows/VCManifestGenerator.cs
+++ b/Engine/Source/Programs/UnrealBuildTool/Platform/Windows/VCManifestGenerator.cs
@@ -601,6 +601,14 @@ namespace UnrealBuildTool
                 Log.TraceError("Consider using the setting [{0}]:PackageName to provide a specific value.", IniSection_PlatformTargetSettings);
             }
 
+			// If specified in the project settings append the users machine name onto the package name to allow sharing of devkits without stomping of deploys
+			bool bPackageNameUseMachineName;
+			if (EngineIni.GetBool(IniSection_PlatformTargetSettings, "bPackageNameUseMachineName", out bPackageNameUseMachineName) && bPackageNameUseMachineName)
+			{
+				var MachineName = Regex.Replace(Environment.MachineName.ToString(), "[^-.A-Za-z0-9]", "");
+				PackageName = PackageName + ".NOT.SHIPPABLE." + MachineName;
+			}
+
 			return PackageName;
 		}
 
