Refactor iOS Icons

diff --git a/Engine/Build/IOS/Resources/Assets.xcassets/AppIcon.appiconset/Contents.json b/Engine/Build/IOS/Resources/Assets.xcassets/AppIcon.appiconset/Contents.json
--- a/Engine/Build/IOS/Resources/Assets.xcassets/AppIcon.appiconset/Contents.json
+++ b/Engine/Build/IOS/Resources/Assets.xcassets/AppIcon.appiconset/Contents.json
@@ -48,48 +48,24 @@
       "filename" : "[REDACTED_EMAIL]",
       "scale" : "3x"
     },
-    {
-      "size" : "20x20",
-      "idiom" : "ipad",
-      "filename" : "IPadIcon20.png",
-      "scale" : "1x"
-    },
     {
       "size" : "20x20",
       "idiom" : "ipad",
       "filename" : "[REDACTED_EMAIL]",
       "scale" : "2x"
     },
-    {
-      "size" : "29x29",
-      "idiom" : "ipad",
-      "filename" : "IPadIcon29.png",
-      "scale" : "1x"
-    },
     {
       "size" : "29x29",
       "idiom" : "ipad",
       "filename" : "[REDACTED_EMAIL]",
       "scale" : "2x"
     },
-    {
-      "size" : "40x40",
-      "idiom" : "ipad",
-      "filename" : "IPadIcon40.png",
-      "scale" : "1x"
-    },
     {
       "size" : "40x40",
       "idiom" : "ipad",
       "filename" : "[REDACTED_EMAIL]",
       "scale" : "2x"
     },
-    {
-      "size" : "76x76",
-      "idiom" : "ipad",
-      "filename" : "IPadIcon76.png",
-      "scale" : "1x"
-    },
     {
       "size" : "76x76",
       "idiom" : "ipad",
diff --git a/Engine/Source/Programs/AutomationTool/IOS/IOSPlatform.Automation.cs b/Engine/Source/Programs/AutomationTool/IOS/IOSPlatform.Automation.cs
--- a/Engine/Source/Programs/AutomationTool/IOS/IOSPlatform.Automation.cs
+++ b/Engine/Source/Programs/AutomationTool/IOS/IOSPlatform.Automation.cs
@@ -622,10 +622,10 @@ public class IOSPlatform : Platform
 		IOSExports.GetProvisioningData(InProject, bDistribution, out MobileProvision, out SigningCertificate, out TeamUUID, out bAutomaticSigning);
 	}
 
-	public virtual bool DeployGeneratePList(FileReference ProjectFile, UnrealTargetConfiguration Config, DirectoryReference ProjectDirectory, bool bIsUE4Game, string GameName, bool bIsClient, string ProjectName, DirectoryReference InEngineDir, DirectoryReference AppDirectory, string InExecutablePath, out bool bSupportsPortrait, out bool bSupportsLandscape, out bool bSkipIcons)
+	public virtual bool DeployGeneratePList(FileReference ProjectFile, UnrealTargetConfiguration Config, DirectoryReference ProjectDirectory, bool bIsUE4Game, string GameName, bool bIsClient, string ProjectName, DirectoryReference InEngineDir, DirectoryReference AppDirectory, string InExecutablePath, out bool bSupportsPortrait, out bool bSupportsLandscape)
 	{
 		FileReference TargetReceiptFileName = GetTargetReceiptFileName(Config, InExecutablePath, InEngineDir, ProjectDirectory, bIsUE4Game);
-		return IOSExports.GeneratePList(ProjectFile, Config, ProjectDirectory, bIsUE4Game, GameName, bIsClient, ProjectName, InEngineDir, AppDirectory, TargetReceiptFileName, out bSupportsPortrait, out bSupportsLandscape, out bSkipIcons);
+		return IOSExports.GeneratePList(ProjectFile, Config, ProjectDirectory, bIsUE4Game, GameName, bIsClient, ProjectName, InEngineDir, AppDirectory, TargetReceiptFileName, out bSupportsPortrait, out bSupportsLandscape);
 	}
 
 	protected string MakeIPAFileName(UnrealTargetConfiguration TargetConfiguration, ProjectParams Params, DeploymentContext SC, bool bAllowDistroPrefix)
@@ -1339,7 +1339,6 @@ public class IOSPlatform : Platform
 					var TargetConfiguration = SC.StageTargetConfigurations[0];
 					bool bSupportsPortrait = false;
 					bool bSupportsLandscape = false;
-					bool bSkipIcons = false;
 
 					DeployGeneratePList(
 							SC.RawProjectPath,
@@ -1352,23 +1351,10 @@ public class IOSPlatform : Platform
 							DirectoryReference.Combine((SC.IsCodeBasedProject ? SC.ProjectRoot : DirectoryReference.Combine(SC.LocalRoot, "Engine")), "Binaries", PlatformName, "Payload", (SC.IsCodeBasedProject ? SC.ShortProjectName : "UnrealGame") + ".app"),
 							SC.StageExecutables[0],
 							out bSupportsPortrait,
-							out bSupportsLandscape,
-							out bSkipIcons);
+							out bSupportsLandscape);
 
 					// copy the plist to the stage dir
 					SC.StageFile(StagedFileType.SystemNonUFS, TargetPListFile, new StagedFileReference("Info.plist"));
-
-					// copy the icons from the engine
-					{
-						DirectoryReference GraphicsDataPath = DirectoryReference.Combine(SC.EngineRoot, "Build", "IOS", "Resources", "Graphics");
-						StageImageAndIconFiles(Params, GraphicsDataPath, bSupportsPortrait, bSupportsLandscape, SC, bSkipIcons);
-					}
-
-					// copy the icons from the game (may stomp the engine copies)
-					{
-						DirectoryReference GraphicsDataPath = DirectoryReference.Combine(SC.ProjectRoot, "Build", "IOS", "Resources", "Graphics");
-						StageImageAndIconFiles(Params, GraphicsDataPath, bSupportsPortrait, bSupportsLandscape, SC, bSkipIcons);
-					}
 				}
 
 				// copy the udebugsymbols if they exist
@@ -1427,17 +1413,6 @@ public class IOSPlatform : Platform
 		}
 	}
 
-	private void StageImageAndIconFiles(ProjectParams Params, DirectoryReference GraphicsDataPath, bool bSupportsPortrait, bool bSupportsLandscape, DeploymentContext SC, bool bSkipIcons)
-	{
-		if (DirectoryReference.Exists(GraphicsDataPath))
-		{
-			if (!bSkipIcons)
-			{
-				SC.StageFiles(StagedFileType.SystemNonUFS, GraphicsDataPath, "Icon*.png", StageFilesSearch.TopDirectoryOnly, StagedDirectoryReference.Root);
-			}
-		}
-	}
-
 	protected void StageMovieFiles(DirectoryReference InputDir, DeploymentContext SC)
 	{
 		if (DirectoryReference.Exists(InputDir))
diff --git a/Engine/Source/Programs/AutomationTool/TVOS/TVOSPlatform.Automation.cs b/Engine/Source/Programs/AutomationTool/TVOS/TVOSPlatform.Automation.cs
--- a/Engine/Source/Programs/AutomationTool/TVOS/TVOSPlatform.Automation.cs
+++ b/Engine/Source/Programs/AutomationTool/TVOS/TVOSPlatform.Automation.cs
@@ -37,7 +37,7 @@ public class TVOSPlatform : IOSPlatform
 		TVOSExports.GetProvisioningData(InProject, bDistribution, out MobileProvision, out SigningCertificate, out Team, out bAutomaticSigning);
     }
 
-	public override bool DeployGeneratePList(FileReference ProjectFile, UnrealTargetConfiguration Config, DirectoryReference ProjectDirectory, bool bIsUE4Game, string GameName, bool bIsClient, string ProjectName, DirectoryReference InEngineDir, DirectoryReference AppDirectory, string InExecutablePath, out bool bSupportsPortrait, out bool bSupportsLandscape, out bool bSkipIcons)
+	public override bool DeployGeneratePList(FileReference ProjectFile, UnrealTargetConfiguration Config, DirectoryReference ProjectDirectory, bool bIsUE4Game, string GameName, bool bIsClient, string ProjectName, DirectoryReference InEngineDir, DirectoryReference AppDirectory, string InExecutablePath, out bool bSupportsPortrait, out bool bSupportsLandscape)
 	{
 		string TargetName = Path.GetFileNameWithoutExtension(InExecutablePath).Split("-".ToCharArray())[0];
 		FileReference TargetReceiptFileName;
@@ -49,7 +49,7 @@ public class TVOSPlatform : IOSPlatform
 		{
 			TargetReceiptFileName = TargetReceipt.GetDefaultPath(ProjectDirectory, TargetName, UnrealTargetPlatform.TVOS, Config, "");
 		}
-		return TVOSExports.GeneratePList(ProjectFile, Config, ProjectDirectory, bIsUE4Game, GameName, bIsClient, ProjectName, InEngineDir, AppDirectory, TargetReceiptFileName, out bSupportsPortrait, out bSupportsLandscape, out bSkipIcons);
+		return TVOSExports.GeneratePList(ProjectFile, Config, ProjectDirectory, bIsUE4Game, GameName, bIsClient, ProjectName, InEngineDir, AppDirectory, TargetReceiptFileName, out bSupportsPortrait, out bSupportsLandscape);
 	}
 
     public override string GetCookPlatform(bool bDedicatedServer, bool bIsClientOnly)
@@ -111,8 +111,7 @@ public class TVOSPlatform : IOSPlatform
 
                     bool bSupportsPortrait = false;
                     bool bSupportsLandscape = false;
-                    bool bSkipIcons = false;
-                    DeployGeneratePList(SC.RawProjectPath, TargetConfiguration, (SC.IsCodeBasedProject ? SC.ProjectRoot : SC.EngineRoot), !SC.IsCodeBasedProject, (SC.IsCodeBasedProject ? SC.ShortProjectName : "UnrealGame"), Params.Client, SC.ShortProjectName, SC.EngineRoot, DirectoryReference.Combine((SC.IsCodeBasedProject ? SC.ProjectRoot : SC.EngineRoot), "Binaries", "TVOS", "Payload", (SC.IsCodeBasedProject ? SC.ShortProjectName : "UnrealGame") + ".app"), SC.StageExecutables[0], out bSupportsPortrait, out bSupportsLandscape, out bSkipIcons);
+                    DeployGeneratePList(SC.RawProjectPath, TargetConfiguration, (SC.IsCodeBasedProject ? SC.ProjectRoot : SC.EngineRoot), !SC.IsCodeBasedProject, (SC.IsCodeBasedProject ? SC.ShortProjectName : "UnrealGame"), Params.Client, SC.ShortProjectName, SC.EngineRoot, DirectoryReference.Combine((SC.IsCodeBasedProject ? SC.ProjectRoot : SC.EngineRoot), "Binaries", "TVOS", "Payload", (SC.IsCodeBasedProject ? SC.ShortProjectName : "UnrealGame") + ".app"), SC.StageExecutables[0], out bSupportsPortrait, out bSupportsLandscape);
                 }
 
 
diff --git a/Engine/Source/Programs/UnrealBuildTool/Platform/IOS/IOSExports.cs b/Engine/Source/Programs/UnrealBuildTool/Platform/IOS/IOSExports.cs
--- a/Engine/Source/Programs/UnrealBuildTool/Platform/IOS/IOSExports.cs
+++ b/Engine/Source/Programs/UnrealBuildTool/Platform/IOS/IOSExports.cs
@@ -102,12 +102,11 @@ namespace UnrealBuildTool
 		/// <param name="BuildReceiptFileName"></param>
 		/// <param name="bSupportsPortrait"></param>
 		/// <param name="bSupportsLandscape"></param>
-		/// <param name="bSkipIcons"></param>
 		/// <returns></returns>
-		public static bool GeneratePList(FileReference ProjectFile, UnrealTargetConfiguration Config, DirectoryReference ProjectDirectory, bool bIsUE4Game, string GameName, bool bIsClient, string ProjectName, DirectoryReference InEngineDir, DirectoryReference AppDirectory, FileReference BuildReceiptFileName, out bool bSupportsPortrait, out bool bSupportsLandscape, out bool bSkipIcons)
+		public static bool GeneratePList(FileReference ProjectFile, UnrealTargetConfiguration Config, DirectoryReference ProjectDirectory, bool bIsUE4Game, string GameName, bool bIsClient, string ProjectName, DirectoryReference InEngineDir, DirectoryReference AppDirectory, FileReference BuildReceiptFileName, out bool bSupportsPortrait, out bool bSupportsLandscape)
 		{
 			TargetReceipt Receipt = TargetReceipt.Read(BuildReceiptFileName);
-			return new UEDeployIOS().GeneratePList(ProjectFile, Config, ProjectDirectory.FullName, bIsUE4Game, GameName, bIsClient, ProjectName, InEngineDir.FullName, AppDirectory.FullName, Receipt, out bSupportsPortrait, out bSupportsLandscape, out bSkipIcons);
+			return new UEDeployIOS().GeneratePList(ProjectFile, Config, ProjectDirectory.FullName, bIsUE4Game, GameName, bIsClient, ProjectName, InEngineDir.FullName, AppDirectory.FullName, Receipt, out bSupportsPortrait, out bSupportsLandscape);
 		}
 
 		/// <summary>
diff --git a/Engine/Source/Programs/UnrealBuildTool/Platform/IOS/UEDeployIOS.cs b/Engine/Source/Programs/UnrealBuildTool/Platform/IOS/UEDeployIOS.cs
--- a/Engine/Source/Programs/UnrealBuildTool/Platform/IOS/UEDeployIOS.cs
+++ b/Engine/Source/Programs/UnrealBuildTool/Platform/IOS/UEDeployIOS.cs
@@ -216,7 +216,7 @@ namespace UnrealBuildTool
 			return MinVersionToReturn;
 		}
 
-		public static bool GenerateIOSPList(FileReference ProjectFile, UnrealTargetConfiguration Config, string ProjectDirectory, bool bIsUnrealGame, string GameName, bool bIsClient, string ProjectName, string InEngineDir, string AppDirectory, VersionNumber SdkVersion, UnrealPluginLanguage UPL, string BundleID, bool bBuildAsFramework, out bool bSupportsPortrait, out bool bSupportsLandscape, out bool bSkipIcons)
+		public static bool GenerateIOSPList(FileReference ProjectFile, UnrealTargetConfiguration Config, string ProjectDirectory, bool bIsUnrealGame, string GameName, bool bIsClient, string ProjectName, string InEngineDir, string AppDirectory, VersionNumber SdkVersion, UnrealPluginLanguage UPL, string BundleID, bool bBuildAsFramework, out bool bSupportsPortrait, out bool bSupportsLandscape)
 		{
 			// generate the Info.plist for future use
 			string BuildDirectory = ProjectDirectory + "/Build/IOS";
@@ -466,92 +466,35 @@ namespace UnrealBuildTool
 			}
 			Text.AppendLine("\t</array>");
 
-			if (SupportsIconCatalog(SdkVersion))
-			{
-				bSkipIcons = true;
-				Text.AppendLine("\t<key>CFBundleIcons</key>");
-				Text.AppendLine("\t<dict>");
-				Text.AppendLine("\t\t<key>CFBundlePrimaryIcon</key>");
-				Text.AppendLine("\t\t<dict>");
-				Text.AppendLine("\t\t\t<key>CFBundleIconFiles</key>");
-				Text.AppendLine("\t\t\t<array>");
-				Text.AppendLine("\t\t\t\t<string>AppIcon20x20</string>");
-				Text.AppendLine("\t\t\t\t<string>AppIcon29x29</string>");
-				Text.AppendLine("\t\t\t\t<string>AppIcon40x40</string>");
-				Text.AppendLine("\t\t\t\t<string>AppIcon60x60</string>");
-				Text.AppendLine("\t\t\t</array>");
-				Text.AppendLine("\t\t\t<key>CFBundleIconName</key>");
-				Text.AppendLine("\t\t\t<string>AppIcon</string>");
-				Text.AppendLine("\t\t\t<key>UIPrerenderedIcon</key>");
-				Text.AppendLine("\t\t\t<true/>");
-				Text.AppendLine("\t\t</dict>");
-				Text.AppendLine("\t</dict>");
-				Text.AppendLine("\t<key>CFBundleIcons~ipad</key>");
-				Text.AppendLine("\t<dict>");
-				Text.AppendLine("\t\t<key>CFBundlePrimaryIcon</key>");
-				Text.AppendLine("\t\t<dict>");
-				Text.AppendLine("\t\t\t<key>CFBundleIconFiles</key>");
-				Text.AppendLine("\t\t\t<array>");
-				Text.AppendLine("\t\t\t\t<string>AppIcon20x20</string>");
-				Text.AppendLine("\t\t\t\t<string>AppIcon29x29</string>");
-				Text.AppendLine("\t\t\t\t<string>AppIcon40x40</string>");
-				Text.AppendLine("\t\t\t\t<string>AppIcon60x60</string>");
-				Text.AppendLine("\t\t\t\t<string>AppIcon76x76</string>");
-				Text.AppendLine("\t\t\t\t<string>AppIcon83.5x83.5</string>");
-				Text.AppendLine("\t\t\t</array>");
-				Text.AppendLine("\t\t\t<key>CFBundleIconName</key>");
-				Text.AppendLine("\t\t\t<string>AppIcon</string>");
-				Text.AppendLine("\t\t\t<key>UIPrerenderedIcon</key>");
-				Text.AppendLine("\t\t\t<true/>");
-				Text.AppendLine("\t\t</dict>");
-				Text.AppendLine("\t</dict>");
-			}
-			else
-			{
-				bSkipIcons = false;
-				Text.AppendLine("\t<key>CFBundleIcons</key>");
-				Text.AppendLine("\t<dict>");
-				Text.AppendLine("\t\t<key>CFBundlePrimaryIcon</key>");
-				Text.AppendLine("\t\t<dict>");
-				Text.AppendLine("\t\t\t<key>CFBundleIconFiles</key>");
-				Text.AppendLine("\t\t\t<array>");
-				Text.AppendLine("\t\t\t\t<string>Icon29.png</string>");
-				Text.AppendLine("\t\t\t\t<string>[REDACTED_EMAIL]</string>");
-				Text.AppendLine("\t\t\t\t<string>Icon40.png</string>");
-				Text.AppendLine("\t\t\t\t<string>[REDACTED_EMAIL]</string>");
-				Text.AppendLine("\t\t\t\t<string>Icon57.png</string>");
-				Text.AppendLine("\t\t\t\t<string>[REDACTED_EMAIL]</string>");
-				Text.AppendLine("\t\t\t\t<string>[REDACTED_EMAIL]</string>");
-				Text.AppendLine("\t\t\t\t<string>[REDACTED_EMAIL]</string>");
-				Text.AppendLine("\t\t\t</array>");
-				Text.AppendLine("\t\t\t<key>UIPrerenderedIcon</key>");
-				Text.AppendLine("\t\t\t<true/>");
-				Text.AppendLine("\t\t</dict>");
-				Text.AppendLine("\t</dict>");
-				Text.AppendLine("\t<key>CFBundleIcons~ipad</key>");
-				Text.AppendLine("\t<dict>");
-				Text.AppendLine("\t\t<key>CFBundlePrimaryIcon</key>");
-				Text.AppendLine("\t\t<dict>");
-				Text.AppendLine("\t\t\t<key>CFBundleIconFiles</key>");
-				Text.AppendLine("\t\t\t<array>");
-				Text.AppendLine("\t\t\t\t<string>Icon29.png</string>");
-				Text.AppendLine("\t\t\t\t<string>[REDACTED_EMAIL]</string>");
-				Text.AppendLine("\t\t\t\t<string>Icon40.png</string>");
-				Text.AppendLine("\t\t\t\t<string>[REDACTED_EMAIL]</string>");
-				Text.AppendLine("\t\t\t\t<string>Icon50.png</string>");
-				Text.AppendLine("\t\t\t\t<string>[REDACTED_EMAIL]</string>");
-				Text.AppendLine("\t\t\t\t<string>Icon72.png</string>");
-				Text.AppendLine("\t\t\t\t<string>[REDACTED_EMAIL]</string>");
-				Text.AppendLine("\t\t\t\t<string>Icon76.png</string>");
-				Text.AppendLine("\t\t\t\t<string>[REDACTED_EMAIL]</string>");
-				Text.AppendLine("\t\t\t\t<string>[REDACTED_EMAIL]</string>");
-				Text.AppendLine("\t\t\t</array>");
-				Text.AppendLine("\t\t\t<key>UIPrerenderedIcon</key>");
-				Text.AppendLine("\t\t\t<true/>");
-				Text.AppendLine("\t\t</dict>");
-				Text.AppendLine("\t</dict>");
-			}
-
+			Text.AppendLine("\t<key>CFBundleIcons</key>");
+			Text.AppendLine("\t<dict>");
+			Text.AppendLine("\t\t<key>CFBundlePrimaryIcon</key>");
+			Text.AppendLine("\t\t<dict>");
+			Text.AppendLine("\t\t\t<key>CFBundleIconFiles</key>");
+			Text.AppendLine("\t\t\t<array>");
+			Text.AppendLine("\t\t\t\t<string>AppIcon60x60</string>");
+			Text.AppendLine("\t\t\t</array>");
+			Text.AppendLine("\t\t\t<key>CFBundleIconName</key>");
+			Text.AppendLine("\t\t\t<string>AppIcon</string>");
+			Text.AppendLine("\t\t\t<key>UIPrerenderedIcon</key>");
+			Text.AppendLine("\t\t\t<true/>");
+			Text.AppendLine("\t\t</dict>");
+			Text.AppendLine("\t</dict>");
+			Text.AppendLine("\t<key>CFBundleIcons~ipad</key>");
+			Text.AppendLine("\t<dict>");
+			Text.AppendLine("\t\t<key>CFBundlePrimaryIcon</key>");
+			Text.AppendLine("\t\t<dict>");
+			Text.AppendLine("\t\t\t<key>CFBundleIconFiles</key>");
+			Text.AppendLine("\t\t\t<array>");
+			Text.AppendLine("\t\t\t\t<string>AppIcon60x60</string>");
+			Text.AppendLine("\t\t\t\t<string>AppIcon76x76</string>");
+			Text.AppendLine("\t\t\t</array>");
+			Text.AppendLine("\t\t\t<key>CFBundleIconName</key>");
+			Text.AppendLine("\t\t\t<string>AppIcon</string>");
+			Text.AppendLine("\t\t\t<key>UIPrerenderedIcon</key>");
+			Text.AppendLine("\t\t\t<true/>");
+			Text.AppendLine("\t\t</dict>");
+			Text.AppendLine("\t</dict>");
 			Text.AppendLine("\t<key>UILaunchStoryboardName</key>");
 			Text.AppendLine("\t<string>LaunchScreen</string>");
 
@@ -710,11 +653,6 @@ namespace UnrealBuildTool
 			return true;
 		}
 
-        static bool SupportsIconCatalog(VersionNumber SdkVersion)
-        {
-			return SdkVersion != null && SdkVersion >= new VersionNumber(11);
-		}
-
         public static VersionNumber GetSdkVersion(TargetReceipt Receipt)
         {
 			VersionNumber SdkVersion = null;
@@ -742,15 +680,15 @@ namespace UnrealBuildTool
 			return false;
 		}
 
- 		public bool GeneratePList(FileReference ProjectFile, UnrealTargetConfiguration Config, string ProjectDirectory, bool bIsUnrealGame, string GameName, bool bIsClient, string ProjectName, string InEngineDir, string AppDirectory, TargetReceipt Receipt, out bool bSupportsPortrait, out bool bSupportsLandscape, out bool bSkipIcons)
+ 		public bool GeneratePList(FileReference ProjectFile, UnrealTargetConfiguration Config, string ProjectDirectory, bool bIsUnrealGame, string GameName, bool bIsClient, string ProjectName, string InEngineDir, string AppDirectory, TargetReceipt Receipt, out bool bSupportsPortrait, out bool bSupportsLandscape)
 		{
 			List<string> UPLScripts = CollectPluginDataPaths(Receipt.AdditionalProperties);
 			VersionNumber SdkVersion = GetSdkVersion(Receipt);
 			bool bBuildAsFramework = GetCompileAsDll(Receipt);
-			return GeneratePList(ProjectFile, Config, ProjectDirectory, bIsUnrealGame, GameName, bIsClient, ProjectName, InEngineDir, AppDirectory, UPLScripts, SdkVersion, "", bBuildAsFramework, out bSupportsPortrait, out bSupportsLandscape, out bSkipIcons);
+			return GeneratePList(ProjectFile, Config, ProjectDirectory, bIsUnrealGame, GameName, bIsClient, ProjectName, InEngineDir, AppDirectory, UPLScripts, SdkVersion, "", bBuildAsFramework, out bSupportsPortrait, out bSupportsLandscape);
 		}
 
-		public virtual bool GeneratePList(FileReference ProjectFile, UnrealTargetConfiguration Config, string ProjectDirectory, bool bIsUnrealGame, string GameName, bool bIsClient, string ProjectName, string InEngineDir, string AppDirectory, List<string> UPLScripts, VersionNumber SdkVersion, string BundleID, bool bBuildAsFramework, out bool bSupportsPortrait, out bool bSupportsLandscape, out bool bSkipIcons)
+		public virtual bool GeneratePList(FileReference ProjectFile, UnrealTargetConfiguration Config, string ProjectDirectory, bool bIsUnrealGame, string GameName, bool bIsClient, string ProjectName, string InEngineDir, string AppDirectory, List<string> UPLScripts, VersionNumber SdkVersion, string BundleID, bool bBuildAsFramework, out bool bSupportsPortrait, out bool bSupportsLandscape)
 		{
 			// remember name with -IOS-Shipping, etc
 			// string ExeName = GameName;
@@ -782,7 +720,7 @@ namespace UnrealBuildTool
 			// Passing in true for distribution is not ideal here but given the way that ios packaging happens and this call chain it seems unavoidable for now, maybe there is a way to correctly pass it in that I can't find?
 			UPL.Init(ProjectArches, true, RelativeEnginePath, BundlePath, ProjectDirectory, Config.ToString(), false);
 
-			return GenerateIOSPList(ProjectFile, Config, ProjectDirectory, bIsUnrealGame, GameName, bIsClient, ProjectName, InEngineDir, AppDirectory, SdkVersion, UPL, BundleID, bBuildAsFramework, out bSupportsPortrait, out bSupportsLandscape, out bSkipIcons);
+			return GenerateIOSPList(ProjectFile, Config, ProjectDirectory, bIsUnrealGame, GameName, bIsClient, ProjectName, InEngineDir, AppDirectory, SdkVersion, UPL, BundleID, bBuildAsFramework, out bSupportsPortrait, out bSupportsLandscape);
 		}
 
 		protected virtual void CopyCloudResources(string InEngineDir, string AppDirectory)
@@ -836,18 +774,6 @@ namespace UnrealBuildTool
 			}
 		}
 
-		protected virtual void CopyIconResources(string InEngineDir, string AppDirectory, string BuildDirectory)
-		{
-			CopyFiles(InEngineDir + "/Build/IOS/Resources/Graphics", AppDirectory, "Icon*.png", true);
-	
-			// merge game assets on top
-			// @todo tvos: Do we want to copy IOS and TVOS both in? (Engine/IOS -> Game/IOS -> Game/TVOS)?
-			if (Directory.Exists(BuildDirectory + "/Resources/Graphics"))
-			{
-					CopyFiles(BuildDirectory + "/Resources/Graphics", AppDirectory, "Icon*.png", true);
-			}
-		}
-
 		protected virtual void CopyLocalizationsResources(string InEngineDir, string AppDirectory, string BuildDirectory, string IntermediateDir)
 		{
 			string LocalizationsPath = BuildDirectory + "/Resources/Localizations";
@@ -1048,8 +974,7 @@ namespace UnrealBuildTool
 
 			bool bSupportsPortrait = true;
 			bool bSupportsLandscape = false;
-			bool bSkipIcons = false;
-			GeneratePList(ProjectFile, Config, InProjectDirectory, bIsUnrealGame, GameExeName, false, InProjectName, InEngineDir, AppDirectory, UPLScripts, SdkVersion, BundleID, bBuildAsFramework, out bSupportsPortrait, out bSupportsLandscape, out bSkipIcons);
+			GeneratePList(ProjectFile, Config, InProjectDirectory, bIsUnrealGame, GameExeName, false, InProjectName, InEngineDir, AppDirectory, UPLScripts, SdkVersion, BundleID, bBuildAsFramework, out bSupportsPortrait, out bSupportsLandscape);
 
 			// ensure the destination is writable
 			if (File.Exists(AppDirectory + "/" + GameName))
@@ -1074,10 +999,6 @@ namespace UnrealBuildTool
 
 			if (!bCreateStubIPA)
 			{
-				if (!bSkipIcons)
-				{
-					CopyIconResources(InEngineDir, AppDirectory, BuildDirectory);
-				}
 				CopyLocalizationsResources(InEngineDir, AppDirectory, BuildDirectory, IntermediateDirectory);
 
 				CopyCloudResources(InProjectDirectory, AppDirectory);
@@ -1139,8 +1060,8 @@ namespace UnrealBuildTool
 			else
 			{
 				// @todo tvos merge: This used to copy the bundle back - where did that code go? It needs to be fixed up for TVOS directories
-				bool bSupportPortrait, bSupportLandscape, bSkipIcons;
-				GeneratePList(ProjectFile, Configuration, ProjectDirectory, bIsUnrealGame, GameName, false, (ProjectFile == null) ? "" : Path.GetFileNameWithoutExtension(ProjectFile.FullName), "../../Engine", "", UPLScripts, SdkVersion, BundleID, bBuildAsFramework, out bSupportPortrait, out bSupportLandscape, out bSkipIcons);
+				bool bSupportPortrait, bSupportLandscape;
+				GeneratePList(ProjectFile, Configuration, ProjectDirectory, bIsUnrealGame, GameName, false, (ProjectFile == null) ? "" : Path.GetFileNameWithoutExtension(ProjectFile.FullName), "../../Engine", "", UPLScripts, SdkVersion, BundleID, bBuildAsFramework, out bSupportPortrait, out bSupportLandscape);
 			}
 			return true;
 		}
diff --git a/Engine/Source/Programs/UnrealBuildTool/Platform/TVOS/TVOSExports.cs b/Engine/Source/Programs/UnrealBuildTool/Platform/TVOS/TVOSExports.cs
--- a/Engine/Source/Programs/UnrealBuildTool/Platform/TVOS/TVOSExports.cs
+++ b/Engine/Source/Programs/UnrealBuildTool/Platform/TVOS/TVOSExports.cs
@@ -99,12 +99,11 @@ namespace UnrealBuildTool
 		/// <param name="BuildReceiptFileName"></param>
 		/// <param name="bSupportsPortrait"></param>
 		/// <param name="bSupportsLandscape"></param>
-		/// <param name="bSkipIcons"></param>
 		/// <returns></returns>
-		public static bool GeneratePList(FileReference ProjectFile, UnrealTargetConfiguration Config, DirectoryReference ProjectDirectory, bool bIsUnrealGame, string GameName, bool bIsClient, string ProjectName, DirectoryReference InEngineDir, DirectoryReference AppDirectory, FileReference BuildReceiptFileName, out bool bSupportsPortrait, out bool bSupportsLandscape, out bool bSkipIcons)
+		public static bool GeneratePList(FileReference ProjectFile, UnrealTargetConfiguration Config, DirectoryReference ProjectDirectory, bool bIsUnrealGame, string GameName, bool bIsClient, string ProjectName, DirectoryReference InEngineDir, DirectoryReference AppDirectory, FileReference BuildReceiptFileName, out bool bSupportsPortrait, out bool bSupportsLandscape)
 		{
 			TargetReceipt Receipt = TargetReceipt.Read(BuildReceiptFileName);
-			return new UEDeployTVOS().GeneratePList(ProjectFile, Config, ProjectDirectory.FullName, bIsUnrealGame, GameName, bIsClient, ProjectName, InEngineDir.FullName, AppDirectory.FullName, Receipt, out bSupportsPortrait, out bSupportsLandscape, out bSkipIcons);
+			return new UEDeployTVOS().GeneratePList(ProjectFile, Config, ProjectDirectory.FullName, bIsUnrealGame, GameName, bIsClient, ProjectName, InEngineDir.FullName, AppDirectory.FullName, Receipt, out bSupportsPortrait, out bSupportsLandscape);
 		}
 	}
 }
diff --git a/Engine/Source/Programs/UnrealBuildTool/Platform/TVOS/UEDeployTVOS.cs b/Engine/Source/Programs/UnrealBuildTool/Platform/TVOS/UEDeployTVOS.cs
--- a/Engine/Source/Programs/UnrealBuildTool/Platform/TVOS/UEDeployTVOS.cs
+++ b/Engine/Source/Programs/UnrealBuildTool/Platform/TVOS/UEDeployTVOS.cs
@@ -204,16 +204,10 @@ namespace UnrealBuildTool
 			return bSkipDefaultPNGs;
 		}
 
-		public override bool GeneratePList(FileReference ProjectFile, UnrealTargetConfiguration Config, string ProjectDirectory, bool bIsUnrealGame, string GameName, bool bIsClient, string ProjectName, string InEngineDir, string AppDirectory, List<string> UPLScripts, VersionNumber SdkVersion, string BundleID, bool bBuildAsFramework, out bool bSupportsPortrait, out bool bSupportsLandscape, out bool bSkipIcons)
+		public override bool GeneratePList(FileReference ProjectFile, UnrealTargetConfiguration Config, string ProjectDirectory, bool bIsUnrealGame, string GameName, bool bIsClient, string ProjectName, string InEngineDir, string AppDirectory, List<string> UPLScripts, VersionNumber SdkVersion, string BundleID, bool bBuildAsFramework, out bool bSupportsPortrait, out bool bSupportsLandscape)
 		{
 			bSupportsLandscape = bSupportsPortrait = true;
-            bSkipIcons = true;
 			return GenerateTVOSPList(ProjectDirectory, bIsUnrealGame, GameName, bIsClient, ProjectName, InEngineDir, AppDirectory, null, BundleID);
 		}
-
-		protected override void CopyIconResources(string InEngineDir, string AppDirectory, string BuildDirectory)
-		{
-			// do nothing on TVOS
-		}
 	}
 }
diff --git a/Engine/Source/Programs/UnrealBuildTool/ProjectFiles/Xcode/XcodeProject.cs b/Engine/Source/Programs/UnrealBuildTool/ProjectFiles/Xcode/XcodeProject.cs
--- a/Engine/Source/Programs/UnrealBuildTool/ProjectFiles/Xcode/XcodeProject.cs
+++ b/Engine/Source/Programs/UnrealBuildTool/ProjectFiles/Xcode/XcodeProject.cs
@@ -1385,12 +1385,12 @@ namespace UnrealBuildTool
 								ReceiptFilename = TargetReceipt.GetDefaultPath(ProjectPath, GameName, UnrealTargetPlatform.IOS, Config.BuildConfig, "");
 							}
 							Directory.CreateDirectory(Path.GetDirectoryName(IOSInfoPlistPath));
-                            bool bSupportPortrait, bSupportLandscape, bSkipIcons;
+                            bool bSupportPortrait, bSupportLandscape;
 							TargetReceipt Receipt;
 							TargetReceipt.TryRead(ReceiptFilename, out Receipt);
 							VersionNumber SdkVersion = UEDeployIOS.GetSdkVersion(Receipt);
 							bool bBuildAsFramework = UEDeployIOS.GetCompileAsDll(Receipt);
-							UEDeployIOS.GenerateIOSPList(ProjectFile, Config.BuildConfig, ProjectPath.FullName, bIsUnrealGame, GameName, bIsClient, Config.BuildTarget, EngineDir.FullName, ProjectPath + "/Binaries/IOS/Payload", SdkVersion, null, BundleIdentifier, bBuildAsFramework, out bSupportPortrait, out bSupportLandscape, out bSkipIcons);
+							UEDeployIOS.GenerateIOSPList(ProjectFile, Config.BuildConfig, ProjectPath.FullName, bIsUnrealGame, GameName, bIsClient, Config.BuildTarget, EngineDir.FullName, ProjectPath + "/Binaries/IOS/Payload", SdkVersion, null, BundleIdentifier, bBuildAsFramework, out bSupportPortrait, out bSupportLandscape);
 						}
 						if (bCreateTVOSInfoPlist)
 						{
