UnrealBuildTool: Use BuildToolOverride for -2022 in VcProject generation

diff --git a/Engine/Source/Programs/UnrealBuildTool/ProjectFiles/VisualStudio/VCProject.cs b/Engine/Source/Programs/UnrealBuildTool/ProjectFiles/VisualStudio/VCProject.cs
--- a/Engine/Source/Programs/UnrealBuildTool/ProjectFiles/VisualStudio/VCProject.cs
+++ b/Engine/Source/Programs/UnrealBuildTool/ProjectFiles/VisualStudio/VCProject.cs
@@ -1708,11 +1708,6 @@ namespace UnrealBuildTool
 						BuildArguments.Append(ProjGenerator.GetExtraBuildArguments(Platform, Configuration));
 					}
 
-					if (ProjectFileFormat == VCProjectFileFormat.VisualStudio2022)
-					{
-						BuildArguments.Append(" -2022");
-					}
-
 					// NMake Build command line
 					VCProjectFileContent.AppendLine("    <NMakeBuildCommandLine>{0} {1}</NMakeBuildCommandLine>", EscapePath(NormalizeProjectPath(FileReference.Combine(BatchFilesDirectory, "Build.bat"))), BuildArguments.ToString());
 					VCProjectFileContent.AppendLine("    <NMakeReBuildCommandLine>{0} {1}</NMakeReBuildCommandLine>", EscapePath(NormalizeProjectPath(FileReference.Combine(BatchFilesDirectory, "Rebuild.bat"))), BuildArguments.ToString());
diff --git a/Engine/Source/Programs/UnrealBuildTool/ProjectFiles/VisualStudio/VCProjectFileGenerator.cs b/Engine/Source/Programs/UnrealBuildTool/ProjectFiles/VisualStudio/VCProjectFileGenerator.cs
--- a/Engine/Source/Programs/UnrealBuildTool/ProjectFiles/VisualStudio/VCProjectFileGenerator.cs
+++ b/Engine/Source/Programs/UnrealBuildTool/ProjectFiles/VisualStudio/VCProjectFileGenerator.cs
@@ -135,6 +135,10 @@ namespace UnrealBuildTool
 			{
 				BuildToolOverride = "-2019";
 			}
+			else if (InArguments.HasOption("-2022"))
+			{
+				BuildToolOverride = "-2022";
+			}
 		}
 
 		public override string[] GetTargetArguments(string[] Arguments)
