[Workspace] - back out the Lowercasing of the workspace name

diff --git a/Engine/Source/Programs/UnrealBuildTool/System/DeployExports.cs b/Engine/Source/Programs/UnrealBuildTool/System/DeployExports.cs
--- a/Engine/Source/Programs/UnrealBuildTool/System/DeployExports.cs
+++ b/Engine/Source/Programs/UnrealBuildTool/System/DeployExports.cs
@@ -30,7 +30,7 @@ namespace UnrealBuildTool
 			string DeployFolder = string.Format("{0}-{1}-{2}",
 				ShortProjectName,
 				UnrealBuildTool.RootDirectory.GetDirectoryName(),
-				Environment.UserName).ToLower();
+				Environment.UserName);
 
 			DeployFolder = DeployFolder.Replace(" ", "");
 			DeployFolder = DeployFolder.Replace(".", "");
