Add a common C# method for detecting the current host platform to EpicGames.Core, and move many a test to use that code.

diff --git a/Engine/Source/Programs/AutomationTool/Android/AndroidPlatform.Automation.cs b/Engine/Source/Programs/AutomationTool/Android/AndroidPlatform.Automation.cs
--- a/Engine/Source/Programs/AutomationTool/Android/AndroidPlatform.Automation.cs
+++ b/Engine/Source/Programs/AutomationTool/Android/AndroidPlatform.Automation.cs
@@ -479,7 +479,7 @@ public class AndroidPlatform : Platform
 
 	public static string GetStorageQueryCommand(bool bForcePC = false)
     {
-		if (!bForcePC && !Utils.IsRunningOnWindows)
+		if (!bForcePC && !RuntimePlatform.IsWindows)
 		{
 			return "shell 'echo $EXTERNAL_STORAGE'";
 		}
@@ -1048,7 +1048,7 @@ public class AndroidPlatform : Platform
 						File.WriteAllLines(SymbolizeBatchName, BatchLines);
 					}
 
-					if (!Utils.IsRunningOnWindows)
+					if (!RuntimePlatform.IsWindows)
 					{
 						if (bHaveAPK)
 						{
@@ -1617,7 +1617,7 @@ public class AndroidPlatform : Platform
 
 	public static IProcessResult RunAdbCommand(ProjectParams Params, string SerialNumber, string Args, string Input = null, ERunOptions Options = ERunOptions.Default, bool bShouldLogCommand = false)
 	{
-		string AdbCommand = Environment.ExpandEnvironmentVariables("%ANDROID_HOME%/platform-tools/adb" + (Utils.IsRunningOnWindows ? ".exe" : ""));
+		string AdbCommand = Environment.ExpandEnvironmentVariables("%ANDROID_HOME%/platform-tools/adb" + (RuntimePlatform.IsWindows ? ".exe" : ""));
 		if (Options.HasFlag(ERunOptions.AllowSpew) || Options.HasFlag(ERunOptions.SpewIsVerbose))
 		{
 			LastSpewFilename = "";
@@ -1628,7 +1628,7 @@ public class AndroidPlatform : Platform
 
 	private string RunAndLogAdbCommand(ProjectParams Params, string SerialNumber, string Args, out int SuccessCode)
 	{
-		string AdbCommand = Environment.ExpandEnvironmentVariables("%ANDROID_HOME%/platform-tools/adb" + (Utils.IsRunningOnWindows ? ".exe" : ""));
+		string AdbCommand = Environment.ExpandEnvironmentVariables("%ANDROID_HOME%/platform-tools/adb" + (RuntimePlatform.IsWindows ? ".exe" : ""));
 		LastSpewFilename = "";
 		return RunAndLog(CmdEnv, AdbCommand, GetAdbCommandLine(Params, SerialNumber, Args), out SuccessCode, SpewFilterCallback: new ProcessResult.SpewFilterCallbackType(ADBSpewFilter));
 	}
@@ -2476,7 +2476,7 @@ public class AndroidPlatform : Platform
 		uint BestVersion = 0;
 		foreach (string CandidateDir in Subdirs)
 		{
-			string AaptFilename = Path.Combine(CandidateDir, Utils.IsRunningOnWindows ? "aapt.exe" : "aapt");
+			string AaptFilename = Path.Combine(CandidateDir, RuntimePlatform.IsWindows ? "aapt.exe" : "aapt");
 			uint RevisionValue = 0;
 
 			if (File.Exists(AaptFilename))
diff --git a/Engine/Source/Programs/AutomationTool/AutomationUtils/CommandEnvironment.cs b/Engine/Source/Programs/AutomationTool/AutomationUtils/CommandEnvironment.cs
--- a/Engine/Source/Programs/AutomationTool/AutomationUtils/CommandEnvironment.cs
+++ b/Engine/Source/Programs/AutomationTool/AutomationUtils/CommandEnvironment.cs
@@ -7,7 +7,6 @@ using System.Reflection;
 using Microsoft.Win32;
 using System.Diagnostics;
 using EpicGames.Core;
-using UnrealBuildTool;
 using System.Text.RegularExpressions;
 
 namespace AutomationTool
@@ -73,7 +72,7 @@ namespace AutomationTool
 		{
 			// Get the path to the UAT executable
 			// the entry assembly is the .dll but it is easier to use apphost so change extension to the executable
-			UATExe = Path.ChangeExtension(Assembly.GetEntryAssembly().GetOriginalLocation(), Utils.IsRunningOnWindows ? "exe" : null);
+			UATExe = Path.ChangeExtension(Assembly.GetEntryAssembly().GetOriginalLocation(), RuntimePlatform.IsWindows ? "exe" : null);
 
 			if (!CommandUtils.FileExists(UATExe))
 			{
@@ -127,7 +126,7 @@ namespace AutomationTool
 
 			RobocopyExe = GetSystemExePath("robocopy.exe");
 			MountExe = GetSystemExePath("mount.exe");
-			CmdExe = Utils.IsRunningOnWindows ? GetSystemExePath("cmd.exe") : "/bin/sh";
+			CmdExe = RuntimePlatform.IsWindows ? GetSystemExePath("cmd.exe") : "/bin/sh";
 			MallocNanoZone = "0";
 			CommandUtils.SetEnvVar(EnvVarNames.MacMallocNanoZone, MallocNanoZone);
 
diff --git a/Engine/Source/Programs/AutomationTool/AutomationUtils/CommandUtils.cs b/Engine/Source/Programs/AutomationTool/AutomationUtils/CommandUtils.cs
--- a/Engine/Source/Programs/AutomationTool/AutomationUtils/CommandUtils.cs
+++ b/Engine/Source/Programs/AutomationTool/AutomationUtils/CommandUtils.cs
@@ -2096,7 +2096,7 @@ namespace AutomationTool
                 }
                 else
                 {
-                    CachedRootBuildStorageDirectory = Utils.IsRunningOnWindows ? CombinePaths("P:", "Builds") : "/Volumes/Builds";
+                    CachedRootBuildStorageDirectory = RuntimePlatform.IsWindows ? CombinePaths("P:", "Builds") : "/Volumes/Builds";
                 }
             }
             return CachedRootBuildStorageDirectory;
@@ -2266,7 +2266,7 @@ namespace AutomationTool
 		public static IEnumerable<string> LegacyUnzipFiles(string ZipFileName, string BaseDirectory)
 		{
 			List<string> OutputFileNames = new List<string>();
-			if (!Utils.IsRunningOnWindows)
+			if (!RuntimePlatform.IsWindows)
 			{
 				CommandUtils.CreateDirectory(BaseDirectory);
 
@@ -2966,7 +2966,7 @@ namespace AutomationTool
 		/// </summary>
 		public static void SignSingleExecutableIfEXEOrDLL(string Filename, bool bIgnoreExtension = false)
 		{
-            if (!Utils.IsRunningOnWindows)
+            if (!RuntimePlatform.IsWindows)
             {
                 CommandUtils.LogLog(String.Format("Can't sign '{0}' on non-Windows platform.", Filename));
                 return;
@@ -3159,7 +3159,7 @@ namespace AutomationTool
 
 		public static void SignMultipleFilesIfEXEOrDLL(List<FileReference> Files, bool bIgnoreExtension = false)
 		{
-			if (!Utils.IsRunningOnWindows)
+			if (!RuntimePlatform.IsWindows)
 			{
 				CommandUtils.LogLog(String.Format("Can't sign on non-Windows platform."));
 				return;
diff --git a/Engine/Source/Programs/AutomationTool/AutomationUtils/HostPlatform.cs b/Engine/Source/Programs/AutomationTool/AutomationUtils/HostPlatform.cs
--- a/Engine/Source/Programs/AutomationTool/AutomationUtils/HostPlatform.cs
+++ b/Engine/Source/Programs/AutomationTool/AutomationUtils/HostPlatform.cs
@@ -36,33 +36,21 @@ namespace AutomationTool
 		/// </summary>
 		public static void Initialize()
 		{
-			PlatformID Platform = Environment.OSVersion.Platform;
-			switch (Platform)
+			if (RuntimePlatform.IsWindows)
 			{
-				case PlatformID.Win32NT:
-				case PlatformID.Win32S:
-				case PlatformID.Win32Windows:
-				case PlatformID.WinCE:
-					RunningPlatform = new WindowsHostPlatform();
-					break;
-
-				case PlatformID.Unix:
-					if (File.Exists ("/System/Library/CoreServices/SystemVersion.plist"))
-					{
-						RunningPlatform = new MacHostPlatform();
-					} 
-					else 
-					{
-						RunningPlatform = new LinuxHostPlatform();
-					}
-					break;
-
-				case PlatformID.MacOSX:
-					RunningPlatform = new MacHostPlatform();
-					break;
-
-				default:
-					throw new Exception ("Unhandled runtime platform " + Platform);
+				RunningPlatform = new WindowsHostPlatform();
+			}
+			else if (RuntimePlatform.IsMac)
+			{
+				RunningPlatform = new MacHostPlatform();
+			}
+			else if (RuntimePlatform.IsLinux)
+			{
+				RunningPlatform = new LinuxHostPlatform();
+			}
+			else
+			{
+				throw new Exception ("Unhandled runtime platform " + Environment.OSVersion.Platform);
 			}
 		}
 
diff --git a/Engine/Source/Programs/AutomationTool/AutomationUtils/UE4Build.cs b/Engine/Source/Programs/AutomationTool/AutomationUtils/UE4Build.cs
--- a/Engine/Source/Programs/AutomationTool/AutomationUtils/UE4Build.cs
+++ b/Engine/Source/Programs/AutomationTool/AutomationUtils/UE4Build.cs
@@ -1317,7 +1317,7 @@ namespace AutomationTool
 
 		public static string GetUBTExecutable()
 		{
-			return CommandUtils.CombinePaths(CommandUtils.CmdEnv.LocalRoot, @"Engine/Binaries/DotNET/UnrealBuildTool/UnrealBuildTool" + (Utils.IsRunningOnWindows ? ".exe" : ""));
+			return CommandUtils.CombinePaths(CommandUtils.CmdEnv.LocalRoot, @"Engine/Binaries/DotNET/UnrealBuildTool/UnrealBuildTool" + (RuntimePlatform.IsWindows ? ".exe" : ""));
 		}
 
 		public string UBTExecutable
diff --git a/Engine/Source/Programs/AutomationTool/HoloLens/HoloLensPlatform.Automation.cs b/Engine/Source/Programs/AutomationTool/HoloLens/HoloLensPlatform.Automation.cs
--- a/Engine/Source/Programs/AutomationTool/HoloLens/HoloLensPlatform.Automation.cs
+++ b/Engine/Source/Programs/AutomationTool/HoloLens/HoloLensPlatform.Automation.cs
@@ -1290,7 +1290,7 @@ namespace HoloLens.Automation
 		private void DeployToLocalDevice(ProjectParams Params, DeploymentContext SC)
 		{
 #if !__MonoCS__
-            if (!Utils.IsRunningOnWindows)
+            if (!RuntimePlatform.IsWindows)
             {
                 return;
             }
@@ -1407,7 +1407,7 @@ namespace HoloLens.Automation
 		private void DeployToRemoteDevice(string DeviceAddress, ProjectParams Params, DeploymentContext SC)
 		{
 #if !__MonoCS__
-			if (!Utils.IsRunningOnWindows)
+			if (!RuntimePlatform.IsWindows)
             {
                 return;
             }
@@ -1496,7 +1496,7 @@ namespace HoloLens.Automation
 
 		private IProcessResult RunUsingLauncherTool(string DeviceAddress, ERunOptions ClientRunFlags, string ClientApp, string ClientCmdLine, ProjectParams Params)
 		{
-            if (!Utils.IsRunningOnWindows)
+            if (!RuntimePlatform.IsWindows)
             {
                 return null;
             }
@@ -1588,7 +1588,7 @@ namespace HoloLens.Automation
 		private IProcessResult RunUsingDevicePortal(string DeviceAddress, ERunOptions ClientRunFlags, string ClientApp, string ClientCmdLine, ProjectParams Params)
 		{
 #if !__MonoCS__
-            if (!Utils.IsRunningOnWindows)
+            if (!RuntimePlatform.IsWindows)
             {
                 return null;
             }
@@ -1693,7 +1693,7 @@ namespace HoloLens.Automation
 		private bool ShouldAcceptCertificate(System.Security.Cryptography.X509Certificates.X509Certificate2 Certificate, bool Unattended)
 		{
 #if !__MonoCS__
-            if (!Utils.IsRunningOnWindows)
+            if (!RuntimePlatform.IsWindows)
             {
                 return false;
             }
diff --git a/Engine/Source/Programs/AutomationTool/IOS/IOSPlatform.Automation.cs b/Engine/Source/Programs/AutomationTool/IOS/IOSPlatform.Automation.cs
--- a/Engine/Source/Programs/AutomationTool/IOS/IOSPlatform.Automation.cs
+++ b/Engine/Source/Programs/AutomationTool/IOS/IOSPlatform.Automation.cs
@@ -1463,7 +1463,7 @@ public class IOSPlatform : Platform
 			PlatformGameConfig.GetBool("/Script/IOSRuntimeSettings.IOSRuntimeSettings", "bGenerateXCArchive", out bXCArchive);
 		}
 
-		if (bXCArchive && !Utils.IsRunningOnWindows)
+		if (bXCArchive && !RuntimePlatform.IsWindows)
 		{
 			// Always put the archive in the current user's Library/Developer/Xcode/Archives path if not on the build machine
 			WindowsIdentity id = WindowsIdentity.GetCurrent();
@@ -1631,7 +1631,7 @@ public class IOSPlatform : Platform
 			Text.AppendLine("</plist>");
 			File.WriteAllText(Path.Combine(ArchiveName, "Info.plist"), Text.ToString());
 		}
-		else if (bXCArchive && Utils.IsRunningOnWindows)
+		else if (bXCArchive && RuntimePlatform.IsWindows)
 		{
 			LogWarning("Can not produce an XCArchive on windows");
 		}
diff --git a/Engine/Source/Programs/AutomationTool/Linux/LinuxPlatform.Automation.cs b/Engine/Source/Programs/AutomationTool/Linux/LinuxPlatform.Automation.cs
--- a/Engine/Source/Programs/AutomationTool/Linux/LinuxPlatform.Automation.cs
+++ b/Engine/Source/Programs/AutomationTool/Linux/LinuxPlatform.Automation.cs
@@ -133,7 +133,7 @@ public abstract class BaseLinuxPlatform : Platform
 		// write out the 
 		FileReference.WriteAllText(IntermediateFile, Script.ToString());
 
-		if (!Utils.IsRunningOnWindows)
+		if (!RuntimePlatform.IsWindows)
 		{
 			var Result = CommandUtils.Run("env", string.Format("-- \"chmod\" \"+x\" \"{0}\"", IntermediateFile.ToString().Replace("'", "'\"'\"'")));
 			if (Result.ExitCode != 0)
diff --git a/Engine/Source/Programs/AutomationTool/Lumin/LuminPlatform.Automation.cs b/Engine/Source/Programs/AutomationTool/Lumin/LuminPlatform.Automation.cs
--- a/Engine/Source/Programs/AutomationTool/Lumin/LuminPlatform.Automation.cs
+++ b/Engine/Source/Programs/AutomationTool/Lumin/LuminPlatform.Automation.cs
@@ -781,7 +781,7 @@ public class LuminPlatform : Platform
 		File.Copy(UninstallBatchName, SC.CookSourceRuntimeRootDir + "\\" + Path.GetFileName(UninstallBatchName).ToString(), true);
 
 		// If needed, make the batch files able to execute
-		if (!Utils.IsRunningOnWindows)
+		if (!RuntimePlatform.IsWindows)
 		{
 			CommandUtils.FixUnixFilePermissions(BatchName);
 			CommandUtils.FixUnixFilePermissions(UninstallBatchName);
@@ -1294,22 +1294,13 @@ public class LuminPlatform : Platform
 				throw new AutomationException(ExitCode.Error_AndroidBuildToolsPathNotFound, "Failed to find a %MLSDK% directory. Please set MLSDK environment variable to point to ML SDK.");
 			}
 
-			switch (Environment.OSVersion.Platform)
+			if (RuntimePlatform.IsMac || RuntimePlatform.IsLinux)
 			{
-				case PlatformID.MacOSX:
-					return Path.Combine(
-						envValue,
-						"tools", "mldb", "mldb");
-
-				case PlatformID.Unix:
-					return Path.Combine(
-						envValue,
-						"tools", "mldb", "mldb");
-
-				default:
-					return Path.Combine(
-						envValue,
-						"tools", "mldb", "mldb.exe");
+				return Path.Combine(envValue, "tools", "mldb", "mldb");
+			}
+			else
+			{
+				return Path.Combine(envValue, "tools", "mldb", "mldb.exe");
 			}
 		}
 	}
diff --git a/Engine/Source/Programs/AutomationTool/Program.cs b/Engine/Source/Programs/AutomationTool/Program.cs
--- a/Engine/Source/Programs/AutomationTool/Program.cs
+++ b/Engine/Source/Programs/AutomationTool/Program.cs
@@ -124,7 +124,7 @@ namespace AutomationTool
                 NoThrow(() => CommandUtils.ClearDirStack(), "Clear Dir Stack");
 
                 // Try to kill process before app domain exits to leave the other KillAll call to extreme edge cases
-                NoThrow(() => { if (ShouldKillProcesses && Utils.IsRunningOnWindows) ProcessManager.KillAll(); }, "Kill All Processes");
+                NoThrow(() => { if (ShouldKillProcesses && RuntimePlatform.IsWindows) ProcessManager.KillAll(); }, "Kill All Processes");
 
 				// Write the exit code
                 Log.TraceInformation("AutomationTool executed for {0}", Timer.Elapsed.ToString("h'h 'm'm 's's'"));
@@ -168,7 +168,7 @@ namespace AutomationTool
 		static void Domain_ProcessExit(object sender, EventArgs e)
 		{
 			// Kill all spawned processes (Console instead of Log because logging is closed at this time anyway)
-			if (ShouldKillProcesses && Utils.IsRunningOnWindows)
+			if (ShouldKillProcesses && RuntimePlatform.IsWindows)
 			{			
 				ProcessManager.KillAll();
 			}
diff --git a/Engine/Source/Programs/AutomationTool/Scripts/BuildHlslcc.Automation.cs b/Engine/Source/Programs/AutomationTool/Scripts/BuildHlslcc.Automation.cs
--- a/Engine/Source/Programs/AutomationTool/Scripts/BuildHlslcc.Automation.cs
+++ b/Engine/Source/Programs/AutomationTool/Scripts/BuildHlslcc.Automation.cs
@@ -519,7 +519,7 @@ class BuildHlslcc : BuildCommand
 	}
 	private static void SetupStaticBuildEnvironment()
 	{
-		if (Utils.IsRunningOnWindows)
+		if (RuntimePlatform.IsWindows)
 		{
 			string VS2019Path = GetMsDevExe(WindowsCompiler.VisualStudio2019);
 			if (VS2019Path != null)
diff --git a/Engine/Source/Programs/AutomationTool/Scripts/ZipUtils.Automation.cs b/Engine/Source/Programs/AutomationTool/Scripts/ZipUtils.Automation.cs
--- a/Engine/Source/Programs/AutomationTool/Scripts/ZipUtils.Automation.cs
+++ b/Engine/Source/Programs/AutomationTool/Scripts/ZipUtils.Automation.cs
@@ -6,7 +6,6 @@ using System.Collections.Generic;
 using System.Linq;
 using System.Text;
 using AutomationTool;
-using UnrealBuildTool;
 using Ionic.Zip;
 using EpicGames.Core;
 
@@ -158,7 +157,7 @@ public class ZipUtils : BuildCommand
 					}
 				}
 
-				if (!Utils.IsRunningOnWindows && CommandUtils.IsProbablyAMacOrIOSExe(OutputFileName))
+				if (!RuntimePlatform.IsWindows && CommandUtils.IsProbablyAMacOrIOSExe(OutputFileName))
 				{
 					FixUnixFilePermissions(OutputFileName);
 				}
diff --git a/Engine/Source/Programs/AutomationTool/Turnkey/Turnkey.Automation.cs b/Engine/Source/Programs/AutomationTool/Turnkey/Turnkey.Automation.cs
--- a/Engine/Source/Programs/AutomationTool/Turnkey/Turnkey.Automation.cs
+++ b/Engine/Source/Programs/AutomationTool/Turnkey/Turnkey.Automation.cs
@@ -3,7 +3,6 @@
 using System;
 using System.Collections.Generic;
 using EpicGames.Core;
-using UnrealBuildTool;
 using AutomationTool;
 using System.Threading;
 
@@ -103,7 +102,7 @@ namespace Turnkey
 		private static void SetupVisuals()
 		{
 			// make the form look good on modern displays!
-			if (Utils.IsRunningOnWindows && Environment.OSVersion.Version.Major >= 6)
+			if (RuntimePlatform.IsWindows && Environment.OSVersion.Version.Major >= 6)
 			{
 				SetProcessDPIAware();
 			}
diff --git a/Engine/Source/Programs/DotNETCommon/BuildUtilities/UEBuildPlatformSDK.cs b/Engine/Source/Programs/DotNETCommon/BuildUtilities/UEBuildPlatformSDK.cs
--- a/Engine/Source/Programs/DotNETCommon/BuildUtilities/UEBuildPlatformSDK.cs
+++ b/Engine/Source/Programs/DotNETCommon/BuildUtilities/UEBuildPlatformSDK.cs
@@ -468,28 +468,17 @@ namespace EpicGames.Core
 		protected const string AutoSetupEnvVar = "AutoSDKSetup";
 
 
-		protected static bool IsWindows()
-		{
-			return Environment.OSVersion.Platform == PlatformID.Win32NT || Environment.OSVersion.Platform == PlatformID.Win32S || Environment.OSVersion.Platform == PlatformID.Win32Windows;
-		}
-
-		protected static bool IsMac()
-		{
-			// mono tends to return Unix on Mac, so check for a Mac-only file
-			return Environment.OSVersion.Platform == PlatformID.MacOSX || (Environment.OSVersion.Platform == PlatformID.Unix && File.Exists ("/System/Library/CoreServices/SystemVersion.plist"));
-		}
-
 		private static string GetAutoSDKHostPlatform()
 		{
-			if (IsWindows())
+			if (RuntimePlatform.IsWindows)
 			{
 				return "Win64";
 			}
-			else if (IsMac())
+			else if (RuntimePlatform.IsMac)
 			{
 				return "Mac";
 			}
-			else if (Environment.OSVersion.Platform == PlatformID.Unix)
+			else if (RuntimePlatform.IsLinux)
 			{
 				return "Linux";
 			}
@@ -746,7 +735,7 @@ namespace EpicGames.Core
 		/// <param name="Hook">Hook type</param>
 		protected virtual string GetHookExecutableName(SDKHookType Hook)
 		{
-			if (IsWindows())
+			if (RuntimePlatform.IsWindows)
 			{
 				if (Hook == SDKHookType.Uninstall)
 				{
@@ -939,7 +928,7 @@ namespace EpicGames.Core
 
                     // actually perform the PATH stripping / adding.
                     String OrigPathVar = Environment.GetEnvironmentVariable("PATH");
-                    String PathDelimiter = IsWindows() ? ";" : ":";
+                    String PathDelimiter = RuntimePlatform.IsWindows ? ";" : ":";
                     String[] PathVars = { };
                     if (!String.IsNullOrEmpty(OrigPathVar))
                     {
diff --git a/Engine/Source/Programs/Shared/EpicGames.Core/FileUtils.cs b/Engine/Source/Programs/Shared/EpicGames.Core/FileUtils.cs
--- a/Engine/Source/Programs/Shared/EpicGames.Core/FileUtils.cs
+++ b/Engine/Source/Programs/Shared/EpicGames.Core/FileUtils.cs
@@ -46,12 +46,12 @@ namespace EpicGames.Core
 		/// <summary>
 		/// Comparer that should be used for native path comparisons
 		/// </summary>
-		public static IEqualityComparer<string> PlatformPathComparer = Environment.OSVersion.Platform == PlatformID.Unix ? StringComparer.Ordinal : StringComparer.OrdinalIgnoreCase;
+		public static IEqualityComparer<string> PlatformPathComparer = RuntimePlatform.IsLinux ? StringComparer.Ordinal : StringComparer.OrdinalIgnoreCase;
 
 		/// <summary>
 		/// Utf8 string comparer that should be used for native path comparisons
 		/// </summary>
-		public static IEqualityComparer<ReadOnlyUtf8String> PlatformPathComparerUtf8 = Environment.OSVersion.Platform == PlatformID.Unix ? ReadOnlyUtf8StringComparer.Ordinal : ReadOnlyUtf8StringComparer.OrdinalIgnoreCase;
+		public static IEqualityComparer<ReadOnlyUtf8String> PlatformPathComparerUtf8 = RuntimePlatform.IsLinux ? ReadOnlyUtf8StringComparer.Ordinal : ReadOnlyUtf8StringComparer.OrdinalIgnoreCase;
 
 		/// <summary>
 		/// Read all text for a file
@@ -119,7 +119,7 @@ namespace EpicGames.Core
 		/// <param name="FileName">Name of the file to delete</param>
 		public static void ForceDeleteFile(string FileName)
 		{
-			if (Environment.OSVersion.Platform == PlatformID.Win32NT)
+			if (RuntimePlatform.IsWindows)
 			{
 				ForceDeleteFileWin32(FileName);
 			}
@@ -135,7 +135,7 @@ namespace EpicGames.Core
 		/// <param name="File">The file to delete</param>
 		public static void ForceDeleteFile(FileInfo File)
 		{
-			if (Environment.OSVersion.Platform == PlatformID.Win32NT)
+			if (RuntimePlatform.IsWindows)
 			{
 				ForceDeleteFileWin32(File.FullName);
 				File.Refresh();
@@ -163,7 +163,7 @@ namespace EpicGames.Core
 		/// <param name="Location">The file to delete</param>
 		public static void ForceDeleteFile(FileReference Location)
 		{
-			if (Environment.OSVersion.Platform == PlatformID.Win32NT)
+			if (RuntimePlatform.IsWindows)
 			{
 				ForceDeleteFileWin32(Location.FullName);
 			}
@@ -179,7 +179,7 @@ namespace EpicGames.Core
 		/// <param name="DirectoryName">Directory to delete</param>
 		public static void ForceDeleteDirectory(string DirectoryName)
 		{
-			if (Environment.OSVersion.Platform == PlatformID.Win32NT)
+			if (RuntimePlatform.IsWindows)
 			{
 				ForceDeleteLongDirectoryWin32("\\\\?\\" + DirectoryName);
 			}
@@ -197,7 +197,7 @@ namespace EpicGames.Core
 		{
 			if (Directory.Exists)
 			{
-				if (Environment.OSVersion.Platform == PlatformID.Win32NT)
+				if (RuntimePlatform.IsWindows)
 				{
 					ForceDeleteLongDirectoryWin32("\\\\?\\" + Directory.FullName);
 				}
@@ -249,7 +249,7 @@ namespace EpicGames.Core
 		/// <param name="Directory">Directory to delete</param>
 		public static void ForceDeleteDirectoryContents(DirectoryInfo Directory)
 		{
-			if (Environment.OSVersion.Platform == PlatformID.Win32NT)
+			if (RuntimePlatform.IsWindows)
 			{
 				ForceDeleteLongDirectoryContentsWin32(Directory.FullName);
 			}
diff --git a/Engine/Source/Programs/Shared/EpicGames.Core/ManagedProcess.cs b/Engine/Source/Programs/Shared/EpicGames.Core/ManagedProcess.cs
--- a/Engine/Source/Programs/Shared/EpicGames.Core/ManagedProcess.cs
+++ b/Engine/Source/Programs/Shared/EpicGames.Core/ManagedProcess.cs
@@ -99,7 +99,7 @@ namespace EpicGames.Core
 		/// </summary>
 		static internal bool SupportsJobObjects
 		{
-			get { return Environment.OSVersion.Platform == PlatformID.Win32NT || Environment.OSVersion.Platform == PlatformID.Win32S || Environment.OSVersion.Platform == PlatformID.Win32Windows; }
+			get { return RuntimePlatform.IsWindows; }
 		}
 
 		/// <summary>
diff --git a/Engine/Source/Programs/Shared/EpicGames.Core/RuntimePlatform.cs b/Engine/Source/Programs/Shared/EpicGames.Core/RuntimePlatform.cs
new file mode 100644
--- /dev/null
+++ b/Engine/Source/Programs/Shared/EpicGames.Core/RuntimePlatform.cs
@@ -0,0 +1,27 @@
+// Copyright Epic Games, Inc. All Rights Reserved.
+
+using System;
+using System.Collections.Generic;
+using System.Runtime.InteropServices;
+using System.Text;
+
+namespace EpicGames.Core
+{
+	public class RuntimePlatform
+	{
+		/// <summary>
+		/// Whether we are currently running on Linux.
+		/// </summary>
+		public static readonly bool IsLinux = RuntimeInformation.IsOSPlatform(OSPlatform.Linux);
+
+		/// <summary>
+		/// Whether we are currently running on a MacOS platform.
+		/// </summary>
+		public static readonly bool IsMac = RuntimeInformation.IsOSPlatform(OSPlatform.OSX);
+
+		/// <summary>
+		/// Whether we are currently running a Windows platform.
+		/// </summary>
+		public static readonly bool IsWindows = RuntimeInformation.IsOSPlatform(OSPlatform.Windows);
+	}
+}
diff --git a/Engine/Source/Programs/Shared/EpicGames.Perforce/PerforceChildProcess.cs b/Engine/Source/Programs/Shared/EpicGames.Perforce/PerforceChildProcess.cs
--- a/Engine/Source/Programs/Shared/EpicGames.Perforce/PerforceChildProcess.cs
+++ b/Engine/Source/Programs/Shared/EpicGames.Perforce/PerforceChildProcess.cs
@@ -128,7 +128,7 @@ namespace EpicGames.Perforce
 		public static string GetExecutable()
 		{
 			string PerforceFileName;
-			if (Environment.OSVersion.Platform == PlatformID.Win32NT || Environment.OSVersion.Platform == PlatformID.Win32S || Environment.OSVersion.Platform == PlatformID.Win32Windows)
+			if (RuntimePlatform.IsWindows)
 			{
 				PerforceFileName = "p4.exe";
 			}
diff --git a/Engine/Source/Programs/UnrealBuildTool/Configuration/UEBuildTarget.cs b/Engine/Source/Programs/UnrealBuildTool/Configuration/UEBuildTarget.cs
--- a/Engine/Source/Programs/UnrealBuildTool/Configuration/UEBuildTarget.cs
+++ b/Engine/Source/Programs/UnrealBuildTool/Configuration/UEBuildTarget.cs
@@ -3742,7 +3742,7 @@ namespace UnrealBuildTool
 				UEBuildBinary ExecutableBinary = Binaries[0];
 
 				bool IsCurrentPlatform;
-				if (!Utils.IsRunningOnWindows)
+				if (!RuntimePlatform.IsWindows)
 				{
 					IsCurrentPlatform = Platform == UnrealTargetPlatform.Mac || (UEBuildPlatform.IsPlatformInGroup(Platform, UnrealPlatformGroup.Unix) && Platform == BuildHostPlatform.Current.Platform);
 				}
diff --git a/Engine/Source/Programs/UnrealBuildTool/Executors/ParallelExecutor.cs b/Engine/Source/Programs/UnrealBuildTool/Executors/ParallelExecutor.cs
--- a/Engine/Source/Programs/UnrealBuildTool/Executors/ParallelExecutor.cs
+++ b/Engine/Source/Programs/UnrealBuildTool/Executors/ParallelExecutor.cs
@@ -113,7 +113,7 @@ namespace UnrealBuildTool
 		/// <returns>True if the parallel executor can be used</returns>
 		public static bool IsAvailable()
 		{
-			return Environment.OSVersion.Platform == PlatformID.Win32NT;
+			return RuntimePlatform.IsWindows;
 		}
 
 		/// <summary>
diff --git a/Engine/Source/Programs/UnrealBuildTool/Platform/Android/AndroidPlatformSDK.cs b/Engine/Source/Programs/UnrealBuildTool/Platform/Android/AndroidPlatformSDK.cs
--- a/Engine/Source/Programs/UnrealBuildTool/Platform/Android/AndroidPlatformSDK.cs
+++ b/Engine/Source/Programs/UnrealBuildTool/Platform/Android/AndroidPlatformSDK.cs
@@ -217,7 +217,7 @@ namespace UnrealBuildTool
 				}
 
 				// If we are not running on Windows and we are still missing a key then go and find it from the .bash_profile
-				if (!Utils.IsRunningOnWindows && !EnvVarNames.All(s => AndroidEnv.ContainsKey(s.Key)))
+				if (!RuntimePlatform.IsWindows && !EnvVarNames.All(s => AndroidEnv.ContainsKey(s.Key)))
 				{
 					string BashProfilePath = Path.Combine(Environment.GetFolderPath(Environment.SpecialFolder.Personal), ".bash_profile");
 					if (!File.Exists(BashProfilePath))
diff --git a/Engine/Source/Programs/UnrealBuildTool/Platform/Android/AndroidToolChain.cs b/Engine/Source/Programs/UnrealBuildTool/Platform/Android/AndroidToolChain.cs
--- a/Engine/Source/Programs/UnrealBuildTool/Platform/Android/AndroidToolChain.cs
+++ b/Engine/Source/Programs/UnrealBuildTool/Platform/Android/AndroidToolChain.cs
@@ -1336,7 +1336,7 @@ namespace UnrealBuildTool
 			foreach (string Definition in CompileEnvironment.Definitions)
 			{
 				// We must change the \" to escaped double quotes to save it properly for clang .rsp
-				BaseArguments += string.Format(" -D \"{0}\"", Definition.Replace("\"", Utils.IsRunningOnWindows ? "\"\"" : "\\\""));
+				BaseArguments += string.Format(" -D \"{0}\"", Definition.Replace("\"", RuntimePlatform.IsWindows ? "\"\"" : "\\\""));
 			}
 
 			//LUMIN_MERGE
@@ -2197,7 +2197,7 @@ namespace UnrealBuildTool
 					// To work around this we call clang to just get the command it would execute and generate a
 					// second response file to directly call ld with the right arguments instead of calling through clang.
 /* disable while tracking down some linker errors this introduces
-					if (Utils.IsRunningOnWindows)
+					if (RuntimePlatform.IsWindows)
 					{
 						// capture the actual link command without running it
 						ProcessStartInfo StartInfo = new ProcessStartInfo();
diff --git a/Engine/Source/Programs/UnrealBuildTool/Platform/Android/UEDeployAndroid.cs b/Engine/Source/Programs/UnrealBuildTool/Platform/Android/UEDeployAndroid.cs
--- a/Engine/Source/Programs/UnrealBuildTool/Platform/Android/UEDeployAndroid.cs
+++ b/Engine/Source/Programs/UnrealBuildTool/Platform/Android/UEDeployAndroid.cs
@@ -341,7 +341,7 @@ namespace UnrealBuildTool
 			uint BestVersion = 0;
 			foreach (string CandidateDir in Subdirs)
 			{
-				string AaptFilename = Path.Combine(CandidateDir, Utils.IsRunningOnWindows ? "aapt.exe" : "aapt");
+				string AaptFilename = Path.Combine(CandidateDir, RuntimePlatform.IsWindows ? "aapt.exe" : "aapt");
 				string RevisionString = "";
 				uint RevisionValue = 0;
 
@@ -1497,13 +1497,13 @@ namespace UnrealBuildTool
 			string SanitizerFullLibName = "libclang_rt." + LibName + Architecture + "-android.so";
 
 			string WrapSh = Path.Combine(UnrealBuildTool.EngineDirectory.ToString(), "Build", "Android", "ClangSanitizers", "wrap.sh");
+
 			string PlatformName = "windows-x86_64";
-			int PlatformId = (int)Environment.OSVersion.Platform;
-			if ((PlatformId == 4) || (PlatformId == 128))
+			if (RuntimePlatform.IsLinux)
 			{
 				PlatformName = "linux-x86_64";
 			}
-			else if (PlatformId == 6)
+			else if (RuntimePlatform.IsMac)
 			{
 				PlatformName = "darwin-x86_64";
 			}
@@ -3653,7 +3653,7 @@ namespace UnrealBuildTool
 			string BuildToolsVersion = GetBuildToolsVersion();
 
 			// cache some tools paths
-			//string NDKBuildPath = Environment.ExpandEnvironmentVariables("%NDKROOT%/ndk-build" + (Utils.IsRunningOnWindows ? ".cmd" : ""));
+			//string NDKBuildPath = Environment.ExpandEnvironmentVariables("%NDKROOT%/ndk-build" + (RuntimePlatform.IsWindows ? ".cmd" : ""));
 			//bool HasNDKPath = File.Exists(NDKBuildPath);
 
 			// set up some directory info
@@ -4232,7 +4232,7 @@ namespace UnrealBuildTool
 				if (!bSkipGradleBuild)
 				{
 					string GradleScriptPath = Path.Combine(UnrealBuildGradlePath, "gradlew");
-					if (!Utils.IsRunningOnWindows)
+					if (!RuntimePlatform.IsWindows)
 					{
 						// fix permissions for Mac/Linux
 						RunCommandLineProgramWithException(UnrealBuildGradlePath, "/bin/sh", string.Format("-c 'chmod 0755 \"{0}\"'", GradleScriptPath.Replace("'", "'\"'\"'")), "Fix gradlew permissions");
@@ -4265,9 +4265,9 @@ namespace UnrealBuildTool
 							Directory.CreateDirectory(Path.GetDirectoryName(DestApkName));
 
 							// Use gradle to build the .apk file
-							string ShellExecutable = Utils.IsRunningOnWindows ? "cmd.exe" : "/bin/sh";
-							string ShellParametersBegin = Utils.IsRunningOnWindows ? "/c " : "-c '";
-							string ShellParametersEnd = Utils.IsRunningOnWindows ? "" : "'";
+							string ShellExecutable = RuntimePlatform.IsWindows ? "cmd.exe" : "/bin/sh";
+							string ShellParametersBegin = RuntimePlatform.IsWindows ? "/c " : "-c '";
+							string ShellParametersEnd = RuntimePlatform.IsWindows ? "" : "'";
 							RunCommandLineProgramWithExceptionAndFiltering(UnrealBuildGradlePath, ShellExecutable, ShellParametersBegin + "\"" + GradleScriptPath + "\" " + GradleOptions + ShellParametersEnd, "Making .apk with Gradle...");
 
 							if (GradleSecondCallOptions != "")
@@ -4613,7 +4613,7 @@ namespace UnrealBuildTool
 					File.WriteAllText(GradlePropertiesFilename, GradlePropertiesContent);
 
 					string GradleScriptPath = Path.Combine(UnrealBuildGradlePath, "gradlew");
-					if (!Utils.IsRunningOnWindows)
+					if (!RuntimePlatform.IsWindows)
 					{
 						// fix permissions for Mac/Linux
 						RunCommandLineProgramWithException(UnrealBuildGradlePath, "/bin/sh", string.Format("-c 'chmod 0755 \"{0}\"'", GradleScriptPath.Replace("'", "'\"'\"'")), "Fix gradlew permissions");
@@ -4640,9 +4640,9 @@ namespace UnrealBuildTool
 					//Directory.CreateDirectory(Path.GetDirectoryName(DestApkName));
 
 					// Use gradle to build the .apk file
-					string ShellExecutable = Utils.IsRunningOnWindows ? "cmd.exe" : "/bin/sh";
-					string ShellParametersBegin = Utils.IsRunningOnWindows ? "/c " : "-c '";
-					string ShellParametersEnd = Utils.IsRunningOnWindows ? "" : "'";
+					string ShellExecutable = RuntimePlatform.IsWindows ? "cmd.exe" : "/bin/sh";
+					string ShellParametersBegin = RuntimePlatform.IsWindows ? "/c " : "-c '";
+					string ShellParametersEnd = RuntimePlatform.IsWindows ? "" : "'";
 					RunCommandLineProgramWithExceptionAndFiltering(UnrealBuildGradlePath, ShellExecutable, ShellParametersBegin + "\"" + GradleScriptPath + "\" " + GradleOptions + ShellParametersEnd, "Making .aab with Gradle...");
 
 					if (GradleSecondCallOptions != "")
@@ -4676,7 +4676,7 @@ namespace UnrealBuildTool
 						string UnrealBuildGradlePath = Path.Combine(UnrealBuildPath, "gradle");
 
 						string GradleScriptPath = Path.Combine(UnrealBuildGradlePath, "gradlew");
-						if (!Utils.IsRunningOnWindows)
+						if (!RuntimePlatform.IsWindows)
 						{
 							// fix permissions for Mac/Linux
 							RunCommandLineProgramWithException(UnrealBuildGradlePath, "/bin/sh", string.Format("-c 'chmod 0755 \"{0}\"'", GradleScriptPath.Replace("'", "'\"'\"'")), "Fix gradlew permissions");
@@ -4703,9 +4703,9 @@ namespace UnrealBuildTool
 						//Directory.CreateDirectory(Path.GetDirectoryName(DestApkName));
 
 						// Use gradle to build the .apk file
-						string ShellExecutable = Utils.IsRunningOnWindows ? "cmd.exe" : "/bin/sh";
-						string ShellParametersBegin = Utils.IsRunningOnWindows ? "/c " : "-c '";
-						string ShellParametersEnd = Utils.IsRunningOnWindows ? "" : "'";
+						string ShellExecutable = RuntimePlatform.IsWindows ? "cmd.exe" : "/bin/sh";
+						string ShellParametersBegin = RuntimePlatform.IsWindows ? "/c " : "-c '";
+						string ShellParametersEnd = RuntimePlatform.IsWindows ? "" : "'";
 						RunCommandLineProgramWithExceptionAndFiltering(UnrealBuildGradlePath, ShellExecutable, ShellParametersBegin + "\"" + GradleScriptPath + "\" " + GradleOptions + ShellParametersEnd, "Making .aab with Gradle...");
 
 						if (GradleSecondCallOptions != "")
diff --git a/Engine/Source/Programs/UnrealBuildTool/Platform/IOS/IOSToolChain.cs b/Engine/Source/Programs/UnrealBuildTool/Platform/IOS/IOSToolChain.cs
--- a/Engine/Source/Programs/UnrealBuildTool/Platform/IOS/IOSToolChain.cs
+++ b/Engine/Source/Programs/UnrealBuildTool/Platform/IOS/IOSToolChain.cs
@@ -235,7 +235,7 @@ namespace UnrealBuildTool
 			Result += " -pipe";
 			Result += " -fpascal-strings";
 
-			if (Utils.IsRunningOnWindows)
+			if (RuntimePlatform.IsWindows)
 			{
 				Result += " -fdiagnostics-format=msvc";
 			}
diff --git a/Engine/Source/Programs/UnrealBuildTool/Platform/Lumin/LuminToolChain.cs b/Engine/Source/Programs/UnrealBuildTool/Platform/Lumin/LuminToolChain.cs
--- a/Engine/Source/Programs/UnrealBuildTool/Platform/Lumin/LuminToolChain.cs
+++ b/Engine/Source/Programs/UnrealBuildTool/Platform/Lumin/LuminToolChain.cs
@@ -35,7 +35,7 @@ namespace UnrealBuildTool
 			{
 				if (MabuPath_ == null)
 				{
-					MabuPath_ = Environment.ExpandEnvironmentVariables("%MLSDK%/mabu" + (Utils.IsRunningOnWindows ? ".cmd" : ""));
+					MabuPath_ = Environment.ExpandEnvironmentVariables("%MLSDK%/mabu" + (RuntimePlatform.IsWindows ? ".cmd" : ""));
 
 					MabuPath_ = MabuPath_.Replace("\"", "");
 
diff --git a/Engine/Source/Programs/UnrealBuildTool/Platform/Mac/ApplePlatformSDK.cs b/Engine/Source/Programs/UnrealBuildTool/Platform/Mac/ApplePlatformSDK.cs
--- a/Engine/Source/Programs/UnrealBuildTool/Platform/Mac/ApplePlatformSDK.cs
+++ b/Engine/Source/Programs/UnrealBuildTool/Platform/Mac/ApplePlatformSDK.cs
@@ -21,7 +21,7 @@ namespace UnrealBuildTool
 
 		public override void GetValidVersionRange(out string MinVersion, out string MaxVersion)
 		{
-			if (IsMac())
+			if (RuntimePlatform.IsMac)
 			{
 				MinVersion = "11.0.0";
 				MaxVersion = "12.9.9";
@@ -42,14 +42,14 @@ namespace UnrealBuildTool
 		public override string GetInstalledSDKVersion()
 		{
 			// get xcode version on Mac
-			if (IsMac())
+			if (RuntimePlatform.IsMac)
 			{
 				string Output = Utils.RunLocalProcessAndReturnStdOut("sh", "-c 'xcodebuild -version'");
 				Match Result = Regex.Match(Output, @"Xcode (\S*)");
 				return Result.Success ? Result.Groups[1].Value : "";
 			}
 
-			if (IsWindows())
+			if (RuntimePlatform.IsWindows)
 			{
 				// otherwise, get iTunes "Version"
 				string DllPath = Registry.GetValue("HKEY_LOCAL_MACHINE\\SOFTWARE\\Wow6432Node\\Apple Inc.\\Apple Mobile Device Support\\Shared", "iTunesMobileDeviceDLL", null) as string;
diff --git a/Engine/Source/Programs/UnrealBuildTool/Platform/Mac/MacToolChain.cs b/Engine/Source/Programs/UnrealBuildTool/Platform/Mac/MacToolChain.cs
--- a/Engine/Source/Programs/UnrealBuildTool/Platform/Mac/MacToolChain.cs
+++ b/Engine/Source/Programs/UnrealBuildTool/Platform/Mac/MacToolChain.cs
@@ -774,7 +774,7 @@ namespace UnrealBuildTool
 
 			// Tell the action that we're building an import library here and it should conditionally be
 			// ignored as a prerequisite for other actions
-			LinkAction.bProducesImportLibrary = Utils.IsRunningOnWindows && (bBuildImportLibraryOnly || LinkEnvironment.bIsBuildingDLL);
+			LinkAction.bProducesImportLibrary = RuntimePlatform.IsWindows && (bBuildImportLibraryOnly || LinkEnvironment.bIsBuildingDLL);
 			
 			// Add the output file as a production of the link action.
 			FileItem OutputFile = FileItem.GetItemByFileReference(LinkEnvironment.OutputFilePath);
@@ -919,7 +919,7 @@ namespace UnrealBuildTool
 
 			if (!bIsBuildingLibrary)
 			{
-				if (UnrealBuildTool.IsEngineInstalled() || (!Utils.IsRunningOnWindows && LinkEnvironment.bIsCrossReferenced == false))
+				if (UnrealBuildTool.IsEngineInstalled() || (!RuntimePlatform.IsWindows && LinkEnvironment.bIsCrossReferenced == false))
 				{
 					foreach (string Library in EngineAndGameLibraries)
 					{
@@ -994,7 +994,7 @@ namespace UnrealBuildTool
 
 				StreamWriter FixDylibDepsScript = File.AppendText(FixDylibDepsScriptPath.FullName);
 
-				if (LinkEnvironment.bIsCrossReferenced || Utils.IsRunningOnWindows)
+				if (LinkEnvironment.bIsCrossReferenced || RuntimePlatform.IsWindows)
 				{
 					string EngineAndGameLibrariesString = "";
 					foreach (string Library in EngineAndGameLibraries)
diff --git a/Engine/Source/Programs/UnrealBuildTool/Platform/Windows/MicrosoftPlatformSDK.cs b/Engine/Source/Programs/UnrealBuildTool/Platform/Windows/MicrosoftPlatformSDK.cs
--- a/Engine/Source/Programs/UnrealBuildTool/Platform/Windows/MicrosoftPlatformSDK.cs
+++ b/Engine/Source/Programs/UnrealBuildTool/Platform/Windows/MicrosoftPlatformSDK.cs
@@ -42,7 +42,7 @@ namespace UnrealBuildTool
 
 		public override string? GetInstalledSDKVersion()
 		{
-			if (!Utils.IsRunningOnWindows)
+			if (!RuntimePlatform.IsWindows)
 			{
 				return null;
 			}
diff --git a/Engine/Source/Programs/UnrealBuildTool/ProjectFiles/ProjectFileGenerator.cs b/Engine/Source/Programs/UnrealBuildTool/ProjectFiles/ProjectFileGenerator.cs
--- a/Engine/Source/Programs/UnrealBuildTool/ProjectFiles/ProjectFileGenerator.cs
+++ b/Engine/Source/Programs/UnrealBuildTool/ProjectFiles/ProjectFileGenerator.cs
@@ -1797,7 +1797,7 @@ namespace UnrealBuildTool
 		{
 			if (ShouldGenerateIntelliSenseData() && Targets.Count > 0)
 			{
-				string ProgressInfoText = Utils.IsRunningOnWindows ? "Binding IntelliSense data..." : "Generating data for project indexing...";
+				string ProgressInfoText = RuntimePlatform.IsWindows ? "Binding IntelliSense data..." : "Generating data for project indexing...";
 				using (ProgressWriter Progress = new ProgressWriter(ProgressInfoText, true))
 				{
 					for (int TargetIndex = 0; TargetIndex < Targets.Count; ++TargetIndex)
diff --git a/Engine/Source/Programs/UnrealBuildTool/ProjectFiles/VisualStudioCode/VSCodeProjectFileGenerator.cs b/Engine/Source/Programs/UnrealBuildTool/ProjectFiles/VisualStudioCode/VSCodeProjectFileGenerator.cs
--- a/Engine/Source/Programs/UnrealBuildTool/ProjectFiles/VisualStudioCode/VSCodeProjectFileGenerator.cs
+++ b/Engine/Source/Programs/UnrealBuildTool/ProjectFiles/VisualStudioCode/VSCodeProjectFileGenerator.cs
@@ -950,7 +950,7 @@ namespace UnrealBuildTool
 						{
 							OutFile.AddField("label", TaskName);
 							OutFile.AddField("group", "build");
-							if (!Utils.IsRunningOnWindows)
+							if (!RuntimePlatform.IsWindows)
 							{
 								OutFile.AddField("command", MakePathString(FileReference.Combine(UE4ProjectRoot, "Engine", "Build", "BatchFiles", HostPlatform.ToString(), "RunDotnet.sh")));
 							}
@@ -1562,7 +1562,7 @@ namespace UnrealBuildTool
 					WorkspaceFile.AddUnnamedField("ms-dotnettools.csharp");
 
 					// If the platform we run the generator on uses mono, there are additional debugging extensions to add.
-					if (!Utils.IsRunningOnWindows)
+					if (!RuntimePlatform.IsWindows)
 					{
 						WorkspaceFile.AddUnnamedField("vadimcn.vscode-lldb");
 						WorkspaceFile.AddUnnamedField("ms-vscode.mono-debug");
diff --git a/Engine/Source/Programs/UnrealBuildTool/System/BuildHostPlatform.cs b/Engine/Source/Programs/UnrealBuildTool/System/BuildHostPlatform.cs
--- a/Engine/Source/Programs/UnrealBuildTool/System/BuildHostPlatform.cs
+++ b/Engine/Source/Programs/UnrealBuildTool/System/BuildHostPlatform.cs
@@ -33,26 +33,6 @@ namespace UnrealBuildTool
 	public abstract class BuildHostPlatform
 	{
 		private static BuildHostPlatform? CurrentPlatform;
-		private static bool bIsMac = File.Exists("/System/Library/CoreServices/SystemVersion.plist");
-
-		/// <summary>
-		/// Returns the name of platform UBT is running on. Internal use only. If you need access this this enum, use BuildHostPlatform.Current.Platform */
-		/// </summary>
-		private static UnrealTargetPlatform GetRuntimePlatform()
-		{
-			PlatformID Platform = Environment.OSVersion.Platform;
-			switch (Platform)
-			{
-				case PlatformID.Win32NT:
-					return UnrealTargetPlatform.Win64;
-				case PlatformID.Unix:
-					return bIsMac ? UnrealTargetPlatform.Mac : UnrealTargetPlatform.Linux;
-				case PlatformID.MacOSX:
-					return UnrealTargetPlatform.Mac;
-				default:
-					throw new BuildException("Unhandled runtime platform " + Platform);
-			}
-		}
 
 		/// <summary>
 		/// Host platform singleton.
@@ -63,16 +43,15 @@ namespace UnrealBuildTool
 			{
 				if (CurrentPlatform == null)
 				{
-					UnrealTargetPlatform RuntimePlatform = GetRuntimePlatform();
-					if (RuntimePlatform == UnrealTargetPlatform.Win64)
+					if (RuntimePlatform.IsWindows)
 					{
 						CurrentPlatform = new WindowsBuildHostPlatform();
 					}
-					else if (RuntimePlatform == UnrealTargetPlatform.Mac)
+					else if (RuntimePlatform.IsMac)
 					{
 						CurrentPlatform = new MacBuildHostPlatform();
 					}
-					else if (RuntimePlatform == UnrealTargetPlatform.Linux)
+					else if (RuntimePlatform.IsLinux)
 					{
 						CurrentPlatform = new LinuxBuildHostPlatform();
 					}
diff --git a/Engine/Source/Programs/UnrealBuildTool/System/ConfigHierarchy.cs b/Engine/Source/Programs/UnrealBuildTool/System/ConfigHierarchy.cs
--- a/Engine/Source/Programs/UnrealBuildTool/System/ConfigHierarchy.cs
+++ b/Engine/Source/Programs/UnrealBuildTool/System/ConfigHierarchy.cs
@@ -942,7 +942,7 @@ namespace UnrealBuildTool
 			if (!String.IsNullOrEmpty(PersonalFolder))
 			{
 				PersonalConfigFolder = PersonalFolder;
-				if (BuildHostPlatform.Current.Platform == UnrealTargetPlatform.Mac || Environment.OSVersion.Platform == PlatformID.Unix)
+				if (RuntimePlatform.IsMac || RuntimePlatform.IsLinux)
 				{
 					PersonalConfigFolder = System.IO.Path.Combine(PersonalConfigFolder, "Documents");
 				}
diff --git a/Engine/Source/Programs/UnrealBuildTool/System/UnrealPluginLanguage.cs b/Engine/Source/Programs/UnrealBuildTool/System/UnrealPluginLanguage.cs
--- a/Engine/Source/Programs/UnrealBuildTool/System/UnrealPluginLanguage.cs
+++ b/Engine/Source/Programs/UnrealBuildTool/System/UnrealPluginLanguage.cs
@@ -3303,11 +3303,7 @@ namespace UnrealBuildTool
 
 			DirectoryReference UserSettingsFolder = Utils.GetUserSettingDirectory(); // Match FPlatformProcess::UserSettingsDir()
 			DirectoryReference PersonalFolder = null; // Match FPlatformProcess::UserDir()
-			if (BuildHostPlatform.Current.Platform == UnrealTargetPlatform.Mac)
-			{
-				PersonalFolder = new DirectoryReference(Path.Combine(Environment.GetFolderPath(Environment.SpecialFolder.Personal), "Documents"));
-			}
-			else if (Environment.OSVersion.Platform == PlatformID.Unix)
+			if (RuntimePlatform.IsMac || RuntimePlatform.IsLinux)
 			{
 				PersonalFolder = new DirectoryReference(Path.Combine(Environment.GetFolderPath(Environment.SpecialFolder.Personal), "Documents"));
 			}
@@ -3315,7 +3311,7 @@ namespace UnrealBuildTool
 			{
 				// Not all user accounts have a local application data directory (eg. SYSTEM, used by Jenkins for builds).
 				string PersonalFolderSetting = Environment.GetFolderPath(Environment.SpecialFolder.Personal);
-				if(!String.IsNullOrEmpty(PersonalFolderSetting))
+				if (!String.IsNullOrEmpty(PersonalFolderSetting))
 				{
 					PersonalFolder = new DirectoryReference(PersonalFolderSetting);
 				}
diff --git a/Engine/Source/Programs/UnrealBuildTool/System/Utils.cs b/Engine/Source/Programs/UnrealBuildTool/System/Utils.cs
--- a/Engine/Source/Programs/UnrealBuildTool/System/Utils.cs
+++ b/Engine/Source/Programs/UnrealBuildTool/System/Utils.cs
@@ -20,21 +20,6 @@ namespace UnrealBuildTool
 	/// </summary>
 	public static class Utils
 	{
-		/// <summary>
-		/// Whether we are currently running on Linux.
-		/// </summary>
-		public static readonly bool IsRunningOnLinux = RuntimeInformation.IsOSPlatform(OSPlatform.Linux);
-
-		/// <summary>
-		/// Whether we are currently running on a MacOS platform.
-		/// </summary>
-		public static readonly bool IsRunningOnMac = RuntimeInformation.IsOSPlatform(OSPlatform.OSX);
-		
-		/// <summary>
-		/// Whether we are currently running a Windows platform.
-		/// </summary>
-		public static readonly bool IsRunningOnWindows = RuntimeInformation.IsOSPlatform(OSPlatform.Windows);
-
 		/// <summary>
 		/// Searches for a flag in a set of command-line arguments.
 		/// </summary>
@@ -593,7 +578,7 @@ namespace UnrealBuildTool
 			StringBuilder? CleanPath = null;
 			if (UseDirectorySeparatorChar == '\0')
 			{
-				UseDirectorySeparatorChar = Environment.OSVersion.Platform == PlatformID.Unix ? '/' : '\\';
+				UseDirectorySeparatorChar = Path.DirectorySeparatorChar;
 			}
 			char PrevC = '\0';
 			// Don't check for double separators until we run across a valid dir name. Paths that start with '//' or '\\' can still be valid.			
@@ -996,11 +981,11 @@ namespace UnrealBuildTool
 		/// </summary>
 		public static DirectoryReference? GetUserSettingDirectory()
 		{
-			if (BuildHostPlatform.Current.Platform == UnrealTargetPlatform.Mac)
+			if (RuntimePlatform.IsMac)
 			{
 				return new DirectoryReference(Path.Combine(Environment.GetFolderPath(Environment.SpecialFolder.Personal), "Library", "Application Support", "Epic"));
 			}
-			else if (Environment.OSVersion.Platform == PlatformID.Unix)
+			else if (RuntimePlatform.IsLinux)
 			{
 				return new DirectoryReference(Path.Combine(Environment.GetFolderPath(Environment.SpecialFolder.ApplicationData), "Epic"));
 			}
@@ -1039,7 +1024,7 @@ namespace UnrealBuildTool
 		public static int GetLogicalProcessorCount()
 		{
 			// This function uses Windows P/Invoke calls; if we're not running on Windows, just return the default.
-			if(Utils.IsRunningOnWindows)
+			if(RuntimePlatform.IsWindows)
 			{
 				const int ERROR_INSUFFICIENT_BUFFER = 122;
 
@@ -1096,7 +1081,7 @@ namespace UnrealBuildTool
 		/// <returns>The number of physical cores, or -1 if it could not be obtained</returns>
 		public static int GetPhysicalProcessorCount()
 		{
-			if (Utils.IsRunningOnWindows)
+			if (RuntimePlatform.IsWindows)
 			{
 				const int ERROR_INSUFFICIENT_BUFFER = 122;
 
@@ -1136,7 +1121,7 @@ namespace UnrealBuildTool
 					}
 				}
 			}
-			else if(Utils.IsRunningOnMac)
+			else if(RuntimePlatform.IsMac)
 			{
 				UInt64 Size = 4;
 				if (0 == sysctlbyname("hw.physicalcpu", out int Value, ref Size, IntPtr.Zero, 0))
@@ -1225,7 +1210,7 @@ namespace UnrealBuildTool
 			// On Mac, MemoryInfo.MemoryLoadBytes includes memory used to cache disk-backed files ("Cached Files" in
 			// Activity Monitor), which can result in a significant over-estimate of memory pressure.
 			// We treat memory used for caching of disk-backed files as free for use in compilation tasks.
-			if (Utils.IsRunningOnMac)	
+			if (RuntimePlatform.IsMac)	
 			{
 				// host_statistics64() flavor, from <mach/host_info.h>
 				int HOST_VM_INFO64 = 4;
diff --git a/Engine/Source/Programs/UnrealBuildTool/ToolChain/AppleToolChain.cs b/Engine/Source/Programs/UnrealBuildTool/ToolChain/AppleToolChain.cs
--- a/Engine/Source/Programs/UnrealBuildTool/ToolChain/AppleToolChain.cs
+++ b/Engine/Source/Programs/UnrealBuildTool/ToolChain/AppleToolChain.cs
@@ -145,7 +145,7 @@ namespace UnrealBuildTool
 			}
 
 			// make sure we have a valid SDK directory
-			if (!Utils.IsRunningOnWindows && !Directory.Exists(Path.Combine(BaseSDKDir, OSPrefix + PlatformSDKVersion + ".sdk")))
+			if (!RuntimePlatform.IsWindows && !Directory.Exists(Path.Combine(BaseSDKDir, OSPrefix + PlatformSDKVersion + ".sdk")))
 			{
 				throw new BuildException("Invalid SDK {0}{1}.sdk, not found in {2}", OSPrefix, PlatformSDKVersion, BaseSDKDir);
 			}
