Build fix Ios/Tvos

diff --git a/Engine/Source/Programs/UnrealBuildTool/Platform/IOS/IOSToolChain.cs b/Engine/Source/Programs/UnrealBuildTool/Platform/IOS/IOSToolChain.cs
--- a/Engine/Source/Programs/UnrealBuildTool/Platform/IOS/IOSToolChain.cs
+++ b/Engine/Source/Programs/UnrealBuildTool/Platform/IOS/IOSToolChain.cs
@@ -643,7 +643,8 @@ namespace UnrealBuildTool
 
 			foreach (string Definition in CompileEnvironment.Definitions)
 			{
-				Arguments += string.Format(" -D\"{0}\"", Definition);
+				string DefinitionArgument = Definition.Contains("\"") ? Definition.Replace("\"", "\\\"") : Definition;
+				Arguments += string.Format(" -D \"{0}\"", DefinitionArgument);
 			}
 
 			CPPOutput Result = new CPPOutput();
