UnrealBuildTool: move some path constants into BuildUtilities

RootDirectory, EngineDirectory, UnrealBuildToolPath are now found in BuildUtilities' UnrealBuild namesapce.

The way these are computed has changed. Previously, it was assumed that the application is UnrealBuildTool, and paths were constructed relative to that assembly.

Now, the assumption is that the process is located under a "Engine/Build/DotNET" sub-path and paths are constructed relative to that.

diff --git a/Engine/Source/Programs/DotNETCommon/BuildUtilities/UnrealBuild.cs b/Engine/Source/Programs/DotNETCommon/BuildUtilities/UnrealBuild.cs
new file mode 100644
--- /dev/null
+++ b/Engine/Source/Programs/DotNETCommon/BuildUtilities/UnrealBuild.cs
@@ -0,0 +1,93 @@
+// Copyright Epic Games, Inc. All Rights Reserved.
+
+using EpicGames.Core;
+using System;
+using System.Collections.Generic;
+using System.IO;
+using System.Reflection;
+using System.Text;
+
+namespace UnrealBuildBase
+{
+	public static class UnrealBuild
+	{
+		private static DirectoryReference FindRootDirectory()
+		{
+			// This base library may be used - and so be launched - from more than one location (at time of writing, UnrealBuildTool and AutomationTool)
+			// Programs that use this assembly must be located under "Engine/Binaries/DotNET" and so we look for that sequence of directories in that path of the executing assembly
+			
+			// Use the EntryAssembly (the application path), rather than the ExecutingAssembly (the library path)
+			string AssemblyLocation = Assembly.GetEntryAssembly().GetOriginalLocation();
+
+			DirectoryReference? FoundRootDirectory = DirectoryReference.FromString(AssemblyLocation);
+
+			// Search up through the directory tree for the deepest instance of the sub-path "Engine/Binaries/DotNET"
+			while(FoundRootDirectory != null)
+			{
+				if (String.Equals("DotNET", FoundRootDirectory.GetDirectoryName()))
+				{
+					FoundRootDirectory = FoundRootDirectory.ParentDirectory;
+					if (FoundRootDirectory != null && String.Equals("Binaries", FoundRootDirectory.GetDirectoryName()))
+					{
+						FoundRootDirectory = FoundRootDirectory.ParentDirectory;
+						if (FoundRootDirectory != null && String.Equals("Engine", FoundRootDirectory.GetDirectoryName()))
+						{
+							FoundRootDirectory = FoundRootDirectory.ParentDirectory;
+							break;
+						}
+						continue;
+					}
+					continue;
+				}
+				FoundRootDirectory = FoundRootDirectory.ParentDirectory;
+			}
+
+			if (FoundRootDirectory == null)
+			{
+				throw new Exception($"The BuildUtilities assembly requires that applications are launched from a path containing \"Engine/Binaries/DotNET\". This application was launched from {Path.GetDirectoryName(AssemblyLocation)}");
+			}
+
+			// Confirm that we've found a valid root directory, by testing for the existence of a well-known file
+			FileReference ExpectedExistingFile = FileReference.Combine(FoundRootDirectory, "Engine", "Build", "Build.version");
+			if (!FileReference.Exists(ExpectedExistingFile))
+			{
+				throw new Exception($"Expected file \"Engine/Build/Build.version\" was not found at {ExpectedExistingFile.FullName}");
+			}
+
+			return FoundRootDirectory;
+		}
+
+		private static FileReference FindUnrealBuildTool()
+		{
+			// todo: use UnrealBuildTool.dll (same on all platforms). Will require changes wherever UnrealBuildTool is invoked.
+			string UBTName = RuntimePlatform.IsWindows ? "UnrealBuildTool.exe" : "UnrealBuildTool";
+
+			// the UnrealBuildTool executable is assumed to be located under {RootDirectory}/Engine/Binaries/DotNET/UnrealBuildTool/
+			FileReference UnrealBuildToolPath = FileReference.Combine(EngineDirectory, "Binaries", "DotNET", "UnrealBuildTool", UBTName);
+			
+			FileReference UnrealBuildToolPathCorrected = FileReference.FindCorrectCase(UnrealBuildToolPath);
+
+			if (!FileReference.Exists(UnrealBuildToolPathCorrected))
+			{
+				throw new Exception($"Unable to find {UBTName} in the expected location at {UnrealBuildToolPathCorrected.FullName}");
+			}
+
+			return UnrealBuildToolPathCorrected;
+		}
+
+		/// <summary>
+		/// The full name of the root UE directory
+		/// </summary>
+		public static readonly DirectoryReference RootDirectory = FindRootDirectory();
+
+		/// <summary>
+		/// The full name of the Engine directory
+		/// </summary>
+		public static readonly DirectoryReference EngineDirectory = DirectoryReference.Combine(RootDirectory, "Engine");
+
+		/// <summary>
+		/// The path to UBT
+		/// </summary>
+		public static readonly FileReference UnrealBuildToolPath = FindUnrealBuildTool();
+	}
+}
diff --git a/Engine/Source/Programs/UnrealBuildTool/Configuration/ModuleRules.cs b/Engine/Source/Programs/UnrealBuildTool/Configuration/ModuleRules.cs
--- a/Engine/Source/Programs/UnrealBuildTool/Configuration/ModuleRules.cs
+++ b/Engine/Source/Programs/UnrealBuildTool/Configuration/ModuleRules.cs
@@ -8,6 +8,7 @@ using System.Linq;
 using System.Text;
 using System.Threading.Tasks;
 using EpicGames.Core;
+using UnrealBuildBase;
 
 namespace UnrealBuildTool
 {
@@ -1119,7 +1120,7 @@ namespace UnrealBuildTool
 		{
 			get
 			{
-				return UnrealBuildTool.EngineDirectory.FullName;
+				return UnrealBuild.EngineDirectory.FullName;
 			}
 		}
 
@@ -1415,7 +1416,7 @@ namespace UnrealBuildTool
 					case ModuleRules.PrecompileTargetsType.None:
 						return false;
 					case ModuleRules.PrecompileTargetsType.Default:
-						return (Target.Type == TargetType.Editor || !UnrealBuildTool.GetExtensionDirs(UnrealBuildTool.EngineDirectory, "Source/Developer").Any(Dir => RulesFile.IsUnderDirectory(Dir)) || Plugin != null);
+						return (Target.Type == TargetType.Editor || !UnrealBuildTool.GetExtensionDirs(UnrealBuild.EngineDirectory, "Source/Developer").Any(Dir => RulesFile.IsUnderDirectory(Dir)) || Plugin != null);
 					case ModuleRules.PrecompileTargetsType.Game:
 						return (Target.Type == TargetType.Client || Target.Type == TargetType.Server || Target.Type == TargetType.Game);
 					case ModuleRules.PrecompileTargetsType.Editor:
diff --git a/Engine/Source/Programs/UnrealBuildTool/Configuration/TargetRules.cs b/Engine/Source/Programs/UnrealBuildTool/Configuration/TargetRules.cs
--- a/Engine/Source/Programs/UnrealBuildTool/Configuration/TargetRules.cs
+++ b/Engine/Source/Programs/UnrealBuildTool/Configuration/TargetRules.cs
@@ -5,6 +5,7 @@ using System.Collections.Generic;
 using System.Linq;
 using System.Reflection;
 using EpicGames.Core;
+using UnrealBuildBase;
 
 namespace UnrealBuildTool
 {
@@ -2874,7 +2875,7 @@ namespace UnrealBuildTool
 		/// </summary>
 		public string RelativeEnginePath
 		{
-			get { return UnrealBuildTool.EngineDirectory.MakeRelativeTo(DirectoryReference.GetCurrentDirectory()); }
+			get { return UnrealBuild.EngineDirectory.MakeRelativeTo(DirectoryReference.GetCurrentDirectory()); }
 		}
 
 		/// <summary>
diff --git a/Engine/Source/Programs/UnrealBuildTool/Configuration/UEBuildBinary.cs b/Engine/Source/Programs/UnrealBuildTool/Configuration/UEBuildBinary.cs
--- a/Engine/Source/Programs/UnrealBuildTool/Configuration/UEBuildBinary.cs
+++ b/Engine/Source/Programs/UnrealBuildTool/Configuration/UEBuildBinary.cs
@@ -804,7 +804,7 @@ namespace UnrealBuildTool
 						ResourceCompileEnvironment.Definitions.Add(String.Format("BUILD_VERSION={0}", Target.BuildVersion));
 
 						// Otherwise compile the default resource file per-binary, so that it gets the correct ORIGINAL_FILE_NAME macro.
-						FileItem DefaultResourceFile = FileItem.GetItemByFileReference(FileReference.Combine(UnrealBuildTool.EngineDirectory, "Build", "Windows", "Resources", "Default.rc2"));
+						FileItem DefaultResourceFile = FileItem.GetItemByFileReference(FileReference.Combine(UnrealBuild.EngineDirectory, "Build", "Windows", "Resources", "Default.rc2"));
 						CPPOutput DefaultResourceOutput = ToolChain.CompileRCFiles(ResourceCompileEnvironment, new List<FileItem> { DefaultResourceFile }, ResourceIntermediateDirectory, Graph);
 						BinaryLinkEnvironment.InputFiles.AddRange(DefaultResourceOutput.ObjectFiles);
 					}
diff --git a/Engine/Source/Programs/UnrealBuildTool/Configuration/UEBuildModule.cs b/Engine/Source/Programs/UnrealBuildTool/Configuration/UEBuildModule.cs
--- a/Engine/Source/Programs/UnrealBuildTool/Configuration/UEBuildModule.cs
+++ b/Engine/Source/Programs/UnrealBuildTool/Configuration/UEBuildModule.cs
@@ -220,7 +220,7 @@ namespace UnrealBuildTool
 					if (FrameworkRules.IsZipFile())
 					{
 						// If FrameworkPath ends in .zip, it needs to be extracted
-						Framework = new UEBuildFramework(FrameworkRules.Name, FileReference.Combine(ModuleDirectory, FrameworkRules.Path), DirectoryReference.Combine(UnrealBuildTool.EngineDirectory, "Intermediate", "UnzippedFrameworks", FrameworkRules.Name, Path.GetFileNameWithoutExtension(FrameworkRules.Path)), FrameworkRules.CopyBundledAssets, FrameworkRules.bCopyFramework);
+						Framework = new UEBuildFramework(FrameworkRules.Name, FileReference.Combine(ModuleDirectory, FrameworkRules.Path), DirectoryReference.Combine(UnrealBuild.EngineDirectory, "Intermediate", "UnzippedFrameworks", FrameworkRules.Name, Path.GetFileNameWithoutExtension(FrameworkRules.Path)), FrameworkRules.CopyBundledAssets, FrameworkRules.bCopyFramework);
 					}
 					else
 					{
@@ -633,13 +633,13 @@ namespace UnrealBuildTool
 					{
 						if(MatchVariableName(Path, StartIdx, EndIdx, "EngineDir"))
 						{
-							Path = UnrealBuildTool.EngineDirectory + Path.Substring(EndIdx + 1);
+							Path = UnrealBuild.EngineDirectory + Path.Substring(EndIdx + 1);
 						}
 						else if(MatchVariableName(Path, StartIdx, EndIdx, "ProjectDir"))
 						{
 							if(Rules.Target.ProjectFile == null)
 							{
-								Path = UnrealBuildTool.EngineDirectory + Path.Substring(EndIdx + 1);
+								Path = UnrealBuild.EngineDirectory + Path.Substring(EndIdx + 1);
 							}
 							else
 							{
diff --git a/Engine/Source/Programs/UnrealBuildTool/Configuration/UEBuildTarget.cs b/Engine/Source/Programs/UnrealBuildTool/Configuration/UEBuildTarget.cs
--- a/Engine/Source/Programs/UnrealBuildTool/Configuration/UEBuildTarget.cs
+++ b/Engine/Source/Programs/UnrealBuildTool/Configuration/UEBuildTarget.cs
@@ -723,7 +723,7 @@ namespace UnrealBuildTool
 				DirectoryReference DependencyListDir;
 				if(RulesObject.ProjectFile == null)
 				{
-					DependencyListDir = DirectoryReference.Combine(UnrealBuildTool.EngineDirectory, "Intermediate", "DependencyLists", RulesObject.Name, RulesObject.Configuration.ToString(), RulesObject.Platform.ToString());
+					DependencyListDir = DirectoryReference.Combine(UnrealBuild.EngineDirectory, "Intermediate", "DependencyLists", RulesObject.Name, RulesObject.Configuration.ToString(), RulesObject.Platform.ToString());
 				}
 				else
 				{
@@ -1085,7 +1085,7 @@ namespace UnrealBuildTool
 			}
 			else
 			{
-				ProjectDirectory = UnrealBuildTool.EngineDirectory;
+				ProjectDirectory = UnrealBuild.EngineDirectory;
 			}
 
 			// Build the project intermediate directory
@@ -1098,11 +1098,11 @@ namespace UnrealBuildTool
 			}
 			else if (Configuration == UnrealTargetConfiguration.DebugGame)
 			{
-				EngineIntermediateDirectory = DirectoryReference.Combine(UnrealBuildTool.EngineDirectory, PlatformIntermediateFolder, AppName, UnrealTargetConfiguration.Development.ToString());
+				EngineIntermediateDirectory = DirectoryReference.Combine(UnrealBuild.EngineDirectory, PlatformIntermediateFolder, AppName, UnrealTargetConfiguration.Development.ToString());
 			}
 			else
 			{
-				EngineIntermediateDirectory = DirectoryReference.Combine(UnrealBuildTool.EngineDirectory, PlatformIntermediateFolder, AppName, Configuration.ToString());
+				EngineIntermediateDirectory = DirectoryReference.Combine(UnrealBuild.EngineDirectory, PlatformIntermediateFolder, AppName, Configuration.ToString());
 			}
 
 			// Get the receipt path for this target
@@ -1289,7 +1289,7 @@ namespace UnrealBuildTool
 			// Write the file
 			Log.TraceInformation("Writing dependency list to {0}", Location);
 			DirectoryReference.CreateDirectory(Location.Directory);
-			FileReference.WriteAllLines(Location, Files.Where(x => x.IsUnderDirectory(UnrealBuildTool.RootDirectory)).Select(x => x.MakeRelativeTo(UnrealBuildTool.RootDirectory).Replace(Path.DirectorySeparatorChar, '/')).OrderBy(x => x));
+			FileReference.WriteAllLines(Location, Files.Where(x => x.IsUnderDirectory(UnrealBuild.RootDirectory)).Select(x => x.MakeRelativeTo(UnrealBuild.RootDirectory).Replace(Path.DirectorySeparatorChar, '/')).OrderBy(x => x));
 		}
 
 		private void ResolveLibraryName(string LibraryPath, string LibraryName, string LibraryExtension, HashSet<FileReference> Files)
@@ -1718,7 +1718,7 @@ namespace UnrealBuildTool
 				{
 					if(!Rules.bFormalBuild)
 					{
-						FileReference DefaultResourceLocation = FileReference.Combine(UnrealBuildTool.EngineDirectory, "Build", "Windows", "Resources", "Default.rc2");
+						FileReference DefaultResourceLocation = FileReference.Combine(UnrealBuild.EngineDirectory, "Build", "Windows", "Resources", "Default.rc2");
 						if (!UnrealBuildTool.IsFileInstalled(DefaultResourceLocation))
 						{
 							CppCompileEnvironment DefaultResourceCompileEnvironment = new CppCompileEnvironment(GlobalCompileEnvironment);
@@ -1935,7 +1935,7 @@ namespace UnrealBuildTool
 				if (ProjectFile == null && !Rules.bLegalToDistributeBinary)
 				{
 					List<DirectoryReference> RootDirectories = new List<DirectoryReference>();
-					RootDirectories.Add(UnrealBuildTool.EngineDirectory);
+					RootDirectories.Add(UnrealBuild.EngineDirectory);
 					if (ProjectFile != null)
 					{
 						DirectoryReference ProjectDir = DirectoryReference.FromFile(ProjectFile);
@@ -2060,9 +2060,9 @@ namespace UnrealBuildTool
 		{
 			// Find the base directory for this binary
 			DirectoryReference BaseDir;
-			if(File.IsUnderDirectory(UnrealBuildTool.RootDirectory))
+			if(File.IsUnderDirectory(UnrealBuild.RootDirectory))
 			{
-				BaseDir = UnrealBuildTool.RootDirectory;
+				BaseDir = UnrealBuild.RootDirectory;
 			}
 			else if(ProjectDirectory != null && File.IsUnderDirectory(ProjectDirectory))
 			{
@@ -2299,7 +2299,7 @@ namespace UnrealBuildTool
 
 			if(Module.RulesFile != null)
 			{
-				foreach (DirectoryReference ExtensionDir in UnrealBuildTool.GetExtensionDirs(UnrealBuildTool.EngineDirectory))
+				foreach (DirectoryReference ExtensionDir in UnrealBuildTool.GetExtensionDirs(UnrealBuild.EngineDirectory))
 				{
 					DirectoryReference SourceDeveloperDir = DirectoryReference.Combine(ExtensionDir, "Source/Developer");
 					if(Module.RulesFile.IsUnderDirectory(SourceDeveloperDir))
@@ -2482,8 +2482,8 @@ namespace UnrealBuildTool
 		private Dictionary<string, string> GetTargetVariables(UEBuildPlugin? Plugin)
 		{
 			Dictionary<string, string> Variables = new Dictionary<string,string>();
-			Variables.Add("RootDir", UnrealBuildTool.RootDirectory.FullName);
-			Variables.Add("EngineDir", UnrealBuildTool.EngineDirectory.FullName);
+			Variables.Add("RootDir", UnrealBuild.RootDirectory.FullName);
+			Variables.Add("EngineDir", UnrealBuild.EngineDirectory.FullName);
 			Variables.Add("ProjectDir", ProjectDirectory.FullName);
 			Variables.Add("TargetName", TargetName);
 			Variables.Add("TargetPlatform", Platform.ToString());
@@ -2678,15 +2678,15 @@ namespace UnrealBuildTool
 				List<DirectoryReference> Directories = new List<DirectoryReference>();
 				if (TargetType == TargetType.Editor)
 				{
-					Directories.AddRange(UnrealBuildTool.GetExtensionDirs(UnrealBuildTool.EngineDirectory, "Source/Editor"));
+					Directories.AddRange(UnrealBuildTool.GetExtensionDirs(UnrealBuild.EngineDirectory, "Source/Editor"));
 				}
-				Directories.AddRange(UnrealBuildTool.GetExtensionDirs(UnrealBuildTool.EngineDirectory, "Source/Runtime"));
+				Directories.AddRange(UnrealBuildTool.GetExtensionDirs(UnrealBuild.EngineDirectory, "Source/Runtime"));
 
 				// Also allow anything in the developer directory in non-shipping configurations (though we blacklist by default unless the PrecompileForTargets
 				// setting indicates that it's actually useful at runtime).
 				if(Rules.bBuildDeveloperTools)
 				{
-					Directories.AddRange(UnrealBuildTool.GetExtensionDirs(UnrealBuildTool.EngineDirectory, "Source/Developer"));
+					Directories.AddRange(UnrealBuildTool.GetExtensionDirs(UnrealBuild.EngineDirectory, "Source/Developer"));
 				}
 
 				// Find all the modules that are not part of the standard set
@@ -2711,7 +2711,7 @@ namespace UnrealBuildTool
 							}
 							else
 							{
-								if (!ModuleFileName.ContainsAnyNames(ExcludeFolders, UnrealBuildTool.EngineDirectory))
+								if (!ModuleFileName.ContainsAnyNames(ExcludeFolders, UnrealBuild.EngineDirectory))
 								{
 									FilteredModuleNames.Add(ModuleName);
 								}
@@ -2850,7 +2850,7 @@ namespace UnrealBuildTool
 		/// <returns>The output directory for this target</returns>
 		public static DirectoryReference GetOutputDirectoryForExecutable(DirectoryReference BaseDir, FileReference TargetFile)
 		{
-			return UnrealBuildTool.GetExtensionDirs(BaseDir).Where(x => TargetFile.IsUnderDirectory(x)).OrderByDescending(x => x.FullName.Length).FirstOrDefault() ?? UnrealBuildTool.EngineDirectory;
+			return UnrealBuildTool.GetExtensionDirs(BaseDir).Where(x => TargetFile.IsUnderDirectory(x)).OrderByDescending(x => x.FullName.Length).FirstOrDefault() ?? UnrealBuild.EngineDirectory;
 		}
 
 		/// <summary>
@@ -3114,7 +3114,7 @@ namespace UnrealBuildTool
 			// If this is a program, synthesize references for plugins which are enabled via the config file
 			if(TargetType == TargetType.Program)
 			{
-				ConfigHierarchy EngineConfig = ConfigCache.ReadHierarchy(ConfigHierarchyType.Engine, DirectoryReference.Combine(UnrealBuildTool.EngineDirectory, "Programs", TargetName), Platform);
+				ConfigHierarchy EngineConfig = ConfigCache.ReadHierarchy(ConfigHierarchyType.Engine, DirectoryReference.Combine(UnrealBuild.EngineDirectory, "Programs", TargetName), Platform);
 
 				List<string>? PluginNames;
 				if(EngineConfig.GetArray("Plugins", "ProgramEnabledPlugins", out PluginNames))
@@ -3306,7 +3306,7 @@ namespace UnrealBuildTool
 		{
 			if (Plugin.LoadedFrom == PluginLoadedFrom.Engine)
 			{
-				return Plugin.File.ContainsAnyNames(ExcludeFolders, UnrealBuildTool.EngineDirectory);
+				return Plugin.File.ContainsAnyNames(ExcludeFolders, UnrealBuild.EngineDirectory);
 			}
 			else if(ProjectFile != null)
 			{
@@ -3334,9 +3334,9 @@ namespace UnrealBuildTool
 
 			// Get the intermediate directory for the launch module directory. This can differ from the standard engine intermediate directory because it is always configuration-specific.
 			DirectoryReference IntermediateDirectory;
-			if(LaunchModule.RulesFile.IsUnderDirectory(UnrealBuildTool.EngineDirectory) && !ShouldCompileMonolithic())
+			if(LaunchModule.RulesFile.IsUnderDirectory(UnrealBuild.EngineDirectory) && !ShouldCompileMonolithic())
 			{
-				IntermediateDirectory = DirectoryReference.Combine(UnrealBuildTool.EngineDirectory, PlatformIntermediateFolder, AppName, Configuration.ToString());
+				IntermediateDirectory = DirectoryReference.Combine(UnrealBuild.EngineDirectory, PlatformIntermediateFolder, AppName, Configuration.ToString());
 			}
 			else
 			{
@@ -3351,7 +3351,7 @@ namespace UnrealBuildTool
 			}
 			else
 			{
-				OutputDirectory = GetOutputDirectoryForExecutable(UnrealBuildTool.EngineDirectory, Rules.File);
+				OutputDirectory = GetOutputDirectoryForExecutable(UnrealBuild.EngineDirectory, Rules.File);
 			}
 
 			bool bCompileAsDLL = Rules.bShouldCompileAsDLL && bCompileMonolithic;
@@ -3365,7 +3365,7 @@ namespace UnrealBuildTool
 				bAllowExports: Rules.bHasExports,
 				bBuildAdditionalConsoleApp: Rules.bBuildAdditionalConsoleApp,
 				PrimaryModule: LaunchModule,
-				bUsePrecompiled: LaunchModule.Rules.bUsePrecompiled && OutputPaths[0].IsUnderDirectory(UnrealBuildTool.EngineDirectory)
+				bUsePrecompiled: LaunchModule.Rules.bUsePrecompiled && OutputPaths[0].IsUnderDirectory(UnrealBuild.EngineDirectory)
 			);
 			Binaries.Add(Binary);
 
@@ -3472,7 +3472,7 @@ namespace UnrealBuildTool
 			[REDACTED_EMAIL]: Do any platform specific tool chain initialization here if required
 
 			UnrealTargetConfiguration EngineTargetConfiguration = Configuration == UnrealTargetConfiguration.DebugGame ? UnrealTargetConfiguration.Development : Configuration;
-			DirectoryReference LinkIntermediateDirectory = DirectoryReference.Combine(UnrealBuildTool.EngineDirectory, PlatformIntermediateFolder, AppName, EngineTargetConfiguration.ToString());
+			DirectoryReference LinkIntermediateDirectory = DirectoryReference.Combine(UnrealBuild.EngineDirectory, PlatformIntermediateFolder, AppName, EngineTargetConfiguration.ToString());
 
 			// Installed Engine intermediates go to the project's intermediate folder. Installed Engine never writes to the engine intermediate folder. (Those files are immutable)
 			// Also, when compiling in monolithic, all intermediates go to the project's folder.  This is because a project can change definitions that affects all engine translation
@@ -3497,9 +3497,9 @@ namespace UnrealBuildTool
 			GlobalLinkEnvironment.LocalShadowDirectory = GlobalLinkEnvironment.OutputDirectory;
 
 			DirectoryReference OutputDir = Binaries[0].OutputDir;
-			if (OutputDir.IsUnderDirectory(UnrealBuildTool.EngineDirectory))
+			if (OutputDir.IsUnderDirectory(UnrealBuild.EngineDirectory))
 			{
-				DirectoryReference BaseDir = DirectoryReference.Combine(UnrealBuildTool.EngineDirectory, "Binaries", Platform.ToString());
+				DirectoryReference BaseDir = DirectoryReference.Combine(UnrealBuild.EngineDirectory, "Binaries", Platform.ToString());
 				if (BaseDir != OutputDir)
 				{
 					string RelativeBaseDir = BaseDir.MakeRelativeTo(OutputDir).Replace(Path.DirectorySeparatorChar, '/');
@@ -3704,7 +3704,7 @@ namespace UnrealBuildTool
 				GlobalCompileEnvironment.Definitions.Add("WITH_LIVE_CODING=1");
 				if (Rules.LinkType == TargetLinkType.Monolithic)
 				{
-					GlobalCompileEnvironment.Definitions.Add(String.Format("UE_LIVE_CODING_ENGINE_DIR=\"{0}\"", UnrealBuildTool.EngineDirectory.FullName.Replace("\\", "\\\\")));
+					GlobalCompileEnvironment.Definitions.Add(String.Format("UE_LIVE_CODING_ENGINE_DIR=\"{0}\"", UnrealBuild.EngineDirectory.FullName.Replace("\\", "\\\\")));
 					if(ProjectFile != null)
 					{
 						GlobalCompileEnvironment.Definitions.Add(String.Format("UE_LIVE_CODING_PROJECT=\"{0}\"", ProjectFile.FullName.Replace("\\", "\\\\")));
@@ -3760,7 +3760,7 @@ namespace UnrealBuildTool
 					{
 						OutputFilePath = OutputFilePath.Substring(0, OutputFilePath.LastIndexOf(".app/Contents/MacOS") + 4);
 					}
-					string EnginePath = Utils.CleanDirectorySeparators(UnrealBuildTool.EngineDirectory.MakeRelativeTo(ExecutableBinary.OutputFilePath.Directory), '/');
+					string EnginePath = Utils.CleanDirectorySeparators(UnrealBuild.EngineDirectory.MakeRelativeTo(ExecutableBinary.OutputFilePath.Directory), '/');
 					if (EnginePath.EndsWith("/") == false)
 					{
 						EnginePath += "/";
diff --git a/Engine/Source/Programs/UnrealBuildTool/Executors/Experimental/FASTBuild.cs b/Engine/Source/Programs/UnrealBuildTool/Executors/Experimental/FASTBuild.cs
--- a/Engine/Source/Programs/UnrealBuildTool/Executors/Experimental/FASTBuild.cs
+++ b/Engine/Source/Programs/UnrealBuildTool/Executors/Experimental/FASTBuild.cs
@@ -71,7 +71,7 @@ namespace UnrealBuildTool
 	class FASTBuild : ActionExecutor
 	{
 
-		public readonly static string DefaultExecutableBasePath	= Path.Combine(UnrealBuildTool.EngineDirectory.FullName, "Extras", "ThirdPartyNotUE", "FASTBuild");
+		public readonly static string DefaultExecutableBasePath	= Path.Combine(UnrealBuild.EngineDirectory.FullName, "Extras", "ThirdPartyNotUE", "FASTBuild");
 
 		//////////////////////////////////////////
 		// Tweakables
@@ -451,7 +451,7 @@ namespace UnrealBuildTool
 				if (!DetectBuildType(CompileActions))
 					return false;
 
-				string FASTBuildFilePath = Path.Combine(UnrealBuildTool.EngineDirectory.FullName, "Intermediate", "Build", "fbuild.bff");
+				string FASTBuildFilePath = Path.Combine(UnrealBuild.EngineDirectory.FullName, "Intermediate", "Build", "fbuild.bff");
 				if (!CreateBffFile(CompileActions, FASTBuildFilePath))
 					return false;
 
@@ -867,7 +867,7 @@ namespace UnrealBuildTool
 
 				AddText("Compiler('UECompiler') \n{\n");
 
-				DirectoryReference CLFilterDirectory = DirectoryReference.Combine(UnrealBuildTool.EngineDirectory, "Build", "Windows", "cl-filter");
+				DirectoryReference CLFilterDirectory = DirectoryReference.Combine(UnrealBuild.EngineDirectory, "Build", "Windows", "cl-filter");
 
 				AddText($"\t.Root = '{VCEnv.GetToolPath()}'\n");
 				AddText($"\t.CLFilterRoot = '{CLFilterDirectory.FullName}'\n");
@@ -1451,7 +1451,7 @@ namespace UnrealBuildTool
 			Log.TraceInformation($"FBuild Command Line Arguments: '{FBCommandLine}");
 
 			string FBExecutable		= GetExecutablePath();
-			string WorkingDirectory	= Path.GetFullPath(Path.Combine(UnrealBuildTool.EngineDirectory.MakeRelativeTo(DirectoryReference.GetCurrentDirectory()), "Source"));
+			string WorkingDirectory	= Path.GetFullPath(Path.Combine(UnrealBuild.EngineDirectory.MakeRelativeTo(DirectoryReference.GetCurrentDirectory()), "Source"));
 
 			ProcessStartInfo FBStartInfo		= new ProcessStartInfo(FBExecutable, FBCommandLine);
 			FBStartInfo.UseShellExecute			= false;
diff --git a/Engine/Source/Programs/UnrealBuildTool/Executors/SNDBS.cs b/Engine/Source/Programs/UnrealBuildTool/Executors/SNDBS.cs
--- a/Engine/Source/Programs/UnrealBuildTool/Executors/SNDBS.cs
+++ b/Engine/Source/Programs/UnrealBuildTool/Executors/SNDBS.cs
@@ -9,6 +9,7 @@ using System.Linq;
 using System.ServiceProcess;
 using System.Text.Json;
 using EpicGames.Core;
+using UnrealBuildBase;
 
 #nullable disable
 
@@ -21,7 +22,7 @@ namespace UnrealBuildTool
 		private static readonly string SCERoot = Environment.GetEnvironmentVariable("SCE_ROOT_DIR");
 		private static readonly string SNDBSExecutable = Path.Combine(SCERoot ?? string.Empty, "Common", "SN-DBS", "bin", "dbsbuild.exe");
 
-		private static readonly DirectoryReference IntermediateDir = DirectoryReference.Combine(UnrealBuildTool.EngineDirectory, "Intermediate", "Build", "SNDBS");
+		private static readonly DirectoryReference IntermediateDir = DirectoryReference.Combine(UnrealBuild.EngineDirectory, "Intermediate", "Build", "SNDBS");
 		private static readonly FileReference IncludeRewriteRulesFile = FileReference.Combine(IntermediateDir, "include-rewrite-rules.ini");
 		private static readonly FileReference ScriptFile = FileReference.Combine(IntermediateDir, "sndbs.json");
 
diff --git a/Engine/Source/Programs/UnrealBuildTool/Executors/XGE.cs b/Engine/Source/Programs/UnrealBuildTool/Executors/XGE.cs
--- a/Engine/Source/Programs/UnrealBuildTool/Executors/XGE.cs
+++ b/Engine/Source/Programs/UnrealBuildTool/Executors/XGE.cs
@@ -308,7 +308,7 @@ namespace UnrealBuildTool
 		{
 			for(int FileNum = 0;;FileNum++)
 			{
-				string OutFile = Path.Combine(UnrealBuildTool.EngineDirectory.FullName, "Intermediate", "Build", String.Format("UBTExport.{0}.xge.xml", FileNum.ToString("D3")));
+				string OutFile = Path.Combine(UnrealBuild.EngineDirectory.FullName, "Intermediate", "Build", String.Format("UBTExport.{0}.xge.xml", FileNum.ToString("D3")));
 				if(!File.Exists(OutFile))
 				{
 					ExportActions(ActionsToExecute, OutFile);
@@ -350,7 +350,7 @@ namespace UnrealBuildTool
 			if (Actions.Count > 0)
 			{
 				// Write the actions to execute to a XGE task file.
-				string XGETaskFilePath = FileReference.Combine(UnrealBuildTool.EngineDirectory, "Intermediate", "Build", "XGETasks.xml").FullName;
+				string XGETaskFilePath = FileReference.Combine(UnrealBuild.EngineDirectory, "Intermediate", "Build", "XGETasks.xml").FullName;
 				WriteTaskFile(Actions, XGETaskFilePath, ProgressWriter.bWriteMarkup, false);
 
 				XGEResult = ExecuteTaskFileWithProgressMarkup(XGETaskFilePath, Actions.Count);
diff --git a/Engine/Source/Programs/UnrealBuildTool/Modes/BuildMode.cs b/Engine/Source/Programs/UnrealBuildTool/Modes/BuildMode.cs
--- a/Engine/Source/Programs/UnrealBuildTool/Modes/BuildMode.cs
+++ b/Engine/Source/Programs/UnrealBuildTool/Modes/BuildMode.cs
@@ -140,9 +140,9 @@ namespace UnrealBuildTool
 			Arguments.ApplyTo(BuildConfiguration);
 
 			// Check the root path length isn't too long
-			if (BuildHostPlatform.Current.Platform == UnrealTargetPlatform.Win64 && UnrealBuildTool.RootDirectory.FullName.Length > BuildConfiguration.MaxRootPathLength)
+			if (BuildHostPlatform.Current.Platform == UnrealTargetPlatform.Win64 && UnrealBuild.RootDirectory.FullName.Length > BuildConfiguration.MaxRootPathLength)
 			{
-				Log.TraceWarning("Running from a path with a long directory name (\"{0}\" = {1} characters). Root paths shorter than {2} characters are recommended to avoid exceeding maximum path lengths on Windows.", UnrealBuildTool.RootDirectory, UnrealBuildTool.RootDirectory.FullName.Length, BuildConfiguration.MaxRootPathLength);
+				Log.TraceWarning("Running from a path with a long directory name (\"{0}\" = {1} characters). Root paths shorter than {2} characters are recommended to avoid exceeding maximum path lengths on Windows.", UnrealBuild.RootDirectory, UnrealBuild.RootDirectory.FullName.Length, BuildConfiguration.MaxRootPathLength);
 			}
 
 			// now that we know the available platforms, we can delete other platforms' junk. if we're only building specific modules from the editor, don't touch anything else (it may be in use).
@@ -224,7 +224,7 @@ namespace UnrealBuildTool
 					}
 
 					// Create the working set provider per group.
-					using (ISourceFileWorkingSet WorkingSet = SourceFileWorkingSet.Create(UnrealBuildTool.RootDirectory, ProjectDirs))
+					using (ISourceFileWorkingSet WorkingSet = SourceFileWorkingSet.Create(UnrealBuild.RootDirectory, ProjectDirs))
 					{
 						Build(TargetDescriptors, BuildConfiguration, WorkingSet, Options, WriteOutdatedActionsFile, bSkipPreBuildTargets);
 					}
@@ -508,7 +508,7 @@ namespace UnrealBuildTool
 				// Make sure we're not modifying any engine files
 				if ((Options & BuildOptions.NoEngineChanges) != 0)
 				{
-					List<FileItem> EngineChanges = MergedActionsToExecute.SelectMany(x => x.ProducedItems).Where(x => x.Location.IsUnderDirectory(UnrealBuildTool.EngineDirectory)).Distinct().OrderBy(x => x.FullName).ToList();
+					List<FileItem> EngineChanges = MergedActionsToExecute.SelectMany(x => x.ProducedItems).Where(x => x.Location.IsUnderDirectory(UnrealBuild.EngineDirectory)).Distinct().OrderBy(x => x.FullName).ToList();
 					if (EngineChanges.Count > 0)
 					{
 						StringBuilder Result = new StringBuilder("Building would modify the following engine files:\n");
diff --git a/Engine/Source/Programs/UnrealBuildTool/Modes/CleanMode.cs b/Engine/Source/Programs/UnrealBuildTool/Modes/CleanMode.cs
--- a/Engine/Source/Programs/UnrealBuildTool/Modes/CleanMode.cs
+++ b/Engine/Source/Programs/UnrealBuildTool/Modes/CleanMode.cs
@@ -8,6 +8,7 @@ using System.Linq;
 using System.Text;
 using System.Threading.Tasks;
 using EpicGames.Core;
+using UnrealBuildBase;
 
 namespace UnrealBuildTool
 {
@@ -116,7 +117,7 @@ namespace UnrealBuildTool
 
 				// Find the base folders that can contain binaries
 				List<DirectoryReference> BaseDirs = new List<DirectoryReference>();
-				BaseDirs.Add(UnrealBuildTool.EngineDirectory);
+				BaseDirs.Add(UnrealBuild.EngineDirectory);
 				foreach (FileReference Plugin in Plugins.EnumeratePlugins(Target.ProjectFile))
 				{
 					BaseDirs.Add(Plugin.Directory);
diff --git a/Engine/Source/Programs/UnrealBuildTool/Modes/GenerateClangDatabase.cs b/Engine/Source/Programs/UnrealBuildTool/Modes/GenerateClangDatabase.cs
--- a/Engine/Source/Programs/UnrealBuildTool/Modes/GenerateClangDatabase.cs
+++ b/Engine/Source/Programs/UnrealBuildTool/Modes/GenerateClangDatabase.cs
@@ -133,7 +133,7 @@ namespace UnrealBuildTool
 
 								foreach (FileItem InputFile in InputFiles)
 								{
-									if(FileFilter == null || FileFilter.Matches(InputFile.Location.MakeRelativeTo(UnrealBuildTool.RootDirectory)))
+									if(FileFilter == null || FileFilter.Matches(InputFile.Location.MakeRelativeTo(UnrealBuild.RootDirectory)))
 									{
 										FileToCommand[InputFile.Location] = String.Format("{0} \"{1}\"", CommandBuilder, InputFile.FullName);
 									}
@@ -144,7 +144,7 @@ namespace UnrealBuildTool
 				}
 
 				// Write the compile database
-				FileReference DatabaseFile = FileReference.Combine(UnrealBuildTool.RootDirectory, "compile_commands.json");
+				FileReference DatabaseFile = FileReference.Combine(UnrealBuild.RootDirectory, "compile_commands.json");
 				using (JsonWriter Writer = new JsonWriter(DatabaseFile))
 				{
 					Writer.WriteArrayStart();
diff --git a/Engine/Source/Programs/UnrealBuildTool/Modes/QueryTargetsMode.cs b/Engine/Source/Programs/UnrealBuildTool/Modes/QueryTargetsMode.cs
--- a/Engine/Source/Programs/UnrealBuildTool/Modes/QueryTargetsMode.cs
+++ b/Engine/Source/Programs/UnrealBuildTool/Modes/QueryTargetsMode.cs
@@ -6,6 +6,7 @@ using System.Linq;
 using System.Text;
 using System.Threading.Tasks;
 using EpicGames.Core;
+using UnrealBuildBase;
 
 namespace UnrealBuildTool
 {
@@ -73,7 +74,7 @@ namespace UnrealBuildTool
 		{
 			if (ProjectFile == null)
 			{
-				return FileReference.Combine(UnrealBuildTool.EngineDirectory, "Intermediate", "TargetInfo.json");
+				return FileReference.Combine(UnrealBuild.EngineDirectory, "Intermediate", "TargetInfo.json");
 			}
 			else
 			{
diff --git a/Engine/Source/Programs/UnrealBuildTool/Modes/WriteMetadataMode.cs b/Engine/Source/Programs/UnrealBuildTool/Modes/WriteMetadataMode.cs
--- a/Engine/Source/Programs/UnrealBuildTool/Modes/WriteMetadataMode.cs
+++ b/Engine/Source/Programs/UnrealBuildTool/Modes/WriteMetadataMode.cs
@@ -9,6 +9,7 @@ using System.Runtime.Serialization.Formatters.Binary;
 using System.Text;
 using System.Threading.Tasks;
 using EpicGames.Core;
+using UnrealBuildBase;
 
 namespace UnrealBuildTool
 {
@@ -82,7 +83,7 @@ namespace UnrealBuildTool
 		{
 			// Acquire a different mutex to the regular UBT instance, since this mode will be called as part of a build. We need the mutex to ensure that building two modular configurations 
 			// in parallel don't clash over writing shared *.modules files (eg. DebugGame and Development editors).
-			string MutexName = SingleInstanceMutex.GetUniqueMutexForPath("UnrealBuildTool_WriteMetadata", UnrealBuildTool.RootDirectory.FullName);
+			string MutexName = SingleInstanceMutex.GetUniqueMutexForPath("UnrealBuildTool_WriteMetadata", UnrealBuild.RootDirectory.FullName);
 			using(new SingleInstanceMutex(MutexName, true))
 			{
 				return ExecuteInternal(Arguments);
@@ -230,7 +231,7 @@ namespace UnrealBuildTool
 			foreach(FileReference ManifestFileName in ManifestFiles)
 			{
 				ModuleManifest? Manifest;
-				if(ManifestFileName.IsUnderDirectory(UnrealBuildTool.EngineDirectory) && TryReadManifest(ManifestFileName, out Manifest))
+				if(ManifestFileName.IsUnderDirectory(UnrealBuild.EngineDirectory) && TryReadManifest(ManifestFileName, out Manifest))
 				{
 					if(Manifest.BuildId == BuildId)
 					{
diff --git a/Engine/Source/Programs/UnrealBuildTool/Platform/Android/AndroidToolChain.cs b/Engine/Source/Programs/UnrealBuildTool/Platform/Android/AndroidToolChain.cs
--- a/Engine/Source/Programs/UnrealBuildTool/Platform/Android/AndroidToolChain.cs
+++ b/Engine/Source/Programs/UnrealBuildTool/Platform/Android/AndroidToolChain.cs
@@ -355,7 +355,7 @@ namespace UnrealBuildTool
 			if (Ini.GetBool("/Script/AndroidRuntimeSettings.AndroidRuntimeSettings", "bBuildForx8664", out bBuild) && bBuild
 				|| (AdditionalArches != null && (AdditionalArches.Contains("x64", StringComparer.OrdinalIgnoreCase) || AdditionalArches.Contains("-x64", StringComparer.OrdinalIgnoreCase))))
 			{
-				if (File.Exists(Path.Combine(UnrealBuildTool.EngineDirectory.FullName, "Build", "InstalledBuild.txt")))
+				if (File.Exists(Path.Combine(UnrealBuild.EngineDirectory.FullName, "Build", "InstalledBuild.txt")))
 				{
 					bUnsupportedBinaryBuildArch = true;
 					Log.TraceWarningOnce("Please install source to build for x86_64 (-x64); ignoring this architecture target.");
@@ -1432,7 +1432,7 @@ namespace UnrealBuildTool
 							FileArguments += GetCompileArguments_C(bDisableOptimizations);
 
 							// remove shadow variable warnings for externally included files
-							if (!SourceFile.Location.IsUnderDirectory(UnrealBuildTool.RootDirectory))
+							if (!SourceFile.Location.IsUnderDirectory(UnrealBuild.RootDirectory))
 							{
 								bDisableShadowWarning = true;
 							}
diff --git a/Engine/Source/Programs/UnrealBuildTool/Platform/Android/UEDeployAndroid.cs b/Engine/Source/Programs/UnrealBuildTool/Platform/Android/UEDeployAndroid.cs
--- a/Engine/Source/Programs/UnrealBuildTool/Platform/Android/UEDeployAndroid.cs
+++ b/Engine/Source/Programs/UnrealBuildTool/Platform/Android/UEDeployAndroid.cs
@@ -11,6 +11,7 @@ using Microsoft.Win32;
 using System.Xml.Linq;
 using EpicGames.Core;
 using System.Security.Cryptography;
+using UnrealBuildBase;
 
 #nullable disable
 
@@ -1496,7 +1497,7 @@ namespace UnrealBuildTool
 			
 			string SanitizerFullLibName = "libclang_rt." + LibName + Architecture + "-android.so";
 
-			string WrapSh = Path.Combine(UnrealBuildTool.EngineDirectory.ToString(), "Build", "Android", "ClangSanitizers", "wrap.sh");
+			string WrapSh = Path.Combine(UnrealBuild.EngineDirectory.ToString(), "Build", "Android", "ClangSanitizers", "wrap.sh");
 
 			string PlatformName = "windows-x86_64";
 			if (RuntimePlatform.IsLinux)
@@ -4757,7 +4758,7 @@ namespace UnrealBuildTool
 		{
 			//Log.TraceInformation("$$$$$$$$$$$$$$ PrepTargetForDeployment $$$$$$$$$$$$$$$$$");
 
-			DirectoryReference ProjectDirectory = DirectoryReference.FromFile(Receipt.ProjectFile) ?? UnrealBuildTool.EngineDirectory;
+			DirectoryReference ProjectDirectory = DirectoryReference.FromFile(Receipt.ProjectFile) ?? UnrealBuild.EngineDirectory;
 			string TargetName = (Receipt.ProjectFile == null ? Receipt.TargetName : Receipt.ProjectFile.GetFileNameWithoutAnyExtensions());
 
 			AndroidToolChain ToolChain = ((AndroidPlatform)UEBuildPlatform.GetBuildPlatform(Receipt.Platform)).CreateTempToolChainForProject(Receipt.ProjectFile) as AndroidToolChain;
@@ -4781,7 +4782,7 @@ namespace UnrealBuildTool
 			string BaseSoName = ToolChain.RemoveArchName(OutputPaths[0].FullName);
 
 			// make an apk at the end of compiling, so that we can run without packaging (debugger, cook on the fly, etc)
-			string RelativeEnginePath = UnrealBuildTool.EngineDirectory.MakeRelativeTo(DirectoryReference.GetCurrentDirectory());
+			string RelativeEnginePath = UnrealBuild.EngineDirectory.MakeRelativeTo(DirectoryReference.GetCurrentDirectory());
 
 			MakeApk(ToolChain, TargetName, Receipt.TargetType, ProjectDirectory.FullName, BaseSoName, RelativeEnginePath, bForDistribution: false, CookFlavor: "", Configuration: Receipt.Configuration,
 				bMakeSeparateApks: ShouldMakeSeparateApks(), bIncrementalPackage: true, bDisallowPackagingDataInApk: false, bDisallowExternalFilesDir: true, bSkipGradleBuild: bShouldCompileAsDll);
diff --git a/Engine/Source/Programs/UnrealBuildTool/Platform/HoloLens/HoloLensDeploy.cs b/Engine/Source/Programs/UnrealBuildTool/Platform/HoloLens/HoloLensDeploy.cs
--- a/Engine/Source/Programs/UnrealBuildTool/Platform/HoloLens/HoloLensDeploy.cs
+++ b/Engine/Source/Programs/UnrealBuildTool/Platform/HoloLens/HoloLensDeploy.cs
@@ -7,6 +7,7 @@ using System.IO;
 using System.Diagnostics;
 using System.Linq;
 using EpicGames.Core;
+using UnrealBuildBase;
 
 #nullable disable
 
@@ -216,15 +217,15 @@ namespace UnrealBuildTool
 		private void MakePackage(TargetReceipt Receipt, TargetReceipt NewReceipt, WindowsArchitecture Architecture, List<string> UpdatedFiles)
 		{
 			string OutputName = String.Format("{0}_{1}_{2}_{3}", Receipt.TargetName, Receipt.Platform, Receipt.Configuration, WindowsExports.GetArchitectureSubpath(Architecture));
-			string IntermediateDirectory = Path.Combine(Receipt.ProjectDir != null ? Receipt.ProjectDir.FullName : UnrealBuildTool.EngineDirectory.FullName, "Intermediate", "Deploy", WindowsExports.GetArchitectureSubpath(Architecture));
+			string IntermediateDirectory = Path.Combine(Receipt.ProjectDir != null ? Receipt.ProjectDir.FullName : UnrealBuild.EngineDirectory.FullName, "Intermediate", "Deploy", WindowsExports.GetArchitectureSubpath(Architecture));
 			string OutputDirectory = Receipt.Launch.Directory.FullName;
 			string OutputAppX = Path.Combine(OutputDirectory, OutputName + Extension);
 			string SigningCertificate = @"Build\HoloLens\SigningCertificate.pfx";
-			string SigningCertificatePath = Path.Combine(Receipt.ProjectDir != null ? Receipt.ProjectDir.FullName : UnrealBuildTool.EngineDirectory.FullName, SigningCertificate);
+			string SigningCertificatePath = Path.Combine(Receipt.ProjectDir != null ? Receipt.ProjectDir.FullName : UnrealBuild.EngineDirectory.FullName, SigningCertificate);
 
 			string MapFilename = Path.Combine(IntermediateDirectory, OutputName + ".pkgmap");
 			var LocalRoot = Receipt.ProjectDir;
-			var EngineRoot = UnrealBuildTool.RootDirectory;
+			var EngineRoot = UnrealBuild.RootDirectory;
 			var AddedFiles = new Dictionary<string, string>();
 			bool PackageFileNeedToBeUpdated = !File.Exists(OutputAppX);
 
@@ -477,8 +478,8 @@ namespace UnrealBuildTool
 
 			// Note: TargetReceipt.Read now expands path variables internally.
 			TargetReceipt NewReceipt = null;
-			FileReference ReceiptFileName = TargetReceipt.GetDefaultPath(Receipt.ProjectDir != null ? Receipt.ProjectDir : UnrealBuildTool.EngineDirectory, Receipt.TargetName, Receipt.Platform, Receipt.Configuration, "Multi");
-			if (!TargetReceipt.TryRead(ReceiptFileName, UnrealBuildTool.EngineDirectory, out NewReceipt))
+			FileReference ReceiptFileName = TargetReceipt.GetDefaultPath(Receipt.ProjectDir != null ? Receipt.ProjectDir : UnrealBuild.EngineDirectory, Receipt.TargetName, Receipt.Platform, Receipt.Configuration, "Multi");
+			if (!TargetReceipt.TryRead(ReceiptFileName, UnrealBuild.EngineDirectory, out NewReceipt))
 			{
 				NewReceipt = new TargetReceipt(Receipt.ProjectFile, Receipt.TargetName, Receipt.TargetType, Receipt.Platform, Receipt.Configuration, Receipt.Version, "Multi");
 			}
@@ -491,12 +492,12 @@ namespace UnrealBuildTool
 			}
 			HoloLensExports.InitWindowsSdkToolPath(SDK);
 
-			AddWinMDReferencesFromReceipt(Receipt, Receipt.ProjectDir != null ? Receipt.ProjectDir : UnrealBuildTool.EngineDirectory, UnrealBuildTool.EngineDirectory.ParentDirectory.FullName, SDK);
+			AddWinMDReferencesFromReceipt(Receipt, Receipt.ProjectDir != null ? Receipt.ProjectDir : UnrealBuild.EngineDirectory, UnrealBuild.EngineDirectory.ParentDirectory.FullName, SDK);
 
 			//PrepForUATPackageOrDeploy(InTarget.ProjectFile, InAppName, InTarget.ProjectDirectory.FullName, InTarget.OutputPath.FullName, TargetBuildEnvironment.RelativeEnginePath, false, "", false);
 			List<UnrealTargetConfiguration> TargetConfigs = new List<UnrealTargetConfiguration> { Receipt.Configuration };
 			List<string> ExePaths = new List<string> { Receipt.Launch.FullName };
-			string RelativeEnginePath = UnrealBuildTool.EngineDirectory.MakeRelativeTo(DirectoryReference.GetCurrentDirectory());
+			string RelativeEnginePath = UnrealBuild.EngineDirectory.MakeRelativeTo(DirectoryReference.GetCurrentDirectory());
 
 			WindowsArchitecture Arch = WindowsArchitecture.ARM64;
 			if (Receipt.Architecture.ToLower() == "x64")
@@ -507,14 +508,14 @@ namespace UnrealBuildTool
 
 			string AbsoluteExeDirectory = Path.GetDirectoryName(ExePaths[0]);
 			UnrealTargetPlatform Platform = UnrealTargetPlatform.HoloLens;
-			string IntermediateDirectory = Path.Combine(Receipt.ProjectDir != null ? Receipt.ProjectDir.FullName : UnrealBuildTool.EngineDirectory.FullName, "Intermediate", "Deploy", WindowsExports.GetArchitectureSubpath(Arch));
-			List<string> UpdatedFiles = new HoloLensManifestGenerator().CreateManifest(Platform, Arch, AbsoluteExeDirectory, IntermediateDirectory, Receipt.ProjectFile, Receipt.ProjectDir != null ? Receipt.ProjectDir.FullName : UnrealBuildTool.EngineDirectory.FullName, TargetConfigs, ExePaths, WinMDReferences);
+			string IntermediateDirectory = Path.Combine(Receipt.ProjectDir != null ? Receipt.ProjectDir.FullName : UnrealBuild.EngineDirectory.FullName, "Intermediate", "Deploy", WindowsExports.GetArchitectureSubpath(Arch));
+			List<string> UpdatedFiles = new HoloLensManifestGenerator().CreateManifest(Platform, Arch, AbsoluteExeDirectory, IntermediateDirectory, Receipt.ProjectFile, Receipt.ProjectDir != null ? Receipt.ProjectDir.FullName : UnrealBuild.EngineDirectory.FullName, TargetConfigs, ExePaths, WinMDReferences);
 
-			PrepForUATPackageOrDeploy(Receipt.ProjectFile, ProjectName, Receipt.ProjectDir != null ? Receipt.ProjectDir.FullName : UnrealBuildTool.EngineDirectory.FullName, Arch, TargetConfigs, ExePaths, RelativeEnginePath, false, "", false);
+			PrepForUATPackageOrDeploy(Receipt.ProjectFile, ProjectName, Receipt.ProjectDir != null ? Receipt.ProjectDir.FullName : UnrealBuild.EngineDirectory.FullName, Arch, TargetConfigs, ExePaths, RelativeEnginePath, false, "", false);
 			MakePackage(Receipt, NewReceipt, Arch, UpdatedFiles);
 			CopyDataAndSymbolsBetweenReceipts(Receipt, NewReceipt, Arch);
 
-			NewReceipt.Write(ReceiptFileName, UnrealBuildTool.EngineDirectory);
+			NewReceipt.Write(ReceiptFileName, UnrealBuild.EngineDirectory);
 
 			// Log out the time taken to deploy...
 			double PrepDeployDuration = (DateTime.UtcNow - PrepDeployStartTime).TotalSeconds;
@@ -535,7 +536,7 @@ namespace UnrealBuildTool
 					if (File.Exists(WinMDFile))
 					{
 						string DestPath = Dep.Path.FullName;
-						DestPath = Dep.Path.FullName.Replace(UnrealBuildTool.EngineDirectory.FullName, Path.Combine(DestRelativeTo, "Engine"));
+						DestPath = Dep.Path.FullName.Replace(UnrealBuild.EngineDirectory.FullName, Path.Combine(DestRelativeTo, "Engine"));
 						DestPath = DestPath.Replace(SourceProjectDir.FullName, Path.Combine(DestRelativeTo, SourceProjectDir.GetDirectoryName()));
 						DestPath = Utils.MakePathRelativeTo(DestPath, DestRelativeTo);
 						WinMDReferences.Add(new WinMDRegistrationInfo(new FileReference(WinMDFile), DestPath, SDKVersion));
diff --git a/Engine/Source/Programs/UnrealBuildTool/Platform/HoloLens/HoloLensProjectGenerator.cs b/Engine/Source/Programs/UnrealBuildTool/Platform/HoloLens/HoloLensProjectGenerator.cs
--- a/Engine/Source/Programs/UnrealBuildTool/Platform/HoloLens/HoloLensProjectGenerator.cs
+++ b/Engine/Source/Programs/UnrealBuildTool/Platform/HoloLens/HoloLensProjectGenerator.cs
@@ -7,6 +7,7 @@ using System.Diagnostics;
 using System.IO;
 using System.Xml.Linq;
 using EpicGames.Core;
+using UnrealBuildBase;
 
 #nullable disable
 
@@ -177,7 +178,7 @@ namespace UnrealBuildTool
 				string TempTargetFilePath = InTargetFilePath.FullName.Replace("\\", "/");
 				if (TempTargetFilePath.Contains("/Templates/"))
 				{
-					string AbsoluteEnginePath = UnrealBuildTool.EngineDirectory.FullName;
+					string AbsoluteEnginePath = UnrealBuild.EngineDirectory.FullName;
 					AbsoluteEnginePath = AbsoluteEnginePath.Replace("\\", "/");
 					if (AbsoluteEnginePath.EndsWith("/") == false)
 					{
diff --git a/Engine/Source/Programs/UnrealBuildTool/Platform/HoloLens/HoloLensToolChain.cs b/Engine/Source/Programs/UnrealBuildTool/Platform/HoloLens/HoloLensToolChain.cs
--- a/Engine/Source/Programs/UnrealBuildTool/Platform/HoloLens/HoloLensToolChain.cs
+++ b/Engine/Source/Programs/UnrealBuildTool/Platform/HoloLens/HoloLensToolChain.cs
@@ -709,7 +709,7 @@ namespace UnrealBuildTool
 					CommandArguments.Add(CompileAction.CommandArguments);
 					CommandArguments.Add("/showIncludes");
 					CompileAction.CommandArguments = string.Join(" ", CommandArguments);
-					CompileAction.CommandPath = FileReference.Combine(UnrealBuildTool.EngineDirectory, "Build", "Windows", "cl-filter", "cl-filter.exe");
+					CompileAction.CommandPath = FileReference.Combine(UnrealBuild.EngineDirectory, "Build", "Windows", "cl-filter", "cl-filter.exe");
 				}
 
 				if (CompileEnvironment.PrecompiledHeaderAction == PrecompiledHeaderAction.Create)
diff --git a/Engine/Source/Programs/UnrealBuildTool/Platform/HoloLens/ManifestGenerator.cs b/Engine/Source/Programs/UnrealBuildTool/Platform/HoloLens/ManifestGenerator.cs
--- a/Engine/Source/Programs/UnrealBuildTool/Platform/HoloLens/ManifestGenerator.cs
+++ b/Engine/Source/Programs/UnrealBuildTool/Platform/HoloLens/ManifestGenerator.cs
@@ -9,6 +9,7 @@ using System.Text.RegularExpressions;
 using System.Resources;
 using System.Xml;
 using EpicGames.Core;
+using UnrealBuildBase;
 
 #nullable disable
 
@@ -435,7 +436,7 @@ namespace UnrealBuildTool
 			{
 				if (AllowEngineFallback)
 				{
-					SourcePath = Path.Combine(UnrealBuildTool.EngineDirectory.FullName, BuildResourceProjectRelativePath, EngineResourceSubPath);
+					SourcePath = Path.Combine(UnrealBuild.EngineDirectory.FullName, BuildResourceProjectRelativePath, EngineResourceSubPath);
 					bFileExists = File.Exists(Path.Combine(SourcePath, ResourceFileName));
 				}
 			}
@@ -2053,7 +2054,7 @@ namespace UnrealBuildTool
 				ConfigPostfix = TargetConfig.ToString();
 			}
 
-			string MakeRelativeTo = IsGameSpecificExe ? Path.Combine(ProjectPath, "..") : UnrealBuildTool.EngineDirectory.FullName;
+			string MakeRelativeTo = IsGameSpecificExe ? Path.Combine(ProjectPath, "..") : UnrealBuild.EngineDirectory.FullName;
 			string RelativeExePath = IsGameSpecificExe ?
 				Utils.MakePathRelativeTo(ExecutablePath, MakeRelativeTo) :
 				Path.Combine(RelativeProjectRootForStage, Utils.MakePathRelativeTo(ExecutablePath, MakeRelativeTo));
@@ -2161,7 +2162,7 @@ namespace UnrealBuildTool
 						, to
 						//, Path.Combine(IntermediatePath, "tempGltf") //temp path doesn't work
 						);
-					string Exe = Path.Combine(UnrealBuildTool.EngineDirectory.FullName, "Binaries", "Win64", "WindowsMRAssetConverter.exe");
+					string Exe = Path.Combine(UnrealBuild.EngineDirectory.FullName, "Binaries", "Win64", "WindowsMRAssetConverter.exe");
 					System.Diagnostics.ProcessStartInfo StartInfo = new System.Diagnostics.ProcessStartInfo(Exe, Args);
 					StartInfo.UseShellExecute = false;
 					StartInfo.RedirectStandardOutput = true;
diff --git a/Engine/Source/Programs/UnrealBuildTool/Platform/IOS/IOSPostBuildSyncMode.cs b/Engine/Source/Programs/UnrealBuildTool/Platform/IOS/IOSPostBuildSyncMode.cs
--- a/Engine/Source/Programs/UnrealBuildTool/Platform/IOS/IOSPostBuildSyncMode.cs
+++ b/Engine/Source/Programs/UnrealBuildTool/Platform/IOS/IOSPostBuildSyncMode.cs
@@ -6,6 +6,7 @@ using System.Linq;
 using System.Text;
 using System.Threading.Tasks;
 using EpicGames.Core;
+using UnrealBuildBase;
 
 #nullable disable
 
@@ -45,7 +46,7 @@ namespace UnrealBuildTool
 			this.SdkVersion = SdkVersion;
 			this.bCreateStubIPA = Target.IOSPlatform.bCreateStubIPA;
 			this.bSkipCrashlytics = Target.IOSPlatform.bSkipCrashlytics;
-			this.ProjectDirectory = DirectoryReference.FromFile(Target.ProjectFile) ?? UnrealBuildTool.EngineDirectory;
+			this.ProjectDirectory = DirectoryReference.FromFile(Target.ProjectFile) ?? UnrealBuild.EngineDirectory;
 			this.ProjectIntermediateDirectory = ProjectIntermediateDirectory;
 			this.ImportProvision = Target.IOSPlatform.ImportProvision;
 			this.ImportCertificate = Target.IOSPlatform.ImportCertificate;
diff --git a/Engine/Source/Programs/UnrealBuildTool/Platform/IOS/IOSToolChain.cs b/Engine/Source/Programs/UnrealBuildTool/Platform/IOS/IOSToolChain.cs
--- a/Engine/Source/Programs/UnrealBuildTool/Platform/IOS/IOSToolChain.cs
+++ b/Engine/Source/Programs/UnrealBuildTool/Platform/IOS/IOSToolChain.cs
@@ -902,7 +902,7 @@ namespace UnrealBuildTool
 			else
 			{
 				bool bIsUnrealGame = LinkEnvironment.OutputFilePath.FullName.Contains("UnrealGame");
-				FileReference ResponsePath = FileReference.Combine(((!bIsUnrealGame && ProjectFile != null) ? ProjectFile.Directory : UnrealBuildTool.EngineDirectory), "Intermediate", "Build", LinkEnvironment.Platform.ToString(), "LinkFileList_" + LinkEnvironment.OutputFilePath.GetFileNameWithoutExtension() + ".tmp");
+				FileReference ResponsePath = FileReference.Combine(((!bIsUnrealGame && ProjectFile != null) ? ProjectFile.Directory : UnrealBuild.EngineDirectory), "Intermediate", "Build", LinkEnvironment.Platform.ToString(), "LinkFileList_" + LinkEnvironment.OutputFilePath.GetFileNameWithoutExtension() + ".tmp");
 				Graph.CreateIntermediateTextFile(ResponsePath, InputFileNames);
 				LinkCommandArguments += string.Format(" \\\"@{0}\\\"", ResponsePath.FullName);
 			}
@@ -1109,7 +1109,7 @@ namespace UnrealBuildTool
 
 			// Make the compile action
 			Action GenDebugAction = Graph.CreateAction(ActionType.GenerateDebugInfo);
-			GenDebugAction.WorkingDirectory = DirectoryReference.Combine(UnrealBuildTool.EngineDirectory, "Binaries", "Mac");
+			GenDebugAction.WorkingDirectory = DirectoryReference.Combine(UnrealBuild.EngineDirectory, "Binaries", "Mac");
 
 			GenDebugAction.CommandPath = BuildHostPlatform.Current.Shell;
 			GenDebugAction.CommandArguments = string.Format("-c \"rm -rf \\\"{1}\\\"; dwarfdump --uuid \\\"{3}\\\" | cut -d\\  -f2; chmod 777 ./DsymExporter; ./DsymExporter -UUID=$(dwarfdump --uuid \\\"{3}\\\" | cut -d\\  -f2) \\\"{0}\\\" \\\"{2}\\\"\"",
@@ -1263,7 +1263,7 @@ namespace UnrealBuildTool
 				Action UnzipAction = Graph.CreateAction(ActionType.BuildProject);
 				UnzipAction.CommandPath = new FileReference("/bin/sh");
 				UnzipAction.CommandArguments = Utils.MakePathSafeToUseWithCommandLine(ExtractScriptFileItem.AbsolutePath);
-				UnzipAction.WorkingDirectory = UnrealBuildTool.EngineDirectory;
+				UnzipAction.WorkingDirectory = UnrealBuild.EngineDirectory;
 				UnzipAction.PrerequisiteItems.Add(InputFile);
 				UnzipAction.PrerequisiteItems.Add(ExtractScriptFileItem);
 				UnzipAction.ProducedItems.Add(Framework.ExtractedTokenFile);
@@ -1276,9 +1276,9 @@ namespace UnrealBuildTool
 
 		public static DirectoryReference GenerateAssetCatalog(FileReference ProjectFile, UnrealTargetPlatform Platform, ref bool bUserImagesExist)
 		{
-			string EngineDir = UnrealBuildTool.EngineDirectory.ToString();
-			string BuildDir = (((ProjectFile != null) ? ProjectFile.Directory.ToString() : (string.IsNullOrEmpty(UnrealBuildTool.GetRemoteIniPath()) ? UnrealBuildTool.EngineDirectory.ToString() : UnrealBuildTool.GetRemoteIniPath()))) + "/Build/" + (Platform == UnrealTargetPlatform.IOS ? "IOS" : "TVOS");
-			string IntermediateDir = (((ProjectFile != null) ? ProjectFile.Directory.ToString() : UnrealBuildTool.EngineDirectory.ToString())) + "/Intermediate/" + (Platform == UnrealTargetPlatform.IOS ? "IOS" : "TVOS");
+			string EngineDir = UnrealBuild.EngineDirectory.ToString();
+			string BuildDir = (((ProjectFile != null) ? ProjectFile.Directory.ToString() : (string.IsNullOrEmpty(UnrealBuildTool.GetRemoteIniPath()) ? UnrealBuild.EngineDirectory.ToString() : UnrealBuildTool.GetRemoteIniPath()))) + "/Build/" + (Platform == UnrealTargetPlatform.IOS ? "IOS" : "TVOS");
+			string IntermediateDir = (((ProjectFile != null) ? ProjectFile.Directory.ToString() : UnrealBuild.EngineDirectory.ToString())) + "/Intermediate/" + (Platform == UnrealTargetPlatform.IOS ? "IOS" : "TVOS");
 
 			bUserImagesExist = false;
 
@@ -1606,7 +1606,7 @@ namespace UnrealBuildTool
 			Log.TraceInformation("Generating and uploading Crashlytics Data");
 
 			// Clean this folder as it's used for extraction
-			string TempPath = Path.Combine(UnrealBuildTool.EngineDirectory.FullName, "Intermediate", "Unzipped");
+			string TempPath = Path.Combine(UnrealBuild.EngineDirectory.FullName, "Intermediate", "Unzipped");
 
 			if (Directory.Exists(TempPath))
 			{
@@ -1614,7 +1614,7 @@ namespace UnrealBuildTool
 				Directory.Delete(TempPath, true);
 			}
 
-			string FabricPath = UnrealBuildTool.EngineDirectory + "/Intermediate/UnzippedFrameworks/Crashlytics/Fabric.embeddedframework";
+			string FabricPath = UnrealBuild.EngineDirectory + "/Intermediate/UnzippedFrameworks/Crashlytics/Fabric.embeddedframework";
             if (Directory.Exists(FabricPath) && Environment.GetEnvironmentVariable("IsBuildMachine") == "1")
             {
 				//string PlistFile = ProjectDir + "/Intermediate/IOS/" + ProjectName + "-Info.plist";
@@ -1680,7 +1680,7 @@ namespace UnrealBuildTool
 			}
 
 			string IntermediateDir = (((Target.ProjectFile != null) ? Target.ProjectFile.Directory.ToString() :
-				UnrealBuildTool.EngineDirectory.ToString())) + "/Intermediate/" + (Target.Platform == UnrealTargetPlatform.IOS ? "IOS" : "TVOS");
+				UnrealBuild.EngineDirectory.ToString())) + "/Intermediate/" + (Target.Platform == UnrealTargetPlatform.IOS ? "IOS" : "TVOS");
 			// get the settings from the ini file
 			ConfigHierarchy Ini = ConfigCache.ReadHierarchy(ConfigHierarchyType.Engine, DirectoryReference.FromFile(Target.ProjectFile), UnrealTargetPlatform.IOS);
 
@@ -1689,7 +1689,7 @@ namespace UnrealBuildTool
 
 		private static void WriteMissingEmbeddedPlist(FileReference MissingPlistPath, DirectoryReference ProjectDirectory, FileReference ProjectFile)
 		{
-			FileReference PlistTemplatePath = FileReference.Combine(UnrealBuildTool.EngineDirectory, "Build", "IOS", "Resources", "FrameworkPlist", "Embedded.plist");
+			FileReference PlistTemplatePath = FileReference.Combine(UnrealBuild.EngineDirectory, "Build", "IOS", "Resources", "FrameworkPlist", "Embedded.plist");
 			if (!DirectoryReference.Exists(MissingPlistPath.Directory))
 			{
 				DirectoryReference.CreateDirectory(MissingPlistPath.Directory);
@@ -1727,7 +1727,7 @@ namespace UnrealBuildTool
 			string ProjectName = ProjectDirectory.GetDirectoryName();
 			string FrameworkName = Target.TargetName;
 			string BundleId = ProjectFiles.Xcode.XcodeFrameworkWrapperUtils.GetBundleID(Target.ProjectDirectory, Target.ProjectFile);
-			string EnginePath = UnrealBuildTool.EngineDirectory.ToString();
+			string EnginePath = UnrealBuild.EngineDirectory.ToString();
 			string SrcFrameworkPath = DirectoryReference.Combine(Target.ProjectDirectory , "Binaries", "IOS", Target.Configuration.ToString()).ToString(); // We use Target.ProjectDirectory because if it is a UnrealGame we want the engine dir and not the actual project dir.
 			string CookedDataPath = DirectoryReference.Combine(ProjectDirectory, "Saved", "StagedBuilds", "IOS", "cookeddata").ToString();
 			bool GenerateFrameworkWrapperProject = ProjectFiles.Xcode.XcodeFrameworkWrapperUtils.GetGenerateFrameworkWrapperProject(Target.ProjectDirectory);
@@ -1819,9 +1819,9 @@ namespace UnrealBuildTool
 				DirectoryReference XcodeWorkspaceDir = null;
                 if (!Target.bBuildAsFramework)
                 {
-					if (AppName == "UnrealGame" || AppName == "UnrealClient" || Target.ProjectFile == null || Target.ProjectFile.IsUnderDirectory(UnrealBuildTool.EngineDirectory))
+					if (AppName == "UnrealGame" || AppName == "UnrealClient" || Target.ProjectFile == null || Target.ProjectFile.IsUnderDirectory(UnrealBuild.EngineDirectory))
 					{
-						XcodeWorkspaceDir = DirectoryReference.Combine(UnrealBuildTool.RootDirectory, String.Format("UE5_{0}.xcworkspace", (Target.Platform == UnrealTargetPlatform.IOS ? "IOS" : "TVOS")));
+						XcodeWorkspaceDir = DirectoryReference.Combine(UnrealBuild.RootDirectory, String.Format("UE5_{0}.xcworkspace", (Target.Platform == UnrealTargetPlatform.IOS ? "IOS" : "TVOS")));
 					}
 					else
 					{
@@ -1995,7 +1995,7 @@ namespace UnrealBuildTool
 
 				if (!Target.bBuildAsFramework)
 				{
-				    if (AppName == "UnrealGame" || AppName == "UnrealClient" || Target.ProjectFile == null || Target.ProjectFile.IsUnderDirectory(UnrealBuildTool.EngineDirectory))
+				    if (AppName == "UnrealGame" || AppName == "UnrealClient" || Target.ProjectFile == null || Target.ProjectFile.IsUnderDirectory(UnrealBuild.EngineDirectory))
 				    {
 					    GenerateProjectFiles(Target.ProjectFile, new string[] { "-platforms=" + (Target.Platform == UnrealTargetPlatform.IOS ? "IOS" : "TVOS"), "-NoIntellIsense", (Target.Platform == UnrealTargetPlatform.IOS ? "-iosdeployonly" : "-tvosdeployonly"), "-ignorejunk", (Target.bForDistribution ? "-distribution" : "-development"), "-bundleID=" + BundleID, "-includetemptargets", "-appname=" + AppName });
 				    }
diff --git a/Engine/Source/Programs/UnrealBuildTool/Platform/IOS/UEBuildIOS.cs b/Engine/Source/Programs/UnrealBuildTool/Platform/IOS/UEBuildIOS.cs
--- a/Engine/Source/Programs/UnrealBuildTool/Platform/IOS/UEBuildIOS.cs
+++ b/Engine/Source/Programs/UnrealBuildTool/Platform/IOS/UEBuildIOS.cs
@@ -7,6 +7,7 @@ using System.Diagnostics;
 using System.IO;
 using System.Linq;
 using EpicGames.Core;
+using UnrealBuildBase;
 
 #nullable disable
 
@@ -464,7 +465,7 @@ namespace UnrealBuildTool
             string MobileProvision = ProjectSettings.MobileProvision;
 
 			FileReference ProjectFile = ProjectSettings.ProjectFile;
-			FileReference IPhonePackager = FileReference.Combine(UnrealBuildTool.EngineDirectory, "Binaries/DotNET/IOS/IPhonePackager.exe");
+			FileReference IPhonePackager = FileReference.Combine(UnrealBuild.EngineDirectory, "Binaries/DotNET/IOS/IPhonePackager.exe");
 
 			if (!string.IsNullOrEmpty(SigningCertificate))
             {
@@ -473,13 +474,13 @@ namespace UnrealBuildTool
 
 				string IPPCmd = "certificates " + ((ProjectFile != null) ? ("\"" + ProjectFile.ToString() + "\"") : "Engine") + " -bundlename " + ProjectSettings.BundleIdentifier + (bForDistribtion ? " -distribution" : "");
 				
-				IPPProcess.StartInfo.WorkingDirectory = UnrealBuildTool.EngineDirectory.ToString();
+				IPPProcess.StartInfo.WorkingDirectory = UnrealBuild.EngineDirectory.ToString();
 				IPPProcess.OutputDataReceived += new DataReceivedEventHandler(IPPDataReceivedHandler);
 				IPPProcess.ErrorDataReceived += new DataReceivedEventHandler(IPPDataReceivedHandler);
 
 				if (BuildHostPlatform.Current.Platform == UnrealTargetPlatform.Mac)
                 {
-                    IPPProcess.StartInfo.FileName = FileReference.Combine(UnrealBuildTool.EngineDirectory, "Build/BatchFiles/Mac/RunMono.sh").FullName;
+                    IPPProcess.StartInfo.FileName = FileReference.Combine(UnrealBuild.EngineDirectory, "Build/BatchFiles/Mac/RunMono.sh").FullName;
 					IPPProcess.StartInfo.Arguments = string.Format("\"{0}\" {1}", IPhonePackager ,IPPCmd);
 				}
                 else
@@ -527,13 +528,13 @@ namespace UnrealBuildTool
 
 				IPPProcess.OutputDataReceived += new DataReceivedEventHandler(IPPDataReceivedHandler);
 				IPPProcess.ErrorDataReceived += new DataReceivedEventHandler(IPPDataReceivedHandler);
-				IPPProcess.StartInfo.WorkingDirectory = UnrealBuildTool.EngineDirectory.ToString();
+				IPPProcess.StartInfo.WorkingDirectory = UnrealBuild.EngineDirectory.ToString();
 
 				string IPPCmd = "signing_match " + ((ProjectFile != null) ? ("\"" + ProjectFile.ToString() + "\"") : "Engine") + " -bundlename " + ProjectSettings.BundleIdentifier + (bIsTVOS ? " -tvos" : "") + (bForDistribtion ? " -distribution" : "");
 
 				if (BuildHostPlatform.Current.Platform == UnrealTargetPlatform.Mac)
 				{
-					IPPProcess.StartInfo.FileName = FileReference.Combine(UnrealBuildTool.EngineDirectory, "Build/BatchFiles/Mac/RunMono.sh").FullName;
+					IPPProcess.StartInfo.FileName = FileReference.Combine(UnrealBuild.EngineDirectory, "Build/BatchFiles/Mac/RunMono.sh").FullName;
 					IPPProcess.StartInfo.Arguments = string.Format("\"{0}\" {1}", IPhonePackager, IPPCmd);
 				}
 				else
diff --git a/Engine/Source/Programs/UnrealBuildTool/Platform/IOS/UEDeployIOS.cs b/Engine/Source/Programs/UnrealBuildTool/Platform/IOS/UEDeployIOS.cs
--- a/Engine/Source/Programs/UnrealBuildTool/Platform/IOS/UEDeployIOS.cs
+++ b/Engine/Source/Programs/UnrealBuildTool/Platform/IOS/UEDeployIOS.cs
@@ -10,6 +10,7 @@ using System.IO;
 using System.Diagnostics;
 using System.Xml.Linq;
 using EpicGames.Core;
+using UnrealBuildBase;
 
 #nullable disable
 
@@ -704,8 +705,8 @@ namespace UnrealBuildTool
 			// get the receipt
 			if (bIsUnrealGame)
 			{
-				//               ReceiptFilename = TargetReceipt.GetDefaultPath(UnrealBuildTool.EngineDirectory, "UnrealGame", UnrealTargetPlatform.IOS, Config, "");
-				BundlePath = Path.Combine(UnrealBuildTool.EngineDirectory.ToString(), "Intermediate", "IOS-Deploy", "UnrealGame", Config.ToString(), "Payload", "UnrealGame.app");
+				//               ReceiptFilename = TargetReceipt.GetDefaultPath(UnrealBuild.EngineDirectory, "UnrealGame", UnrealTargetPlatform.IOS, Config, "");
+				BundlePath = Path.Combine(UnrealBuild.EngineDirectory.ToString(), "Intermediate", "IOS-Deploy", "UnrealGame", Config.ToString(), "Payload", "UnrealGame.app");
 			}
 			else
 			{
@@ -713,7 +714,7 @@ namespace UnrealBuildTool
 				BundlePath = AppDirectory;//Path.Combine(ProjectDirectory, "Binaries", "IOS", "Payload", ProjectName + ".app");
 			}
 
-			string RelativeEnginePath = UnrealBuildTool.EngineDirectory.MakeRelativeTo(DirectoryReference.GetCurrentDirectory());
+			string RelativeEnginePath = UnrealBuild.EngineDirectory.MakeRelativeTo(DirectoryReference.GetCurrentDirectory());
 
 			UnrealPluginLanguage UPL = new UnrealPluginLanguage(ProjectFile, UPLScripts, ProjectArches, "", "", UnrealTargetPlatform.IOS);
 
@@ -1037,7 +1038,7 @@ namespace UnrealBuildTool
 			string SubDir = GetTargetPlatformName();
 
 			string GameName = TargetName;
-			string ProjectDirectory = (DirectoryReference.FromFile(ProjectFile) ?? UnrealBuildTool.EngineDirectory).FullName;
+			string ProjectDirectory = (DirectoryReference.FromFile(ProjectFile) ?? UnrealBuild.EngineDirectory).FullName;
 			string BuildPath = (GameName == "UnrealGame" ? "../../Engine" : ProjectDirectory) + "/Binaries/" + SubDir;
 			bool bIsUnrealGame = GameName.Contains("UnrealGame");
 
diff --git a/Engine/Source/Programs/UnrealBuildTool/Platform/Linux/LinuxPlatformSDK.cs b/Engine/Source/Programs/UnrealBuildTool/Platform/Linux/LinuxPlatformSDK.cs
--- a/Engine/Source/Programs/UnrealBuildTool/Platform/Linux/LinuxPlatformSDK.cs
+++ b/Engine/Source/Programs/UnrealBuildTool/Platform/Linux/LinuxPlatformSDK.cs
@@ -5,6 +5,7 @@ using System.Collections.Generic;
 using System.IO;
 using EpicGames.Core;
 using System.Text.RegularExpressions;
+using UnrealBuildBase;
 
 #nullable disable
 
@@ -163,7 +164,7 @@ namespace UnrealBuildTool
 		/// </summary>
 		private static DirectoryReference GetInTreeSDKRoot()
 		{
-			return DirectoryReference.Combine(UnrealBuildTool.RootDirectory, "Engine/Extras/ThirdPartyNotUE/SDKs", "Host" + BuildHostPlatform.Current.Platform, TargetPlatformName);
+			return DirectoryReference.Combine(UnrealBuild.RootDirectory, "Engine/Extras/ThirdPartyNotUE/SDKs", "Host" + BuildHostPlatform.Current.Platform, TargetPlatformName);
 		}
 
 		/// <summary>
diff --git a/Engine/Source/Programs/UnrealBuildTool/Platform/Linux/LinuxToolChain.cs b/Engine/Source/Programs/UnrealBuildTool/Platform/Linux/LinuxToolChain.cs
--- a/Engine/Source/Programs/UnrealBuildTool/Platform/Linux/LinuxToolChain.cs
+++ b/Engine/Source/Programs/UnrealBuildTool/Platform/Linux/LinuxToolChain.cs
@@ -104,8 +104,8 @@ namespace UnrealBuildTool
 			bool bHasValidCompiler = false;
 
 			// these are supplied by the engine and do not change depending on the circumstances
-			DumpSymsPath = Path.Combine(UnrealBuildTool.EngineDirectory.FullName, "Binaries", "Linux", "dump_syms");
-			BreakpadEncoderPath = Path.Combine(UnrealBuildTool.EngineDirectory.FullName, "Binaries", "Linux", "BreakpadSymbolEncoder");
+			DumpSymsPath = Path.Combine(UnrealBuild.EngineDirectory.FullName, "Binaries", "Linux", "dump_syms");
+			BreakpadEncoderPath = Path.Combine(UnrealBuild.EngineDirectory.FullName, "Binaries", "Linux", "BreakpadSymbolEncoder");
 
 			if (bForceUseSystemCompiler)
 			{
@@ -1295,7 +1295,7 @@ namespace UnrealBuildTool
 			foreach (DirectoryReference IncludePath in CompileEnvironment.UserIncludePaths)
 			{
 				string IncludePathString;
-				if (IncludePath.IsUnderDirectory(UnrealBuildTool.RootDirectory))
+				if (IncludePath.IsUnderDirectory(UnrealBuild.RootDirectory))
 				{
 					IncludePathString = IncludePath.MakeRelativeTo(UnrealBuildTool.EngineSourceDirectory);
 				}
@@ -1679,21 +1679,21 @@ namespace UnrealBuildTool
 
 					if (LinkEnvironment.bIsBuildingDLL)
 					{
-						// Remove the root UnrealBuildTool.RootDirectory from the RuntimeLibaryPath
-						string AdditionalLibraryRootPath = new FileReference(AdditionalLibrary).Directory.MakeRelativeTo(UnrealBuildTool.RootDirectory);
+						// Remove the root UnrealBuild.RootDirectory from the RuntimeLibaryPath
+						string AdditionalLibraryRootPath = new FileReference(AdditionalLibrary).Directory.MakeRelativeTo(UnrealBuild.RootDirectory);
 
 						// Figure out how many dirs we need to go back
-						string RelativeRootPath = UnrealBuildTool.RootDirectory.MakeRelativeTo(OutputFile.Location.Directory);
+						string RelativeRootPath = UnrealBuild.RootDirectory.MakeRelativeTo(OutputFile.Location.Directory);
 
 						// Combine the two together ie. number of ../ + the path after the root
 						RelativePath = Path.Combine(RelativeRootPath, AdditionalLibraryRootPath);
 					}
 
 					// On Windows, MakeRelativeTo can silently fail if the engine and the project are located on different drives
-					if (CrossCompiling() && RelativePath.StartsWith(UnrealBuildTool.RootDirectory.FullName))
+					if (CrossCompiling() && RelativePath.StartsWith(UnrealBuild.RootDirectory.FullName))
 					{
 						// do not replace directly, but take care to avoid potential double slashes or missed slashes
-						string PathFromRootDir = RelativePath.Replace(UnrealBuildTool.RootDirectory.FullName, "");
+						string PathFromRootDir = RelativePath.Replace(UnrealBuild.RootDirectory.FullName, "");
 						// Path.Combine doesn't combine these properly
 						RelativePath = ((PathFromRootDir.StartsWith("\\") || PathFromRootDir.StartsWith("/")) ? "..\\..\\.." : "..\\..\\..\\") + PathFromRootDir;
 					}
@@ -1714,18 +1714,18 @@ namespace UnrealBuildTool
 				{
 					if (LinkEnvironment.bIsBuildingDLL)
 					{
-						// Remove the root UnrealBuildTool.RootDirectory from the RuntimeLibaryPath
-						string RuntimeLibraryRootPath = new DirectoryReference(RuntimeLibaryPath).MakeRelativeTo(UnrealBuildTool.RootDirectory);
+						// Remove the root UnrealBuild.RootDirectory from the RuntimeLibaryPath
+						string RuntimeLibraryRootPath = new DirectoryReference(RuntimeLibaryPath).MakeRelativeTo(UnrealBuild.RootDirectory);
 
 						// Figure out how many dirs we need to go back
-						string RelativeRootPath = UnrealBuildTool.RootDirectory.MakeRelativeTo(OutputFile.Location.Directory);
+						string RelativeRootPath = UnrealBuild.RootDirectory.MakeRelativeTo(OutputFile.Location.Directory);
 
 						// Combine the two together ie. number of ../ + the path after the root
 						RelativePath = Path.Combine(RelativeRootPath, RuntimeLibraryRootPath);
 					}
 					else
 					{
-						string RelativeRootPath = new DirectoryReference(RuntimeLibaryPath).MakeRelativeTo(UnrealBuildTool.RootDirectory);
+						string RelativeRootPath = new DirectoryReference(RuntimeLibaryPath).MakeRelativeTo(UnrealBuild.RootDirectory);
 
 						// We're assuming that the binary will be placed according to our ProjectName/Binaries/Platform scheme
 						RelativePath = Path.Combine("..", "..", "..", RelativeRootPath);
@@ -1733,10 +1733,10 @@ namespace UnrealBuildTool
 				}
 
 				// On Windows, MakeRelativeTo can silently fail if the engine and the project are located on different drives
-				if (CrossCompiling() && RelativePath.StartsWith(UnrealBuildTool.RootDirectory.FullName))
+				if (CrossCompiling() && RelativePath.StartsWith(UnrealBuild.RootDirectory.FullName))
 				{
 					// do not replace directly, but take care to avoid potential double slashes or missed slashes
-					string PathFromRootDir = RelativePath.Replace(UnrealBuildTool.RootDirectory.FullName, "");
+					string PathFromRootDir = RelativePath.Replace(UnrealBuild.RootDirectory.FullName, "");
 					// Path.Combine doesn't combine these properly
 					RelativePath = ((PathFromRootDir.StartsWith("\\") || PathFromRootDir.StartsWith("/")) ? "..\\..\\.." : "..\\..\\..\\") + PathFromRootDir;
 				}
diff --git a/Engine/Source/Programs/UnrealBuildTool/Platform/Linux/UEBuildLinux.cs b/Engine/Source/Programs/UnrealBuildTool/Platform/Linux/UEBuildLinux.cs
--- a/Engine/Source/Programs/UnrealBuildTool/Platform/Linux/UEBuildLinux.cs
+++ b/Engine/Source/Programs/UnrealBuildTool/Platform/Linux/UEBuildLinux.cs
@@ -7,6 +7,7 @@ using System.Diagnostics;
 using System.IO;
 using EpicGames.Core;
 using System.Text.RegularExpressions;
+using UnrealBuildBase;
 
 #nullable disable
 
@@ -465,7 +466,7 @@ namespace UnrealBuildTool
 
 			if (CompileEnvironment.bPGOOptimize)
 			{
-				DirectoryReference BaseDir = UnrealBuildTool.EngineDirectory;
+				DirectoryReference BaseDir = UnrealBuild.EngineDirectory;
 				if (Target.ProjectFile != null)
 				{
 					BaseDir = DirectoryReference.FromFile(Target.ProjectFile);
diff --git a/Engine/Source/Programs/UnrealBuildTool/Platform/Lumin/UEDeployLumin.cs b/Engine/Source/Programs/UnrealBuildTool/Platform/Lumin/UEDeployLumin.cs
--- a/Engine/Source/Programs/UnrealBuildTool/Platform/Lumin/UEDeployLumin.cs
+++ b/Engine/Source/Programs/UnrealBuildTool/Platform/Lumin/UEDeployLumin.cs
@@ -11,6 +11,7 @@ using System.Linq;
 using System.Xml;
 using System.Xml.Serialization;
 using System.Xml.Linq;
+using UnrealBuildBase;
 
 #nullable disable
 
@@ -492,10 +493,10 @@ namespace UnrealBuildTool
 
 		public void InitUPL(TargetReceipt Receipt)
 		{
-			DirectoryReference ProjectDirectory = Receipt.ProjectDir ?? UnrealBuildTool.EngineDirectory;
+			DirectoryReference ProjectDirectory = Receipt.ProjectDir ?? UnrealBuild.EngineDirectory;
 
 			string UE4BuildPath = Path.Combine(ProjectDirectory.FullName, "Intermediate/Lumin/Mabu");
-			string RelativeEnginePath = UnrealBuildTool.EngineDirectory.MakeRelativeTo(DirectoryReference.GetCurrentDirectory());
+			string RelativeEnginePath = UnrealBuild.EngineDirectory.MakeRelativeTo(DirectoryReference.GetCurrentDirectory());
 			string RelativeProjectPath = ProjectDirectory.MakeRelativeTo(DirectoryReference.GetCurrentDirectory());//.MakeRelativeTo(ProjectDirectory);
 
 			string ConfigurationString = Receipt.Configuration.ToString();
diff --git a/Engine/Source/Programs/UnrealBuildTool/Platform/Mac/MacToolChain.cs b/Engine/Source/Programs/UnrealBuildTool/Platform/Mac/MacToolChain.cs
--- a/Engine/Source/Programs/UnrealBuildTool/Platform/Mac/MacToolChain.cs
+++ b/Engine/Source/Programs/UnrealBuildTool/Platform/Mac/MacToolChain.cs
@@ -706,7 +706,7 @@ namespace UnrealBuildTool
 				// If building an app bundle, we also need an RPATH for use in packaged game and a separate one for staged builds
 				if (bIsBuildingAppBundle)
 				{
-					string EngineDir = UnrealBuildTool.RootDirectory.ToString();
+					string EngineDir = UnrealBuild.RootDirectory.ToString();
 					string ProjectDir = (ProjectFile != null ? ProjectFile.Directory.FullName : null);
 
 					// In packaged games dylibs are stored in Contents/UE4 subfolders, for example in GameName.app/Contents/UE4/Engine/Binaries/ThirdParty/PhysX/Mac
@@ -893,7 +893,7 @@ namespace UnrealBuildTool
 			foreach (FileItem InputFile in LinkEnvironment.InputFiles)
 			{
 				string InputFilePath = InputFile.AbsolutePath;
-				if (InputFile.Location.IsUnderDirectory(UnrealBuildTool.RootDirectory))
+				if (InputFile.Location.IsUnderDirectory(UnrealBuild.RootDirectory))
 				{
 					InputFilePath = InputFile.Location.MakeRelativeTo(UnrealBuildTool.EngineSourceDirectory);
 				}
diff --git a/Engine/Source/Programs/UnrealBuildTool/Platform/Windows/PVSToolChain.cs b/Engine/Source/Programs/UnrealBuildTool/Platform/Windows/PVSToolChain.cs
--- a/Engine/Source/Programs/UnrealBuildTool/Platform/Windows/PVSToolChain.cs
+++ b/Engine/Source/Programs/UnrealBuildTool/Platform/Windows/PVSToolChain.cs
@@ -356,7 +356,7 @@ namespace UnrealBuildTool
 			AnalyzerFile = FileReference.Combine(new DirectoryReference(Environment.GetFolderPath(Environment.SpecialFolder.ProgramFilesX86)), "PVS-Studio", "x64", "PVS-Studio.exe");
 			if(!FileReference.Exists(AnalyzerFile))
 			{
-				FileReference EngineAnalyzerFile = FileReference.Combine(UnrealBuildTool.RootDirectory, "Engine", "Restricted", "NoRedist", "Extras", "ThirdPartyNotUE", "PVS-Studio", "PVS-Studio.exe");
+				FileReference EngineAnalyzerFile = FileReference.Combine(UnrealBuild.RootDirectory, "Engine", "Restricted", "NoRedist", "Extras", "ThirdPartyNotUE", "PVS-Studio", "PVS-Studio.exe");
 				if (FileReference.Exists(EngineAnalyzerFile))
 				{
 					AnalyzerFile = EngineAnalyzerFile;
@@ -380,7 +380,7 @@ namespace UnrealBuildTool
 
 				if (!String.IsNullOrEmpty(ApplicationSettings.UserName) && !String.IsNullOrEmpty(ApplicationSettings.SerialNumber))
 				{
-					LicenseFile = FileReference.Combine(UnrealBuildTool.EngineDirectory, "Intermediate", "PVS", "PVS-Studio.lic");
+					LicenseFile = FileReference.Combine(UnrealBuild.EngineDirectory, "Intermediate", "PVS", "PVS-Studio.lic");
 					Utils.WriteFileIfChanged(LicenseFile, String.Format("{0}\n{1}\n", ApplicationSettings.UserName, ApplicationSettings.SerialNumber), StringComparison.Ordinal);
 				}
 			}
@@ -612,7 +612,7 @@ namespace UnrealBuildTool
 			FileReference OutputFile;
 			if (Target.ProjectFile == null)
 			{
-				OutputFile = FileReference.Combine(UnrealBuildTool.EngineDirectory, "Saved", "PVS-Studio", String.Format("{0}.pvslog", Target.Name));
+				OutputFile = FileReference.Combine(UnrealBuild.EngineDirectory, "Saved", "PVS-Studio", String.Format("{0}.pvslog", Target.Name));
 			}
 			else
 			{
diff --git a/Engine/Source/Programs/UnrealBuildTool/Platform/Windows/UEBuildWindows.cs b/Engine/Source/Programs/UnrealBuildTool/Platform/Windows/UEBuildWindows.cs
--- a/Engine/Source/Programs/UnrealBuildTool/Platform/Windows/UEBuildWindows.cs
+++ b/Engine/Source/Programs/UnrealBuildTool/Platform/Windows/UEBuildWindows.cs
@@ -12,6 +12,7 @@ using Microsoft.VisualStudio.Setup.Configuration;
 using System.Runtime.InteropServices;
 using System.Diagnostics.CodeAnalysis;
 using System.Buffers.Binary;
+using UnrealBuildBase;
 
 namespace UnrealBuildTool
 {
@@ -403,7 +404,7 @@ namespace UnrealBuildTool
 		{
 			this.Target = Target;
 
-			ManifestFile = FileReference.Combine(UnrealBuildTool.EngineDirectory, "Build", "Windows", "Resources", String.Format("Default-{0}.manifest", Target.Platform)).FullName;
+			ManifestFile = FileReference.Combine(UnrealBuild.EngineDirectory, "Build", "Windows", "Resources", String.Format("Default-{0}.manifest", Target.Platform)).FullName;
 		}
 	}
 
@@ -1821,7 +1822,7 @@ namespace UnrealBuildTool
 			}
 
 			// Otherwise use the default
-			return FileReference.Combine(UnrealBuildTool.EngineDirectory, "Build", "Windows", "Resources", "Default.ico");
+			return FileReference.Combine(UnrealBuild.EngineDirectory, "Build", "Windows", "Resources", "Default.ico");
 		}
 
 		/// <summary>
diff --git a/Engine/Source/Programs/UnrealBuildTool/Platform/Windows/VCCompileAction.cs b/Engine/Source/Programs/UnrealBuildTool/Platform/Windows/VCCompileAction.cs
--- a/Engine/Source/Programs/UnrealBuildTool/Platform/Windows/VCCompileAction.cs
+++ b/Engine/Source/Programs/UnrealBuildTool/Platform/Windows/VCCompileAction.cs
@@ -228,7 +228,7 @@ namespace UnrealBuildTool
 			{
 				if (UsingClFilter)
 				{
-					return FileReference.Combine(UnrealBuildTool.EngineDirectory, "Build", "Windows", "cl-filter", "cl-filter.exe");
+					return FileReference.Combine(UnrealBuild.EngineDirectory, "Build", "Windows", "cl-filter", "cl-filter.exe");
 				}
 				else
 				{
diff --git a/Engine/Source/Programs/UnrealBuildTool/Platform/Windows/VCManifestGenerator.cs b/Engine/Source/Programs/UnrealBuildTool/Platform/Windows/VCManifestGenerator.cs
--- a/Engine/Source/Programs/UnrealBuildTool/Platform/Windows/VCManifestGenerator.cs
+++ b/Engine/Source/Programs/UnrealBuildTool/Platform/Windows/VCManifestGenerator.cs
@@ -10,6 +10,7 @@ using EpicGames.Core;
 using System.Xml.Linq;
 using System.Text;
 using System.Diagnostics;
+using UnrealBuildBase;
 
 #nullable disable
 
@@ -292,13 +293,13 @@ namespace UnrealBuildTool
 			// look in Engine, if allowed
 			if (!bFileExists && AllowEngineFallback)
 			{
-				SourcePath = Path.Combine(UnrealBuildTool.EngineDirectory.FullName, "Build", Platform.ToString(), EngineResourceSubPath);
+				SourcePath = Path.Combine(UnrealBuild.EngineDirectory.FullName, "Build", Platform.ToString(), EngineResourceSubPath);
 				bFileExists = File.Exists(Path.Combine(SourcePath, ResourceFileName));
 
 				// look in Platform extensions too
 				if (!bFileExists)
 				{
-					SourcePath = Path.Combine(UnrealBuildTool.EngineDirectory.FullName, "Platforms", Platform.ToString(), "Build", EngineResourceSubPath);
+					SourcePath = Path.Combine(UnrealBuild.EngineDirectory.FullName, "Platforms", Platform.ToString(), "Build", EngineResourceSubPath);
 					bFileExists = File.Exists(Path.Combine(SourcePath, ResourceFileName));
 				}
 			}
diff --git a/Engine/Source/Programs/UnrealBuildTool/Platform/Windows/VCPreludeCache.cs b/Engine/Source/Programs/UnrealBuildTool/Platform/Windows/VCPreludeCache.cs
--- a/Engine/Source/Programs/UnrealBuildTool/Platform/Windows/VCPreludeCache.cs
+++ b/Engine/Source/Programs/UnrealBuildTool/Platform/Windows/VCPreludeCache.cs
@@ -6,6 +6,7 @@ using System.Diagnostics;
 using System.IO;
 using System.Linq;
 using EpicGames.Core;
+using UnrealBuildBase;
 
 namespace UnrealBuildTool
 {
@@ -35,12 +36,12 @@ namespace UnrealBuildTool
 		/// <summary>
 		/// File containing a list of macros to be used to generate a prelude file
 		/// </summary>
-		static readonly FileReference MacroNamesFile = FileReference.Combine(UnrealBuildTool.EngineDirectory, "Build", "Windows", "PreludeMacros.txt");
+		static readonly FileReference MacroNamesFile = FileReference.Combine(UnrealBuild.EngineDirectory, "Build", "Windows", "PreludeMacros.txt");
 
 		/// <summary>
 		/// Base directory for the cache
 		/// </summary>
-		static readonly DirectoryReference PreludeCacheDir = DirectoryReference.Combine(UnrealBuildTool.EngineDirectory, "Saved", "UnrealBuildTool", "Prelude", "MSVC");
+		static readonly DirectoryReference PreludeCacheDir = DirectoryReference.Combine(UnrealBuild.EngineDirectory, "Saved", "UnrealBuildTool", "Prelude", "MSVC");
 
 		/// <summary>
 		/// The parsed list of macros
diff --git a/Engine/Source/Programs/UnrealBuildTool/Platform/Windows/VCToolChain.cs b/Engine/Source/Programs/UnrealBuildTool/Platform/Windows/VCToolChain.cs
--- a/Engine/Source/Programs/UnrealBuildTool/Platform/Windows/VCToolChain.cs
+++ b/Engine/Source/Programs/UnrealBuildTool/Platform/Windows/VCToolChain.cs
@@ -121,7 +121,7 @@ namespace UnrealBuildTool
 		{
 			// Try to use a relative path to shorten command line length. Always need the full path when preprocessing because the output file will be in a different place, where include paths cannot be relative.
 			string IncludePathString;
-			if(IncludePath.IsUnderDirectory(UnrealBuildTool.RootDirectory) && Compiler != WindowsCompiler.Clang && !bPreprocessOnly)
+			if(IncludePath.IsUnderDirectory(UnrealBuild.RootDirectory) && Compiler != WindowsCompiler.Clang && !bPreprocessOnly)
 			{
 				IncludePathString = IncludePath.MakeRelativeTo(UnrealBuildTool.EngineSourceDirectory);
 			}
@@ -1831,7 +1831,7 @@ namespace UnrealBuildTool
 			}
 			else
 			{
-				LinkAction.CommandPath = FileReference.Combine(UnrealBuildTool.EngineDirectory, "Build", "Windows", "link-filter", "link-filter.exe");
+				LinkAction.CommandPath = FileReference.Combine(UnrealBuild.EngineDirectory, "Build", "Windows", "link-filter", "link-filter.exe");
 				LinkAction.CommandArguments = String.Format("-- \"{0}\" @\"{1}\"", EnvVars.LinkerPath, ResponseFileName);
 			}
 			LinkAction.CommandVersion = EnvVars.ToolChainVersion.ToString();
diff --git a/Engine/Source/Programs/UnrealBuildTool/ProjectFiles/CMake/CMakefileGenerator.cs b/Engine/Source/Programs/UnrealBuildTool/ProjectFiles/CMake/CMakefileGenerator.cs
--- a/Engine/Source/Programs/UnrealBuildTool/ProjectFiles/CMake/CMakefileGenerator.cs
+++ b/Engine/Source/Programs/UnrealBuildTool/ProjectFiles/CMake/CMakefileGenerator.cs
@@ -5,6 +5,7 @@ using System.Collections.Generic;
 using System.Text;
 using System.IO;
 using EpicGames.Core;
+using UnrealBuildBase;
 
 #nullable disable
 
@@ -279,7 +280,7 @@ namespace UnrealBuildTool
 			StringBuilder IncludeDirectoriesList = new StringBuilder("include_directories( \n");
 			StringBuilder PreprocessorDefinitionsList = new StringBuilder("add_definitions( \n");
 
-            string UE4RootPath = Utils.CleanDirectorySeparators(UnrealBuildTool.RootDirectory.FullName, '/');
+            string UE4RootPath = Utils.CleanDirectorySeparators(UnrealBuild.RootDirectory.FullName, '/');
             string CMakeGameRootPath = "";
 			string GameProjectPath = "";
 			string CMakeGameProjectFile = "";
@@ -368,9 +369,9 @@ namespace UnrealBuildTool
 					string IncludeDirectory = GetIncludeDirectory(IncludeSearchPath, Path.GetDirectoryName(CurProject.ProjectFilePath.FullName));
 					if (IncludeDirectory != null && !IncludeDirectories.Contains(IncludeDirectory))
 					{
-						if (IncludeDirectory.Contains(UnrealBuildTool.RootDirectory.FullName))
+						if (IncludeDirectory.Contains(UnrealBuild.RootDirectory.FullName))
 						{
-							IncludeDirectories.Add(IncludeDirectory.Replace(UnrealBuildTool.RootDirectory.FullName, UE4RootPath));
+							IncludeDirectories.Add(IncludeDirectory.Replace(UnrealBuild.RootDirectory.FullName, UE4RootPath));
 						}
 						else
 						{
@@ -415,7 +416,7 @@ namespace UnrealBuildTool
 				List<FileReference> FoundFiles = SourceFileSearch.FindModuleSourceFiles(CurModuleFile);
 				foreach (FileReference CurSourceFile in FoundFiles)
 				{
-					string SourceFileRelativeToRoot = CurSourceFile.MakeRelativeTo(UnrealBuildTool.EngineDirectory);
+					string SourceFileRelativeToRoot = CurSourceFile.MakeRelativeTo(UnrealBuild.EngineDirectory);
 
 					// Exclude files/folders on a per-platform basis.
 					if (!IsPathExcludedOnPlatform(SourceFileRelativeToRoot, BuildHostPlatform.Current.Platform))
@@ -465,10 +466,10 @@ namespace UnrealBuildTool
 			}
 
 			// Add Engine/Shaders files (game are added via modules)
-			List<FileReference> EngineShaderFiles = SourceFileSearch.FindFiles(DirectoryReference.Combine(UnrealBuildTool.EngineDirectory, "Shaders"));
+			List<FileReference> EngineShaderFiles = SourceFileSearch.FindFiles(DirectoryReference.Combine(UnrealBuild.EngineDirectory, "Shaders"));
 			foreach (FileReference CurSourceFile in EngineShaderFiles)
 			{
-				string SourceFileRelativeToRoot = CurSourceFile.MakeRelativeTo(UnrealBuildTool.EngineDirectory);
+				string SourceFileRelativeToRoot = CurSourceFile.MakeRelativeTo(UnrealBuild.EngineDirectory);
 				if (SourceFileRelativeToRoot.EndsWith(".usf") || SourceFileRelativeToRoot.EndsWith(".ush"))
 				{
 					AppendCleanedPathToList(CMakeEngineShaderFilesList, CMakeProjectShaderFilesList, SourceFileRelativeToRoot, CurSourceFile.FullName, GameProjectPath, UE4RootPath, CMakeGameRootPath);
@@ -476,10 +477,10 @@ namespace UnrealBuildTool
 			}
 
 			// Add Engine/Config ini files (game are added via modules)
-			List<FileReference> EngineConfigFiles = SourceFileSearch.FindFiles(DirectoryReference.Combine(UnrealBuildTool.EngineDirectory, "Config"));
+			List<FileReference> EngineConfigFiles = SourceFileSearch.FindFiles(DirectoryReference.Combine(UnrealBuild.EngineDirectory, "Config"));
 			foreach (FileReference CurSourceFile in EngineConfigFiles)
 			{
-				string SourceFileRelativeToRoot = CurSourceFile.MakeRelativeTo(UnrealBuildTool.EngineDirectory);
+				string SourceFileRelativeToRoot = CurSourceFile.MakeRelativeTo(UnrealBuild.EngineDirectory);
 				if (SourceFileRelativeToRoot.EndsWith(".ini"))
 				{
 					AppendCleanedPathToList(CMakeEngineConfigFilesList, CMakeProjectConfigFilesList, SourceFileRelativeToRoot, CurSourceFile.FullName, GameProjectPath, UE4RootPath, CMakeGameRootPath);
@@ -725,7 +726,7 @@ namespace UnrealBuildTool
 			// Check for paths outside of both the engine and the project
 			if (Path.IsPathRooted(IncludeDir) &&
 				!IncludeDir.StartsWith(FullProjectPath) &&
-				!IncludeDir.StartsWith(UnrealBuildTool.RootDirectory.FullName))
+				!IncludeDir.StartsWith(UnrealBuild.RootDirectory.FullName))
 			{
 				// Full path to a folder outside of project
 				FullPath = IncludeDir;
@@ -733,7 +734,7 @@ namespace UnrealBuildTool
 			else
 			{
 				FullPath = Path.GetFullPath(Path.Combine(ProjectDir, IncludeDir));
-				if (!FullPath.StartsWith(UnrealBuildTool.RootDirectory.FullName))
+				if (!FullPath.StartsWith(UnrealBuild.RootDirectory.FullName))
 				{
 					FullPath = Utils.MakePathRelativeTo(FullPath, FullProjectPath);
 				}
diff --git a/Engine/Source/Programs/UnrealBuildTool/ProjectFiles/CodeLite/CodeLiteProject.cs b/Engine/Source/Programs/UnrealBuildTool/ProjectFiles/CodeLite/CodeLiteProject.cs
--- a/Engine/Source/Programs/UnrealBuildTool/ProjectFiles/CodeLite/CodeLiteProject.cs
+++ b/Engine/Source/Programs/UnrealBuildTool/ProjectFiles/CodeLite/CodeLiteProject.cs
@@ -7,6 +7,7 @@ using System.IO;
 using System.Xml;
 using System.Xml.Linq;
 using EpicGames.Core;
+using UnrealBuildBase;
 
 #nullable disable
 
@@ -45,7 +46,7 @@ namespace UnrealBuildTool
 			string ProjectPlatformName = BuildHostPlatform.Current.Platform.ToString();
 
 			// Get the output directory
-			string EngineRootDirectory = UnrealBuildTool.EngineDirectory.FullName;
+			string EngineRootDirectory = UnrealBuild.EngineDirectory.FullName;
 
 			//
 			// Build the working directory of the Game executable.
@@ -129,7 +130,7 @@ namespace UnrealBuildTool
 							TargetName.Equals("UnrealGame") ||
 							TargetName.Equals("UnrealEditor"))
 						{
-							int Idx = UnrealBuildTool.EngineDirectory.FullName.Length;
+							int Idx = UnrealBuild.EngineDirectory.FullName.Length;
 							CurrentFilePath = Path.GetDirectoryName(Path.GetFullPath(CurrentFile.Reference.FullName)).Substring(Idx);
 						}
 						else
@@ -423,7 +424,7 @@ namespace UnrealBuildTool
 							PlatformName = "Mac";
 						}
 
-						BuildTarget = Path.Combine(UnrealBuildTool.EngineDirectory.FullName, "Build/BatchFiles", PlatformName, "Build.sh") + " " + BuildTarget;
+						BuildTarget = Path.Combine(UnrealBuild.EngineDirectory.FullName, "Build/BatchFiles", PlatformName, "Build.sh") + " " + BuildTarget;
 					}
 					else
 					{
diff --git a/Engine/Source/Programs/UnrealBuildTool/ProjectFiles/KDevelop/KDevelopGenerator.cs b/Engine/Source/Programs/UnrealBuildTool/ProjectFiles/KDevelop/KDevelopGenerator.cs
--- a/Engine/Source/Programs/UnrealBuildTool/ProjectFiles/KDevelop/KDevelopGenerator.cs
+++ b/Engine/Source/Programs/UnrealBuildTool/ProjectFiles/KDevelop/KDevelopGenerator.cs
@@ -6,6 +6,7 @@ using System.Text;
 using System.IO;
 using System.Linq;
 using EpicGames.Core;
+using UnrealBuildBase;
 
 #nullable disable
 
@@ -157,7 +158,7 @@ namespace UnrealBuildTool
 		{
 			int BuildConfigIndex = 1;
 
-			string UnrealRootPath = UnrealBuildTool.RootDirectory.FullName;
+			string UnrealRootPath = UnrealBuild.RootDirectory.FullName;
 			FileContent.Append("[CustomBuildSystem]\n");
 			FileContent.Append("CurrentConfiguration=BuildConfig0\n\n"); //
 
@@ -270,7 +271,7 @@ namespace UnrealBuildTool
 			List<string> IncludeDirectories = new List<string>();
 			List<string> SystemIncludeDirectories = new List<string>();
 
-			string UnrealEngineRootPath = UnrealBuildTool.RootDirectory.FullName;
+			string UnrealEngineRootPath = UnrealBuild.RootDirectory.FullName;
 
 			int IncludeIndex = 1;
 			// Iterate through all the include paths that
diff --git a/Engine/Source/Programs/UnrealBuildTool/ProjectFiles/Make/MakefileGenerator.cs b/Engine/Source/Programs/UnrealBuildTool/ProjectFiles/Make/MakefileGenerator.cs
--- a/Engine/Source/Programs/UnrealBuildTool/ProjectFiles/Make/MakefileGenerator.cs
+++ b/Engine/Source/Programs/UnrealBuildTool/ProjectFiles/Make/MakefileGenerator.cs
@@ -5,6 +5,7 @@ using System.Collections.Generic;
 using System.Text;
 using System.IO;
 using EpicGames.Core;
+using UnrealBuildBase;
 
 #nullable disable
 
@@ -76,7 +77,7 @@ namespace UnrealBuildTool
 
 			string MakeGameProjectFile = "";
 
-			string UnrealRootPath = UnrealBuildTool.RootDirectory.FullName;
+			string UnrealRootPath = UnrealBuild.RootDirectory.FullName;
 
 			if (!String.IsNullOrEmpty(GameProjectName))
 			{
diff --git a/Engine/Source/Programs/UnrealBuildTool/ProjectFiles/ProjectFileGenerator.cs b/Engine/Source/Programs/UnrealBuildTool/ProjectFiles/ProjectFileGenerator.cs
--- a/Engine/Source/Programs/UnrealBuildTool/ProjectFiles/ProjectFileGenerator.cs
+++ b/Engine/Source/Programs/UnrealBuildTool/ProjectFiles/ProjectFileGenerator.cs
@@ -334,7 +334,7 @@ namespace UnrealBuildTool
 		/// <summary>
 		/// Relative path to the directory where the master project file will be saved to
 		/// </summary>
-		public static DirectoryReference MasterProjectPath = UnrealBuildTool.RootDirectory; // We'll save the master project to our "root" folder
+		public static DirectoryReference MasterProjectPath = UnrealBuild.RootDirectory; // We'll save the master project to our "root" folder
 
 		/// <summary>
 		/// Name of the UE4 engine project that contains all of the engine code, config files and other files
@@ -345,7 +345,7 @@ namespace UnrealBuildTool
 		/// When ProjectsAreIntermediate is true, this is the directory to store generated project files
 		/// @todo projectfiles: Ideally, projects for game modules/targets would be created in the game's Intermediate folder!
 		/// </summary>
-		public static DirectoryReference IntermediateProjectFilesPath = DirectoryReference.Combine(UnrealBuildTool.EngineDirectory, "Intermediate", "ProjectFiles");
+		public static DirectoryReference IntermediateProjectFilesPath = DirectoryReference.Combine(UnrealBuild.EngineDirectory, "Intermediate", "ProjectFiles");
 
 		/// <summary>
 		/// Path to timestamp file, recording when was the last time projects were created.
@@ -450,7 +450,7 @@ namespace UnrealBuildTool
 					AutomationProjectFiles.Add(Project);
 					Folder.ChildProjects.Add(Project);
 
-					if (!ProjectFile.IsUnderDirectory(UnrealBuildTool.EngineDirectory))
+					if (!ProjectFile.IsUnderDirectory(UnrealBuild.EngineDirectory))
 					{
 						FileReference PropsFile = new FileReference(ProjectFile.FullName + ".props");
 						CreateAutomationProjectPropsFile(PropsFile);
@@ -499,7 +499,7 @@ namespace UnrealBuildTool
 					Writer.WriteLine("<?xml version=\"1.0\" encoding=\"utf-8\"?>");
 					Writer.WriteLine("<Project ToolsVersion=\"Current\" xmlns=\"http://schemas.microsoft.com/developer/msbuild/2003\">");
 					Writer.WriteLine("\t<PropertyGroup>");
-					Writer.WriteLine("\t\t<EngineDir Condition=\"'$(EngineDir)' == ''\">{0}</EngineDir>", UnrealBuildTool.EngineDirectory);
+					Writer.WriteLine("\t\t<EngineDir Condition=\"'$(EngineDir)' == ''\">{0}</EngineDir>", UnrealBuild.EngineDirectory);
 					Writer.WriteLine("\t</PropertyGroup>");
 					Writer.WriteLine("</Project>");
 				}
@@ -516,11 +516,11 @@ namespace UnrealBuildTool
 			List<FileReference> FoundProjects = new List<FileReference>();
 			List<DirectoryReference> ProjectDirs = new List<DirectoryReference>();
 
-			DirectoryReference EngineExtras = DirectoryReference.Combine(UnrealBuildTool.EngineDirectory, "Extras");
+			DirectoryReference EngineExtras = DirectoryReference.Combine(UnrealBuild.EngineDirectory, "Extras");
 			ProjectDirs.Add(EngineExtras);
 			DiscoverCSharpProgramProjectsRecursively(EngineExtras, FoundProjects);
 
-			List<DirectoryReference> AllEngineDirectories = UnrealBuildTool.GetExtensionDirs(UnrealBuildTool.EngineDirectory, "Source/Programs");
+			List<DirectoryReference> AllEngineDirectories = UnrealBuildTool.GetExtensionDirs(UnrealBuild.EngineDirectory, "Source/Programs");
 			ProjectDirs = ProjectDirs.Union(AllEngineDirectories).ToList();
 			foreach (DirectoryReference EngineDir in AllEngineDirectories)
 			{
@@ -705,7 +705,7 @@ namespace UnrealBuildTool
 					{
 						if (BuildHostPlatform.Current.Platform == UnrealTargetPlatform.Mac)
 						{
-							MasterProjectPath = UnrealBuildTool.EngineDirectory;
+							MasterProjectPath = UnrealBuild.EngineDirectory;
 							GameProjectName = "UnrealGame";
 						}
 						if (!DirectoryReference.Exists(DirectoryReference.Combine(MasterProjectPath, "Source")))
@@ -733,7 +733,7 @@ namespace UnrealBuildTool
 				}
 
 				// Write out the name of the master project file, so the runtime knows to use it
-				FileReference MasterProjectNameLocation = FileReference.Combine(UnrealBuildTool.EngineDirectory, "Intermediate", "ProjectFiles", "MasterProjectName.txt");
+				FileReference MasterProjectNameLocation = FileReference.Combine(UnrealBuild.EngineDirectory, "Intermediate", "ProjectFiles", "MasterProjectName.txt");
 				DirectoryReference.CreateDirectory(MasterProjectNameLocation.Directory);
 				FileReference.WriteAllText(MasterProjectNameLocation, MasterProjectName);
 			}
@@ -757,7 +757,7 @@ namespace UnrealBuildTool
 				}
 
 				// the master project name is always read from our intermediate directory and not the overriden one for this set of platforms
-				FileReference MasterProjectNameLocation = FileReference.Combine(UnrealBuildTool.EngineDirectory, "Intermediate", "ProjectFiles", "MasterProjectName.txt");
+				FileReference MasterProjectNameLocation = FileReference.Combine(UnrealBuild.EngineDirectory, "Intermediate", "ProjectFiles", "MasterProjectName.txt");
 				DirectoryReference.CreateDirectory(MasterProjectNameLocation.Directory);
 				FileReference.WriteAllText(MasterProjectNameLocation, MasterProjectName);
 			}
@@ -829,7 +829,7 @@ namespace UnrealBuildTool
 						FileReference ProjectFilePath = FileReference.Combine(IntermediateProjectFilesPath, EngineProjectFileNameBase + ProjectFileExtension);
 
 						bool bAlreadyExisted;
-						EngineProject = FindOrAddProject(ProjectFilePath, UnrealBuildTool.EngineDirectory, true, out bAlreadyExisted);
+						EngineProject = FindOrAddProject(ProjectFilePath, UnrealBuild.EngineDirectory, true, out bAlreadyExisted);
 
 						EngineProject.IsForeignProject = false;
 						EngineProject.IsGeneratedProject = true;
@@ -902,8 +902,8 @@ namespace UnrealBuildTool
 						RootFolder.AddSubFolder("Mods").ChildProjects.Add(CurModProject);
 					}
 
-					DirectoryReference TemplatesDirectory = DirectoryReference.Combine(UnrealBuildTool.RootDirectory, "Templates");
-					DirectoryReference SamplesDirectory = DirectoryReference.Combine(UnrealBuildTool.RootDirectory, "Samples");
+					DirectoryReference TemplatesDirectory = DirectoryReference.Combine(UnrealBuild.RootDirectory, "Templates");
+					DirectoryReference SamplesDirectory = DirectoryReference.Combine(UnrealBuild.RootDirectory, "Samples");
 
 					foreach (ProjectFile CurGameProject in GameProjects)
 					{
@@ -1331,10 +1331,10 @@ namespace UnrealBuildTool
 		/// Adds all engine localization text files to the specified project
 		private void AddEngineLocalizationFiles(ProjectFile EngineProject)
 		{
-			DirectoryReference EngineLocalizationDirectory = DirectoryReference.Combine(UnrealBuildTool.EngineDirectory, "Content", "Localization");
+			DirectoryReference EngineLocalizationDirectory = DirectoryReference.Combine(UnrealBuild.EngineDirectory, "Content", "Localization");
 			if (DirectoryReference.Exists(EngineLocalizationDirectory))
 			{
-				EngineProject.AddFilesToProject(SourceFileSearch.FindFiles(EngineLocalizationDirectory), UnrealBuildTool.EngineDirectory);
+				EngineProject.AddFilesToProject(SourceFileSearch.FindFiles(EngineLocalizationDirectory), UnrealBuild.EngineDirectory);
 			}
 		}
 
@@ -1342,10 +1342,10 @@ namespace UnrealBuildTool
 		/// Adds all engine template text files to the specified project
 		private void AddEngineTemplateFiles(ProjectFile EngineProject)
 		{
-			DirectoryReference EngineTemplateDirectory = DirectoryReference.Combine(UnrealBuildTool.EngineDirectory, "Content", "Editor", "Templates");
+			DirectoryReference EngineTemplateDirectory = DirectoryReference.Combine(UnrealBuild.EngineDirectory, "Content", "Editor", "Templates");
 			if (DirectoryReference.Exists(EngineTemplateDirectory))
 			{
-				EngineProject.AddFilesToProject(SourceFileSearch.FindFiles(EngineTemplateDirectory), UnrealBuildTool.EngineDirectory);
+				EngineProject.AddFilesToProject(SourceFileSearch.FindFiles(EngineTemplateDirectory), UnrealBuild.EngineDirectory);
 			}
 		}
 
@@ -1353,10 +1353,10 @@ namespace UnrealBuildTool
 		/// Adds all engine config files to the specified project
 		private void AddEngineConfigFiles(ProjectFile EngineProject)
 		{
-			DirectoryReference EngineConfigDirectory = DirectoryReference.Combine(UnrealBuildTool.EngineDirectory, "Config");
+			DirectoryReference EngineConfigDirectory = DirectoryReference.Combine(UnrealBuild.EngineDirectory, "Config");
 			if (DirectoryReference.Exists(EngineConfigDirectory))
 			{
-				EngineProject.AddFilesToProject(SourceFileSearch.FindFiles(EngineConfigDirectory), UnrealBuildTool.EngineDirectory);
+				EngineProject.AddFilesToProject(SourceFileSearch.FindFiles(EngineConfigDirectory), UnrealBuild.EngineDirectory);
 			}
 		}
 
@@ -1370,16 +1370,16 @@ namespace UnrealBuildTool
 		{
 			// @todo: this will add the same files to the solution (like the UBT source files that also get added to UnrealBuildTool project).
 			// not sure of a good filtering method here
-			foreach (DirectoryReference ExtensionDir in UnrealBuildTool.GetExtensionDirs(UnrealBuildTool.EngineDirectory))
+			foreach (DirectoryReference ExtensionDir in UnrealBuildTool.GetExtensionDirs(UnrealBuild.EngineDirectory))
 			{
-				if (ExtensionDir != UnrealBuildTool.EngineDirectory)
+				if (ExtensionDir != UnrealBuild.EngineDirectory)
 				{
 					List<string> SubdirectoryNamesToExclude = new List<string>();
 					SubdirectoryNamesToExclude.Add("AutomationTool"); //automation files are added separately to the AutomationTool project
 					SubdirectoryNamesToExclude.Add("Binaries");
 					SubdirectoryNamesToExclude.Add("Content");
 
-					EngineProject.AddFilesToProject(SourceFileSearch.FindFiles(ExtensionDir, SubdirectoryNamesToExclude), UnrealBuildTool.EngineDirectory);
+					EngineProject.AddFilesToProject(SourceFileSearch.FindFiles(ExtensionDir, SubdirectoryNamesToExclude), UnrealBuild.EngineDirectory);
 				}
 			}
 		}
@@ -1387,10 +1387,10 @@ namespace UnrealBuildTool
 		/// Adds UnrealHeaderTool config files to the specified project
 		private void AddUnrealHeaderToolConfigFiles(ProjectFile EngineProject)
 		{
-			DirectoryReference UHTConfigDirectory = DirectoryReference.Combine(UnrealBuildTool.EngineDirectory, "Programs", "UnrealHeaderTool", "Config");
+			DirectoryReference UHTConfigDirectory = DirectoryReference.Combine(UnrealBuild.EngineDirectory, "Programs", "UnrealHeaderTool", "Config");
 			if (DirectoryReference.Exists(UHTConfigDirectory))
 			{
-				EngineProject.AddFilesToProject(SourceFileSearch.FindFiles(UHTConfigDirectory), UnrealBuildTool.EngineDirectory);
+				EngineProject.AddFilesToProject(SourceFileSearch.FindFiles(UHTConfigDirectory), UnrealBuild.EngineDirectory);
 			}
 		}
 
@@ -1453,7 +1453,7 @@ namespace UnrealBuildTool
 				string CleanTargetFileName = Utils.CleanDirectorySeparators(CurTargetFile.FullName);
 
 				// remove the local root
-				string LocalRoot = UnrealBuildTool.RootDirectory.FullName;
+				string LocalRoot = UnrealBuild.RootDirectory.FullName;
 				string Search = CleanTargetFileName;
 				if (Search.StartsWith(LocalRoot, StringComparison.InvariantCultureIgnoreCase))
 				{
@@ -1777,7 +1777,7 @@ namespace UnrealBuildTool
 					// @todo projectfiles: Fragile assumption here about Programs always being under /Engine/Programs
 
 					DirectoryReference ProgramDirectory;
-					ProgramDirectory = DirectoryReference.Combine(UnrealBuildTool.EngineDirectory, "Programs", ProgramName);
+					ProgramDirectory = DirectoryReference.Combine(UnrealBuild.EngineDirectory, "Programs", ProgramName);
 
 					DirectoryReference ProgramConfigDirectory = DirectoryReference.Combine(ProgramDirectory, "Config");
 					if (DirectoryReference.Exists(ProgramConfigDirectory))
@@ -2015,12 +2015,12 @@ namespace UnrealBuildTool
 				bool WantProjectFileForModule = true;
 
 				// We'll keep track of whether this is an "engine" or "external" module.  This is determined below while loading module rules.
-				bool IsThirdPartyModule = CurModuleFile.ContainsName("ThirdParty", UnrealBuildTool.RootDirectory);
+				bool IsThirdPartyModule = CurModuleFile.ContainsName("ThirdParty", UnrealBuild.RootDirectory);
 
 				// check for engine, or platform extension engine folders
 				if (!bIncludeEngineSource)
 				{
-					if (CurModuleFile.IsUnderDirectory(UnrealBuildTool.EngineDirectory))
+					if (CurModuleFile.IsUnderDirectory(UnrealBuild.EngineDirectory))
 					{
 						// We were asked to exclude engine modules from the generated projects
 						WantProjectFileForModule = false;
@@ -2155,11 +2155,11 @@ namespace UnrealBuildTool
 				}
 
 				// check for engine, or platform extension engine folders
-				foreach (DirectoryReference ExtensionDir in UnrealBuildTool.GetExtensionDirs(UnrealBuildTool.EngineDirectory))
+				foreach (DirectoryReference ExtensionDir in UnrealBuildTool.GetExtensionDirs(UnrealBuild.EngineDirectory))
 				{
 					if (Path == ExtensionDir)
 					{
-						BaseFolder = UnrealBuildTool.EngineDirectory;
+						BaseFolder = UnrealBuild.EngineDirectory;
 						return FindOrAddProjectHelper(EngineProjectFileNameBase, BaseFolder);
 					}
 				}
@@ -2227,16 +2227,16 @@ namespace UnrealBuildTool
 				// Check to see if this is an Engine target.  That is, the target is located under the "Engine" folder
 				bool IsEngineTarget = false;
 				bool WantProjectFileForTarget = true;
-				if (TargetFilePath.IsUnderDirectory(UnrealBuildTool.EngineDirectory))
+				if (TargetFilePath.IsUnderDirectory(UnrealBuild.EngineDirectory))
 				{
 					// This is an engine target
 					IsEngineTarget = true;
 
-					if (UnrealBuildTool.GetExtensionDirs(UnrealBuildTool.EngineDirectory, "Source/Programs").Any(x => TargetFilePath.IsUnderDirectory(x)))
+					if (UnrealBuildTool.GetExtensionDirs(UnrealBuild.EngineDirectory, "Source/Programs").Any(x => TargetFilePath.IsUnderDirectory(x)))
 					{
 						WantProjectFileForTarget = IncludeEnginePrograms;
 					}
-					else if (UnrealBuildTool.GetExtensionDirs(UnrealBuildTool.EngineDirectory, "Source").Any(x => TargetFilePath.IsUnderDirectory(x)))
+					else if (UnrealBuildTool.GetExtensionDirs(UnrealBuild.EngineDirectory, "Source").Any(x => TargetFilePath.IsUnderDirectory(x)))
 					{
 						WantProjectFileForTarget = bIncludeEngineSource;
 					}
@@ -2353,7 +2353,7 @@ namespace UnrealBuildTool
 					}
 					else if (IsEngineTarget)
 					{
-						BaseFolder = UnrealBuildTool.EngineDirectory;
+						BaseFolder = UnrealBuild.EngineDirectory;
 					}
 					else
 					{
@@ -2498,7 +2498,7 @@ namespace UnrealBuildTool
 		protected void AddEngineShaderSource(ProjectFile EngineProject)
 		{
 			// Setup a project file entry for this module's project.  Remember, some projects may host multiple modules!
-			DirectoryReference ShadersDirectory = DirectoryReference.Combine(UnrealBuildTool.EngineDirectory, "Shaders");
+			DirectoryReference ShadersDirectory = DirectoryReference.Combine(UnrealBuild.EngineDirectory, "Shaders");
 			if (DirectoryReference.Exists(ShadersDirectory))
 			{
 				List<string> SubdirectoryNamesToExclude = new List<string>();
@@ -2510,7 +2510,7 @@ namespace UnrealBuildTool
 					SubdirectoryNamesToExclude.Add("WorkingDirectory");
 				}
 
-				EngineProject.AddFilesToProject(SourceFileSearch.FindFiles(ShadersDirectory, SubdirectoryNamesToExclude), UnrealBuildTool.EngineDirectory);
+				EngineProject.AddFilesToProject(SourceFileSearch.FindFiles(ShadersDirectory, SubdirectoryNamesToExclude), UnrealBuild.EngineDirectory);
 			}
 		}
 
@@ -2518,12 +2518,12 @@ namespace UnrealBuildTool
 		/// Adds engine build infrastructure files to the specified project
 		protected void AddEngineBuildFiles(ProjectFile EngineProject)
 		{
-			DirectoryReference BuildDirectory = DirectoryReference.Combine(UnrealBuildTool.EngineDirectory, "Build");
+			DirectoryReference BuildDirectory = DirectoryReference.Combine(UnrealBuild.EngineDirectory, "Build");
 
 			List<string> SubdirectoryNamesToExclude = new List<string>();
 			SubdirectoryNamesToExclude.Add("Receipts");
 
-			EngineProject.AddFilesToProject(SourceFileSearch.FindFiles(BuildDirectory, SubdirectoryNamesToExclude), UnrealBuildTool.EngineDirectory);
+			EngineProject.AddFilesToProject(SourceFileSearch.FindFiles(BuildDirectory, SubdirectoryNamesToExclude), UnrealBuild.EngineDirectory);
 		}
 
 
@@ -2532,8 +2532,8 @@ namespace UnrealBuildTool
 		protected void AddEngineDocumentation(ProjectFile EngineProject)
 		{
 			// NOTE: The project folder added here will actually be collapsed away later if not needed
-			DirectoryReference DocumentationProjectDirectory = DirectoryReference.Combine(UnrealBuildTool.EngineDirectory, "Documentation");
-			DirectoryReference DocumentationSourceDirectory = DirectoryReference.Combine(UnrealBuildTool.EngineDirectory, "Documentation", "Source");
+			DirectoryReference DocumentationProjectDirectory = DirectoryReference.Combine(UnrealBuild.EngineDirectory, "Documentation");
+			DirectoryReference DocumentationSourceDirectory = DirectoryReference.Combine(UnrealBuild.EngineDirectory, "Documentation", "Source");
 			DirectoryInfo DirInfo = new DirectoryInfo(DocumentationProjectDirectory.FullName);
 			if (DirInfo.Exists && DirectoryReference.Exists(DocumentationSourceDirectory))
 			{
@@ -2579,7 +2579,7 @@ namespace UnrealBuildTool
 					DocumentationFiles = FilteredDocumentationFiles;
 				}
 
-				EngineProject.AddFilesToProject(DocumentationFiles, UnrealBuildTool.EngineDirectory);
+				EngineProject.AddFilesToProject(DocumentationFiles, UnrealBuild.EngineDirectory);
 			}
 			else
 			{
diff --git a/Engine/Source/Programs/UnrealBuildTool/ProjectFiles/QMake/QMakefileGenerator.cs b/Engine/Source/Programs/UnrealBuildTool/ProjectFiles/QMake/QMakefileGenerator.cs
--- a/Engine/Source/Programs/UnrealBuildTool/ProjectFiles/QMake/QMakefileGenerator.cs
+++ b/Engine/Source/Programs/UnrealBuildTool/ProjectFiles/QMake/QMakefileGenerator.cs
@@ -6,6 +6,7 @@ using System.Text;
 using System.IO;
 using System.Linq;
 using EpicGames.Core;
+using UnrealBuildBase;
 
 #nullable disable
 
@@ -232,7 +233,7 @@ namespace UnrealBuildTool
 				BuildCommand = "build=bash $$unrealRootPath/Engine/Build/BatchFiles/Linux/Build.sh\n";
 			}
 
-			string UnrealRootPath = UnrealBuildTool.RootDirectory.FullName;
+			string UnrealRootPath = UnrealBuild.RootDirectory.FullName;
 
 			string FileName = MasterProjectName + ".pro";
 
@@ -286,7 +287,7 @@ namespace UnrealBuildTool
 				List<FileReference> FoundFiles = SourceFileSearch.FindModuleSourceFiles(CurModuleFile);
 				foreach (FileReference CurSourceFile in FoundFiles)
 				{
-					string SourceFileRelativeToRoot = CurSourceFile.MakeRelativeTo(UnrealBuildTool.EngineDirectory);
+					string SourceFileRelativeToRoot = CurSourceFile.MakeRelativeTo(UnrealBuild.EngineDirectory);
 					// Exclude some directories that we don't compile (note that we still want Windows/Mac etc for code navigation)
 					if (!SourceFileRelativeToRoot.Contains("Source/ThirdParty/"))
 					{
diff --git a/Engine/Source/Programs/UnrealBuildTool/ProjectFiles/Rider/RiderProjectFile.cs b/Engine/Source/Programs/UnrealBuildTool/ProjectFiles/Rider/RiderProjectFile.cs
--- a/Engine/Source/Programs/UnrealBuildTool/ProjectFiles/Rider/RiderProjectFile.cs
+++ b/Engine/Source/Programs/UnrealBuildTool/ProjectFiles/Rider/RiderProjectFile.cs
@@ -7,6 +7,7 @@ using System.IO;
 using System.Linq;
 using System.Text.RegularExpressions;
 using EpicGames.Core;
+using UnrealBuildBase;
 
 #nullable disable
 
@@ -492,7 +493,7 @@ namespace UnrealBuildTool
 			else if(UEBuildPlatform.IsPlatformInGroup(Target.Platform, UnrealPlatformGroup.Linux) ||
 			        UEBuildPlatform.IsPlatformInGroup(Target.Platform, UnrealPlatformGroup.Unix))
 			{
-				var EngineDirectory = UnrealBuildTool.EngineDirectory.ToString();
+				var EngineDirectory = UnrealBuild.EngineDirectory.ToString();
 
 				string UseLibcxxEnvVarOverride = Environment.GetEnvironmentVariable("UE4_LINUX_USE_LIBCXX");
 				if (string.IsNullOrEmpty(UseLibcxxEnvVarOverride) || UseLibcxxEnvVarOverride == "1")
diff --git a/Engine/Source/Programs/UnrealBuildTool/ProjectFiles/Rider/RiderProjectFileGenerator.cs b/Engine/Source/Programs/UnrealBuildTool/ProjectFiles/Rider/RiderProjectFileGenerator.cs
--- a/Engine/Source/Programs/UnrealBuildTool/ProjectFiles/Rider/RiderProjectFileGenerator.cs
+++ b/Engine/Source/Programs/UnrealBuildTool/ProjectFiles/Rider/RiderProjectFileGenerator.cs
@@ -5,6 +5,7 @@ using System.Collections.Generic;
 using System.IO;
 using System.Linq;
 using EpicGames.Core;
+using UnrealBuildBase;
 
 #nullable disable
 
@@ -174,7 +175,7 @@ namespace UnrealBuildTool
 			}
 			else
 			{
-				return FileReference.Combine( UnrealBuildTool.EngineDirectory, "Intermediate", "ProjectFiles", ".Rider", GeneratedProjectName);
+				return FileReference.Combine( UnrealBuild.EngineDirectory, "Intermediate", "ProjectFiles", ".Rider", GeneratedProjectName);
 			}
 		}
 
@@ -204,7 +205,7 @@ namespace UnrealBuildTool
 				// Check to see if this is an Engine target.  That is, the target is located under the "Engine" folder
 				bool IsEngineTarget = false;
 				bool WantProjectFileForTarget = true;
-				if (TargetFilePath.IsUnderDirectory(UnrealBuildTool.EngineDirectory))
+				if (TargetFilePath.IsUnderDirectory(UnrealBuild.EngineDirectory))
 				{
 					// This is an engine target
 					IsEngineTarget = true;
@@ -304,7 +305,7 @@ namespace UnrealBuildTool
 					}
 					else if (IsEngineTarget)
 					{
-						BaseFolder = UnrealBuildTool.EngineDirectory;
+						BaseFolder = UnrealBuild.EngineDirectory;
 					}
 					else
 					{
@@ -432,7 +433,7 @@ namespace UnrealBuildTool
 					{
 						if (BuildHostPlatform.Current.Platform == UnrealTargetPlatform.Mac)
 						{
-							MasterProjectPath = UnrealBuildTool.EngineDirectory;
+							MasterProjectPath = UnrealBuild.EngineDirectory;
 							GameProjectName = "UnrealGame";
 						}
 
@@ -484,7 +485,7 @@ namespace UnrealBuildTool
 					EngineProjectFileNameBase + ProjectFileExtension);
 
 				bool bAlreadyExisted;
-				EngineProject = FindOrAddProject(ProjectFilePath, UnrealBuildTool.EngineDirectory, true, out bAlreadyExisted);
+				EngineProject = FindOrAddProject(ProjectFilePath, UnrealBuild.EngineDirectory, true, out bAlreadyExisted);
 
 				EngineProject.IsForeignProject = false;
 				EngineProject.IsGeneratedProject = true;
diff --git a/Engine/Source/Programs/UnrealBuildTool/ProjectFiles/VisualStudio/VCProject.cs b/Engine/Source/Programs/UnrealBuildTool/ProjectFiles/VisualStudio/VCProject.cs
--- a/Engine/Source/Programs/UnrealBuildTool/ProjectFiles/VisualStudio/VCProject.cs
+++ b/Engine/Source/Programs/UnrealBuildTool/ProjectFiles/VisualStudio/VCProject.cs
@@ -12,6 +12,7 @@ using System.Text.RegularExpressions;
 using System.Diagnostics;
 using System.Security;
 using EpicGames.Core;
+using UnrealBuildBase;
 
 #nullable disable
 
@@ -1059,7 +1060,7 @@ namespace UnrealBuildTool
 			{
 				const string InvalidMessage = "echo The selected platform/configuration is not valid for this target.";
 
-				string ProjectRelativeUnusedDirectory = NormalizeProjectPath(DirectoryReference.Combine(UnrealBuildTool.EngineDirectory, "Intermediate", "Build", "Unused"));
+				string ProjectRelativeUnusedDirectory = NormalizeProjectPath(DirectoryReference.Combine(UnrealBuild.EngineDirectory, "Intermediate", "Build", "Unused"));
 
 				VCProjectFileContent.AppendLine("  <PropertyGroup Condition=\"'$(Configuration)|$(Platform)'=='Invalid|{0}'\">", InvalidConfigPlatformName);
 				VCProjectFileContent.AppendLine("    <NMakeBuildCommandLine>{0}</NMakeBuildCommandLine>", InvalidMessage);
@@ -1535,7 +1536,7 @@ namespace UnrealBuildTool
 
 				if (IsStubProject)
 				{
-					string ProjectRelativeUnusedDirectory = NormalizeProjectPath(DirectoryReference.Combine(UnrealBuildTool.EngineDirectory, "Intermediate", "Build", "Unused"));
+					string ProjectRelativeUnusedDirectory = NormalizeProjectPath(DirectoryReference.Combine(UnrealBuild.EngineDirectory, "Intermediate", "Build", "Unused"));
 
 					VCProjectFileContent.AppendLine("  <PropertyGroup {0}>", ConditionString);
 					VCProjectFileContent.AppendLine("    <OutDir>{0}{1}</OutDir>", ProjectRelativeUnusedDirectory, Path.DirectorySeparatorChar);
@@ -1549,7 +1550,7 @@ namespace UnrealBuildTool
 				else if (UnrealBuildTool.IsEngineInstalled() && Combination.ProjectTarget != null && Combination.ProjectTarget.TargetRules != null &&
 					(Combination.Platform == null || !Combination.ProjectTarget.SupportedPlatforms.Contains(Combination.Platform.Value)))
 				{
-					string ProjectRelativeUnusedDirectory = NormalizeProjectPath(DirectoryReference.Combine(UnrealBuildTool.EngineDirectory, "Intermediate", "Build", "Unused"));
+					string ProjectRelativeUnusedDirectory = NormalizeProjectPath(DirectoryReference.Combine(UnrealBuild.EngineDirectory, "Intermediate", "Build", "Unused"));
 
 					string TargetName = Combination.ProjectTarget.TargetFilePath.GetFileNameWithoutAnyExtensions();
 					string ValidPlatforms = String.Join(", ", Combination.ProjectTarget.SupportedPlatforms.Select(x => x.ToString()));
@@ -1593,7 +1594,7 @@ namespace UnrealBuildTool
 					}
 					else
 					{
-						RootOutputDirectory = UEBuildTarget.GetOutputDirectoryForExecutable(UnrealBuildTool.EngineDirectory, TargetRulesObject.File);
+						RootOutputDirectory = UEBuildTarget.GetOutputDirectoryForExecutable(UnrealBuild.EngineDirectory, TargetRulesObject.File);
 					}
 
 					// Get the output directory
@@ -1696,7 +1697,7 @@ namespace UnrealBuildTool
 						BuildArguments.Append(" -FastPDB");
 					}
 
-					DirectoryReference BatchFilesDirectory = DirectoryReference.Combine(UnrealBuildTool.EngineDirectory, "Build", "BatchFiles");
+					DirectoryReference BatchFilesDirectory = DirectoryReference.Combine(UnrealBuild.EngineDirectory, "Build", "BatchFiles");
 
 					if (BuildToolOverride != null)
 					{
diff --git a/Engine/Source/Programs/UnrealBuildTool/ProjectFiles/VisualStudio/VCProjectFileGenerator.cs b/Engine/Source/Programs/UnrealBuildTool/ProjectFiles/VisualStudio/VCProjectFileGenerator.cs
--- a/Engine/Source/Programs/UnrealBuildTool/ProjectFiles/VisualStudio/VCProjectFileGenerator.cs
+++ b/Engine/Source/Programs/UnrealBuildTool/ProjectFiles/VisualStudio/VCProjectFileGenerator.cs
@@ -8,6 +8,7 @@ using System.Security.Cryptography;
 using System.Text;
 using System.Xml.Linq;
 using EpicGames.Core;
+using UnrealBuildBase;
 
 #nullable disable
 
@@ -605,7 +606,7 @@ namespace UnrealBuildTool
 				}
 
 				// Get the path to the visualizers file. Try to make it relative to the solution directory, but fall back to a full path if it's a foreign project.
-				FileReference VisualizersFile = FileReference.Combine(UnrealBuildTool.EngineDirectory, "Extras", "VisualStudioDebugging", "Unreal.natvis");
+				FileReference VisualizersFile = FileReference.Combine(UnrealBuild.EngineDirectory, "Extras", "VisualStudioDebugging", "Unreal.natvis");
 
 				// Add the visualizers at the solution level. Doesn't seem to be picked up from a makefile project in VS2017 15.8.5.
 				VCSolutionFileContent.AppendLine(String.Format("Project(\"{0}\") = \"Visualizers\", \"Visualizers\", \"{{1CCEC849-CC72-4C59-8C36-2F7C38706D4C}}\"", SolutionFolderEntryGUID));
diff --git a/Engine/Source/Programs/UnrealBuildTool/ProjectFiles/VisualStudioCode/VSCodeProjectFileGenerator.cs b/Engine/Source/Programs/UnrealBuildTool/ProjectFiles/VisualStudioCode/VSCodeProjectFileGenerator.cs
--- a/Engine/Source/Programs/UnrealBuildTool/ProjectFiles/VisualStudioCode/VSCodeProjectFileGenerator.cs
+++ b/Engine/Source/Programs/UnrealBuildTool/ProjectFiles/VisualStudioCode/VSCodeProjectFileGenerator.cs
@@ -7,6 +7,7 @@ using System.IO;
 using EpicGames.Core;
 using System.Diagnostics;
 using System.Linq;
+using UnrealBuildBase;
 
 #nullable disable
 
@@ -265,7 +266,7 @@ namespace UnrealBuildTool
 			VSCodeDir = DirectoryReference.Combine(MasterProjectPath, ".vscode");
 			DirectoryReference.CreateDirectory(VSCodeDir);
 
-			UE4ProjectRoot = UnrealBuildTool.RootDirectory;
+			UE4ProjectRoot = UnrealBuild.RootDirectory;
 			bForeignProject = !VSCodeDir.IsUnderDirectory(UE4ProjectRoot);
 
 			List<ProjectFile> Projects;
@@ -299,7 +300,7 @@ namespace UnrealBuildTool
 			if (bForeignProject && bIncludeEngineSource)
 			{
 				// for installed builds we need to write the cpp properties file under the installed engine as well for intellisense to work
-				DirectoryReference Ue4CodeDirectory = DirectoryReference.Combine(UnrealBuildTool.RootDirectory, ".vscode");
+				DirectoryReference Ue4CodeDirectory = DirectoryReference.Combine(UnrealBuild.RootDirectory, ".vscode");
 				WriteCppPropertiesFile(Ue4CodeDirectory, ProjectData);
 			}
 
@@ -809,7 +810,7 @@ namespace UnrealBuildTool
 					FileReference ResponseFile = null;
 					if (!DirectoryToResponseFile.TryGetValue(Directory, out ResponseFile))
 					{
-						for (DirectoryReference ParentDir = Directory; ParentDir != null && ParentDir != UnrealBuildTool.RootDirectory; ParentDir = ParentDir.ParentDirectory)
+						for (DirectoryReference ParentDir = Directory; ParentDir != null && ParentDir != UnrealBuild.RootDirectory; ParentDir = ParentDir.ParentDirectory)
 						{
 							if (DirectoryToResponseFile.TryGetValue(ParentDir, out ResponseFile))
 							{
@@ -1032,7 +1033,7 @@ namespace UnrealBuildTool
 			TargetType TargetRulesType = Target.TargetRules == null ? TargetType.Program : Target.TargetRules.Type;
 
 			// Get the output directory
-			DirectoryReference RootDirectory = UnrealBuildTool.EngineDirectory;
+			DirectoryReference RootDirectory = UnrealBuild.EngineDirectory;
 			if (TargetRulesType != TargetType.Program && (bShouldCompileMonolithic || TargetRulesObject.BuildEnvironment == TargetBuildEnvironment.Unique))
 			{
 				if(Target.UnrealProjectFilePath != null)
@@ -1477,7 +1478,7 @@ namespace UnrealBuildTool
 				// but using files.exclude with our full set of excludes breaks vscode for larger code bases so a verbose file explorer
 				// seems like less of an issue and thus we are not adding these to files.exclude.
 				// see [REDACTED_PROFILE_URL]/vscode/issues/109380 for discussions with vscode team
-				DirectoryReference WorkspaceRoot = bForeignProject ? Projects[0].BaseDir : UnrealBuildTool.RootDirectory;
+				DirectoryReference WorkspaceRoot = bForeignProject ? Projects[0].BaseDir : UnrealBuild.RootDirectory;
 				string WorkspaceRootPath = WorkspaceRoot.ToString().Replace('\\', '/') + "/";
 
 				if (!bForeignProject)
@@ -1537,7 +1538,7 @@ namespace UnrealBuildTool
 						WorkspaceFile.BeginObject();
 						{
 							WorkspaceFile.AddField("name", "UE5");
-							WorkspaceFile.AddField("path", MakeUnquotedPathString(UnrealBuildTool.RootDirectory, EPathType.Absolute));
+							WorkspaceFile.AddField("path", MakeUnquotedPathString(UnrealBuild.RootDirectory, EPathType.Absolute));
 						}
 						WorkspaceFile.EndObject();
 					}
diff --git a/Engine/Source/Programs/UnrealBuildTool/ProjectFiles/Xcode/XcodeFrameworkWrapperProject.cs b/Engine/Source/Programs/UnrealBuildTool/ProjectFiles/Xcode/XcodeFrameworkWrapperProject.cs
--- a/Engine/Source/Programs/UnrealBuildTool/ProjectFiles/Xcode/XcodeFrameworkWrapperProject.cs
+++ b/Engine/Source/Programs/UnrealBuildTool/ProjectFiles/Xcode/XcodeFrameworkWrapperProject.cs
@@ -6,6 +6,7 @@ using System.Text.RegularExpressions;
 using System.Linq;
 using System.Collections.Generic;
 using EpicGames.Core;
+using UnrealBuildBase;
 
 #nullable disable
 
@@ -18,7 +19,7 @@ namespace UnrealBuildTool.ProjectFiles.Xcode
 	{
 		private readonly static string PROJECT_FILE_SEARCH_EXPRESSION = "*.pbxproj";
 		private readonly static string TEMPLATE_NAME = "FrameworkWrapper";
-		private readonly static string FRAMEWORK_WRAPPER_TEMPLATE_DIRECTORY = Path.Combine(UnrealBuildTool.EngineDirectory.ToNormalizedPath(), "Build", "IOS", "Resources", TEMPLATE_NAME);
+		private readonly static string FRAMEWORK_WRAPPER_TEMPLATE_DIRECTORY = Path.Combine(UnrealBuild.EngineDirectory.ToNormalizedPath(), "Build", "IOS", "Resources", TEMPLATE_NAME);
 		private readonly static string TEMPLATE_PROJECT_NAME = "PROJECT_NAME";
 		private readonly static string COMMANDLINE_FILENAME = "uecommandline.txt";
 
diff --git a/Engine/Source/Programs/UnrealBuildTool/ProjectFiles/Xcode/XcodeProject.cs b/Engine/Source/Programs/UnrealBuildTool/ProjectFiles/Xcode/XcodeProject.cs
--- a/Engine/Source/Programs/UnrealBuildTool/ProjectFiles/Xcode/XcodeProject.cs
+++ b/Engine/Source/Programs/UnrealBuildTool/ProjectFiles/Xcode/XcodeProject.cs
@@ -9,6 +9,7 @@ using System.Xml.Linq;
 using System.Linq;
 using System.Text;
 using EpicGames.Core;
+using UnrealBuildBase;
 
 #nullable disable
 
@@ -758,7 +759,7 @@ namespace UnrealBuildTool
 			
 			// @todo: look also in Project/Build/Frameworks directory!
 			ProjectDescriptor Project = ProjectDescriptor.FromFile(UProjectPath);
-			List<PluginInfo> AvailablePlugins = Plugins.ReadAvailablePlugins(UnrealBuildTool.EngineDirectory, DirectoryReference.FromFile(UProjectPath), Project.AdditionalPluginDirectories);
+			List<PluginInfo> AvailablePlugins = Plugins.ReadAvailablePlugins(UnrealBuild.EngineDirectory, DirectoryReference.FromFile(UProjectPath), Project.AdditionalPluginDirectories);
 
 			// look in each plugin for frameworks
 			// @todo: Cache this kind of things since every target will re-do this work!
@@ -1378,7 +1379,7 @@ namespace UnrealBuildTool
 							FileReference ReceiptFilename;
 							if (bIsUnrealGame)
 							{
-								ReceiptFilename = TargetReceipt.GetDefaultPath(UnrealBuildTool.EngineDirectory, "UnrealGame", UnrealTargetPlatform.IOS, Config.BuildConfig, "");
+								ReceiptFilename = TargetReceipt.GetDefaultPath(UnrealBuild.EngineDirectory, "UnrealGame", UnrealTargetPlatform.IOS, Config.BuildConfig, "");
 							}
 							else
 							{
@@ -1587,7 +1588,7 @@ namespace UnrealBuildTool
 											string TargetName = ProjectTarget.TargetFilePath.GetFileNameWithoutAnyExtensions();
 
 											// Get the output directory
-											DirectoryReference RootDirectory = UnrealBuildTool.EngineDirectory;
+											DirectoryReference RootDirectory = UnrealBuild.EngineDirectory;
 											if ((ProjectTarget.TargetRules.Type == TargetType.Game || ProjectTarget.TargetRules.Type == TargetType.Client || ProjectTarget.TargetRules.Type == TargetType.Server) && bShouldCompileMonolithic)
 											{
 												if (ProjectTarget.UnrealProjectFilePath != null)
diff --git a/Engine/Source/Programs/UnrealBuildTool/Storage/Impl/FileSystemStorageProvider.cs b/Engine/Source/Programs/UnrealBuildTool/Storage/Impl/FileSystemStorageProvider.cs
--- a/Engine/Source/Programs/UnrealBuildTool/Storage/Impl/FileSystemStorageProvider.cs
+++ b/Engine/Source/Programs/UnrealBuildTool/Storage/Impl/FileSystemStorageProvider.cs
@@ -6,6 +6,7 @@ using System.Diagnostics;
 using System.IO;
 using System.Text;
 using EpicGames.Core;
+using UnrealBuildBase;
 
 namespace UnrealBuildTool.Storage.Impl
 {
@@ -138,7 +139,7 @@ namespace UnrealBuildTool.Storage.Impl
 		static FileReference GetFileForDigest(ContentHash Digest)
 		{
 			string DigestText = Digest.ToString();
-			return FileReference.Combine(UnrealBuildTool.EngineDirectory, "Saved", "UnrealBuildTool", "Cache", String.Format("{0}/{1}/{2}/{3}.bin", DigestText[0], DigestText[1], DigestText[2], DigestText));
+			return FileReference.Combine(UnrealBuild.EngineDirectory, "Saved", "UnrealBuildTool", "Cache", String.Format("{0}/{1}/{2}/{3}.bin", DigestText[0], DigestText[1], DigestText[2], DigestText));
 		}
 	}
 }
diff --git a/Engine/Source/Programs/UnrealBuildTool/System/ActionGraph.cs b/Engine/Source/Programs/UnrealBuildTool/System/ActionGraph.cs
--- a/Engine/Source/Programs/UnrealBuildTool/System/ActionGraph.cs
+++ b/Engine/Source/Programs/UnrealBuildTool/System/ActionGraph.cs
@@ -281,7 +281,7 @@ namespace UnrealBuildTool
 						{
 							FailPaths.Add(ProducedItem.Location);
 						}
-						if (ProducedItem.Location.FullName.Length > UnrealBuildTool.RootDirectory.FullName.Length + BuildConfiguration.MaxNestedPathLength && ProducedItem.Location.IsUnderDirectory(UnrealBuildTool.RootDirectory))
+						if (ProducedItem.Location.FullName.Length > UnrealBuild.RootDirectory.FullName.Length + BuildConfiguration.MaxNestedPathLength && ProducedItem.Location.IsUnderDirectory(UnrealBuild.RootDirectory))
 						{
 							WarnPaths.Add(ProducedItem.Location);
 						}
@@ -305,7 +305,7 @@ namespace UnrealBuildTool
 					Message.AppendFormat("Detected paths more than {0} characters below UE root directory. This may cause portability issues due to the {1} character maximum path length on Windows:\n", BuildConfiguration.MaxNestedPathLength, MAX_PATH);
 					foreach (FileReference Path in WarnPaths)
 					{
-						string RelativePath = Path.MakeRelativeTo(UnrealBuildTool.RootDirectory);
+						string RelativePath = Path.MakeRelativeTo(UnrealBuild.RootDirectory);
 						Message.AppendFormat("\n[{0} characters] {1}", RelativePath.Length, RelativePath);
 					}
 					Message.AppendFormat("\n\nConsider setting {0} = ... in module *.Build.cs files to use alternative names for intermediate paths.", nameof(ModuleRules.ShortName));
diff --git a/Engine/Source/Programs/UnrealBuildTool/System/ActionHistory.cs b/Engine/Source/Programs/UnrealBuildTool/System/ActionHistory.cs
--- a/Engine/Source/Programs/UnrealBuildTool/System/ActionHistory.cs
+++ b/Engine/Source/Programs/UnrealBuildTool/System/ActionHistory.cs
@@ -197,7 +197,7 @@ namespace UnrealBuildTool
 				AppName = UEBuildTarget.GetAppNameForTargetType(TargetType);
 			}
 
-			return FileReference.Combine(UnrealBuildTool.EngineDirectory, UEBuildTarget.GetPlatformIntermediateFolder(Platform, Architecture), AppName, "ActionHistory.bin");
+			return FileReference.Combine(UnrealBuild.EngineDirectory, UEBuildTarget.GetPlatformIntermediateFolder(Platform, Architecture), AppName, "ActionHistory.bin");
 		}
 
 		/// <summary>
@@ -400,7 +400,7 @@ namespace UnrealBuildTool
 		/// </summary>
 		public ActionHistory()
 		{
-			Partitions.Add(new ActionHistoryPartition(UnrealBuildTool.EngineDirectory));
+			Partitions.Add(new ActionHistoryPartition(UnrealBuild.EngineDirectory));
 		}
 
 		/// <summary>
diff --git a/Engine/Source/Programs/UnrealBuildTool/System/BuildVersion.cs b/Engine/Source/Programs/UnrealBuildTool/System/BuildVersion.cs
--- a/Engine/Source/Programs/UnrealBuildTool/System/BuildVersion.cs
+++ b/Engine/Source/Programs/UnrealBuildTool/System/BuildVersion.cs
@@ -8,6 +8,7 @@ using System.Linq;
 using System.Text;
 using System.Threading.Tasks;
 using EpicGames.Core;
+using UnrealBuildBase;
 
 namespace UnrealBuildTool
 {
@@ -99,7 +100,7 @@ namespace UnrealBuildTool
 		/// <returns>Path to the Build.version file</returns>
 		public static FileReference GetDefaultFileName()
 		{
-			return FileReference.Combine(UnrealBuildTool.EngineDirectory, "Build", "Build.version");
+			return FileReference.Combine(UnrealBuild.EngineDirectory, "Build", "Build.version");
 		}
 
 		/// <summary>
diff --git a/Engine/Source/Programs/UnrealBuildTool/System/ConfigHierarchy.cs b/Engine/Source/Programs/UnrealBuildTool/System/ConfigHierarchy.cs
--- a/Engine/Source/Programs/UnrealBuildTool/System/ConfigHierarchy.cs
+++ b/Engine/Source/Programs/UnrealBuildTool/System/ConfigHierarchy.cs
@@ -8,6 +8,7 @@ using System.Threading.Tasks;
 using EpicGames.Core;
 using System.IO;
 using System.Diagnostics.CodeAnalysis;
+using UnrealBuildBase;
 
 namespace UnrealBuildTool
 {
@@ -993,9 +994,9 @@ namespace UnrealBuildTool
 
 		private static string PerformFinalExpansions(string InString, string PlatformName, DirectoryReference? ProjectDir)
 		{
-			string PlatformExtensionEngineConfigDir = DirectoryReference.Combine(UnrealBuildTool.EngineDirectory, "Platforms", PlatformName).FullName;
+			string PlatformExtensionEngineConfigDir = DirectoryReference.Combine(UnrealBuild.EngineDirectory, "Platforms", PlatformName).FullName;
 
-			string OutString = InString.Replace("{ENGINE}", UnrealBuildTool.EngineDirectory.FullName);
+			string OutString = InString.Replace("{ENGINE}", UnrealBuild.EngineDirectory.FullName);
 			OutString = OutString.Replace("{EXTENGINE}", PlatformExtensionEngineConfigDir);
 			OutString = OutString.Replace("{PLATFORM}", PlatformName);
 
@@ -1003,11 +1004,11 @@ namespace UnrealBuildTool
 			{
 				DirectoryReference NFLDir;
 				DirectoryReference NRDir;
-				if (ProjectDir.IsUnderDirectory(UnrealBuildTool.EngineDirectory))
+				if (ProjectDir.IsUnderDirectory(UnrealBuild.EngineDirectory))
 				{
-					string RelativeDir = ProjectDir.MakeRelativeTo(UnrealBuildTool.EngineDirectory);
-					NFLDir = DirectoryReference.Combine(UnrealBuildTool.EngineDirectory, "Restricted/NotForLicensees", RelativeDir);
-					NRDir = DirectoryReference.Combine(UnrealBuildTool.EngineDirectory, "Restricted/NoRedist", RelativeDir);
+					string RelativeDir = ProjectDir.MakeRelativeTo(UnrealBuild.EngineDirectory);
+					NFLDir = DirectoryReference.Combine(UnrealBuild.EngineDirectory, "Restricted/NotForLicensees", RelativeDir);
+					NRDir = DirectoryReference.Combine(UnrealBuild.EngineDirectory, "Restricted/NoRedist", RelativeDir);
 				}
 				else
 				{
@@ -1128,7 +1129,7 @@ namespace UnrealBuildTool
 		{
 			if(ProjectDir == null)
 			{
-				return DirectoryReference.Combine(UnrealBuildTool.EngineDirectory, "Saved", "Config");
+				return DirectoryReference.Combine(UnrealBuild.EngineDirectory, "Saved", "Config");
 			}
 			else
 			{
@@ -1150,7 +1151,7 @@ namespace UnrealBuildTool
 					return DirectoryReference.Combine(UserSettingDir, "UnrealEngine", String.Format("{0}.{1}", ReadOnlyBuildVersion.Current.MajorVersion, ReadOnlyBuildVersion.Current.MinorVersion), "Saved");
 				}
 			}
-			return DirectoryReference.Combine(UnrealBuildTool.EngineDirectory, "Saved");
+			return DirectoryReference.Combine(UnrealBuild.EngineDirectory, "Saved");
 		}
 
 		/// <summary>
diff --git a/Engine/Source/Programs/UnrealBuildTool/System/CppDependencyCache.cs b/Engine/Source/Programs/UnrealBuildTool/System/CppDependencyCache.cs
--- a/Engine/Source/Programs/UnrealBuildTool/System/CppDependencyCache.cs
+++ b/Engine/Source/Programs/UnrealBuildTool/System/CppDependencyCache.cs
@@ -312,8 +312,8 @@ namespace UnrealBuildTool
 					AppName = UEBuildTarget.GetAppNameForTargetType(TargetType);
 				}
 
-				FileReference EngineCacheLocation = FileReference.Combine(UnrealBuildTool.EngineDirectory, UEBuildTarget.GetPlatformIntermediateFolder(Platform, Architecture), AppName, Configuration.ToString(), "DependencyCache.bin");
-				FindOrAddPartition(EngineCacheLocation, UnrealBuildTool.EngineDirectory);
+				FileReference EngineCacheLocation = FileReference.Combine(UnrealBuild.EngineDirectory, UEBuildTarget.GetPlatformIntermediateFolder(Platform, Architecture), AppName, Configuration.ToString(), "DependencyCache.bin");
+				FindOrAddPartition(EngineCacheLocation, UnrealBuild.EngineDirectory);
 			}
 
 			if (ProjectFile != null)
diff --git a/Engine/Source/Programs/UnrealBuildTool/System/DataDrivenPlatformInfo.cs b/Engine/Source/Programs/UnrealBuildTool/System/DataDrivenPlatformInfo.cs
--- a/Engine/Source/Programs/UnrealBuildTool/System/DataDrivenPlatformInfo.cs
+++ b/Engine/Source/Programs/UnrealBuildTool/System/DataDrivenPlatformInfo.cs
@@ -6,6 +6,7 @@ using System.Linq;
 using System.Text;
 using EpicGames.Core;
 using System.IO;
+using UnrealBuildBase;
 
 namespace UnrealBuildTool
 {
@@ -60,7 +61,7 @@ namespace UnrealBuildTool
 				Dictionary<string, string> IniParents = new Dictionary<string, string>(StringComparer.OrdinalIgnoreCase);
 
 				// find all platform directories (skipping NFL/NoRedist)
-				foreach (DirectoryReference EngineConfigDir in UnrealBuildTool.GetExtensionDirs(UnrealBuildTool.EngineDirectory, "Config", bIncludeRestrictedDirectories:false))
+				foreach (DirectoryReference EngineConfigDir in UnrealBuildTool.GetExtensionDirs(UnrealBuild.EngineDirectory, "Config", bIncludeRestrictedDirectories:false))
 				{
 					// look through all config dirs looking for the data driven ini file
 					foreach (string FilePath in Directory.EnumerateFiles(EngineConfigDir.FullName, "DataDrivenPlatformInfo.ini", SearchOption.AllDirectories))
@@ -69,7 +70,7 @@ namespace UnrealBuildTool
 
 						// get the platform name from the path
 						string IniPlatformName;
-						if (FileRef.IsUnderDirectory(DirectoryReference.Combine(UnrealBuildTool.EngineDirectory, "Config")))
+						if (FileRef.IsUnderDirectory(DirectoryReference.Combine(UnrealBuild.EngineDirectory, "Config")))
 						{
 							// Foo/Engine/Config/<Platform>/DataDrivenPlatformInfo.ini
 							IniPlatformName = Path.GetFileName(Path.GetDirectoryName(FilePath))!;
diff --git a/Engine/Source/Programs/UnrealBuildTool/System/DeployExports.cs b/Engine/Source/Programs/UnrealBuildTool/System/DeployExports.cs
--- a/Engine/Source/Programs/UnrealBuildTool/System/DeployExports.cs
+++ b/Engine/Source/Programs/UnrealBuildTool/System/DeployExports.cs
@@ -6,6 +6,7 @@ using System.Linq;
 using System.Text;
 using System.Threading.Tasks;
 using EpicGames.Core;
+using UnrealBuildBase;
 
 namespace UnrealBuildTool
 {
@@ -29,7 +30,7 @@ namespace UnrealBuildTool
 		{
 			string DeployFolder = string.Format("{0}-{1}-{2}",
 				ShortProjectName,
-				UnrealBuildTool.RootDirectory.GetDirectoryName(),
+				UnrealBuild.RootDirectory.GetDirectoryName(),
 				Environment.UserName);
 
 			DeployFolder = DeployFolder.Replace(" ", "");
diff --git a/Engine/Source/Programs/UnrealBuildTool/System/ExternalExecution.cs b/Engine/Source/Programs/UnrealBuildTool/System/ExternalExecution.cs
--- a/Engine/Source/Programs/UnrealBuildTool/System/ExternalExecution.cs
+++ b/Engine/Source/Programs/UnrealBuildTool/System/ExternalExecution.cs
@@ -715,7 +715,7 @@ namespace UnrealBuildTool
 			}
 			else
 			{
-				return TargetReceipt.GetDefaultPath(UnrealBuildTool.EngineDirectory, "UnrealHeaderTool", BuildHostPlatform.Current.Platform, Configuration, "");
+				return TargetReceipt.GetDefaultPath(UnrealBuild.EngineDirectory, "UnrealHeaderTool", BuildHostPlatform.Current.Platform, Configuration, "");
 			}
 		}
 
@@ -1076,7 +1076,7 @@ namespace UnrealBuildTool
 				// UHTLite
 				bIsBuildingUHT |= TargetName.Equals("UHTLite", StringComparison.InvariantCultureIgnoreCase);
 
-				string RootLocalPath = UnrealBuildTool.RootDirectory.FullName;
+				string RootLocalPath = UnrealBuild.RootDirectory.FullName;
 
 				UnrealTargetConfiguration UHTConfig = BuildConfiguration.bForceDebugUnrealHeaderTool ? UnrealTargetConfiguration.Debug : UnrealTargetConfiguration.Development;
 
diff --git a/Engine/Source/Programs/UnrealBuildTool/System/FileMetadataPrefetch.cs b/Engine/Source/Programs/UnrealBuildTool/System/FileMetadataPrefetch.cs
--- a/Engine/Source/Programs/UnrealBuildTool/System/FileMetadataPrefetch.cs
+++ b/Engine/Source/Programs/UnrealBuildTool/System/FileMetadataPrefetch.cs
@@ -43,7 +43,7 @@ namespace UnrealBuildTool
 		{
 			lock(QueuedDirectories)
 			{
-				if(QueuedDirectories.Add(UnrealBuildTool.EngineDirectory))
+				if(QueuedDirectories.Add(UnrealBuild.EngineDirectory))
 				{
 					Enqueue(() => ScanEngineDirectory());
 				}
@@ -111,7 +111,7 @@ namespace UnrealBuildTool
 		/// </summary>
 		static void ScanEngineDirectory()
 		{
-			foreach (DirectoryReference ExtensionDir in UnrealBuildTool.GetExtensionDirs(UnrealBuildTool.EngineDirectory))
+			foreach (DirectoryReference ExtensionDir in UnrealBuildTool.GetExtensionDirs(UnrealBuild.EngineDirectory))
 			{
 				DirectoryItem BaseDirectory = DirectoryItem.GetItemByDirectoryReference(ExtensionDir);
 				BaseDirectory.CacheDirectories();
@@ -139,7 +139,7 @@ namespace UnrealBuildTool
 		/// <param name="ProjectDirectory">The project directory to search</param>
 		static void ScanProjectDirectory(DirectoryItem ProjectDirectory)
 		{
-			foreach (DirectoryReference ExtensionDir in UnrealBuildTool.GetExtensionDirs(UnrealBuildTool.EngineDirectory))
+			foreach (DirectoryReference ExtensionDir in UnrealBuildTool.GetExtensionDirs(UnrealBuild.EngineDirectory))
 			{
 				DirectoryItem BaseDirectory = DirectoryItem.GetItemByDirectoryReference(ExtensionDir);
 				BaseDirectory.CacheDirectories();
diff --git a/Engine/Source/Programs/UnrealBuildTool/System/HotReload.cs b/Engine/Source/Programs/UnrealBuildTool/System/HotReload.cs
--- a/Engine/Source/Programs/UnrealBuildTool/System/HotReload.cs
+++ b/Engine/Source/Programs/UnrealBuildTool/System/HotReload.cs
@@ -107,7 +107,7 @@ namespace UnrealBuildTool
 		/// <returns>Location of the hot reload state file</returns>
 		public static FileReference GetLocation(FileReference ProjectFile, string TargetName, UnrealTargetPlatform Platform, UnrealTargetConfiguration Configuration, string Architecture)
 		{
-			DirectoryReference BaseDir = DirectoryReference.FromFile(ProjectFile) ?? UnrealBuildTool.EngineDirectory;
+			DirectoryReference BaseDir = DirectoryReference.FromFile(ProjectFile) ?? UnrealBuild.EngineDirectory;
 			return FileReference.Combine(BaseDir, UEBuildTarget.GetPlatformIntermediateFolder(Platform, Architecture), TargetName, Configuration.ToString(), "HotReload.state");
 		}
 
@@ -297,15 +297,15 @@ namespace UnrealBuildTool
 				FileReference EditorLocation;
 				if (TargetDesc.Platform == UnrealTargetPlatform.Win64)
 				{
-					EditorLocation = FileReference.Combine(UnrealBuildTool.EngineDirectory, "Binaries", "Win64", String.Format("{0}.exe", EditorBaseFileName));
+					EditorLocation = FileReference.Combine(UnrealBuild.EngineDirectory, "Binaries", "Win64", String.Format("{0}.exe", EditorBaseFileName));
 				}
 				else if (TargetDesc.Platform == UnrealTargetPlatform.Mac)
 				{
-					EditorLocation = FileReference.Combine(UnrealBuildTool.EngineDirectory, "Binaries", "Mac", String.Format("{0}.app/Contents/MacOS/{0}", EditorBaseFileName));
+					EditorLocation = FileReference.Combine(UnrealBuild.EngineDirectory, "Binaries", "Mac", String.Format("{0}.app/Contents/MacOS/{0}", EditorBaseFileName));
 				}
 				else if (TargetDesc.Platform == UnrealTargetPlatform.Linux)
 				{
-					EditorLocation = FileReference.Combine(UnrealBuildTool.EngineDirectory, "Binaries", "Linux", EditorBaseFileName);
+					EditorLocation = FileReference.Combine(UnrealBuild.EngineDirectory, "Binaries", "Linux", EditorBaseFileName);
 				}
 				else
 				{
@@ -314,7 +314,7 @@ namespace UnrealBuildTool
 
 				using(Timeline.ScopeEvent("Finding editor processes for hot-reload"))
 				{
-					DirectoryReference EditorRunsDir = DirectoryReference.Combine(UnrealBuildTool.EngineDirectory, "Intermediate", "EditorRuns");
+					DirectoryReference EditorRunsDir = DirectoryReference.Combine(UnrealBuild.EngineDirectory, "Intermediate", "EditorRuns");
 					if (!DirectoryReference.Exists(EditorRunsDir))
 					{
 						return false;
diff --git a/Engine/Source/Programs/UnrealBuildTool/System/InstalledPlatformInfo.cs b/Engine/Source/Programs/UnrealBuildTool/System/InstalledPlatformInfo.cs
--- a/Engine/Source/Programs/UnrealBuildTool/System/InstalledPlatformInfo.cs
+++ b/Engine/Source/Programs/UnrealBuildTool/System/InstalledPlatformInfo.cs
@@ -7,6 +7,7 @@ using System.Text;
 using System.Threading.Tasks;
 using System.IO;
 using EpicGames.Core;
+using UnrealBuildBase;
 
 namespace UnrealBuildTool
 {
@@ -203,7 +204,7 @@ namespace UnrealBuildTool
 			string RequiredFile;
 			if (ParseSubValue(PlatformConfiguration, "RequiredFile=", out RequiredFile))
 			{
-				RequiredFile = FileReference.Combine(UnrealBuildTool.RootDirectory, RequiredFile).ToString();
+				RequiredFile = FileReference.Combine(UnrealBuild.RootDirectory, RequiredFile).ToString();
 			}
 
 			string ProjectTypeName;
diff --git a/Engine/Source/Programs/UnrealBuildTool/System/NativeProjects.cs b/Engine/Source/Programs/UnrealBuildTool/System/NativeProjects.cs
--- a/Engine/Source/Programs/UnrealBuildTool/System/NativeProjects.cs
+++ b/Engine/Source/Programs/UnrealBuildTool/System/NativeProjects.cs
@@ -59,7 +59,7 @@ namespace UnrealBuildTool
 					if(CachedBaseDirectories == null)
 					{
 						HashSet<DirectoryReference> BaseDirs = new HashSet<DirectoryReference>();
-						foreach (FileReference RootFile in DirectoryLookupCache.EnumerateFiles(UnrealBuildTool.RootDirectory))
+						foreach (FileReference RootFile in DirectoryLookupCache.EnumerateFiles(UnrealBuild.RootDirectory))
 						{
 							if(RootFile.HasExtension(".uprojectdirs"))
 							{
@@ -68,14 +68,14 @@ namespace UnrealBuildTool
 									string TrimLine = Line.Trim();
 									if(!TrimLine.StartsWith(";"))
 									{
-										DirectoryReference BaseProjectDir = DirectoryReference.Combine(UnrealBuildTool.RootDirectory, TrimLine);
-										if(BaseProjectDir.IsUnderDirectory(UnrealBuildTool.RootDirectory))
+										DirectoryReference BaseProjectDir = DirectoryReference.Combine(UnrealBuild.RootDirectory, TrimLine);
+										if(BaseProjectDir.IsUnderDirectory(UnrealBuild.RootDirectory))
 										{
 											BaseDirs.Add(BaseProjectDir);
 										}
 										else
 										{
-											Log.TraceWarning("Project search path '{0}' referenced by '{1}' is not under '{2}', ignoring.", TrimLine, RootFile, UnrealBuildTool.RootDirectory);
+											Log.TraceWarning("Project search path '{0}' referenced by '{1}' is not under '{2}', ignoring.", TrimLine, RootFile, UnrealBuild.RootDirectory);
 										}
 									}
 								}
diff --git a/Engine/Source/Programs/UnrealBuildTool/System/Plugins.cs b/Engine/Source/Programs/UnrealBuildTool/System/Plugins.cs
--- a/Engine/Source/Programs/UnrealBuildTool/System/Plugins.cs
+++ b/Engine/Source/Programs/UnrealBuildTool/System/Plugins.cs
@@ -249,7 +249,7 @@ namespace UnrealBuildTool
 		public static IEnumerable<FileReference> EnumeratePlugins(FileReference? ProjectFile)
 		{
 			List<DirectoryReference> BaseDirs = new List<DirectoryReference>();
-			BaseDirs.AddRange(UnrealBuildTool.GetExtensionDirs(UnrealBuildTool.EngineDirectory, "Plugins"));
+			BaseDirs.AddRange(UnrealBuildTool.GetExtensionDirs(UnrealBuild.EngineDirectory, "Plugins"));
 			if(ProjectFile != null)
 			{
 				BaseDirs.AddRange(UnrealBuildTool.GetExtensionDirs(ProjectFile.Directory, "Plugins"));
diff --git a/Engine/Source/Programs/UnrealBuildTool/System/RulesAssembly.cs b/Engine/Source/Programs/UnrealBuildTool/System/RulesAssembly.cs
--- a/Engine/Source/Programs/UnrealBuildTool/System/RulesAssembly.cs
+++ b/Engine/Source/Programs/UnrealBuildTool/System/RulesAssembly.cs
@@ -413,7 +413,7 @@ namespace UnrealBuildTool
 				List<string> DisallowedPlatformsAndGroups = Utils.MakeListOfUnsupportedPlatforms(Target.OptedInModulePlatforms.ToList(), false);
 
 				// check if the module file is disallowed
-				if (ModuleFileName.ContainsAnyNames(DisallowedPlatformsAndGroups, UnrealBuildTool.EngineDirectory) ||
+				if (ModuleFileName.ContainsAnyNames(DisallowedPlatformsAndGroups, UnrealBuild.EngineDirectory) ||
 					(Target.ProjectFile != null && ModuleFileName.ContainsAnyNames(DisallowedPlatformsAndGroups, Target.ProjectFile.Directory)))
 				{
 					throw new BuildException("Platform module file {0} is not allowed (only platforms '{1}', and their groups, are allowed. This indicates a module reference not being checked with something like IsPlatformAvailableForTarget()).",
@@ -628,7 +628,7 @@ namespace UnrealBuildTool
 
 				// look in all plugins' paths to see if any disallowed
 				IEnumerable<PluginInfo> DisallowedPlugins = EnumeratePlugins().Where(Plugin =>
-					Plugin.File.ContainsAnyNames(DisallowedPlatformsAndGroups, UnrealBuildTool.EngineDirectory) ||
+					Plugin.File.ContainsAnyNames(DisallowedPlatformsAndGroups, UnrealBuild.EngineDirectory) ||
 					(Rules.ProjectFile != null && Plugin.File.ContainsAnyNames(DisallowedPlatformsAndGroups, Rules.ProjectFile.Directory)));
 				// log out the plugins we are disabling
 				DisallowedPlugins.ToList().ForEach(x => Log.TraceLog($"Disallowing non-opted-in platform plugin {x.File}"));
diff --git a/Engine/Source/Programs/UnrealBuildTool/System/RulesCompiler.cs b/Engine/Source/Programs/UnrealBuildTool/System/RulesCompiler.cs
--- a/Engine/Source/Programs/UnrealBuildTool/System/RulesCompiler.cs
+++ b/Engine/Source/Programs/UnrealBuildTool/System/RulesCompiler.cs
@@ -72,7 +72,7 @@ namespace UnrealBuildTool
 			// Add all engine source (including third party source)
 			if (bIncludeEngine)
 			{
-				Folders.AddRange(UnrealBuildTool.GetExtensionDirs(UnrealBuildTool.EngineDirectory, "Source"));
+				Folders.AddRange(UnrealBuildTool.GetExtensionDirs(UnrealBuild.EngineDirectory, "Source"));
 			}
 
 			// @todo plugin: Disallow modules from including plugin modules as dependency modules? (except when the module is part of that plugin)
@@ -81,7 +81,7 @@ namespace UnrealBuildTool
 			List<DirectoryReference> RootFolders = new List<DirectoryReference>();
 			if (bIncludeEngine)
 			{
-				RootFolders.AddRange(UnrealBuildTool.GetExtensionDirs(UnrealBuildTool.EngineDirectory));
+				RootFolders.AddRange(UnrealBuildTool.GetExtensionDirs(UnrealBuild.EngineDirectory));
 			}
 			if (GameFolders != null)
 			{
@@ -352,8 +352,8 @@ namespace UnrealBuildTool
 				List<PluginInfo> EnginePlugins = new List<PluginInfo>();
 				List<PluginInfo> MarketplacePlugins = new List<PluginInfo>();
 
-				DirectoryReference MarketplaceDirectory = DirectoryReference.Combine(UnrealBuildTool.EngineDirectory, "Plugins", "Marketplace");
-				foreach (PluginInfo PluginInfo in Plugins.ReadEnginePlugins(UnrealBuildTool.EngineDirectory))
+				DirectoryReference MarketplaceDirectory = DirectoryReference.Combine(UnrealBuild.EngineDirectory, "Plugins", "Marketplace");
+				foreach (PluginInfo PluginInfo in Plugins.ReadEnginePlugins(UnrealBuild.EngineDirectory))
 				{
 					if (PluginInfo.File.IsUnderDirectory(MarketplaceDirectory))
 					{
@@ -365,7 +365,7 @@ namespace UnrealBuildTool
 					}
 				}
 
-				EngineRulesAssembly = CreateEngineRulesAssemblyInternal(UnrealBuildTool.GetExtensionDirs(UnrealBuildTool.EngineDirectory), ProjectFileGenerator.EngineProjectFileNameBase, EnginePlugins, UnrealBuildTool.IsEngineInstalled() || bUsePrecompiled, bSkipCompile, null);
+				EngineRulesAssembly = CreateEngineRulesAssemblyInternal(UnrealBuildTool.GetExtensionDirs(UnrealBuild.EngineDirectory), ProjectFileGenerator.EngineProjectFileNameBase, EnginePlugins, UnrealBuildTool.IsEngineInstalled() || bUsePrecompiled, bSkipCompile, null);
 
 				if (MarketplacePlugins.Count > 0)
 				{
@@ -467,7 +467,7 @@ namespace UnrealBuildTool
 			Dictionary<FileReference, ModuleRulesContext> ModuleFileToContext = new Dictionary<FileReference, ModuleRulesContext>();
 			using (Timeline.ScopeEvent("Finding marketplace plugin modules"))
 			{
-				ModuleRulesContext PluginsModuleContext = new ModuleRulesContext(MarketplaceScope, UnrealBuildTool.EngineDirectory);
+				ModuleRulesContext PluginsModuleContext = new ModuleRulesContext(MarketplaceScope, UnrealBuild.EngineDirectory);
 				FindModuleRulesForPlugins(Plugins, PluginsModuleContext, ModuleFileToContext);
 			}
 
@@ -476,7 +476,7 @@ namespace UnrealBuildTool
 			if (ModuleFileToContext.Count > 0)
 			{
 				FileReference AssemblyFileName = FileReference.Combine(UnrealBuildTool.WritableEngineDirectory, "Intermediate", "Build", "BuildRules", "MarketplaceRules.dll");
-				Result = new RulesAssembly(MarketplaceScope, new List<DirectoryReference> { UnrealBuildTool.EngineDirectory }, Plugins, ModuleFileToContext, new List<FileReference>(), AssemblyFileName, bContainsEngineModules: true, DefaultBuildSettings: BuildSettingsVersion.Latest, bReadOnly: false, bSkipCompile: bSkipCompile, Parent: Parent);
+				Result = new RulesAssembly(MarketplaceScope, new List<DirectoryReference> { UnrealBuild.EngineDirectory }, Plugins, ModuleFileToContext, new List<FileReference>(), AssemblyFileName, bContainsEngineModules: true, DefaultBuildSettings: BuildSettingsVersion.Latest, bReadOnly: false, bSkipCompile: bSkipCompile, Parent: Parent);
 			}
 			return Result;
 		}
diff --git a/Engine/Source/Programs/UnrealBuildTool/System/SourceFileMetadataCache.cs b/Engine/Source/Programs/UnrealBuildTool/System/SourceFileMetadataCache.cs
--- a/Engine/Source/Programs/UnrealBuildTool/System/SourceFileMetadataCache.cs
+++ b/Engine/Source/Programs/UnrealBuildTool/System/SourceFileMetadataCache.cs
@@ -277,8 +277,8 @@ namespace UnrealBuildTool
 
 			if(ProjectFile == null || !UnrealBuildTool.IsEngineInstalled())
 			{
-				FileReference EngineCacheLocation = FileReference.Combine(UnrealBuildTool.EngineDirectory, "Intermediate", "Build", "SourceFileCache.bin");
-				Cache = FindOrAddCache(EngineCacheLocation, UnrealBuildTool.EngineDirectory, Cache);
+				FileReference EngineCacheLocation = FileReference.Combine(UnrealBuild.EngineDirectory, "Intermediate", "Build", "SourceFileCache.bin");
+				Cache = FindOrAddCache(EngineCacheLocation, UnrealBuild.EngineDirectory, Cache);
 			}
 
 			if(ProjectFile != null)
@@ -299,7 +299,7 @@ namespace UnrealBuildTool
 		{
 			if(ProjectFile == null || !UnrealBuildTool.IsEngineInstalled())
 			{
-				yield return FileReference.Combine(UnrealBuildTool.EngineDirectory, "Intermediate", "Build", "SourceFileCache.bin");
+				yield return FileReference.Combine(UnrealBuild.EngineDirectory, "Intermediate", "Build", "SourceFileCache.bin");
 			}
 			if(ProjectFile != null)
 			{
diff --git a/Engine/Source/Programs/UnrealBuildTool/System/TargetMakefile.cs b/Engine/Source/Programs/UnrealBuildTool/System/TargetMakefile.cs
--- a/Engine/Source/Programs/UnrealBuildTool/System/TargetMakefile.cs
+++ b/Engine/Source/Programs/UnrealBuildTool/System/TargetMakefile.cs
@@ -677,7 +677,7 @@ namespace UnrealBuildTool
 		/// <returns>Path to the makefile</returns>
 		public static FileReference GetLocation(FileReference? ProjectFile, string TargetName, UnrealTargetPlatform Platform, string Architecture, UnrealTargetConfiguration Configuration)
 		{
-			DirectoryReference BaseDirectory = DirectoryReference.FromFile(ProjectFile) ?? UnrealBuildTool.EngineDirectory;
+			DirectoryReference BaseDirectory = DirectoryReference.FromFile(ProjectFile) ?? UnrealBuild.EngineDirectory;
 			return FileReference.Combine(BaseDirectory, UEBuildTarget.GetPlatformIntermediateFolder(Platform, Architecture), TargetName, Configuration.ToString(), "Makefile.bin");
 		}
 
diff --git a/Engine/Source/Programs/UnrealBuildTool/System/TargetReceipt.cs b/Engine/Source/Programs/UnrealBuildTool/System/TargetReceipt.cs
--- a/Engine/Source/Programs/UnrealBuildTool/System/TargetReceipt.cs
+++ b/Engine/Source/Programs/UnrealBuildTool/System/TargetReceipt.cs
@@ -10,6 +10,7 @@ using System.ComponentModel;
 using System.Runtime.Serialization;
 using EpicGames.Core;
 using System.Diagnostics.CodeAnalysis;
+using UnrealBuildBase;
 
 namespace UnrealBuildTool
 {
@@ -501,7 +502,7 @@ namespace UnrealBuildTool
 		/// <param name="Location">Filename to read from</param>
 		public static TargetReceipt Read(FileReference Location)
 		{
-			return Read(Location, UnrealBuildTool.EngineDirectory);
+			return Read(Location, UnrealBuild.EngineDirectory);
 		}
 
 		/// <summary>
@@ -655,7 +656,7 @@ namespace UnrealBuildTool
 		/// <returns>True if successful</returns>
 		public static bool TryRead(FileReference Location, [NotNullWhen(true)] out TargetReceipt? Receipt)
 		{
-			return TryRead(Location, UnrealBuildTool.EngineDirectory, out Receipt);
+			return TryRead(Location, UnrealBuild.EngineDirectory, out Receipt);
 		}
 
 		/// <summary>
@@ -691,7 +692,7 @@ namespace UnrealBuildTool
 		/// <param name="Location">Output filename</param>
 		public void Write(FileReference Location)
 		{
-			Write(Location, UnrealBuildTool.EngineDirectory);
+			Write(Location, UnrealBuild.EngineDirectory);
 		}
 
 		/// <summary>
diff --git a/Engine/Source/Programs/UnrealBuildTool/System/UnrealPluginLanguage.cs b/Engine/Source/Programs/UnrealBuildTool/System/UnrealPluginLanguage.cs
--- a/Engine/Source/Programs/UnrealBuildTool/System/UnrealPluginLanguage.cs
+++ b/Engine/Source/Programs/UnrealBuildTool/System/UnrealPluginLanguage.cs
@@ -9,6 +9,7 @@ using System.Xml.Linq;
 using System.Diagnostics;
 using System.IO;
 using EpicGames.Core;
+using UnrealBuildBase;
 
 #nullable disable
 
@@ -2564,7 +2565,7 @@ namespace UnrealBuildTool
 		{
 			if (EngineDirectory == null)
 			{
-				EngineDirectory = UnrealBuildTool.EngineDirectory;
+				EngineDirectory = UnrealBuild.EngineDirectory;
 			}
 
 			// cache base ini for use as the seed for the rest
@@ -2708,7 +2709,7 @@ namespace UnrealBuildTool
 			bIsMergingConfigs = true;
 			if (EngineDirectory == null)
 			{
-				EngineDirectory = UnrealBuildTool.EngineDirectory;
+				EngineDirectory = UnrealBuild.EngineDirectory;
 			}
 
 			if (BaseCache != null)
diff --git a/Engine/Source/Programs/UnrealBuildTool/System/Utils.cs b/Engine/Source/Programs/UnrealBuildTool/System/Utils.cs
--- a/Engine/Source/Programs/UnrealBuildTool/System/Utils.cs
+++ b/Engine/Source/Programs/UnrealBuildTool/System/Utils.cs
@@ -307,7 +307,7 @@ namespace UnrealBuildTool
 		{
 			string AppName = Path.GetFileNameWithoutExtension(StartInfo.FileName);
 			string LogFilenameBase = string.Format("{0}_{1}", AppName, DateTime.Now.ToString("yyyy.MM.dd-HH.mm.ss"));
-			string LogDir = Path.Combine(UnrealBuildTool.EngineDirectory.FullName, "Programs", "AutomationTool", "Saved", "Logs");
+			string LogDir = Path.Combine(UnrealBuild.EngineDirectory.FullName, "Programs", "AutomationTool", "Saved", "Logs");
 			string LogFilename = "";
 			for (int Attempt = 1; Attempt < 100; ++Attempt)
 			{
@@ -996,7 +996,7 @@ namespace UnrealBuildTool
 				string DirectoryName = Environment.GetFolderPath(Environment.SpecialFolder.LocalApplicationData);
 				if(String.IsNullOrEmpty(DirectoryName))
 				{
-					return DirectoryReference.Combine(UnrealBuildTool.EngineDirectory, "Saved");
+					return DirectoryReference.Combine(UnrealBuild.EngineDirectory, "Saved");
 				}
 				else
 				{
diff --git a/Engine/Source/Programs/UnrealBuildTool/System/XmlConfig.cs b/Engine/Source/Programs/UnrealBuildTool/System/XmlConfig.cs
--- a/Engine/Source/Programs/UnrealBuildTool/System/XmlConfig.cs
+++ b/Engine/Source/Programs/UnrealBuildTool/System/XmlConfig.cs
@@ -12,6 +12,7 @@ using System.Xml;
 using System.Xml.Schema;
 using System.Xml.Serialization;
 using EpicGames.Core;
+using UnrealBuildBase;
 
 namespace UnrealBuildTool
 {
@@ -86,13 +87,13 @@ namespace UnrealBuildTool
 			else
 			{
 				// Get the default cache file
-				CacheFile = FileReference.Combine(UnrealBuildTool.EngineDirectory, "Intermediate", "Build", "XmlConfigCache.bin");
+				CacheFile = FileReference.Combine(UnrealBuild.EngineDirectory, "Intermediate", "Build", "XmlConfigCache.bin");
 				if(UnrealBuildTool.IsEngineInstalled())
 				{
 					DirectoryReference? UserSettingsDir = Utils.GetUserSettingDirectory();
 					if(UserSettingsDir != null)
 					{
-						CacheFile = FileReference.Combine(UserSettingsDir, "UnrealEngine", String.Format("XmlConfigCache-{0}.bin", UnrealBuildTool.RootDirectory.FullName.Replace(":", "").Replace(Path.DirectorySeparatorChar, '+')));
+						CacheFile = FileReference.Combine(UserSettingsDir, "UnrealEngine", String.Format("XmlConfigCache-{0}.bin", UnrealBuild.RootDirectory.FullName.Replace(":", "").Replace(Path.DirectorySeparatorChar, '+')));
 					}
 				}
 
@@ -210,7 +211,7 @@ namespace UnrealBuildTool
 		/// <returns>The location of the schema file</returns>
 		public static FileReference GetSchemaLocation()
 		{
-			return FileReference.Combine(UnrealBuildTool.EngineDirectory, "Saved", "UnrealBuildTool", "BuildConfiguration.Schema.xsd");
+			return FileReference.Combine(UnrealBuild.EngineDirectory, "Saved", "UnrealBuildTool", "BuildConfiguration.Schema.xsd");
 		}
 
 		/// <summary>
@@ -225,14 +226,14 @@ namespace UnrealBuildTool
 			if(!UnrealBuildTool.IsEngineInstalled())
 			{
  				// Check for the config file under /Engine/Programs/NotForLicensees/UnrealBuildTool
- 				FileReference NotForLicenseesConfigLocation = FileReference.Combine(UnrealBuildTool.EngineDirectory, "Restricted", "NotForLicensees", "Programs", "UnrealBuildTool", "BuildConfiguration.xml");
+ 				FileReference NotForLicenseesConfigLocation = FileReference.Combine(UnrealBuild.EngineDirectory, "Restricted", "NotForLicensees", "Programs", "UnrealBuildTool", "BuildConfiguration.xml");
  				if(FileReference.Exists(NotForLicenseesConfigLocation))
  				{
  					InputFiles.Add(new InputFile(NotForLicenseesConfigLocation, "NotForLicensees"));
  				}
 
 				// Check for the user config file under /Engine/Saved/UnrealBuildTool
-				FileReference UserConfigLocation = FileReference.Combine(UnrealBuildTool.EngineDirectory, "Saved", "UnrealBuildTool", "BuildConfiguration.xml");
+				FileReference UserConfigLocation = FileReference.Combine(UnrealBuild.EngineDirectory, "Saved", "UnrealBuildTool", "BuildConfiguration.xml");
 				if(!FileReference.Exists(UserConfigLocation))
 				{
 					CreateDefaultConfigFile(UserConfigLocation);
diff --git a/Engine/Source/Programs/UnrealBuildTool/ToolChain/AppleToolChain.cs b/Engine/Source/Programs/UnrealBuildTool/ToolChain/AppleToolChain.cs
--- a/Engine/Source/Programs/UnrealBuildTool/ToolChain/AppleToolChain.cs
+++ b/Engine/Source/Programs/UnrealBuildTool/ToolChain/AppleToolChain.cs
@@ -8,6 +8,7 @@ using System.Diagnostics;
 using System.Text;
 using EpicGames.Core;
 using System.Text.RegularExpressions;
+using UnrealBuildBase;
 
 #nullable disable
 
@@ -274,7 +275,7 @@ namespace UnrealBuildTool
 			// if the installed one is too old, use a fixed up one if it can
 			if (bUseInstalledDsymutil == false)
 			{
-				FileReference PatchedDsymutilLocation = FileReference.Combine(UnrealBuildTool.EngineDirectory, "Restricted/NotForLicensees/Binaries/Mac/LLVM/bin/dsymutil");
+				FileReference PatchedDsymutilLocation = FileReference.Combine(UnrealBuild.EngineDirectory, "Restricted/NotForLicensees/Binaries/Mac/LLVM/bin/dsymutil");
 
 				if (File.Exists(PatchedDsymutilLocation.FullName))
 				{
diff --git a/Engine/Source/Programs/UnrealBuildTool/ToolChain/RemoteMac.cs b/Engine/Source/Programs/UnrealBuildTool/ToolChain/RemoteMac.cs
--- a/Engine/Source/Programs/UnrealBuildTool/ToolChain/RemoteMac.cs
+++ b/Engine/Source/Programs/UnrealBuildTool/ToolChain/RemoteMac.cs
@@ -12,6 +12,7 @@ using System.Text;
 using System.Text.RegularExpressions;
 using EpicGames.Core;
 using System.Security.Cryptography.X509Certificates;
+using UnrealBuildBase;
 
 #nullable disable
 
@@ -131,8 +132,8 @@ namespace UnrealBuildTool
 		/// <param name="ProjectFile">Project to read settings from</param>
 		public RemoteMac(FileReference ProjectFile)
 		{
-			this.RsyncExe = FileReference.Combine(UnrealBuildTool.EngineDirectory, "Extras", "ThirdPartyNotUE", "cwrsync", "bin", "rsync.exe");
-			this.SshExe = FileReference.Combine(UnrealBuildTool.EngineDirectory, "Extras", "ThirdPartyNotUE", "cwrsync", "bin", "ssh.exe");
+			this.RsyncExe = FileReference.Combine(UnrealBuild.EngineDirectory, "Extras", "ThirdPartyNotUE", "cwrsync", "bin", "rsync.exe");
+			this.SshExe = FileReference.Combine(UnrealBuild.EngineDirectory, "Extras", "ThirdPartyNotUE", "cwrsync", "bin", "ssh.exe");
 			this.ProjectFile = ProjectFile;
 			if (ProjectFile != null)
 			{
@@ -218,7 +219,7 @@ namespace UnrealBuildTool
 				Log.TraceWarning("No SSH private key found for {0}@{1}. Launching SSH to generate one.", UserName, ServerName);
 
 				StringBuilder CommandLine = new StringBuilder();
-				CommandLine.AppendFormat("/C \"\"{0}\"", FileReference.Combine(UnrealBuildTool.EngineDirectory, "Build", "BatchFiles", "MakeAndInstallSSHKey.bat"));
+				CommandLine.AppendFormat("/C \"\"{0}\"", FileReference.Combine(UnrealBuild.EngineDirectory, "Build", "BatchFiles", "MakeAndInstallSSHKey.bat"));
 				CommandLine.AppendFormat(" \"{0}\"", SshExe);
 				CommandLine.AppendFormat(" \"{0}\"", ServerPort);
 				CommandLine.AppendFormat(" \"{0}\"", RsyncExe);
@@ -226,7 +227,7 @@ namespace UnrealBuildTool
 				CommandLine.AppendFormat(" \"{0}\"", ServerName);
 				CommandLine.AppendFormat(" \"{0}\"", DirectoryReference.GetSpecialFolder(Environment.SpecialFolder.MyDocuments));
 				CommandLine.AppendFormat(" \"{0}\"", GetLocalCygwinPath(DirectoryReference.GetSpecialFolder(Environment.SpecialFolder.MyDocuments)));
-				CommandLine.AppendFormat(" \"{0}\"", UnrealBuildTool.EngineDirectory);
+				CommandLine.AppendFormat(" \"{0}\"", UnrealBuild.EngineDirectory);
 				CommandLine.Append("\"");
 
 				using(Process ChildProcess = Process.Start(BuildHostPlatform.Current.Shell.FullName, CommandLine.ToString()))
@@ -290,8 +291,8 @@ namespace UnrealBuildTool
 
 			// Build the list of directory mappings between the local and remote machines
 			Mappings = new List<RemoteMapping>();
-			Mappings.Add(new RemoteMapping(UnrealBuildTool.EngineDirectory, GetRemotePath(UnrealBuildTool.EngineDirectory)));
-			if(ProjectFile != null && !ProjectFile.IsUnderDirectory(UnrealBuildTool.EngineDirectory))
+			Mappings.Add(new RemoteMapping(UnrealBuild.EngineDirectory, GetRemotePath(UnrealBuild.EngineDirectory)));
+			if(ProjectFile != null && !ProjectFile.IsUnderDirectory(UnrealBuild.EngineDirectory))
 			{
 				Mappings.Add(new RemoteMapping(ProjectFile.Directory, GetRemotePath(ProjectFile.Directory)));
 			}
@@ -299,7 +300,7 @@ namespace UnrealBuildTool
 			{
 				foreach (DirectoryReference AdditionalPath in AdditionalPaths)
 				{
-					if (!AdditionalPath.IsUnderDirectory(UnrealBuildTool.EngineDirectory) &&
+					if (!AdditionalPath.IsUnderDirectory(UnrealBuild.EngineDirectory) &&
 						!AdditionalPath.IsUnderDirectory(ProjectFile.Directory))
 					{
 						Mappings.Add(new RemoteMapping(AdditionalPath, GetRemotePath(AdditionalPath)));
@@ -325,9 +326,9 @@ namespace UnrealBuildTool
 				Locations.Add(DirectoryReference.Combine(ProjectFile.Directory, "Restricted", "NoRedist", "Build"));
 				Locations.Add(DirectoryReference.Combine(ProjectFile.Directory, "Build"));
 			}
-			Locations.Add(DirectoryReference.Combine(UnrealBuildTool.EngineDirectory, "Restricted", "NotForLicensees", "Build"));
-			Locations.Add(DirectoryReference.Combine(UnrealBuildTool.EngineDirectory, "Restricted", "NoRedist", "Build"));
-			Locations.Add(DirectoryReference.Combine(UnrealBuildTool.EngineDirectory, "Build"));
+			Locations.Add(DirectoryReference.Combine(UnrealBuild.EngineDirectory, "Restricted", "NotForLicensees", "Build"));
+			Locations.Add(DirectoryReference.Combine(UnrealBuild.EngineDirectory, "Restricted", "NoRedist", "Build"));
+			Locations.Add(DirectoryReference.Combine(UnrealBuild.EngineDirectory, "Build"));
 
 			// Find the first that exists
 			foreach (DirectoryReference Location in Locations)
@@ -406,7 +407,7 @@ namespace UnrealBuildTool
 				BuildCommandLine.AppendFormat(" {0}", EscapeShellArgument(RemoteArgument));
 			}
 
-			int Result = Execute(GetRemotePath(UnrealBuildTool.RootDirectory), BuildCommandLine.ToString());
+			int Result = Execute(GetRemotePath(UnrealBuild.RootDirectory), BuildCommandLine.ToString());
 			return Result == 0;
 		}
 
@@ -522,7 +523,7 @@ namespace UnrealBuildTool
 					}
 
 					ProcessStartInfo StartInfo = new ProcessStartInfo();
-					StartInfo.FileName = FileReference.Combine(UnrealBuildTool.EngineDirectory, "Binaries", "DotNET", "IOS", "IPhonePackager.exe").FullName;
+					StartInfo.FileName = FileReference.Combine(UnrealBuild.EngineDirectory, "Binaries", "DotNET", "IOS", "IPhonePackager.exe").FullName;
 					StartInfo.Arguments = Arguments.ToString();
 					if(Utils.RunLocalProcessAndLogOutput(StartInfo) != 0)
 					{
@@ -554,7 +555,7 @@ namespace UnrealBuildTool
 				BuildCommandLine.AppendFormat(" {0}", EscapeShellArgument(RemoteArgument));
 			}
 
-			int Result = Execute(GetRemotePath(UnrealBuildTool.RootDirectory), BuildCommandLine.ToString());
+			int Result = Execute(GetRemotePath(UnrealBuild.RootDirectory), BuildCommandLine.ToString());
 			if(Result != 0)
 			{
 				if(RemoteLogFile != null)
@@ -585,7 +586,7 @@ namespace UnrealBuildTool
 			DownloadFiles(FilesToDownload);
 
 			// Copy remote FrameworkAssets directory as it could contain resource bundles that must be packaged locally.
-			DirectoryReference BaseDir = DirectoryReference.FromFile(TargetDesc.ProjectFile) ?? UnrealBuildTool.EngineDirectory;
+			DirectoryReference BaseDir = DirectoryReference.FromFile(TargetDesc.ProjectFile) ?? UnrealBuild.EngineDirectory;
 			DirectoryReference FrameworkAssetsDir = DirectoryReference.Combine(BaseDir, "Intermediate", TargetDesc.Platform == UnrealTargetPlatform.IOS ? "IOS" : "TVOS", "FrameworkAssets");
 			if(RemoteDirectoryExists(FrameworkAssetsDir))
 			{
@@ -610,7 +611,7 @@ namespace UnrealBuildTool
 		/// <returns>Directory to use for temporary files</returns>
 		static DirectoryReference CreateTempDirectory(TargetDescriptor TargetDesc)
 		{
-			DirectoryReference BaseDir = DirectoryReference.FromFile(TargetDesc.ProjectFile) ?? UnrealBuildTool.EngineDirectory;
+			DirectoryReference BaseDir = DirectoryReference.FromFile(TargetDesc.ProjectFile) ?? UnrealBuild.EngineDirectory;
 			DirectoryReference TempDir = DirectoryReference.Combine(BaseDir, "Intermediate", "Remote", TargetDesc.Name, TargetDesc.Platform.ToString(), TargetDesc.Configuration.ToString());
 			DirectoryReference.CreateDirectory(TempDir);
 			return TempDir;
@@ -863,7 +864,7 @@ namespace UnrealBuildTool
 		void UploadWorkspace(DirectoryReference TempDir)
 		{
 			// Path to the scripts to be uploaded
-			FileReference ScriptPathsFileName = FileReference.Combine(UnrealBuildTool.EngineDirectory, "Build", "Rsync", "RsyncEngineScripts.txt");
+			FileReference ScriptPathsFileName = FileReference.Combine(UnrealBuild.EngineDirectory, "Build", "Rsync", "RsyncEngineScripts.txt");
 
 			// Read the list of scripts to be uploaded
 			List<string> ScriptPaths = new List<string>();
@@ -880,13 +881,13 @@ namespace UnrealBuildTool
 			List<FileReference> TargetFiles = new List<FileReference>();
 			foreach(string ScriptPath in ScriptPaths)
 			{
-				FileReference SourceFile = FileReference.Combine(UnrealBuildTool.EngineDirectory, ScriptPath.TrimStart('/'));
+				FileReference SourceFile = FileReference.Combine(UnrealBuild.EngineDirectory, ScriptPath.TrimStart('/'));
 				if(!FileReference.Exists(SourceFile))
 				{
 					throw new BuildException("Missing script required for remote upload: {0}", SourceFile);
 				}
 
-				FileReference TargetFile = FileReference.Combine(TempDir, SourceFile.MakeRelativeTo(UnrealBuildTool.EngineDirectory));
+				FileReference TargetFile = FileReference.Combine(TempDir, SourceFile.MakeRelativeTo(UnrealBuild.EngineDirectory));
 				if(!FileReference.Exists(TargetFile) || FileReference.GetLastWriteTimeUtc(TargetFile) < FileReference.GetLastWriteTimeUtc(SourceFile))
 				{
 					DirectoryReference.CreateDirectory(TargetFile.Directory);
@@ -908,7 +909,7 @@ namespace UnrealBuildTool
 
 			// Upload these files to the remote
 			Log.TraceInformation("[Remote] Uploading scripts...");
-			UploadFiles(TempDir, GetRemotePath(UnrealBuildTool.EngineDirectory), ScriptPathsFileName);
+			UploadFiles(TempDir, GetRemotePath(UnrealBuild.EngineDirectory), ScriptPathsFileName);
 
 			// Upload the config files
 			Log.TraceInformation("[Remote] Uploading config files...");
@@ -919,23 +920,23 @@ namespace UnrealBuildTool
 			EngineFilters.Add(ScriptProtectList);
 			if(UnrealBuildTool.IsEngineInstalled())
 			{
-				EngineFilters.Add(FileReference.Combine(UnrealBuildTool.EngineDirectory, "Build", "Rsync", "RsyncEngineInstalled.txt"));
+				EngineFilters.Add(FileReference.Combine(UnrealBuild.EngineDirectory, "Build", "Rsync", "RsyncEngineInstalled.txt"));
 			}
-			EngineFilters.Add(FileReference.Combine(UnrealBuildTool.EngineDirectory, "Build", "Rsync", "RsyncEngine.txt"));
+			EngineFilters.Add(FileReference.Combine(UnrealBuild.EngineDirectory, "Build", "Rsync", "RsyncEngine.txt"));
 
 			Log.TraceInformation("[Remote] Uploading engine files...");
-			UploadDirectory(UnrealBuildTool.EngineDirectory, GetRemotePath(UnrealBuildTool.EngineDirectory), EngineFilters);
+			UploadDirectory(UnrealBuild.EngineDirectory, GetRemotePath(UnrealBuild.EngineDirectory), EngineFilters);
 
 			// Upload the project files
 			DirectoryReference ProjectDir = null;
-			if (ProjectFile != null && !ProjectFile.IsUnderDirectory(UnrealBuildTool.EngineDirectory))
+			if (ProjectFile != null && !ProjectFile.IsUnderDirectory(UnrealBuild.EngineDirectory))
 			{
 				ProjectDir = ProjectFile.Directory;
 			}
 			else if (!string.IsNullOrEmpty(UnrealBuildTool.GetRemoteIniPath()))
 			{
 				ProjectDir = new DirectoryReference(UnrealBuildTool.GetRemoteIniPath());
-				if (ProjectDir.IsUnderDirectory(UnrealBuildTool.EngineDirectory))
+				if (ProjectDir.IsUnderDirectory(UnrealBuild.EngineDirectory))
 				{
 					ProjectDir = null;
 				}
@@ -949,7 +950,7 @@ namespace UnrealBuildTool
 				{
 					ProjectFilters.Add(CustomFilter);
 				}
-				ProjectFilters.Add(FileReference.Combine(UnrealBuildTool.EngineDirectory, "Build", "Rsync", "RsyncProject.txt"));
+				ProjectFilters.Add(FileReference.Combine(UnrealBuild.EngineDirectory, "Build", "Rsync", "RsyncProject.txt"));
 
 				Log.TraceInformation("[Remote] Uploading project files...");
 				UploadDirectory(ProjectDir, GetRemotePath(ProjectDir), ProjectFilters);
@@ -966,15 +967,15 @@ namespace UnrealBuildTool
 					{
 						CustomFilters.Add(CustomFilter);
 					}
-					CustomFilters.Add(FileReference.Combine(UnrealBuildTool.EngineDirectory, "Build", "Rsync", "RsyncProject.txt"));
+					CustomFilters.Add(FileReference.Combine(UnrealBuild.EngineDirectory, "Build", "Rsync", "RsyncProject.txt"));
 
 					Log.TraceInformation(string.Format("[Remote] Uploading additional path files [{0}]...", AdditionalPath.FullName));
 					UploadDirectory(AdditionalPath, GetRemotePath(AdditionalPath), CustomFilters);
 				}
 			}
 
-			Execute("/", String.Format("rm -rf {0}/Intermediate/IOS/*.plist", GetRemotePath(UnrealBuildTool.EngineDirectory)));
-			Execute("/", String.Format("rm -rf {0}/Intermediate/TVOS/*.plist", GetRemotePath(UnrealBuildTool.EngineDirectory)));
+			Execute("/", String.Format("rm -rf {0}/Intermediate/IOS/*.plist", GetRemotePath(UnrealBuild.EngineDirectory)));
+			Execute("/", String.Format("rm -rf {0}/Intermediate/TVOS/*.plist", GetRemotePath(UnrealBuild.EngineDirectory)));
 			if (ProjectFile != null)
 			{
 				Execute("/", String.Format("rm -rf {0}/Intermediate/IOS/*.plist", GetRemotePath(ProjectFile.Directory)));
@@ -982,10 +983,10 @@ namespace UnrealBuildTool
 			}
 
 			// Convert CRLF to LF for all shell scripts
-			Execute(RemoteBaseDir, String.Format("for i in {0}/Build/BatchFiles/Mac/*.sh; do mv $i $i.crlf; tr -d '\r' < $i.crlf > $i; done", EscapeShellArgument(GetRemotePath(UnrealBuildTool.EngineDirectory))));
+			Execute(RemoteBaseDir, String.Format("for i in {0}/Build/BatchFiles/Mac/*.sh; do mv $i $i.crlf; tr -d '\r' < $i.crlf > $i; done", EscapeShellArgument(GetRemotePath(UnrealBuild.EngineDirectory))));
 
 			// Fixup permissions on any shell scripts
-			Execute(RemoteBaseDir, String.Format("chmod +x {0}/Build/BatchFiles/Mac/*.sh", EscapeShellArgument(GetRemotePath(UnrealBuildTool.EngineDirectory))));
+			Execute(RemoteBaseDir, String.Format("chmod +x {0}/Build/BatchFiles/Mac/*.sh", EscapeShellArgument(GetRemotePath(UnrealBuild.EngineDirectory))));
 		}
 
 		/// <summary>
@@ -1036,7 +1037,7 @@ namespace UnrealBuildTool
 			{
 				if(FileGroups[Idx].Count > 0)
 				{
-					FileReference DownloadListLocation = FileReference.Combine(UnrealBuildTool.EngineDirectory, "Intermediate", "Rsync", "Download.txt");
+					FileReference DownloadListLocation = FileReference.Combine(UnrealBuild.EngineDirectory, "Intermediate", "Rsync", "Download.txt");
 					DirectoryReference.CreateDirectory(DownloadListLocation.Directory);
 					FileReference.WriteAllLines(DownloadListLocation, FileGroups[Idx].Select(x => x.MakeRelativeTo(Mappings[Idx].LocalDirectory).Replace('\\', '/')));
 
@@ -1062,7 +1063,7 @@ namespace UnrealBuildTool
 		private bool RemoteDirectoryExists(DirectoryReference LocalDirectory)
 		{
 			string RemoteDirectory = GetRemotePath(LocalDirectory);
-			return Execute(UnrealBuildTool.RootDirectory, String.Format("[ -d {0} ]", EscapeShellArgument(RemoteDirectory))) == 0;
+			return Execute(UnrealBuild.RootDirectory, String.Format("[ -d {0} ]", EscapeShellArgument(RemoteDirectory))) == 0;
 		}
 
 		/// <summary>
diff --git a/Engine/Source/Programs/UnrealBuildTool/UnrealBuildTool.cs b/Engine/Source/Programs/UnrealBuildTool/UnrealBuildTool.cs
--- a/Engine/Source/Programs/UnrealBuildTool/UnrealBuildTool.cs
+++ b/Engine/Source/Programs/UnrealBuildTool/UnrealBuildTool.cs
@@ -9,6 +9,7 @@ using System.Reflection;
 using System.Text;
 using System.Threading;
 using EpicGames.Core;
+using UnrealBuildBase;
 
 namespace UnrealBuildTool
 {
@@ -45,59 +46,39 @@ namespace UnrealBuildTool
 		/// </summary>
 		static DirectoryReference? CachedEngineProgramSavedDirectory;
 
-		/// <summary>
-		/// The path to UBT
-		/// </summary>
-		// In net core the actual running entrypoint is the dll, but its easier to run the apphost executable instead to avoid having to run via the dotnet command
-#if WINDOWS
-		// For windows a executable has a .exe extension
-		public static readonly FileReference UnrealBuildToolPath = FileReference.FindCorrectCase(new FileReference(Path.ChangeExtension(Assembly.GetExecutingAssembly().GetOriginalLocation(), "exe")));
-#else
-		// For linux and mac executables have no extension
-		public static readonly FileReference UnrealBuildToolPath = FileReference.FindCorrectCase(new FileReference(Path.ChangeExtension(Assembly.GetExecutingAssembly().GetOriginalLocation(), null)));
-#endif
-		/// <summary>
-		/// The full name of the root UE directory
-		/// </summary>
-		public static readonly DirectoryReference RootDirectory = DirectoryReference.Combine(UnrealBuildToolPath.Directory, "..", "..", "..", "..");
-		/// <summary>
-		/// The full name of the Engine directory
-		/// </summary>
-		public static readonly DirectoryReference EngineDirectory = DirectoryReference.Combine(RootDirectory, "Engine");
-
 		/// <summary>
 		/// The full name of the Engine/Source directory
 		/// </summary>
-		public static readonly DirectoryReference EngineSourceDirectory = DirectoryReference.Combine(EngineDirectory, "Source");
+		public static readonly DirectoryReference EngineSourceDirectory = DirectoryReference.Combine(UnrealBuild.EngineDirectory, "Source");
 
 		/// <summary>
 		/// Full path to the Engine/Source/Runtime directory
 		/// </summary>
-		[Obsolete("Please use UnrealBuildTool.GetExtensionDirs(UnrealBuildTool.EngineDirectory, \"Source/Runtime\") instead.")]
+		[Obsolete("Please use UnrealBuildTool.GetExtensionDirs(UnrealBuild.EngineDirectory, \"Source/Runtime\") instead.")]
 		public static readonly DirectoryReference EngineSourceRuntimeDirectory = DirectoryReference.Combine(EngineSourceDirectory, "Runtime");
 
 		/// <summary>
 		/// Full path to the Engine/Source/Developer directory
 		/// </summary>
-		[Obsolete("Please use UnrealBuildTool.GetExtensionDirs(UnrealBuildTool.EngineDirectory, \"Source/Developer\") instead.")]
+		[Obsolete("Please use UnrealBuildTool.GetExtensionDirs(UnrealBuild.EngineDirectory, \"Source/Developer\") instead.")]
 		public static readonly DirectoryReference EngineSourceDeveloperDirectory = DirectoryReference.Combine(EngineSourceDirectory, "Developer");
 
 		/// <summary>
 		/// Full path to the Engine/Source/Editor directory
 		/// </summary>
-		[Obsolete("Please use UnrealBuildTool.GetExtensionDirs(UnrealBuildTool.EngineDirectory, \"Source/Editor\") instead.")]
+		[Obsolete("Please use UnrealBuildTool.GetExtensionDirs(UnrealBuild.EngineDirectory, \"Source/Editor\") instead.")]
 		public static readonly DirectoryReference EngineSourceEditorDirectory = DirectoryReference.Combine(EngineSourceDirectory, "Editor");
 
 		/// <summary>
 		/// Full path to the Engine/Source/Programs directory
 		/// </summary>
-		[Obsolete("Please use UnrealBuildTool.GetExtensionDirs(UnrealBuildTool.EngineDirectory, \"Source/Programs\") instead.")]
+		[Obsolete("Please use UnrealBuildTool.GetExtensionDirs(UnrealBuild.EngineDirectory, \"Source/Programs\") instead.")]
 		public static readonly DirectoryReference EngineSourceProgramsDirectory = DirectoryReference.Combine(EngineSourceDirectory, "Programs");
 
 		/// <summary>
 		/// Full path to the Engine/Source/ThirdParty directory
 		/// </summary>
-		[Obsolete("Please use UnrealBuildTool.GetExtensionDirs(UnrealBuildTool.EngineDirectory, \"Source/ThirdParty\") instead.")]
+		[Obsolete("Please use UnrealBuildTool.GetExtensionDirs(UnrealBuild.EngineDirectory, \"Source/ThirdParty\") instead.")]
 		public static readonly DirectoryReference EngineSourceThirdPartyDirectory = DirectoryReference.Combine(EngineSourceDirectory, "ThirdParty");
 
 		/// <summary>
@@ -121,7 +102,7 @@ namespace UnrealBuildTool
 					}
 					if (UserDir == null)
 					{
-						CachedWritableEngineDirectory = EngineDirectory;
+						CachedWritableEngineDirectory = UnrealBuild.EngineDirectory;
 					}
 					else
 					{
@@ -143,11 +124,11 @@ namespace UnrealBuildTool
 				{
 					if (IsEngineInstalled())
 					{
-						CachedEngineProgramSavedDirectory = Utils.GetUserSettingDirectory() ?? DirectoryReference.Combine(EngineDirectory, "Programs");
+						CachedEngineProgramSavedDirectory = Utils.GetUserSettingDirectory() ?? DirectoryReference.Combine(UnrealBuild.EngineDirectory, "Programs");
 					}
 					else
 					{
-						CachedEngineProgramSavedDirectory = DirectoryReference.Combine(EngineDirectory, "Programs");
+						CachedEngineProgramSavedDirectory = DirectoryReference.Combine(UnrealBuild.EngineDirectory, "Programs");
 					}
 				}
 				return CachedEngineProgramSavedDirectory;
@@ -198,7 +179,7 @@ namespace UnrealBuildTool
 				}
 
 				// remove any platform directories in non-engine locations if the engine doesn't have the platform 
-				if (BaseDir != UnrealBuildTool.EngineDirectory && CachedDirs.Item1.Count > 0)
+				if (BaseDir != UnrealBuild.EngineDirectory && CachedDirs.Item1.Count > 0)
 				{
 					// if the DDPI.ini file doesn't exist, we haven't synced the platform, so just skip this directory
 					CachedDirs.Item1.RemoveAll(x => DataDrivenPlatformInfo.GetDataDrivenInfoForPlatform(x.GetDirectoryName()) == null);
@@ -249,7 +230,7 @@ namespace UnrealBuildTool
 		{
 			if (!bIsEngineInstalled.HasValue)
 			{
-				bIsEngineInstalled = FileReference.Exists(FileReference.Combine(EngineDirectory, "Build", "InstalledBuild.txt"));
+				bIsEngineInstalled = FileReference.Exists(FileReference.Combine(UnrealBuild.EngineDirectory, "Build", "InstalledBuild.txt"));
 			}
 			return bIsEngineInstalled.Value;
 		}
@@ -262,10 +243,10 @@ namespace UnrealBuildTool
 		{
 			if (!bIsProjectInstalled.HasValue)
 			{
-				FileReference InstalledProjectLocationFile = FileReference.Combine(UnrealBuildTool.RootDirectory, "Engine", "Build", "InstalledProjectBuild.txt");
+				FileReference InstalledProjectLocationFile = FileReference.Combine(UnrealBuild.RootDirectory, "Engine", "Build", "InstalledProjectBuild.txt");
 				if (FileReference.Exists(InstalledProjectLocationFile))
 				{
-					InstalledProjectFile = FileReference.Combine(UnrealBuildTool.RootDirectory, File.ReadAllText(InstalledProjectLocationFile.FullName).Trim());
+					InstalledProjectFile = FileReference.Combine(UnrealBuild.RootDirectory, File.ReadAllText(InstalledProjectLocationFile.FullName).Trim());
 					bIsProjectInstalled = true;
 				}
 				else
@@ -300,7 +281,7 @@ namespace UnrealBuildTool
 		/// <returns>True if the file is part of the installed distribution, false otherwise</returns>
 		static public bool IsFileInstalled(FileReference File)
 		{
-			if(IsEngineInstalled() && File.IsUnderDirectory(EngineDirectory))
+			if(IsEngineInstalled() && File.IsUnderDirectory(UnrealBuild.EngineDirectory))
 			{
 				return true;
 			}
@@ -317,7 +298,7 @@ namespace UnrealBuildTool
 		/// <returns>A string containing the path to the UBT assembly.</returns>
 		static public FileReference GetUBTPath()
 		{
-			return UnrealBuildToolPath;
+			return UnrealBuild.UnrealBuildToolPath;
 		}
 
 		/// <summary>
