UnrealBuildTool:

Move NativeProjects partially into BuildUtilities - class split to avoid pulling in a larger chain of dependencies.

diff --git a/Engine/Source/Programs/DotNETCommon/BuildUtilities/System/NativeProjectsBase.cs b/Engine/Source/Programs/DotNETCommon/BuildUtilities/System/NativeProjectsBase.cs
new file mode 100644
--- /dev/null
+++ b/Engine/Source/Programs/DotNETCommon/BuildUtilities/System/NativeProjectsBase.cs
@@ -0,0 +1,168 @@
+// Copyright Epic Games, Inc. All Rights Reserved.
+
+using System;
+using System.Collections.Generic;
+using System.Diagnostics.CodeAnalysis;
+using System.IO;
+using System.Linq;
+using System.Text;
+using System.Threading.Tasks;
+using EpicGames.Core;
+
+namespace UnrealBuildBase
+{
+	/// <summary>
+	/// Utility functions for querying native projects (ie. those found via a .uprojectdirs query)
+	/// </summary>
+	public class NativeProjectsBase
+	{
+		/// <summary>
+		/// Object used for synchronizing access to static fields
+		/// </summary>
+		protected static object LockObject = new object();
+
+		/// <summary>
+		/// The native project base directories
+		/// </summary>
+		static HashSet<DirectoryReference>? CachedBaseDirectories;
+
+		/// <summary>
+		/// Cached list of project files within all the base directories
+		/// </summary>
+		static HashSet<FileReference>? CachedProjectFiles;
+
+		/// <summary>
+		/// Clear our cached properties. Generally only needed if your script has modified local files...
+		/// </summary>
+		public static void ClearCacheBase()
+		{
+			CachedBaseDirectories = null;
+			CachedProjectFiles = null;
+		}
+
+		/// <summary>
+		/// Retrieve the list of base directories for native projects
+		/// </summary>
+		public static IEnumerable<DirectoryReference> EnumerateBaseDirectories()
+		{
+			if(CachedBaseDirectories == null)
+			{
+				lock(LockObject)
+				{
+					if(CachedBaseDirectories == null)
+					{
+						HashSet<DirectoryReference> BaseDirs = new HashSet<DirectoryReference>();
+						foreach (FileReference RootFile in DirectoryLookupCache.EnumerateFiles(UnrealBuild.RootDirectory))
+						{
+							if(RootFile.HasExtension(".uprojectdirs"))
+							{
+								foreach(string Line in File.ReadAllLines(RootFile.FullName))
+								{
+									string TrimLine = Line.Trim();
+									if(!TrimLine.StartsWith(";"))
+									{
+										DirectoryReference BaseProjectDir = DirectoryReference.Combine(UnrealBuild.RootDirectory, TrimLine);
+										if(BaseProjectDir.IsUnderDirectory(UnrealBuild.RootDirectory))
+										{
+											BaseDirs.Add(BaseProjectDir);
+										}
+										else
+										{
+											Log.TraceWarning("Project search path '{0}' referenced by '{1}' is not under '{2}', ignoring.", TrimLine, RootFile, UnrealBuild.RootDirectory);
+										}
+									}
+								}
+							}
+						}
+						CachedBaseDirectories = BaseDirs;
+					}
+				}
+			}
+			return CachedBaseDirectories;
+		}
+
+		/// <summary>
+		/// Returns a list of all the projects
+		/// </summary>
+		/// <returns>List of projects</returns>
+		public static IEnumerable<FileReference> EnumerateProjectFiles()
+		{
+			if(CachedProjectFiles == null)
+			{
+				lock(LockObject)
+				{
+					if(CachedProjectFiles == null)
+					{
+						HashSet<FileReference> ProjectFiles = new HashSet<FileReference>();
+						foreach(DirectoryReference BaseDirectory in EnumerateBaseDirectories())
+						{
+							if(DirectoryLookupCache.DirectoryExists(BaseDirectory))
+							{
+								foreach(DirectoryReference SubDirectory in DirectoryLookupCache.EnumerateDirectories(BaseDirectory))
+								{
+									// Ignore system directories (specifically for Mac, since temporary folders are created here)
+									string DirectoryName = SubDirectory.GetDirectoryName();
+									if (DirectoryName.StartsWith(".", StringComparison.Ordinal))
+									{
+										continue;
+									}
+
+									foreach(FileReference File in DirectoryLookupCache.EnumerateFiles(SubDirectory))
+									{
+										if(File.HasExtension(".uproject"))
+										{
+											ProjectFiles.Add(File);
+										}
+									}
+								}
+							}
+						}
+						CachedProjectFiles = ProjectFiles;
+					}
+				}
+			}
+			return CachedProjectFiles;
+		}
+
+		/// <summary>
+		/// Finds all target files under a given folder, and add them to the target name to project file map
+		/// </summary>
+		/// <param name="Directory">Directory to search</param>
+		/// <param name="TargetNameToProjectFile">Map from target name to project file</param>
+		/// <param name="ProjectFile">The project file for this directory</param>
+		protected static void FindTargetFiles(DirectoryReference Directory, Dictionary<string, FileReference> TargetNameToProjectFile, FileReference ProjectFile)
+		{
+			// Search for all target files within this directory
+			bool bSearchSubFolders = true;
+			foreach (FileReference File in DirectoryLookupCache.EnumerateFiles(Directory))
+			{
+				if (File.HasExtension(".target.cs"))
+				{
+					string TargetName = Path.GetFileNameWithoutExtension(File.GetFileNameWithoutExtension());
+					TargetNameToProjectFile[TargetName] = ProjectFile;
+					bSearchSubFolders = false;
+				}
+			}
+
+			// If we didn't find anything, recurse through the subfolders
+			if(bSearchSubFolders)
+			{
+				foreach(DirectoryReference SubDirectory in DirectoryLookupCache.EnumerateDirectories(Directory))
+				{
+					FindTargetFiles(SubDirectory, TargetNameToProjectFile, ProjectFile);
+				}
+			}
+		}
+
+		/// <summary>
+		/// Checks if a given project is a native project
+		/// </summary>
+		/// <param name="ProjectFile">The project file to check</param>
+		/// <returns>True if the given project is a native project</returns>
+		public static bool IsNativeProject(FileReference ProjectFile)
+		{
+			EnumerateProjectFiles();
+			return CachedProjectFiles!.Contains(ProjectFile);
+		}
+	}
+}
diff --git a/Engine/Source/Programs/UnrealBuildTool/System/NativeProjects.cs b/Engine/Source/Programs/UnrealBuildTool/System/NativeProjects.cs
--- a/Engine/Source/Programs/UnrealBuildTool/System/NativeProjects.cs
+++ b/Engine/Source/Programs/UnrealBuildTool/System/NativeProjects.cs
@@ -15,23 +15,8 @@ namespace UnrealBuildTool
 	/// <summary>
 	/// Utility functions for querying native projects (ie. those found via a .uprojectdirs query)
 	/// </summary>
-	public static class NativeProjects
+	public class NativeProjects : NativeProjectsBase
 	{
-		/// <summary>
-		/// Object used for synchronizing access to static fields
-		/// </summary>
-		static object LockObject = new object();
-
-		/// <summary>
-		/// The native project base directories
-		/// </summary>
-		static HashSet<DirectoryReference>? CachedBaseDirectories;
-
-		/// <summary>
-		/// Cached list of project files within all the base directories
-		/// </summary>
-		static HashSet<FileReference>? CachedProjectFiles;
-
 		/// <summary>
 		/// Cached map of target names to the project file they belong to
 		/// </summary>
@@ -42,95 +27,10 @@ namespace UnrealBuildTool
 		/// </summary>
 		public static void ClearCache()
 		{
-			CachedBaseDirectories = null;
-			CachedProjectFiles = null;
+			ClearCacheBase();
 			CachedTargetNameToProjectFile = null;
 		}
 
-		/// <summary>
-		/// Retrieve the list of base directories for native projects
-		/// </summary>
-		public static IEnumerable<DirectoryReference> EnumerateBaseDirectories()
-		{
-			if(CachedBaseDirectories == null)
-			{
-				lock(LockObject)
-				{
-					if(CachedBaseDirectories == null)
-					{
-						HashSet<DirectoryReference> BaseDirs = new HashSet<DirectoryReference>();
-						foreach (FileReference RootFile in DirectoryLookupCache.EnumerateFiles(UnrealBuild.RootDirectory))
-						{
-							if(RootFile.HasExtension(".uprojectdirs"))
-							{
-								foreach(string Line in File.ReadAllLines(RootFile.FullName))
-								{
-									string TrimLine = Line.Trim();
-									if(!TrimLine.StartsWith(";"))
-									{
-										DirectoryReference BaseProjectDir = DirectoryReference.Combine(UnrealBuild.RootDirectory, TrimLine);
-										if(BaseProjectDir.IsUnderDirectory(UnrealBuild.RootDirectory))
-										{
-											BaseDirs.Add(BaseProjectDir);
-										}
-										else
-										{
-											Log.TraceWarning("Project search path '{0}' referenced by '{1}' is not under '{2}', ignoring.", TrimLine, RootFile, UnrealBuild.RootDirectory);
-										}
-									}
-								}
-							}
-						}
-						CachedBaseDirectories = BaseDirs;
-					}
-				}
-			}
-			return CachedBaseDirectories;
-		}
-
-		/// <summary>
-		/// Returns a list of all the projects
-		/// </summary>
-		/// <returns>List of projects</returns>
-		public static IEnumerable<FileReference> EnumerateProjectFiles()
-		{
-			if(CachedProjectFiles == null)
-			{
-				lock(LockObject)
-				{
-					if(CachedProjectFiles == null)
-					{
-						HashSet<FileReference> ProjectFiles = new HashSet<FileReference>();
-						foreach(DirectoryReference BaseDirectory in EnumerateBaseDirectories())
-						{
-							if(DirectoryLookupCache.DirectoryExists(BaseDirectory))
-							{
-								foreach(DirectoryReference SubDirectory in DirectoryLookupCache.EnumerateDirectories(BaseDirectory))
-								{
-									// Ignore system directories (specifically for Mac, since temporary folders are created here)
-									string DirectoryName = SubDirectory.GetDirectoryName();
-									if (DirectoryName.StartsWith(".", StringComparison.Ordinal))
-									{
-										continue;
-									}
-
-									foreach(FileReference File in DirectoryLookupCache.EnumerateFiles(SubDirectory))
-									{
-										if(File.HasExtension(".uproject"))
-										{
-											ProjectFiles.Add(File);
-										}
-									}
-								}
-							}
-						}
-						CachedProjectFiles = ProjectFiles;
-					}
-				}
-			}
-			return CachedProjectFiles;
-		}
-
 		/// <summary>
 		/// Get the project folder for the given target name
 		/// </summary>
@@ -169,46 +69,5 @@ namespace UnrealBuildTool
 			}
 			return CachedTargetNameToProjectFile.TryGetValue(InTargetName, out OutProjectFileName);
 		}
-
-		/// <summary>
-		/// Finds all target files under a given folder, and add them to the target name to project file map
-		/// </summary>
-		/// <param name="Directory">Directory to search</param>
-		/// <param name="TargetNameToProjectFile">Map from target name to project file</param>
-		/// <param name="ProjectFile">The project file for this directory</param>
-		private static void FindTargetFiles(DirectoryReference Directory, Dictionary<string, FileReference> TargetNameToProjectFile, FileReference ProjectFile)
-		{
-			// Search for all target files within this directory
-			bool bSearchSubFolders = true;
-			foreach (FileReference File in DirectoryLookupCache.EnumerateFiles(Directory))
-			{
-				if (File.HasExtension(".target.cs"))
-				{
-					string TargetName = Path.GetFileNameWithoutExtension(File.GetFileNameWithoutExtension());
-					TargetNameToProjectFile[TargetName] = ProjectFile;
-					bSearchSubFolders = false;
-				}
-			}
-
-			// If we didn't find anything, recurse through the subfolders
-			if(bSearchSubFolders)
-			{
-				foreach(DirectoryReference SubDirectory in DirectoryLookupCache.EnumerateDirectories(Directory))
-				{
-					FindTargetFiles(SubDirectory, TargetNameToProjectFile, ProjectFile);
-				}
-			}
-		}
-
-		/// <summary>
-		/// Checks if a given project is a native project
-		/// </summary>
-		/// <param name="ProjectFile">The project file to check</param>
-		/// <returns>True if the given project is a native project</returns>
-		public static bool IsNativeProject(FileReference ProjectFile)
-		{
-			EnumerateProjectFiles();
-			return CachedProjectFiles!.Contains(ProjectFile);
-		}
 	}
 }
