UnrealBuildTool:

Add a StringComparison argument to CreateIntermediateTextFile() and require callers to make a decision about the case-sensitivity used when considering whether to re-write a file.

(This is intentionally explicit and verbose at the call-site - we should implement a better solution. This change makes it clear that there is work to be done to be able to simplify the interface and achieve reliable behavior.)

As a first pass, C++ files (headers, cpp files, pch files) will have a case-sensitive comparison (C++ code is, mostly, case sensitive), Linux script files have a case-sensitive comparison (case-sensitive filesystem), and all other files are compared disregarding case - to match pre-existing behavior, and with the desire to not borrow more surprises than necessary.

diff --git a/Engine/Source/Programs/UnrealBuildTool/Configuration/UEBuildModuleCPP.cs b/Engine/Source/Programs/UnrealBuildTool/Configuration/UEBuildModuleCPP.cs
--- a/Engine/Source/Programs/UnrealBuildTool/Configuration/UEBuildModuleCPP.cs
+++ b/Engine/Source/Programs/UnrealBuildTool/Configuration/UEBuildModuleCPP.cs
@@ -1009,7 +1009,7 @@ namespace UnrealBuildTool
 
 				// Write the PCH header
 				FileReference DedicatedPchLocation = FileReference.Combine(IntermediateDirectory, String.Format("PCH.Dedicated.{0}.h", File.Location.GetFileNameWithoutExtension()));
-				FileItem DedicatedPchFile = Graph.CreateIntermediateTextFile(DedicatedPchLocation, WrapperContents.ToString());
+				FileItem DedicatedPchFile = Graph.CreateIntermediateTextFile(DedicatedPchLocation, WrapperContents.ToString(), StringComparison.Ordinal);
 
 				// Create a new C++ environment to compile the PCH
 				CppCompileEnvironment PchEnvironment = new CppCompileEnvironment(CompileEnvironment);
@@ -1094,7 +1094,7 @@ namespace UnrealBuildTool
 							}
 
 							WriteDefinitions(CompileEnvironment.Definitions, Writer);
-							PrivateDefinitionsFileItem = Graph.CreateIntermediateTextFile(PrivateDefinitionsFile, Writer.ToString());
+							PrivateDefinitionsFileItem = Graph.CreateIntermediateTextFile(PrivateDefinitionsFile, Writer.ToString(), StringComparison.Ordinal);
 						}
 
 						CompileEnvironment = new CppCompileEnvironment(CompileEnvironment);
@@ -1137,7 +1137,7 @@ namespace UnrealBuildTool
 					WriteDefinitions(CompileEnvironment.Definitions, Writer);
 					CompileEnvironment.Definitions.Clear();
 
-					FileItem PrivateDefinitionsFileItem = Graph.CreateIntermediateTextFile(PrivateDefinitionsFile, Writer.ToString());
+					FileItem PrivateDefinitionsFileItem = Graph.CreateIntermediateTextFile(PrivateDefinitionsFile, Writer.ToString(), StringComparison.Ordinal);
 					CompileEnvironment.ForceIncludeFiles.Add(PrivateDefinitionsFileItem);
 				}
 			}
@@ -1181,7 +1181,7 @@ namespace UnrealBuildTool
 			}
 
 			// Create the item
-			FileItem WrapperFile = Graph.CreateIntermediateTextFile(OutputFile, WrapperContents.ToString());
+			FileItem WrapperFile = Graph.CreateIntermediateTextFile(OutputFile, WrapperContents.ToString(), StringComparison.Ordinal);
 
 			// Touch it if the included file is newer, to make sure our timestamp dependency checking is accurate.
 			if (IncludedFile.LastWriteTimeUtc > WrapperFile.LastWriteTimeUtc)
diff --git a/Engine/Source/Programs/UnrealBuildTool/Platform/Android/AndroidToolChain.cs b/Engine/Source/Programs/UnrealBuildTool/Platform/Android/AndroidToolChain.cs
--- a/Engine/Source/Programs/UnrealBuildTool/Platform/Android/AndroidToolChain.cs
+++ b/Engine/Source/Programs/UnrealBuildTool/Platform/Android/AndroidToolChain.cs
@@ -1207,7 +1207,7 @@ namespace UnrealBuildTool
 			if (!FileReference.Exists(LinkerExceptionsCPPFilename))
 			{
 				// Create a dummy file in case it doesn't exist yet so that the module does not complain it's not there
-				Graph.CreateIntermediateTextFile(LinkerExceptionsCPPFilename, new List<string>());
+				Graph.CreateIntermediateTextFile(LinkerExceptionsCPPFilename, new List<string>(), StringComparison.Ordinal);
 			}
 
 			List<string> Result = new List<string>();
@@ -1246,7 +1246,7 @@ namespace UnrealBuildTool
 			// If we determined that we should write the file, write it now.
 			if (bShouldWriteFile)
 			{
-				Graph.CreateIntermediateTextFile(LinkerExceptionsCPPFilename, Result);
+				Graph.CreateIntermediateTextFile(LinkerExceptionsCPPFilename, Result, StringComparison.Ordinal);
 			}
 
 			SourceFiles.Add(FileItem.GetItemByFileReference(LinkerExceptionsCPPFilename));
@@ -1541,7 +1541,7 @@ namespace UnrealBuildTool
 
 						// Create the response file
 						FileReference ResponseFileName = CompileAction.ProducedItems[0].Location + ".rsp";
-						FileItem ResponseFileItem = Graph.CreateIntermediateTextFile(ResponseFileName, new List<string> { AllArguments });
+						FileItem ResponseFileItem = Graph.CreateIntermediateTextFile(ResponseFileName, new List<string> { AllArguments }, StringComparison.InvariantCultureIgnoreCase);
 						string ResponseArgument = string.Format("@\"{0}\"", ResponseFileName);
 
 						CompileAction.WorkingDirectory = UnrealBuildTool.EngineSourceDirectory;
@@ -2176,7 +2176,7 @@ namespace UnrealBuildTool
 					FileReference ResponseFileName = GetResponseFileName(LinkEnvironment, OutputFile);
 					InputFileNames.Add(LinkResponseArguments.Replace("\\", "/"));
 
-					FileItem ResponseFileItem = Graph.CreateIntermediateTextFile(ResponseFileName, InputFileNames);
+					FileItem ResponseFileItem = Graph.CreateIntermediateTextFile(ResponseFileName, InputFileNames, StringComparison.InvariantCultureIgnoreCase);
 
 					LinkAction.CommandArguments += string.Format(" @\"{0}\"", ResponseFileName);
 					LinkAction.PrerequisiteItems.Add(ResponseFileItem);
diff --git a/Engine/Source/Programs/UnrealBuildTool/Platform/HoloLens/HoloLensToolChain.cs b/Engine/Source/Programs/UnrealBuildTool/Platform/HoloLens/HoloLensToolChain.cs
--- a/Engine/Source/Programs/UnrealBuildTool/Platform/HoloLens/HoloLensToolChain.cs
+++ b/Engine/Source/Programs/UnrealBuildTool/Platform/HoloLens/HoloLensToolChain.cs
@@ -515,8 +515,8 @@ namespace UnrealBuildTool
 					FileReference PCHCPPPath = CompileEnvironment.PrecompiledHeaderIncludeFilename.ChangeExtension(".cpp");
 					FileItem PCHCPPFile = Graph.CreateIntermediateTextFile(
 						PCHCPPPath,
-						string.Format("#include \"{0}\"\r\n", CompileEnvironment.PrecompiledHeaderIncludeFilename.FullName.Replace('\\', '/'))
-						);
+						string.Format("#include \"{0}\"\r\n", CompileEnvironment.PrecompiledHeaderIncludeFilename.FullName.Replace('\\', '/')),
+						StringComparison.Ordinal);
 
 					// Make sure the original source directory the PCH header file existed in is added as an include
 					// path -- it might be a private PCH header and we need to make sure that its found!
@@ -687,7 +687,7 @@ namespace UnrealBuildTool
 				{
 					FileItem TargetFile = CompileAction.ProducedItems[0];
 					FileReference ResponseFileName = new FileReference(TargetFile.AbsolutePath + ".response");
-					FileItem ResponseFileItem = Graph.CreateIntermediateTextFile(ResponseFileName, SharedArguments.Concat(FileArguments).Concat(AdditionalArguments).Select(x => Utils.ExpandVariables(x)));
+					FileItem ResponseFileItem = Graph.CreateIntermediateTextFile(ResponseFileName, SharedArguments.Concat(FileArguments).Concat(AdditionalArguments).Select(x => Utils.ExpandVariables(x)), StringComparison.InvariantCultureIgnoreCase);
 					CompileAction.CommandArguments = " @\"" + ResponseFileName + "\"";
 					CompileAction.PrerequisiteItems.Add(ResponseFileItem);
 				}
@@ -1005,7 +1005,7 @@ namespace UnrealBuildTool
 			FileReference ResponseFileName = GetResponseFileName(LinkEnvironment, OutputFile);
 			if (!ProjectFileGenerator.bGenerateProjectFiles)
 			{
-				FileItem ResponseFile = Graph.CreateIntermediateTextFile(ResponseFileName, String.Join(Environment.NewLine, Arguments));
+				FileItem ResponseFile = Graph.CreateIntermediateTextFile(ResponseFileName, String.Join(Environment.NewLine, Arguments), StringComparison.InvariantCultureIgnoreCase);
 				PrerequisiteItems.Add(ResponseFile);
 			}
 
diff --git a/Engine/Source/Programs/UnrealBuildTool/Platform/IOS/IOSToolChain.cs b/Engine/Source/Programs/UnrealBuildTool/Platform/IOS/IOSToolChain.cs
--- a/Engine/Source/Programs/UnrealBuildTool/Platform/IOS/IOSToolChain.cs
+++ b/Engine/Source/Programs/UnrealBuildTool/Platform/IOS/IOSToolChain.cs
@@ -903,7 +903,7 @@ namespace UnrealBuildTool
 			{
 				bool bIsUnrealGame = LinkEnvironment.OutputFilePath.FullName.Contains("UnrealGame");
 				FileReference ResponsePath = FileReference.Combine(((!bIsUnrealGame && ProjectFile != null) ? ProjectFile.Directory : UnrealBuild.EngineDirectory), "Intermediate", "Build", LinkEnvironment.Platform.ToString(), "LinkFileList_" + LinkEnvironment.OutputFilePath.GetFileNameWithoutExtension() + ".tmp");
-				Graph.CreateIntermediateTextFile(ResponsePath, InputFileNames);
+				Graph.CreateIntermediateTextFile(ResponsePath, InputFileNames, StringComparison.InvariantCultureIgnoreCase);
 				LinkCommandArguments += string.Format(" \\\"@{0}\\\"", ResponsePath.FullName);
 			}
 
@@ -1258,7 +1258,7 @@ namespace UnrealBuildTool
 				ExtractScript.AppendLine(String.Format("unzip -q -o {0} -d {1}", Utils.MakePathSafeToUseWithCommandLine(Framework.ZipFile.FullName), Utils.MakePathSafeToUseWithCommandLine(Framework.FrameworkDirectory.ParentDirectory.FullName))); // Zip contains folder with the same name, hence ParentDirectory
 				ExtractScript.AppendLine(String.Format("touch {0}", Utils.MakePathSafeToUseWithCommandLine(Framework.ExtractedTokenFile.AbsolutePath)));
 
-				FileItem ExtractScriptFileItem = Graph.CreateIntermediateTextFile(new FileReference(Framework.FrameworkDirectory.FullName + ".sh"), ExtractScript.ToString());
+				FileItem ExtractScriptFileItem = Graph.CreateIntermediateTextFile(new FileReference(Framework.FrameworkDirectory.FullName + ".sh"), ExtractScript.ToString(), StringComparison.InvariantCultureIgnoreCase);
 
 				Action UnzipAction = Graph.CreateAction(ActionType.BuildProject);
 				UnzipAction.CommandPath = new FileReference("/bin/sh");
diff --git a/Engine/Source/Programs/UnrealBuildTool/Platform/Linux/LinuxToolChain.cs b/Engine/Source/Programs/UnrealBuildTool/Platform/Linux/LinuxToolChain.cs
--- a/Engine/Source/Programs/UnrealBuildTool/Platform/Linux/LinuxToolChain.cs
+++ b/Engine/Source/Programs/UnrealBuildTool/Platform/Linux/LinuxToolChain.cs
@@ -1418,7 +1418,7 @@ namespace UnrealBuildTool
 				Debug.Assert(CompileAction.ProducedItems.Count > 0);
 
 				FileReference CompilerResponseFileName = CompileAction.ProducedItems[0].Location + ".rsp";
-				FileItem CompilerResponseFileItem = Graph.CreateIntermediateTextFile(CompilerResponseFileName, AllArguments);
+				FileItem CompilerResponseFileItem = Graph.CreateIntermediateTextFile(CompilerResponseFileName, AllArguments, StringComparison.InvariantCultureIgnoreCase);
 
 				CompileAction.CommandArguments = string.Format(" @\"{0}\"", CompilerResponseFileName);
 				CompileAction.PrerequisiteItems.Add(CompilerResponseFileItem);
@@ -1483,7 +1483,7 @@ namespace UnrealBuildTool
 			FileReference ResponsePath = GetResponseFileName(LinkEnvironment, OutputFile);
 			if (!ProjectFileGenerator.bGenerateProjectFiles)
 			{
-				FileItem ResponseFileItem = Graph.CreateIntermediateTextFile(ResponsePath, InputFileNames);
+				FileItem ResponseFileItem = Graph.CreateIntermediateTextFile(ResponsePath, InputFileNames, StringComparison.Ordinal);
 				ArchiveAction.PrerequisiteItems.Add(ResponseFileItem);
 			}
 			ArchiveAction.CommandArguments += string.Format(" @\"{0}\"", ResponsePath.FullName);
@@ -1835,7 +1835,7 @@ namespace UnrealBuildTool
 			ResponseLines.Add(" --end-group");
 
 			FileReference ResponseFileName = GetResponseFileName(LinkEnvironment, OutputFile);
-			FileItem ResponseFileItem = Graph.CreateIntermediateTextFile(ResponseFileName, ResponseLines);
+			FileItem ResponseFileItem = Graph.CreateIntermediateTextFile(ResponseFileName, ResponseLines, StringComparison.InvariantCultureIgnoreCase);
 
 			LinkCommandString += string.Format(" -Wl,@\"{0}\"", ResponseFileName);
 			LinkAction.PrerequisiteItems.Add(ResponseFileItem);
diff --git a/Engine/Source/Programs/UnrealBuildTool/Platform/Windows/PVSToolChain.cs b/Engine/Source/Programs/UnrealBuildTool/Platform/Windows/PVSToolChain.cs
--- a/Engine/Source/Programs/UnrealBuildTool/Platform/Windows/PVSToolChain.cs
+++ b/Engine/Source/Programs/UnrealBuildTool/Platform/Windows/PVSToolChain.cs
@@ -574,7 +574,7 @@ namespace UnrealBuildTool
 				string BaseFileName = PreprocessedFileItem.Location.GetFileNameWithoutExtension();
 
 				FileReference ConfigFileLocation = FileReference.Combine(OutputDir, BaseFileName + ".cfg");
-				FileItem ConfigFileItem = Graph.CreateIntermediateTextFile(ConfigFileLocation, ConfigFileContents.ToString());
+				FileItem ConfigFileItem = Graph.CreateIntermediateTextFile(ConfigFileLocation, ConfigFileContents.ToString(), StringComparison.InvariantCultureIgnoreCase);
 
 				// Run the analzyer on the preprocessed source file
 				FileReference OutputFileLocation = FileReference.Combine(OutputDir, BaseFileName + ".pvslog");
@@ -624,7 +624,7 @@ namespace UnrealBuildTool
 			// Collect the prerequisite items off of the Compile action added in CompileCPPFiles so that in SingleFileCompile mode the PVSGather step is also not filtered out
 			List<FileItem> AnalyzeActionPrerequisiteItems = Makefile.Actions.SelectMany(x => x.ProducedItems).ToList();
 
-			FileItem InputFileListItem = Makefile.CreateIntermediateTextFile(OutputFile.ChangeExtension(".input"), InputFiles.Select(x => x.FullName));
+			FileItem InputFileListItem = Makefile.CreateIntermediateTextFile(OutputFile.ChangeExtension(".input"), InputFiles.Select(x => x.FullName), StringComparison.InvariantCultureIgnoreCase);
 
 			Action AnalyzeAction = Makefile.CreateAction(ActionType.Compile);
 			AnalyzeAction.CommandPath = UnrealBuildTool.GetUBTPath();
diff --git a/Engine/Source/Programs/UnrealBuildTool/Platform/Windows/VCCompileAction.cs b/Engine/Source/Programs/UnrealBuildTool/Platform/Windows/VCCompileAction.cs
--- a/Engine/Source/Programs/UnrealBuildTool/Platform/Windows/VCCompileAction.cs
+++ b/Engine/Source/Programs/UnrealBuildTool/Platform/Windows/VCCompileAction.cs
@@ -381,7 +381,7 @@ namespace UnrealBuildTool
 		{
 			if (ResponseFile != null)
 			{
-				Graph.CreateIntermediateTextFile(ResponseFile, GetCompilerArguments());
+				Graph.CreateIntermediateTextFile(ResponseFile, GetCompilerArguments(), StringComparison.InvariantCultureIgnoreCase);
 			}
 		}
 
diff --git a/Engine/Source/Programs/UnrealBuildTool/Platform/Windows/VCToolChain.cs b/Engine/Source/Programs/UnrealBuildTool/Platform/Windows/VCToolChain.cs
--- a/Engine/Source/Programs/UnrealBuildTool/Platform/Windows/VCToolChain.cs
+++ b/Engine/Source/Programs/UnrealBuildTool/Platform/Windows/VCToolChain.cs
@@ -1076,7 +1076,7 @@ namespace UnrealBuildTool
 					// Generate a CPP File that just includes the precompiled header.
 					string PchCppFile = string.Format("// Compiler: {0}\n#include \"{1}\"\r\n", EnvVars.CompilerVersion, CompileEnvironment.PrecompiledHeaderIncludeFilename.FullName.Replace('\\', '/'));
 					CompileAction.SourceFile = FileItem.GetItemByFileReference(CompileEnvironment.PrecompiledHeaderIncludeFilename.ChangeExtension(".cpp"));
-					Graph.CreateIntermediateTextFile(CompileAction.SourceFile, PchCppFile);
+					Graph.CreateIntermediateTextFile(CompileAction.SourceFile, PchCppFile, StringComparison.Ordinal);
 
 					// Add the precompiled header file to the produced items list.
 					CompileAction.CreatePchFile = FileItem.GetItemByFileReference(FileReference.Combine(OutputDir, SourceFile.Location.GetFileName() + ".pch"));
@@ -1490,7 +1490,7 @@ namespace UnrealBuildTool
 			}
 			Contents.AppendLine();
 
-			FileItem InputFile = Graph.CreateIntermediateTextFile(OutputFile.ChangeExtension(".cpp"), Contents.ToString());
+			FileItem InputFile = Graph.CreateIntermediateTextFile(OutputFile.ChangeExtension(".cpp"), Contents.ToString(), StringComparison.Ordinal);
 
 			// Build the argument list
 			FileItem ObjectFile = FileItem.GetItemByFileReference(OutputFile.ChangeExtension(".obj"));
@@ -1816,7 +1816,7 @@ namespace UnrealBuildTool
 			FileReference ResponseFileName = GetResponseFileName(LinkEnvironment, OutputFile);
 			if (!ProjectFileGenerator.bGenerateProjectFiles)
 			{
-				FileItem ResponseFile = Graph.CreateIntermediateTextFile(ResponseFileName, String.Join(Environment.NewLine, Arguments));
+				FileItem ResponseFile = Graph.CreateIntermediateTextFile(ResponseFileName, String.Join(Environment.NewLine, Arguments), StringComparison.InvariantCultureIgnoreCase);
 				PrerequisiteItems.Add(ResponseFile);
 			}
 
diff --git a/Engine/Source/Programs/UnrealBuildTool/System/ActionGraphBuilder.cs b/Engine/Source/Programs/UnrealBuildTool/System/ActionGraphBuilder.cs
--- a/Engine/Source/Programs/UnrealBuildTool/System/ActionGraphBuilder.cs
+++ b/Engine/Source/Programs/UnrealBuildTool/System/ActionGraphBuilder.cs
@@ -29,8 +29,9 @@ namespace UnrealBuildTool
 		/// </summary>
 		/// <param name="Location">Location of the response file</param>
 		/// <param name="Contents">Contents of the file</param>
+		/// <param name="ComparisonForContentChanges">The StringComparison mode to use to determine if a file should be re-written</param>
 		/// <returns>New file item</returns>
-		void CreateIntermediateTextFile(FileItem Location, string Contents);
+		void CreateIntermediateTextFile(FileItem Location, string Contents, StringComparison ComparisonForContentChanges);
 
 		/// <summary>
 		/// Adds a file which is in the non-unity working set
@@ -82,9 +83,9 @@ namespace UnrealBuildTool
 		}
 
 		/// <inheritdoc/>
-		public virtual void CreateIntermediateTextFile(FileItem FileItem, string Contents)
+		public virtual void CreateIntermediateTextFile(FileItem FileItem, string Contents, StringComparison ComparisonForContentChanges)
 		{
-			Utils.WriteFileIfChanged(FileItem, Contents, StringComparison.OrdinalIgnoreCase);
+			Utils.WriteFileIfChanged(FileItem, Contents, ComparisonForContentChanges);
 		}
 
 		/// <inheritdoc/>
@@ -144,9 +145,9 @@ namespace UnrealBuildTool
 		}
 
 		/// <inheritdoc/>
-		public virtual void CreateIntermediateTextFile(FileItem FileItem, string Contents)
+		public virtual void CreateIntermediateTextFile(FileItem FileItem, string Contents, StringComparison ComparisonForContentChanges)
 		{
-			Inner.CreateIntermediateTextFile(FileItem, Contents);
+			Inner.CreateIntermediateTextFile(FileItem, Contents, ComparisonForContentChanges);
 		}
 
 		/// <inheritdoc/>
@@ -278,11 +279,12 @@ namespace UnrealBuildTool
 		/// <param name="Graph">The action graph</param>
 		/// <param name="AbsolutePath">Path to the intermediate file to create</param>
 		/// <param name="Contents">Contents of the new file</param>
+		/// <param name="ComparisonForContentChanges">The StringComparison mode to use to determine if a file should be re-written</param>
 		/// <returns>File item for the newly created file</returns>
-		public static FileItem CreateIntermediateTextFile(this IActionGraphBuilder Graph, FileReference AbsolutePath, string Contents)
+		public static FileItem CreateIntermediateTextFile(this IActionGraphBuilder Graph, FileReference AbsolutePath, string Contents, StringComparison ComparisonForContentChanges)
 		{
 			FileItem FileItem = FileItem.GetItemByFileReference(AbsolutePath);
-			Graph.CreateIntermediateTextFile(FileItem, Contents);
+			Graph.CreateIntermediateTextFile(FileItem, Contents, ComparisonForContentChanges);
 			return FileItem;
 		}
 
@@ -293,9 +295,10 @@ namespace UnrealBuildTool
 		/// <param name="Graph">The action graph</param>
 		/// <param name="FileItem">Path to the intermediate file to create</param>
 		/// <param name="Contents">Contents of the new file</param>
-		public static void CreateIntermediateTextFile(this IActionGraphBuilder Graph, FileItem FileItem, IEnumerable<string> Contents)
+		/// <param name="ComparisonForContentChanges">The StringComparison mode to use to determine if a file should be re-written</param>
+		public static void CreateIntermediateTextFile(this IActionGraphBuilder Graph, FileItem FileItem, IEnumerable<string> Contents, StringComparison ComparisonForContentChanges)
 		{
-			Graph.CreateIntermediateTextFile(FileItem, string.Join(Environment.NewLine, Contents));
+			Graph.CreateIntermediateTextFile(FileItem, string.Join(Environment.NewLine, Contents), ComparisonForContentChanges);
 		}
 
 		/// <summary>
@@ -305,10 +308,11 @@ namespace UnrealBuildTool
 		/// <param name="Graph">The action graph</param>
 		/// <param name="AbsolutePath">Path to the intermediate file to create</param>
 		/// <param name="Contents">Contents of the new file</param>
+		/// <param name="ComparisonForContentChanges">The StringComparison mode to use to determine if a file should be re-written</param>
 		/// <returns>File item for the newly created file</returns>
-		public static FileItem CreateIntermediateTextFile(this IActionGraphBuilder Graph, FileReference AbsolutePath, IEnumerable<string> Contents)
+		public static FileItem CreateIntermediateTextFile(this IActionGraphBuilder Graph, FileReference AbsolutePath, IEnumerable<string> Contents, StringComparison ComparisonForContentChanges)
 		{
-			return Graph.CreateIntermediateTextFile(AbsolutePath, string.Join(Environment.NewLine, Contents));
+			return Graph.CreateIntermediateTextFile(AbsolutePath, string.Join(Environment.NewLine, Contents), ComparisonForContentChanges);
 		}
 	}
 }
diff --git a/Engine/Source/Programs/UnrealBuildTool/System/TargetMakefile.cs b/Engine/Source/Programs/UnrealBuildTool/System/TargetMakefile.cs
--- a/Engine/Source/Programs/UnrealBuildTool/System/TargetMakefile.cs
+++ b/Engine/Source/Programs/UnrealBuildTool/System/TargetMakefile.cs
@@ -688,10 +688,10 @@ namespace UnrealBuildTool
 		}
 
 		/// <inheritdoc/>
-		public void CreateIntermediateTextFile(FileItem FileItem, string Contents)
+		public void CreateIntermediateTextFile(FileItem FileItem, string Contents, StringComparison ComparisonForContentChanges)
 		{
 			// Write the file
-			Utils.WriteFileIfChanged(FileItem.Location, Contents, StringComparison.InvariantCultureIgnoreCase);
+			Utils.WriteFileIfChanged(FileItem.Location, Contents, ComparisonForContentChanges);
 
 			// Reset the file info, in case it already knows about the old file
 			InternalDependencies.Add(FileItem);
diff --git a/Engine/Source/Programs/UnrealBuildTool/System/Unity.cs b/Engine/Source/Programs/UnrealBuildTool/System/Unity.cs
--- a/Engine/Source/Programs/UnrealBuildTool/System/Unity.cs
+++ b/Engine/Source/Programs/UnrealBuildTool/System/Unity.cs
@@ -257,7 +257,7 @@ namespace UnrealBuildTool
 				FileReference UnityCPPFilePath = FileReference.Combine(IntermediateDirectory, UnityCPPFileName);
 
 				// Write the unity file to the intermediate folder.
-				FileItem UnityCPPFile = Graph.CreateIntermediateTextFile(UnityCPPFilePath, OutputUnityCPPWriter.ToString());
+				FileItem UnityCPPFile = Graph.CreateIntermediateTextFile(UnityCPPFilePath, OutputUnityCPPWriter.ToString(), StringComparison.Ordinal);
 				NewCPPFiles.Add(UnityCPPFile);
 
 				// Store the mapping of source files to unity files in the makefile
@@ -347,7 +347,7 @@ namespace UnrealBuildTool
 					string HeaderCPPFileName = string.Format($"{ModulePrefix}{BaseName}.Header.{Path.GetFileNameWithoutExtension(HeaderFile.AbsolutePath)}.cpp");
 					FileReference HeaderCPPFilePath = FileReference.Combine(IntermediateDirectory, HeaderCPPFileName);
 
-					FileItem HeaderCPPFile = Graph.CreateIntermediateTextFile(HeaderCPPFilePath, OutputHeaderCPPWriter.ToString());
+					FileItem HeaderCPPFile = Graph.CreateIntermediateTextFile(HeaderCPPFilePath, OutputHeaderCPPWriter.ToString(), StringComparison.Ordinal);
 					AdaptiveFiles.Add(HeaderCPPFile);
 				}
 			}
diff --git a/Engine/Source/Programs/UnrealBuildTool/ToolChain/UEToolChain.cs b/Engine/Source/Programs/UnrealBuildTool/ToolChain/UEToolChain.cs
--- a/Engine/Source/Programs/UnrealBuildTool/ToolChain/UEToolChain.cs
+++ b/Engine/Source/Programs/UnrealBuildTool/ToolChain/UEToolChain.cs
@@ -462,7 +462,7 @@ namespace UnrealBuildTool
 				CompileAction.ProducedItems.Add(ISPCIncludeHeaderFile);
 
 				FileReference ResponseFileName = new FileReference(ISPCIncludeHeaderFile.AbsolutePath + ".response");
-				FileItem ResponseFileItem = Graph.CreateIntermediateTextFile(ResponseFileName, Arguments.Select(x => Utils.ExpandVariables(x)));
+				FileItem ResponseFileItem = Graph.CreateIntermediateTextFile(ResponseFileName, Arguments.Select(x => Utils.ExpandVariables(x)), StringComparison.InvariantCultureIgnoreCase);
 				CompileAction.CommandArguments += String.Format("@\"{0}\"", ResponseFileName);
 				CompileAction.PrerequisiteItems.Add(ResponseFileItem);
 
@@ -649,7 +649,7 @@ namespace UnrealBuildTool
 				Result.ObjectFiles.AddRange(CompiledISPCObjFiles);
 
 				FileReference ResponseFileName = new FileReference(CompiledISPCObjFileNoISA.AbsolutePath + ".response");
-				FileItem ResponseFileItem = Graph.CreateIntermediateTextFile(ResponseFileName, Arguments.Select(x => Utils.ExpandVariables(x)));
+				FileItem ResponseFileItem = Graph.CreateIntermediateTextFile(ResponseFileName, Arguments.Select(x => Utils.ExpandVariables(x)), StringComparison.InvariantCultureIgnoreCase);
 				CompileAction.CommandArguments = " @\"" + ResponseFileName + "\"";
 				CompileAction.PrerequisiteItems.Add(ResponseFileItem);
 
