UBT: Place non-source/non-config files in Platforms subfolders correctly into their relative folder structure in the project file

diff --git a/Engine/Source/Programs/UnrealBuildTool/ProjectFiles/ProjectFileGenerator.cs b/Engine/Source/Programs/UnrealBuildTool/ProjectFiles/ProjectFileGenerator.cs
--- a/Engine/Source/Programs/UnrealBuildTool/ProjectFiles/ProjectFileGenerator.cs
+++ b/Engine/Source/Programs/UnrealBuildTool/ProjectFiles/ProjectFileGenerator.cs
@@ -1285,7 +1285,7 @@ namespace UnrealBuildTool
 
 					foreach (DirectoryReference ExtensionDir in UnrealBuildTool.GetExtensionDirs(GameProjectDirectory))
 					{
-						GameProject.AddFilesToProject(SourceFileSearch.FindFiles(ExtensionDir, SearchSubdirectories: false), ExtensionDir);
+						GameProject.AddFilesToProject(SourceFileSearch.FindFiles(ExtensionDir, SearchSubdirectories: false), GameProjectDirectory);
 					}
 
 					// Game restricted source files, since they won't be added via a module FileReference
