Allow TargetRules to specify pointer behavior overrides that will be applied when calling UnrealHeaderTool.  This is to permit individual projects (eg: ShooterGame) to enforce use of TObjectPtr in place of raw pointer members without requiring that other projects do the same.  Example usage from ShooterGameEditor.target.cs:

NativePointerMemberBehaviorOverride = PointerMemberBehavior.Disallow;

diff --git a/Engine/Source/Programs/UnrealBuildTool/Configuration/TargetRules.cs b/Engine/Source/Programs/UnrealBuildTool/Configuration/TargetRules.cs
--- a/Engine/Source/Programs/UnrealBuildTool/Configuration/TargetRules.cs
+++ b/Engine/Source/Programs/UnrealBuildTool/Configuration/TargetRules.cs
@@ -83,6 +83,29 @@ namespace UnrealBuildTool
 		Unique,
 	}
 
+	/// <summary>
+	/// Specifies how UnrealHeaderTool should enforce member pointer declarations in UCLASSes and USTRUCTs.  This should match (by name, not value) the
+	/// EPointerMemberBehavior enum in BaseParser.h so that it can be interpreted correctly by UnrealHeaderTool.
+	/// </summary>
+	[Serializable]
+	public enum PointerMemberBehavior
+	{
+		/// <summary>
+		/// Pointer members of the associated type will be disallowed and result in an error emitted from UnrealHeaderTool.
+		/// </summary>
+		Disallow,
+
+		/// <summary>
+		/// Pointer members of the associated type will be allowed and not emit any messages to log or screen.
+		/// </summary>
+		AllowSilently,
+
+		/// <summary>
+		/// Pointer members of the associated type will be allowed but will emit messages to log.
+		/// </summary>
+		AllowAndLog,
+	}
+
 	/// <summary>
 	/// Determines which version of the engine to take default build settings from. This allows for backwards compatibility as new options are enabled by default.
 	/// </summary>
@@ -810,6 +833,11 @@ namespace UnrealBuildTool
 		[RequiresUniqueBuildEnvironment]
 		public bool bEventDrivenLoader;
 
+		/// <summary>
+		/// Used to override the behavior controlling whether UCLASSes and USTRUCTs are allowed to have native pointer members, if disallowed they will be a UHT error and should be substituted with TObjectPtr members instead.
+		/// </summary>
+		public PointerMemberBehavior? NativePointerMemberBehaviorOverride = null;
+
 		/// <summary>
 		/// Whether the XGE controller worker and modules should be included in the engine build.
 		/// These are required for distributed shader compilation using the XGE interception interface.
@@ -2332,6 +2360,11 @@ namespace UnrealBuildTool
 			get { return Inner.bEventDrivenLoader; }
 		}
 
+		public PointerMemberBehavior? NativePointerMemberBehaviorOverride
+		{
+			get { return Inner.NativePointerMemberBehaviorOverride; }
+		}
+
 		[Obsolete("Use DefaultBuildSettings to determine the default settings used for this target instead")]
 		public bool bUseBackwardsCompatibleDefaults
 		{
diff --git a/Engine/Source/Programs/UnrealBuildTool/Configuration/UEBuildTarget.cs b/Engine/Source/Programs/UnrealBuildTool/Configuration/UEBuildTarget.cs
--- a/Engine/Source/Programs/UnrealBuildTool/Configuration/UEBuildTarget.cs
+++ b/Engine/Source/Programs/UnrealBuildTool/Configuration/UEBuildTarget.cs
@@ -1697,12 +1697,24 @@ namespace UnrealBuildTool
 				ExternalExecution.SetupUObjectModules(ModulesToGenerateHeadersFor, Rules.Platform, ProjectDescriptor, Makefile.UObjectModules, Makefile.UObjectModuleHeaders, Rules.GeneratedCodeVersion, bIsAssemblingBuild, MetadataCache);
 			}
 
+			List<string> UHTAdditionalArguments = new List<string>();
+			if (Rules.NativePointerMemberBehaviorOverride != null)
+			{
+				UHTAdditionalArguments.Add("-ini:Engine:[UnrealHeaderTool]:EngineNativePointerMemberBehavior=" + Rules.NativePointerMemberBehaviorOverride);
+				UHTAdditionalArguments.Add("-ini:Engine:[UnrealHeaderTool]:NonEngineNativePointerMemberBehavior=" + Rules.NativePointerMemberBehaviorOverride);
+			}
+
+			if ((UHTAdditionalArguments != null) && (UHTAdditionalArguments.Any()))
+			{
+				Makefile.UHTAdditionalArguments = UHTAdditionalArguments.ToArray();
+			}
+
 			// NOTE: Even in Gather mode, we need to run UHT to make sure the files exist for the static action graph to be setup correctly.  This is because UHT generates .cpp
 			// files that are injected as top level prerequisites.  If UHT only emitted included header files, we wouldn't need to run it during the Gather phase at all.
 			if (Makefile.UObjectModules.Count > 0)
 			{
 				FileReference ModuleInfoFileName = FileReference.Combine(ProjectIntermediateDirectory, TargetName + ".uhtmanifest");
-				ExternalExecution.ExecuteHeaderToolIfNecessary(BuildConfiguration, ProjectFile, TargetName, TargetType, bHasProjectScriptPlugin, Makefile.UObjectModules, ModuleInfoFileName, true, bIsAssemblingBuild, WorkingSet);
+				ExternalExecution.ExecuteHeaderToolIfNecessary(BuildConfiguration, ProjectFile, TargetName, TargetType, bHasProjectScriptPlugin, Makefile.UObjectModules, ModuleInfoFileName, true, bIsAssemblingBuild, WorkingSet, Makefile.UHTAdditionalArguments);
 			}
 
 			// Find all the shared PCHs.
diff --git a/Engine/Source/Programs/UnrealBuildTool/Modes/BuildMode.cs b/Engine/Source/Programs/UnrealBuildTool/Modes/BuildMode.cs
--- a/Engine/Source/Programs/UnrealBuildTool/Modes/BuildMode.cs
+++ b/Engine/Source/Programs/UnrealBuildTool/Modes/BuildMode.cs
@@ -757,7 +757,7 @@ namespace UnrealBuildTool
 					const bool bIsAssemblingBuild = true;
 
 					FileReference ModuleInfoFileName = FileReference.Combine(Makefile.ProjectIntermediateDirectory, TargetDescriptor.Name + ".uhtmanifest");
-					ExternalExecution.ExecuteHeaderToolIfNecessary(BuildConfiguration, TargetDescriptor.ProjectFile, TargetDescriptor.Name, Makefile.TargetType, Makefile.bHasProjectScriptPlugin, UObjectModules: Makefile.UObjectModules, ModuleInfoFileName: ModuleInfoFileName, bIsGatheringBuild: bIsGatheringBuild, bIsAssemblingBuild: bIsAssemblingBuild, WorkingSet: WorkingSet);
+					ExternalExecution.ExecuteHeaderToolIfNecessary(BuildConfiguration, TargetDescriptor.ProjectFile, TargetDescriptor.Name, Makefile.TargetType, Makefile.bHasProjectScriptPlugin, UObjectModules: Makefile.UObjectModules, ModuleInfoFileName: ModuleInfoFileName, bIsGatheringBuild: bIsGatheringBuild, bIsAssemblingBuild: bIsAssemblingBuild, WorkingSet: WorkingSet, UHTAdditionalArguments: Makefile.UHTAdditionalArguments);
 				}
 			}
 			return Makefile;
diff --git a/Engine/Source/Programs/UnrealBuildTool/System/ExternalExecution.cs b/Engine/Source/Programs/UnrealBuildTool/System/ExternalExecution.cs
--- a/Engine/Source/Programs/UnrealBuildTool/System/ExternalExecution.cs
+++ b/Engine/Source/Programs/UnrealBuildTool/System/ExternalExecution.cs
@@ -1062,7 +1062,7 @@ namespace UnrealBuildTool
 		/// Builds and runs the header tool and touches the header directories.
 		/// Performs any early outs if headers need no changes, given the UObject modules, tool path, game name, and configuration
 		/// </summary>
-		public static void ExecuteHeaderToolIfNecessary(BuildConfiguration BuildConfiguration, FileReference? ProjectFile, string TargetName, TargetType TargetType, bool bHasProjectScriptPlugin, List<UHTModuleInfo> UObjectModules, FileReference ModuleInfoFileName, bool bIsGatheringBuild, bool bIsAssemblingBuild, ISourceFileWorkingSet WorkingSet)
+		public static void ExecuteHeaderToolIfNecessary(BuildConfiguration BuildConfiguration, FileReference? ProjectFile, string TargetName, TargetType TargetType, bool bHasProjectScriptPlugin, List<UHTModuleInfo> UObjectModules, FileReference ModuleInfoFileName, bool bIsGatheringBuild, bool bIsAssemblingBuild, ISourceFileWorkingSet WorkingSet, string[]? UHTAdditionalArguments)
 		{
 			if (ProgressWriter.bWriteMarkup)
 			{
@@ -1195,6 +1195,14 @@ namespace UnrealBuildTool
 					string CmdLine = (ProjectFile != null) ? "\"" + ProjectFile.FullName + "\"" : TargetName;
 					CmdLine += " \"" + ModuleInfoFileName + "\" -LogCmds=\"loginit warning, logexit warning, logdatabase error\" -Unattended -WarningsAsErrors";
 
+					if (UHTAdditionalArguments != null)
+					{
+						foreach (string UHTArgument in UHTAdditionalArguments)
+						{
+							CmdLine += " " + UHTArgument;
+						}
+					}
+
 					if (Log.OutputFile != null)
 					{
 						string LogFileName = Log.OutputFile.GetFileNameWithoutExtension();
diff --git a/Engine/Source/Programs/UnrealBuildTool/System/TargetMakefile.cs b/Engine/Source/Programs/UnrealBuildTool/System/TargetMakefile.cs
--- a/Engine/Source/Programs/UnrealBuildTool/System/TargetMakefile.cs
+++ b/Engine/Source/Programs/UnrealBuildTool/System/TargetMakefile.cs
@@ -22,7 +22,7 @@ namespace UnrealBuildTool
 		/// <summary>
 		/// The version number to write
 		/// </summary>
-		public const int CurrentVersion = 23;
+		public const int CurrentVersion = 24;
 
 		/// <summary>
 		/// The time at which the makefile was created
@@ -84,6 +84,11 @@ namespace UnrealBuildTool
 		/// </summary>
 		public string[]? AdditionalArguments;
 
+		/// <summary>
+		/// The array of command-line arguments passed to UnrealHeaderTool.
+		/// </summary>
+		public string[]? UHTAdditionalArguments;
+
 		/// <summary>
 		/// Scripts which should be run before building anything
 		/// </summary>
@@ -228,6 +233,7 @@ namespace UnrealBuildTool
 			bDeployAfterCompile = Reader.ReadBool();
 			bHasProjectScriptPlugin = Reader.ReadBool();
 			AdditionalArguments = Reader.ReadArray(() => Reader.ReadString());
+			UHTAdditionalArguments = Reader.ReadArray(() => Reader.ReadString());
 			PreBuildScripts = Reader.ReadArray(() => Reader.ReadFileReference());
 			PreBuildTargets = Reader.ReadArray(() => new TargetInfo(Reader));
 			Actions = Reader.ReadList(() => Reader.ReadAction());
@@ -263,6 +269,7 @@ namespace UnrealBuildTool
 			Writer.WriteBool(bDeployAfterCompile);
 			Writer.WriteBool(bHasProjectScriptPlugin);
 			Writer.WriteArray(AdditionalArguments, Item => Writer.WriteString(Item));
+			Writer.WriteArray(UHTAdditionalArguments, Item => Writer.WriteString(Item));
 			Writer.WriteArray(PreBuildScripts, Item => Writer.WriteFileReference(Item));
 			Writer.WriteArray(PreBuildTargets, Item => Item.Write(Writer));
 			Writer.WriteList(Actions, x => Writer.WriteAction(x));
