UnrealBuildTool:

Move IsEngineInstalled() into BuildUtilities

diff --git a/Engine/Source/Programs/DotNETCommon/BuildUtilities/UnrealBuild.cs b/Engine/Source/Programs/DotNETCommon/BuildUtilities/UnrealBuild.cs
--- a/Engine/Source/Programs/DotNETCommon/BuildUtilities/UnrealBuild.cs
+++ b/Engine/Source/Programs/DotNETCommon/BuildUtilities/UnrealBuild.cs
@@ -91,5 +91,23 @@ namespace UnrealBuildBase
 		/// The path to UBT
 		/// </summary>
 		public static readonly FileReference UnrealBuildToolPath = FindUnrealBuildTool();
+
+		/// <summary>
+		/// Whether we're running with engine installed
+		/// </summary>
+		static private bool? bIsEngineInstalled;
+
+		/// <summary>
+		/// Returns true if UnrealBuildTool is running using installed Engine components
+		/// </summary>
+		/// <returns>True if running using installed Engine components</returns>
+		static public bool IsEngineInstalled()
+		{
+			if (!bIsEngineInstalled.HasValue)
+			{
+				bIsEngineInstalled = FileReference.Exists(FileReference.Combine(UnrealBuild.EngineDirectory, "Build", "InstalledBuild.txt"));
+			}
+			return bIsEngineInstalled.Value;
+		}
 	}
 }
diff --git a/Engine/Source/Programs/UnrealBuildTool/Configuration/TargetRules.cs b/Engine/Source/Programs/UnrealBuildTool/Configuration/TargetRules.cs
--- a/Engine/Source/Programs/UnrealBuildTool/Configuration/TargetRules.cs
+++ b/Engine/Source/Programs/UnrealBuildTool/Configuration/TargetRules.cs
@@ -1484,7 +1484,7 @@ namespace UnrealBuildTool
 				{
 					return TargetBuildEnvironment.Unique;
 				}
-				else if (UnrealBuildTool.IsEngineInstalled() || LinkType != TargetLinkType.Monolithic)
+				else if (UnrealBuild.IsEngineInstalled() || LinkType != TargetLinkType.Monolithic)
 				{
 					return TargetBuildEnvironment.Shared;
 				}
@@ -1858,7 +1858,7 @@ namespace UnrealBuildTool
 		/// <returns>Flag for whether the engine is installed</returns>
 		public bool bIsEngineInstalled
 		{
-			get { return UnrealBuildTool.IsEngineInstalled(); }
+			get { return UnrealBuild.IsEngineInstalled(); }
 		}
 
 		/// <summary>
diff --git a/Engine/Source/Programs/UnrealBuildTool/Configuration/UEBuildTarget.cs b/Engine/Source/Programs/UnrealBuildTool/Configuration/UEBuildTarget.cs
--- a/Engine/Source/Programs/UnrealBuildTool/Configuration/UEBuildTarget.cs
+++ b/Engine/Source/Programs/UnrealBuildTool/Configuration/UEBuildTarget.cs
@@ -1415,7 +1415,7 @@ namespace UnrealBuildTool
 			}
 
 			// If this is an installed engine build, clear the promoted flag on the output binaries. This will ensure we will rebuild them.
-			if (Version.IsPromotedBuild && UnrealBuildTool.IsEngineInstalled())
+			if (Version.IsPromotedBuild && UnrealBuild.IsEngineInstalled())
 			{
 				Version.IsPromotedBuild = false;
 			}
@@ -1824,7 +1824,7 @@ namespace UnrealBuildTool
 			BuildProducts.AddRange(RuntimeDependencyTargetFileToSourceFile.Select(x => new KeyValuePair<FileReference, BuildProductType>(x.Key, BuildProductType.RequiredResource)));
 
 			// Remove any installed build products that don't exist. They may be part of an optional install.
-			if(UnrealBuildTool.IsEngineInstalled())
+			if(UnrealBuild.IsEngineInstalled())
 			{
 				BuildProducts.RemoveAll(x => UnrealBuildTool.IsFileInstalled(x.Key) && !FileReference.Exists(x.Key));
 			}
@@ -3489,7 +3489,7 @@ namespace UnrealBuildTool
 			// Installed Engine intermediates go to the project's intermediate folder. Installed Engine never writes to the engine intermediate folder. (Those files are immutable)
 			// Also, when compiling in monolithic, all intermediates go to the project's folder.  This is because a project can change definitions that affects all engine translation
 			// units too, so they can't be shared between different targets.  They are effectively project-specific engine intermediates.
-			if (UnrealBuildTool.IsEngineInstalled() || (ProjectFile != null && ShouldCompileMonolithic()))
+			if (UnrealBuild.IsEngineInstalled() || (ProjectFile != null && ShouldCompileMonolithic()))
 			{
 				if (ProjectFile != null)
 				{
diff --git a/Engine/Source/Programs/UnrealBuildTool/Modes/BuildMode.cs b/Engine/Source/Programs/UnrealBuildTool/Modes/BuildMode.cs
--- a/Engine/Source/Programs/UnrealBuildTool/Modes/BuildMode.cs
+++ b/Engine/Source/Programs/UnrealBuildTool/Modes/BuildMode.cs
@@ -146,7 +146,7 @@ namespace UnrealBuildTool
 			}
 
 			// now that we know the available platforms, we can delete other platforms' junk. if we're only building specific modules from the editor, don't touch anything else (it may be in use).
-			if (!bIgnoreJunk && !UnrealBuildTool.IsEngineInstalled())
+			if (!bIgnoreJunk && !UnrealBuild.IsEngineInstalled())
 			{
 				using(Timeline.ScopeEvent("DeleteJunk()"))
 				{
diff --git a/Engine/Source/Programs/UnrealBuildTool/Platform/Mac/MacToolChain.cs b/Engine/Source/Programs/UnrealBuildTool/Platform/Mac/MacToolChain.cs
--- a/Engine/Source/Programs/UnrealBuildTool/Platform/Mac/MacToolChain.cs
+++ b/Engine/Source/Programs/UnrealBuildTool/Platform/Mac/MacToolChain.cs
@@ -920,7 +920,7 @@ namespace UnrealBuildTool
 
 			if (!bIsBuildingLibrary)
 			{
-				if (UnrealBuildTool.IsEngineInstalled() || (!RuntimePlatform.IsWindows && LinkEnvironment.bIsCrossReferenced == false))
+				if (UnrealBuild.IsEngineInstalled() || (!RuntimePlatform.IsWindows && LinkEnvironment.bIsCrossReferenced == false))
 				{
 					foreach (string Library in EngineAndGameLibraries)
 					{
diff --git a/Engine/Source/Programs/UnrealBuildTool/ProjectFiles/ProjectFileGenerator.cs b/Engine/Source/Programs/UnrealBuildTool/ProjectFiles/ProjectFileGenerator.cs
--- a/Engine/Source/Programs/UnrealBuildTool/ProjectFiles/ProjectFileGenerator.cs
+++ b/Engine/Source/Programs/UnrealBuildTool/ProjectFiles/ProjectFileGenerator.cs
@@ -690,7 +690,7 @@ namespace UnrealBuildTool
 			bool IncludeAllPlatforms = true;
 			ConfigureProjectFileGeneration(Arguments, ref IncludeAllPlatforms);
 
-			if (bGeneratingGameProjectFiles || UnrealBuildTool.IsEngineInstalled())
+			if (bGeneratingGameProjectFiles || UnrealBuild.IsEngineInstalled())
 			{
 				Log.TraceInformation("Discovering modules, targets and source code for project...");
 
@@ -1240,7 +1240,7 @@ namespace UnrealBuildTool
 				}
 			}
 
-			if (bGeneratingGameProjectFiles || UnrealBuildTool.IsEngineInstalled())
+			if (bGeneratingGameProjectFiles || UnrealBuild.IsEngineInstalled())
 			{
 				if (OnlyGameProject == null)
 				{
@@ -1253,7 +1253,7 @@ namespace UnrealBuildTool
 					throw new BuildException("A valid game project was not found in the specified location (" + OnlyGameProject.Directory.FullName + ")");
 				}
 
-				bool bInstalledEngineWithSource = UnrealBuildTool.IsEngineInstalled() && DirectoryReference.Exists(UnrealBuildTool.EngineSourceDirectory);
+				bool bInstalledEngineWithSource = UnrealBuild.IsEngineInstalled() && DirectoryReference.Exists(UnrealBuildTool.EngineSourceDirectory);
 
 				bIncludeEngineSource = bAlwaysIncludeEngineModules || bInstalledEngineWithSource;
 				bIncludeDocumentation = false;
@@ -2378,7 +2378,7 @@ namespace UnrealBuildTool
 					else if (IsEngineTarget)
 					{
 						EngineProject = ProjectFile;
-						if (UnrealBuildTool.IsEngineInstalled())
+						if (UnrealBuild.IsEngineInstalled())
 						{
 							// Allow engine projects to be created but not built for Installed Engine builds
 							EngineProject.IsForeignProject = false;
diff --git a/Engine/Source/Programs/UnrealBuildTool/ProjectFiles/Rider/RiderProjectFileGenerator.cs b/Engine/Source/Programs/UnrealBuildTool/ProjectFiles/Rider/RiderProjectFileGenerator.cs
--- a/Engine/Source/Programs/UnrealBuildTool/ProjectFiles/Rider/RiderProjectFileGenerator.cs
+++ b/Engine/Source/Programs/UnrealBuildTool/ProjectFiles/Rider/RiderProjectFileGenerator.cs
@@ -333,7 +333,7 @@ namespace UnrealBuildTool
 					else if (IsEngineTarget)
 					{
 						EngineProject = ProjectFile;
-						if (UnrealBuildTool.IsEngineInstalled())
+						if (UnrealBuild.IsEngineInstalled())
 						{
 							// Allow engine projects to be created but not built for Installed Engine builds
 							EngineProject.IsForeignProject = false;
diff --git a/Engine/Source/Programs/UnrealBuildTool/ProjectFiles/VisualStudio/VCProject.cs b/Engine/Source/Programs/UnrealBuildTool/ProjectFiles/VisualStudio/VCProject.cs
--- a/Engine/Source/Programs/UnrealBuildTool/ProjectFiles/VisualStudio/VCProject.cs
+++ b/Engine/Source/Programs/UnrealBuildTool/ProjectFiles/VisualStudio/VCProject.cs
@@ -1232,7 +1232,7 @@ namespace UnrealBuildTool
 			}
 
 			// For Installed engine builds, include engine source in the source search paths if it exists. We never build it locally, so the debugger can't find it.
-			if (UnrealBuildTool.IsEngineInstalled() && !IsStubProject)
+			if (UnrealBuild.IsEngineInstalled() && !IsStubProject)
 			{
 				VCProjectFileContent.AppendLine("  <PropertyGroup>");
 				VCProjectFileContent.Append("    <SourcePath>");
@@ -1547,7 +1547,7 @@ namespace UnrealBuildTool
 					VCProjectFileContent.AppendLine("    <NMakeOutput/>");
 					VCProjectFileContent.AppendLine("  </PropertyGroup>");
 				}
-				else if (UnrealBuildTool.IsEngineInstalled() && Combination.ProjectTarget != null && Combination.ProjectTarget.TargetRules != null &&
+				else if (UnrealBuild.IsEngineInstalled() && Combination.ProjectTarget != null && Combination.ProjectTarget.TargetRules != null &&
 					(Combination.Platform == null || !Combination.ProjectTarget.SupportedPlatforms.Contains(Combination.Platform.Value)))
 				{
 					string ProjectRelativeUnusedDirectory = NormalizeProjectPath(DirectoryReference.Combine(UnrealBuild.EngineDirectory, "Intermediate", "Build", "Unused"));
@@ -1658,11 +1658,11 @@ namespace UnrealBuildTool
 					List<string> ExtraTargets = new List<string>();
 					if (!bUsePrecompiled)
 					{
-						if (TargetRulesObject.Type == TargetType.Editor && Settings.bEditorDependsOnShaderCompileWorker && !UnrealBuildTool.IsEngineInstalled())
+						if (TargetRulesObject.Type == TargetType.Editor && Settings.bEditorDependsOnShaderCompileWorker && !UnrealBuild.IsEngineInstalled())
 						{
 							ExtraTargets.Add("ShaderCompileWorker Win64 Development");
 						}
-						if (TargetRulesObject.bWithLiveCoding && Settings.bBuildLiveCodingConsole && !UnrealBuildTool.IsEngineInstalled() && TargetRulesObject.Name != "LiveCodingConsole")
+						if (TargetRulesObject.bWithLiveCoding && Settings.bBuildLiveCodingConsole && !UnrealBuild.IsEngineInstalled() && TargetRulesObject.Name != "LiveCodingConsole")
 						{
 							ExtraTargets.Add(TargetRulesObject.bUseDebugLiveCodingConsole ? "LiveCodingConsole Win64 Debug" : "LiveCodingConsole Win64 Development");
 						}
diff --git a/Engine/Source/Programs/UnrealBuildTool/ProjectFiles/Xcode/XcodeProject.cs b/Engine/Source/Programs/UnrealBuildTool/ProjectFiles/Xcode/XcodeProject.cs
--- a/Engine/Source/Programs/UnrealBuildTool/ProjectFiles/Xcode/XcodeProject.cs
+++ b/Engine/Source/Programs/UnrealBuildTool/ProjectFiles/Xcode/XcodeProject.cs
@@ -1025,7 +1025,7 @@ namespace UnrealBuildTool
 					if (InProjectFile == null)
 					{
 						// For misc tools we default to Intel for installed builds because we know all of that works. 
-						TargetArchitectures = UnrealBuildTool.IsEngineInstalled() ? new[] { MacExports.IntelArchitecture } : AllArchitectures;
+						TargetArchitectures = UnrealBuild.IsEngineInstalled() ? new[] { MacExports.IntelArchitecture } : AllArchitectures;
 					}
 					else
 					{
diff --git a/Engine/Source/Programs/UnrealBuildTool/System/ActionHistory.cs b/Engine/Source/Programs/UnrealBuildTool/System/ActionHistory.cs
--- a/Engine/Source/Programs/UnrealBuildTool/System/ActionHistory.cs
+++ b/Engine/Source/Programs/UnrealBuildTool/System/ActionHistory.cs
@@ -224,7 +224,7 @@ namespace UnrealBuildTool
 		/// <returns>Dependency cache hierarchy for the given project</returns>
 		public static IEnumerable<FileReference> GetFilesToClean(FileReference ProjectFile, string TargetName, UnrealTargetPlatform Platform, TargetType TargetType, string Architecture)
 		{
-			if(ProjectFile == null || !UnrealBuildTool.IsEngineInstalled())
+			if(ProjectFile == null || !UnrealBuild.IsEngineInstalled())
 			{
 				yield return GetEngineLocation(TargetName, Platform, TargetType, Architecture);
 			}
diff --git a/Engine/Source/Programs/UnrealBuildTool/System/ConfigHierarchy.cs b/Engine/Source/Programs/UnrealBuildTool/System/ConfigHierarchy.cs
--- a/Engine/Source/Programs/UnrealBuildTool/System/ConfigHierarchy.cs
+++ b/Engine/Source/Programs/UnrealBuildTool/System/ConfigHierarchy.cs
@@ -1143,7 +1143,7 @@ namespace UnrealBuildTool
 		/// <returns></returns>
 		public static DirectoryReference GetGameAgnosticSavedDir()
 		{
-			if(UnrealBuildTool.IsEngineInstalled())
+			if(UnrealBuild.IsEngineInstalled())
 			{
 				DirectoryReference? UserSettingDir = Utils.GetUserSettingDirectory();
 				if (UserSettingDir != null)
diff --git a/Engine/Source/Programs/UnrealBuildTool/System/CppDependencyCache.cs b/Engine/Source/Programs/UnrealBuildTool/System/CppDependencyCache.cs
--- a/Engine/Source/Programs/UnrealBuildTool/System/CppDependencyCache.cs
+++ b/Engine/Source/Programs/UnrealBuildTool/System/CppDependencyCache.cs
@@ -300,7 +300,7 @@ namespace UnrealBuildTool
 		/// <returns>Dependency cache hierarchy for the given project</returns>
 		public void Mount(FileReference? ProjectFile, string TargetName, UnrealTargetPlatform Platform, UnrealTargetConfiguration Configuration, TargetType TargetType, string Architecture)
 		{
-			if (ProjectFile == null || !UnrealBuildTool.IsEngineInstalled())
+			if (ProjectFile == null || !UnrealBuild.IsEngineInstalled())
 			{
 				string AppName;
 				if (TargetType == TargetType.Program)
diff --git a/Engine/Source/Programs/UnrealBuildTool/System/ExternalExecution.cs b/Engine/Source/Programs/UnrealBuildTool/System/ExternalExecution.cs
--- a/Engine/Source/Programs/UnrealBuildTool/System/ExternalExecution.cs
+++ b/Engine/Source/Programs/UnrealBuildTool/System/ExternalExecution.cs
@@ -783,7 +783,7 @@ namespace UnrealBuildTool
 		private static DateTime GetCoreGeneratedTimestampUtc(string ModuleName, string ModuleGeneratedCodeDirectory)
 		{
 			// In Installed Builds, we don't check the timestamps on engine headers.  Default to a very old date.
-			if (UnrealBuildTool.IsEngineInstalled())
+			if (UnrealBuild.IsEngineInstalled())
 			{
 				return DateTime.MinValue;
 			}
@@ -1136,7 +1136,7 @@ namespace UnrealBuildTool
 				if (!bIsBuildingUHT && bUHTNeedsToRun)
 				{
 					// Always build UnrealHeaderTool if header regeneration is required, unless we're running within an installed ecosystem or hot-reloading
-					if ((!UnrealBuildTool.IsEngineInstalled() || bHasProjectScriptPlugin) &&
+					if ((!UnrealBuild.IsEngineInstalled() || bHasProjectScriptPlugin) &&
 						!BuildConfiguration.bDoNotBuildUHT &&
 						!(bHaveHeaderTool && !bIsGatheringBuild && bIsAssemblingBuild))	// If running in "assembler only" mode, we assume UHT is already up to date for much faster iteration!
 					{
@@ -1212,7 +1212,7 @@ namespace UnrealBuildTool
 						CmdLine += " -abslog=\"" + LogFileName + "\"";
 					}
 
-					if (UnrealBuildTool.IsEngineInstalled())
+					if (UnrealBuild.IsEngineInstalled())
 					{
 						CmdLine += " -installed";
 					}
diff --git a/Engine/Source/Programs/UnrealBuildTool/System/InstalledPlatformInfo.cs b/Engine/Source/Programs/UnrealBuildTool/System/InstalledPlatformInfo.cs
--- a/Engine/Source/Programs/UnrealBuildTool/System/InstalledPlatformInfo.cs
+++ b/Engine/Source/Programs/UnrealBuildTool/System/InstalledPlatformInfo.cs
@@ -289,7 +289,7 @@ namespace UnrealBuildTool
 		public static bool IsValidPlatform(UnrealTargetPlatform Platform, EProjectType ProjectType = EProjectType.Any)
 		{
 			// HACK: For installed builds, we always need to treat Mac as a valid platform for generating project files. When remote building from PC, we won't have all the libraries to do this, so we need to fake it.
-			if(Platform == UnrealTargetPlatform.Mac && ProjectType == EProjectType.Any && BuildHostPlatform.Current.Platform == UnrealTargetPlatform.Mac && UnrealBuildTool.IsEngineInstalled())
+			if(Platform == UnrealTargetPlatform.Mac && ProjectType == EProjectType.Any && BuildHostPlatform.Current.Platform == UnrealTargetPlatform.Mac && UnrealBuild.IsEngineInstalled())
 			{
 				return true;
 			}
@@ -334,7 +334,7 @@ namespace UnrealBuildTool
 		/// <returns>True if the target can be built</returns>
 		public static bool IsValid(TargetType? TargetType, UnrealTargetPlatform? Platform, UnrealTargetConfiguration? Configuration, EProjectType ProjectType, InstalledPlatformState State)
 		{
-			if(!UnrealBuildTool.IsEngineInstalled() || InstalledPlatformConfigurations == null)
+			if(!UnrealBuild.IsEngineInstalled() || InstalledPlatformConfigurations == null)
 			{
 				return true;
 			}
@@ -372,7 +372,7 @@ namespace UnrealBuildTool
 
 		private static bool ContainsValidConfiguration(Predicate<InstalledPlatformConfiguration> ConfigFilter)
 		{
-			if (UnrealBuildTool.IsEngineInstalled() && InstalledPlatformConfigurations != null)
+			if (UnrealBuild.IsEngineInstalled() && InstalledPlatformConfigurations != null)
 			{
 				foreach (InstalledPlatformConfiguration PlatformConfiguration in InstalledPlatformConfigurations)
 				{
diff --git a/Engine/Source/Programs/UnrealBuildTool/System/RulesAssembly.cs b/Engine/Source/Programs/UnrealBuildTool/System/RulesAssembly.cs
--- a/Engine/Source/Programs/UnrealBuildTool/System/RulesAssembly.cs
+++ b/Engine/Source/Programs/UnrealBuildTool/System/RulesAssembly.cs
@@ -572,7 +572,7 @@ namespace UnrealBuildTool
 			}
 
 			// Set the default value for whether to use the shared build environment
-			if(Rules.BuildEnvironment == TargetBuildEnvironment.Unique && UnrealBuildTool.IsEngineInstalled())
+			if(Rules.BuildEnvironment == TargetBuildEnvironment.Unique && UnrealBuild.IsEngineInstalled())
 			{
 				throw new BuildException("Targets with a unique build environment cannot be built an installed engine.");
 			}
diff --git a/Engine/Source/Programs/UnrealBuildTool/System/RulesCompiler.cs b/Engine/Source/Programs/UnrealBuildTool/System/RulesCompiler.cs
--- a/Engine/Source/Programs/UnrealBuildTool/System/RulesCompiler.cs
+++ b/Engine/Source/Programs/UnrealBuildTool/System/RulesCompiler.cs
@@ -365,7 +365,7 @@ namespace UnrealBuildTool
 					}
 				}
 
-				EngineRulesAssembly = CreateEngineRulesAssemblyInternal(UnrealBuildTool.GetExtensionDirs(UnrealBuild.EngineDirectory), ProjectFileGenerator.EngineProjectFileNameBase, EnginePlugins, UnrealBuildTool.IsEngineInstalled() || bUsePrecompiled, bSkipCompile, null);
+				EngineRulesAssembly = CreateEngineRulesAssemblyInternal(UnrealBuildTool.GetExtensionDirs(UnrealBuild.EngineDirectory), ProjectFileGenerator.EngineProjectFileNameBase, EnginePlugins, UnrealBuild.IsEngineInstalled() || bUsePrecompiled, bSkipCompile, null);
 
 				if (MarketplacePlugins.Count > 0)
 				{
diff --git a/Engine/Source/Programs/UnrealBuildTool/System/SourceFileMetadataCache.cs b/Engine/Source/Programs/UnrealBuildTool/System/SourceFileMetadataCache.cs
--- a/Engine/Source/Programs/UnrealBuildTool/System/SourceFileMetadataCache.cs
+++ b/Engine/Source/Programs/UnrealBuildTool/System/SourceFileMetadataCache.cs
@@ -275,7 +275,7 @@ namespace UnrealBuildTool
 		{
 			SourceFileMetadataCache? Cache = null;
 
-			if(ProjectFile == null || !UnrealBuildTool.IsEngineInstalled())
+			if(ProjectFile == null || !UnrealBuild.IsEngineInstalled())
 			{
 				FileReference EngineCacheLocation = FileReference.Combine(UnrealBuild.EngineDirectory, "Intermediate", "Build", "SourceFileCache.bin");
 				Cache = FindOrAddCache(EngineCacheLocation, UnrealBuild.EngineDirectory, Cache);
@@ -297,7 +297,7 @@ namespace UnrealBuildTool
 		/// <returns>Dependency cache hierarchy for the given project</returns>
 		public static IEnumerable<FileReference> GetFilesToClean(FileReference? ProjectFile)
 		{
-			if(ProjectFile == null || !UnrealBuildTool.IsEngineInstalled())
+			if(ProjectFile == null || !UnrealBuild.IsEngineInstalled())
 			{
 				yield return FileReference.Combine(UnrealBuild.EngineDirectory, "Intermediate", "Build", "SourceFileCache.bin");
 			}
diff --git a/Engine/Source/Programs/UnrealBuildTool/System/TargetMakefile.cs b/Engine/Source/Programs/UnrealBuildTool/System/TargetMakefile.cs
--- a/Engine/Source/Programs/UnrealBuildTool/System/TargetMakefile.cs
+++ b/Engine/Source/Programs/UnrealBuildTool/System/TargetMakefile.cs
@@ -339,7 +339,7 @@ namespace UnrealBuildTool
 				// @todo ubtmake: Some platforms may not save any files into this folder.  We should delete + generate a "touch" file to force the directory timestamp to be updated (or just check the timestamp file itself.  We could put it ANYWHERE, actually)
 
 				// Installed Build doesn't need to check engine projects for outdatedness
-				if (!UnrealBuildTool.IsEngineInstalled())
+				if (!UnrealBuild.IsEngineInstalled())
 				{
 					if (DirectoryReference.Exists(ProjectFileGenerator.IntermediateProjectFilesPath))
 					{
diff --git a/Engine/Source/Programs/UnrealBuildTool/System/XmlConfig.cs b/Engine/Source/Programs/UnrealBuildTool/System/XmlConfig.cs
--- a/Engine/Source/Programs/UnrealBuildTool/System/XmlConfig.cs
+++ b/Engine/Source/Programs/UnrealBuildTool/System/XmlConfig.cs
@@ -88,7 +88,7 @@ namespace UnrealBuildTool
 			{
 				// Get the default cache file
 				CacheFile = FileReference.Combine(UnrealBuild.EngineDirectory, "Intermediate", "Build", "XmlConfigCache.bin");
-				if(UnrealBuildTool.IsEngineInstalled())
+				if(UnrealBuild.IsEngineInstalled())
 				{
 					DirectoryReference? UserSettingsDir = Utils.GetUserSettingDirectory();
 					if(UserSettingsDir != null)
@@ -122,7 +122,7 @@ namespace UnrealBuildTool
 
 					// Create a schema for the config files
 					XmlSchema Schema = CreateSchema(CategoryToFields);
-					if(!UnrealBuildTool.IsEngineInstalled())
+					if(!UnrealBuild.IsEngineInstalled())
 					{
 						WriteSchema(Schema, SchemaFile);
 					}
@@ -223,7 +223,7 @@ namespace UnrealBuildTool
 			List<InputFile> InputFiles = new List<InputFile>();
 
 			// Skip all the config files under the Engine folder if it's an installed build
-			if(!UnrealBuildTool.IsEngineInstalled())
+			if(!UnrealBuild.IsEngineInstalled())
 			{
  				// Check for the config file under /Engine/Programs/NotForLicensees/UnrealBuildTool
  				FileReference NotForLicenseesConfigLocation = FileReference.Combine(UnrealBuild.EngineDirectory, "Restricted", "NotForLicensees", "Programs", "UnrealBuildTool", "BuildConfiguration.xml");
diff --git a/Engine/Source/Programs/UnrealBuildTool/ToolChain/AppleToolChain.cs b/Engine/Source/Programs/UnrealBuildTool/ToolChain/AppleToolChain.cs
--- a/Engine/Source/Programs/UnrealBuildTool/ToolChain/AppleToolChain.cs
+++ b/Engine/Source/Programs/UnrealBuildTool/ToolChain/AppleToolChain.cs
@@ -56,7 +56,7 @@ namespace UnrealBuildTool
 			}
 
 			// Installed engine requires Xcode 11
-			if (UnrealBuildTool.IsEngineInstalled())
+			if (UnrealBuild.IsEngineInstalled())
 			{
 				string XcodeBuilderVersionOutput = Utils.RunLocalProcessAndReturnStdOut("xcodebuild", "-version");
 				if (XcodeBuilderVersionOutput.Length > 10)
diff --git a/Engine/Source/Programs/UnrealBuildTool/ToolChain/RemoteMac.cs b/Engine/Source/Programs/UnrealBuildTool/ToolChain/RemoteMac.cs
--- a/Engine/Source/Programs/UnrealBuildTool/ToolChain/RemoteMac.cs
+++ b/Engine/Source/Programs/UnrealBuildTool/ToolChain/RemoteMac.cs
@@ -918,7 +918,7 @@ namespace UnrealBuildTool
 			// Upload the engine files
 			List<FileReference> EngineFilters = new List<FileReference>();
 			EngineFilters.Add(ScriptProtectList);
-			if(UnrealBuildTool.IsEngineInstalled())
+			if(UnrealBuild.IsEngineInstalled())
 			{
 				EngineFilters.Add(FileReference.Combine(UnrealBuild.EngineDirectory, "Build", "Rsync", "RsyncEngineInstalled.txt"));
 			}
diff --git a/Engine/Source/Programs/UnrealBuildTool/UnrealBuildTool.cs b/Engine/Source/Programs/UnrealBuildTool/UnrealBuildTool.cs
--- a/Engine/Source/Programs/UnrealBuildTool/UnrealBuildTool.cs
+++ b/Engine/Source/Programs/UnrealBuildTool/UnrealBuildTool.cs
@@ -26,11 +26,6 @@ namespace UnrealBuildTool
 		/// </summary>
 		static public System.Collections.IDictionary? InitialEnvironment;
 
-		/// <summary>
-		/// Whether we're running with engine installed
-		/// </summary>
-		static private bool? bIsEngineInstalled;
-
 		/// <summary>
 		/// Whether we're running with an installed project
 		/// </summary>
@@ -96,7 +91,7 @@ namespace UnrealBuildTool
 				if (CachedWritableEngineDirectory == null)
 				{
 					DirectoryReference? UserDir = null;
-					if (IsEngineInstalled())
+					if (UnrealBuild.IsEngineInstalled())
 					{
 						UserDir = Utils.GetUserSettingDirectory();
 					}
@@ -122,7 +117,7 @@ namespace UnrealBuildTool
 			{
 				if (CachedEngineProgramSavedDirectory == null)
 				{
-					if (IsEngineInstalled())
+					if (UnrealBuild.IsEngineInstalled())
 					{
 						CachedEngineProgramSavedDirectory = Utils.GetUserSettingDirectory() ?? DirectoryReference.Combine(UnrealBuild.EngineDirectory, "Programs");
 					}
@@ -222,19 +217,6 @@ namespace UnrealBuildTool
 		/// </summary>
 		static string? RemoteIniPath = null;
 
-		/// <summary>
-		/// Returns true if UnrealBuildTool is running using installed Engine components
-		/// </summary>
-		/// <returns>True if running using installed Engine components</returns>
-		static public bool IsEngineInstalled()
-		{
-			if (!bIsEngineInstalled.HasValue)
-			{
-				bIsEngineInstalled = FileReference.Exists(FileReference.Combine(UnrealBuild.EngineDirectory, "Build", "InstalledBuild.txt"));
-			}
-			return bIsEngineInstalled.Value;
-		}
-
 		/// <summary>
 		/// Returns true if UnrealBuildTool is running using an installed project (ie. a mod kit)
 		/// </summary>
@@ -281,7 +263,7 @@ namespace UnrealBuildTool
 		/// <returns>True if the file is part of the installed distribution, false otherwise</returns>
 		static public bool IsFileInstalled(FileReference File)
 		{
-			if(IsEngineInstalled() && File.IsUnderDirectory(UnrealBuild.EngineDirectory))
+			if(UnrealBuild.IsEngineInstalled() && File.IsUnderDirectory(UnrealBuild.EngineDirectory))
 			{
 				return true;
 			}
