UnrealBuildTool: Remove LocalExecutor and use ParallelExecutor for all host platforms
Update ManagedProcess to capture stderr if using C# Process objects

diff --git a/Engine/Source/Programs/Shared/EpicGames.Core/ManagedProcess.cs b/Engine/Source/Programs/Shared/EpicGames.Core/ManagedProcess.cs
--- a/Engine/Source/Programs/Shared/EpicGames.Core/ManagedProcess.cs
+++ b/Engine/Source/Programs/Shared/EpicGames.Core/ManagedProcess.cs
@@ -600,6 +600,14 @@ namespace EpicGames.Core
 		/// <param name="Flags">Flags for the new process</param>
 		private void CreateManagedProcessPortable(ManagedProcessGroup? Group, string FileName, string CommandLine, string? WorkingDirectory, IReadOnlyDictionary<string, string>? Environment, ProcessPriorityClass Priority, ManagedProcessFlags Flags)
 		{
+			// TODO: Process Arguments follow windows conventions in .NET Core
+			// Which means single quotes ' are not considered quotes.
+			// see [REDACTED_PROFILE_URL]/runtime/issues/29857
+			// also see UE-102580
+			// for rules see https://docs.microsoft.com/en-us/cpp/cpp/main-function-command-line-args
+			CommandLine = CommandLine.Replace('\'', '\"');
+
+
 			// for non-Windows platforms
 			FrameworkProcess = new Process();
 			FrameworkProcess.StartInfo.FileName = FileName;
@@ -607,6 +615,7 @@ namespace EpicGames.Core
 			FrameworkProcess.StartInfo.WorkingDirectory = WorkingDirectory;
 			FrameworkProcess.StartInfo.RedirectStandardInput = true;
 			FrameworkProcess.StartInfo.RedirectStandardOutput = true;
+			FrameworkProcess.StartInfo.RedirectStandardError = true;
 			FrameworkProcess.StartInfo.UseShellExecute = false;
 			FrameworkProcess.StartInfo.CreateNoWindow = true;
 
@@ -631,6 +640,9 @@ namespace EpicGames.Core
 			Id = FrameworkProcess.Id;
 			StdIn = FrameworkProcess.StandardInput.BaseStream;
 			StdOut = FrameworkProcess.StandardOutput.BaseStream;
+			StdOutText = FrameworkProcess.StandardOutput;
+			StdErr = FrameworkProcess.StandardError.BaseStream;
+			StdErrText = FrameworkProcess.StandardError;
 		}
 
 		/// <summary>
@@ -684,18 +696,23 @@ namespace EpicGames.Core
 		public int Read(byte[] Buffer, int Offset, int Count)
 		{
 			// Fill the buffer, reentering managed code every 20ms to allow thread abort exceptions to be thrown
-			Task<int> ReadTask;
-			if(FrameworkProcess == null)
+			Task<int> ReadTask = StdOut!.ReadAsync(Buffer, Offset, Count);
+			while(!ReadTask.Wait(20))
 			{
-				ReadTask = StdOut!.ReadAsync(Buffer, Offset, Count);
+				// Spin through managed code to allow things like ThreadAbortExceptions to be thrown.
 			}
-			else
+			if (ReadTask.Result != 0)
 			{
-				ReadTask = FrameworkProcess.StandardOutput.BaseStream.ReadAsync(Buffer, Offset, Count);
+				return ReadTask.Result;
 			}
-			while(!ReadTask.Wait(20))
+
+			if (StdOut != StdErr)
 			{
-				// Spin through managed code to allow things like ThreadAbortExceptions to be thrown.
+				ReadTask = StdErr!.ReadAsync(Buffer, Offset, Count);
+				while (!ReadTask.Wait(20))
+				{
+					// Spin through managed code to allow things like ThreadAbortExceptions to be thrown.
+				}
 			}
 			return ReadTask.Result;
 		}
@@ -709,14 +726,17 @@ namespace EpicGames.Core
 		/// <returns>Number of bytes read</returns>
 		public async Task<int> ReadAsync(byte[] Buffer, int Offset, int Count, CancellationToken CancellationToken)
 		{
-			if (FrameworkProcess == null)
+			Task<int> ReadTask = StdOut!.ReadAsync(Buffer, Offset, Count, CancellationToken);
+			if (StdOut != StdErr)
 			{
-				return await StdOut!.ReadAsync(Buffer, Offset, Count, CancellationToken);
-			}
-			else
-			{
-				return await FrameworkProcess.StandardOutput.BaseStream.ReadAsync(Buffer, Offset, Count, CancellationToken);
+				await ReadTask;
+				if (ReadTask.Result == 0)
+				{
+					ReadTask = StdErr!.ReadAsync(Buffer, Offset, Count, CancellationToken);
+				}
 			}
+
+			return await ReadTask;
 		}
 
 		/// <summary>
@@ -727,14 +747,16 @@ namespace EpicGames.Core
 		/// <returns></returns>
 		public Task CopyToAsync(Stream OutputStream, CancellationToken CancellationToken)
 		{
-			if(FrameworkProcess == null)
+			if (StdOut == StdErr)
 			{
 				return StdOut!.CopyToAsync(OutputStream, CancellationToken);
 			}
-			else
+
+			return new Task(async () =>
 			{
-				return FrameworkProcess.StandardOutput.BaseStream.CopyToAsync(OutputStream, CancellationToken);
-			}
+				await StdOut!.CopyToAsync(OutputStream, CancellationToken);
+				await StdErr!.CopyToAsync(OutputStream, CancellationToken);
+			});
 		}
 
 		/// <summary>
@@ -761,9 +783,23 @@ namespace EpicGames.Core
 			using (CancellationTokenRegistration Registration = CancellationToken.Register(() => TaskCompletionSource.SetResult(false)))
 			{
 				byte[] Buffer = new byte[BufferSize];
-				for (; ; )
+				for (;;)
+				{
+					Task<int> ReadTask = StdOut!.ReadAsync(Buffer, 0, BufferSize, CancellationToken);
+
+					Task CompletedTask = await Task.WhenAny(ReadTask, TaskCompletionSource.Task);
+					CancellationToken.ThrowIfCancellationRequested();
+
+					int Bytes = await ReadTask;
+					if (Bytes == 0)
+					{
+						break;
+					}
+					await WriteOutputAsync(Buffer, 0, Bytes, CancellationToken);
+				}
+				for (;;)
 				{
-					Task<int> ReadTask = StdOut.ReadAsync(Buffer, 0, BufferSize, CancellationToken);
+					Task<int> ReadTask = StdErr!.ReadAsync(Buffer, 0, BufferSize, CancellationToken);
 
 					Task CompletedTask = await Task.WhenAny(ReadTask, TaskCompletionSource.Task);
 					CancellationToken.ThrowIfCancellationRequested();
@@ -831,6 +867,7 @@ namespace EpicGames.Core
 				Array.Copy(Buffer, LastStartIdx, Buffer, 0, Buffer.Length - LastStartIdx);
 				NumBytesInBuffer -= LastStartIdx;
 			}
+			WaitForExit();
 			return OutputLines;
 		}
 
@@ -840,14 +877,17 @@ namespace EpicGames.Core
 		/// <returns>New line</returns>
 		public async Task<string?> ReadLineAsync()
 		{
-			if (FrameworkProcess == null)
+			Task<string?> ReadLineTask = StdOutText!.ReadLineAsync();
+			if (StdOutText != StdErrText)
 			{
-				return await StdOutText!.ReadLineAsync();
-			}
-			else
-			{
-				return await FrameworkProcess.StandardOutput.ReadLineAsync();
+				await ReadLineTask;
+				if (ReadLineTask.Result == null)
+				{
+					ReadLineTask = StdErrText!.ReadLineAsync();
+				}
 			}
+
+			return await ReadLineTask;
 		}
 
 		/// <summary>
@@ -861,15 +901,7 @@ namespace EpicGames.Core
 			try
 			{
 				// Busy wait for the ReadLine call to finish, so we can get interrupted by thread abort exceptions
-				Task<string?> ReadLineTask;
-				if(FrameworkProcess == null)
-				{
-					ReadLineTask = StdOutText!.ReadLineAsync();
-				}
-				else
-				{
-					ReadLineTask = FrameworkProcess.StandardOutput.ReadLineAsync();
-				}
+				Task<string?> ReadLineTask = StdOutText!.ReadLineAsync();
 				for (;;)
 				{
 					const int MillisecondsTimeout = 20;
@@ -878,16 +910,37 @@ namespace EpicGames.Core
 							: ReadLineTask.Wait(MillisecondsTimeout))
 					{
 						Line = ReadLineTask.IsCompleted ? ReadLineTask.Result : null;
-						return Line != null;
+						if (Line != null)
+						{
+							return true;
+						}
+						break;
+					}
+				}
+
+				if (StdOutText != StdErrText)
+				{
+					ReadLineTask = StdErrText!.ReadLineAsync();
+					for (;;)
+					{
+						const int MillisecondsTimeout = 20;
+						if (Token.HasValue
+								? ReadLineTask.Wait(MillisecondsTimeout, Token.Value)
+								: ReadLineTask.Wait(MillisecondsTimeout))
+						{
+							Line = ReadLineTask.IsCompleted ? ReadLineTask.Result : null;
+							return Line != null;
+						}
 					}
 				}
 			}
 			catch (OperationCanceledException)
 			{
 				// If the cancel token is signalled, just return false.
-				Line = null;
-				return false;
 			}
+
+			Line = null;
+			return false;
 		}
 
 		/// <summary>
diff --git a/Engine/Source/Programs/UnrealBuildTool/Configuration/BuildConfiguration.cs b/Engine/Source/Programs/UnrealBuildTool/Configuration/BuildConfiguration.cs
--- a/Engine/Source/Programs/UnrealBuildTool/Configuration/BuildConfiguration.cs
+++ b/Engine/Source/Programs/UnrealBuildTool/Configuration/BuildConfiguration.cs
@@ -87,12 +87,6 @@ namespace UnrealBuildTool
 		[CommandLine("-NoUBTMakefiles", Value = "false")]
 		public bool bUseUBTMakefiles = true;
 
-		/// <summary>
-		/// Whether to allow using parallel executor on Windows.
-		/// </summary>
-		[XmlConfigFile]
-		public bool bAllowParallelExecutor = true;
-
 		/// <summary>
 		/// Number of actions that can be executed in parallel. If 0 then code will pick a default based
 		/// on the number of cores and memory available. Applies to the ParallelExecutor, HybridExecutor, and LocalExecutor
diff --git a/Engine/Source/Programs/UnrealBuildTool/Executors/Experimental/FASTBuild.cs b/Engine/Source/Programs/UnrealBuildTool/Executors/Experimental/FASTBuild.cs
--- a/Engine/Source/Programs/UnrealBuildTool/Executors/Experimental/FASTBuild.cs
+++ b/Engine/Source/Programs/UnrealBuildTool/Executors/Experimental/FASTBuild.cs
@@ -70,6 +70,10 @@ namespace UnrealBuildTool
 
 	class FASTBuild : ActionExecutor
 	{
+		/// <summary>
+		/// Executor to use for local actions
+		/// </summary>
+		ActionExecutor LocalExecutor;
 
 		public readonly static string DefaultExecutableBasePath	= Path.Combine(UnrealBuild.EngineDirectory.FullName, "Extras", "ThirdPartyNotUE", "FASTBuild");
 
@@ -150,6 +154,16 @@ namespace UnrealBuildTool
 
 		//////////////////////////////////////////
 
+		/// <summary>
+		/// Constructor
+		/// </summary>
+		public FASTBuild(int MaxLocalActions)
+		{
+			XmlConfig.ApplyTo(this);
+
+			this.LocalExecutor = new ParallelExecutor(MaxLocalActions);
+		}
+
 		public override string Name
 		{
 			get { return "FASTBuild"; }
@@ -437,7 +451,7 @@ namespace UnrealBuildTool
 
 			if (PreCompileActions.Any())
 			{
-				bool bResult = new LocalExecutor().ExecuteActions(PreCompileActions);
+				bool bResult = LocalExecutor.ExecuteActions(PreCompileActions);
 
 				if (!bResult)
 					return false;
@@ -466,7 +480,7 @@ namespace UnrealBuildTool
 
 			if (PostCompileActions.Any())
 			{
-				bool bResult = new LocalExecutor().ExecuteActions(PostCompileActions);
+				bool bResult = LocalExecutor.ExecuteActions(PostCompileActions);
 
 				if (!bResult)
 					return false;
diff --git a/Engine/Source/Programs/UnrealBuildTool/Executors/LocalExecutor.cs b/Engine/Source/Programs/UnrealBuildTool/Executors/LocalExecutor.cs
deleted file mode 100644
diff --git a/Engine/Source/Programs/UnrealBuildTool/Executors/ParallelExecutor.cs b/Engine/Source/Programs/UnrealBuildTool/Executors/ParallelExecutor.cs
--- a/Engine/Source/Programs/UnrealBuildTool/Executors/ParallelExecutor.cs
+++ b/Engine/Source/Programs/UnrealBuildTool/Executors/ParallelExecutor.cs
@@ -1,13 +1,11 @@
 // Copyright Epic Games, Inc. All Rights Reserved.
 
+using EpicGames.Core;
 using System;
 using System.Collections.Generic;
 using System.Diagnostics;
 using System.Linq;
-using System.Text;
 using System.Threading;
-using System.Threading.Tasks;
-using EpicGames.Core;
 
 namespace UnrealBuildTool
 {
@@ -31,10 +29,6 @@ namespace UnrealBuildTool
 			public List<string> LogLines = new List<string>();
 			public int ExitCode = -1;
 
-			// BEGIN TEMPORARY TO CATCH PVS-STUDIO ISSUES
-			public bool GotOutput = false;
-			// END TEMPORARY
-
 			public BuildAction(LinkedAction Inner)
 			{
 				this.Inner = Inner;
@@ -113,7 +107,7 @@ namespace UnrealBuildTool
 		/// <returns>True if the parallel executor can be used</returns>
 		public static bool IsAvailable()
 		{
-			return RuntimePlatform.IsWindows;
+			return true;
 		}
 
 		/// <summary>
@@ -122,11 +116,12 @@ namespace UnrealBuildTool
 		/// <returns>True if all the tasks successfully executed, or false if any of them failed.</returns>
 		public override bool ExecuteActions(List<LinkedAction> InputActions)
 		{
-			Log.TraceInformation("Building {0} {1} with {2} {3}...", InputActions.Count, (InputActions.Count == 1) ? "action" : "actions", NumParallelProcesses, (NumParallelProcesses == 1)? "process" : "processes");
+			int ActualNumParallelProcesses = Math.Min(InputActions.Count, NumParallelProcesses);
+			Log.TraceInformation("Building {0} {1} with {2} {3}...", InputActions.Count, (InputActions.Count == 1) ? "action" : "actions", ActualNumParallelProcesses, (ActualNumParallelProcesses == 1) ? "process" : "processes");
 
 			// Create actions with all our internal metadata
 			List<BuildAction> Actions = new List<BuildAction>();
-			for(int Idx = 0; Idx < InputActions.Count; Idx++)
+			for (int Idx = 0; Idx < InputActions.Count; Idx++)
 			{
 				BuildAction Action = new BuildAction(InputActions[Idx]);
 				Action.SortIndex = Idx;
@@ -143,10 +138,10 @@ namespace UnrealBuildTool
 			Dictionary<LinkedAction, BuildAction> LinkedActionToBuildAction = Actions.ToDictionary(x => x.Inner, x => x);
 			foreach (BuildAction Action in Actions)
 			{
-				foreach(LinkedAction PrerequisiteAction in Action.Inner.PrerequisiteActions)
+				foreach (LinkedAction PrerequisiteAction in Action.Inner.PrerequisiteActions)
 				{
 					BuildAction? Dependency;
-					if(LinkedActionToBuildAction.TryGetValue(PrerequisiteAction, out Dependency))
+					if (LinkedActionToBuildAction.TryGetValue(PrerequisiteAction, out Dependency))
 					{
 						Action.Dependencies.Add(Dependency);
 						Dependency.Dependants.Add(Action);
@@ -156,7 +151,7 @@ namespace UnrealBuildTool
 
 			// Figure out the recursive dependency count
 			HashSet<BuildAction> VisitedActions = new HashSet<BuildAction>();
-			foreach(BuildAction Action in Actions)
+			foreach (BuildAction Action in Actions)
 			{
 				Action.MissingDependencyCount = Action.Dependencies.Count;
 				RecursiveIncDependents(Action, VisitedActions);
@@ -164,9 +159,9 @@ namespace UnrealBuildTool
 
 			// Create the list of things to process
 			List<BuildAction> QueuedActions = new List<BuildAction>();
-			foreach(BuildAction Action in Actions)
+			foreach (BuildAction Action in Actions)
 			{
-				if(Action.MissingDependencyCount == 0)
+				if (Action.MissingDependencyCount == 0)
 				{
 					QueuedActions.Add(Action);
 				}
@@ -177,29 +172,30 @@ namespace UnrealBuildTool
 			{
 				// Create a job object for all the child processes
 				bool bResult = true;
-				Dictionary<BuildAction, Thread> ExecutingActions = new Dictionary<BuildAction,Thread>();
+				Dictionary<BuildAction, Thread> ExecutingActions = new Dictionary<BuildAction, Thread>();
 				List<BuildAction> CompletedActions = new List<BuildAction>();
 
-				using(ManagedProcessGroup ProcessGroup = new ManagedProcessGroup())
+				using (ManagedProcessGroup ProcessGroup = new ManagedProcessGroup())
 				{
-					using(AutoResetEvent CompletedEvent = new AutoResetEvent(false))
+					using (AutoResetEvent CompletedEvent = new AutoResetEvent(false))
 					{
 						int NumCompletedActions = 0;
 						using (ProgressWriter ProgressWriter = new ProgressWriter("Compiling C++ source code...", false))
 						{
-							while(QueuedActions.Count > 0 || ExecutingActions.Count > 0)
+							while (QueuedActions.Count > 0 || ExecutingActions.Count > 0)
 							{
 								// Sort the actions by the number of things dependent on them
-								QueuedActions.Sort((A, B) => (A.TotalDependantCount == B.TotalDependantCount)? (B.SortIndex - A.SortIndex) : (B.TotalDependantCount - A.TotalDependantCount));
+								QueuedActions.Sort((A, B) => (A.TotalDependantCount == B.TotalDependantCount) ? (B.SortIndex - A.SortIndex) : (B.TotalDependantCount - A.TotalDependantCount));
 
 								// Create threads up to the maximum number of actions
-								while(ExecutingActions.Count < NumParallelProcesses && QueuedActions.Count > 0)
+								while (ExecutingActions.Count < ActualNumParallelProcesses && QueuedActions.Count > 0)
 								{
 									BuildAction Action = QueuedActions[QueuedActions.Count - 1];
 									QueuedActions.RemoveAt(QueuedActions.Count - 1);
 
 									Thread ExecutingThread = new Thread(() => { ExecuteAction(ProcessGroup, Action, CompletedActions, CompletedEvent); });
-									ExecutingThread.Name = String.Format("Build:{0}", Action.Inner.StatusDescription);
+									string Description = $"{(Action.Inner.CommandDescription != null ? Action.Inner.CommandDescription : Action.Inner.CommandPath.GetFileName())} {Action.Inner.StatusDescription}".Trim();
+									ExecutingThread.Name = String.Format("Build:{0}", Description);
 									ExecutingThread.Start();
 
 									ExecutingActions.Add(Action, ExecutingThread);
@@ -209,9 +205,9 @@ namespace UnrealBuildTool
 								CompletedEvent.WaitOne();
 
 								// Wait for something to finish and flush it to the log
-								lock(CompletedActions)
+								lock (CompletedActions)
 								{
-									foreach(BuildAction CompletedAction in CompletedActions)
+									foreach (BuildAction CompletedAction in CompletedActions)
 									{
 										// Join the thread
 										Thread CompletedThread = ExecutingActions[CompletedAction];
@@ -222,23 +218,26 @@ namespace UnrealBuildTool
 										NumCompletedActions++;
 										ProgressWriter.Write(NumCompletedActions, InputActions.Count);
 
+										string Description = $"{(CompletedAction.Inner.CommandDescription != null ? CompletedAction.Inner.CommandDescription : CompletedAction.Inner.CommandPath.GetFileName())} {CompletedAction.Inner.StatusDescription}".Trim();
+
 										// Write it to the log
-										if(CompletedAction.LogLines.Count > 0)
+										if (CompletedAction.Inner.bShouldOutputStatusDescription)
 										{
-											Log.TraceInformation("[{0}/{1}] {2}", NumCompletedActions, InputActions.Count, CompletedAction.LogLines[0]);
-											for(int LineIdx = 1; LineIdx < CompletedAction.LogLines.Count; LineIdx++)
-											{
-												Log.TraceInformation("{0}", CompletedAction.LogLines[LineIdx]);
-											}
+											Log.TraceInformation("[{0}/{1}] {2}", NumCompletedActions, InputActions.Count, Description);
+										}
+
+										foreach (string Line in CompletedAction.LogLines)
+										{
+											Log.TraceInformation(Line);
 										}
 
 										// Check the exit code
-										if(CompletedAction.ExitCode == 0)
+										if (CompletedAction.ExitCode == 0)
 										{
 											// Mark all the dependents as done
-											foreach(BuildAction DependantAction in CompletedAction.Dependants)
+											foreach (BuildAction DependantAction in CompletedAction.Dependants)
 											{
-												if(--DependantAction.MissingDependencyCount == 0)
+												if (--DependantAction.MissingDependencyCount == 0)
 												{
 													QueuedActions.Add(DependantAction);
 												}
@@ -247,16 +246,15 @@ namespace UnrealBuildTool
 										else
 										{
 											// BEGIN TEMPORARY TO CATCH PVS-STUDIO ISSUES
-											if (!CompletedAction.GotOutput)
+											if (CompletedAction.LogLines.Count == 0)
 											{
-												string StatusDescription = CompletedAction.LogLines.Count > 0 ? CompletedAction.LogLines[0] : "";
-												Log.TraceInformation("[{0}/{1}] {2} - Error but no output", NumCompletedActions, InputActions.Count, StatusDescription);
-												Log.TraceInformation("[{0}/{1}] {2} - {3} {4} {5} {6}", NumCompletedActions, InputActions.Count, StatusDescription, CompletedAction.ExitCode, 
+												Log.TraceInformation("[{0}/{1}] {2} - Error but no output", NumCompletedActions, InputActions.Count, Description);
+												Log.TraceInformation("[{0}/{1}] {2} - {3} {4} {5} {6}", NumCompletedActions, InputActions.Count, Description, CompletedAction.ExitCode,
 													CompletedAction.Inner.WorkingDirectory, CompletedAction.Inner.CommandPath, CompletedAction.Inner.CommandArguments);
 											}
 											// END TEMPORARY
 											// Update the exit code if it's not already set
-											if(bResult && CompletedAction.ExitCode != 0)
+											if (bResult && CompletedAction.ExitCode != 0)
 											{
 												bResult = false;
 											}
@@ -266,7 +264,7 @@ namespace UnrealBuildTool
 								}
 
 								// If we've already got a non-zero exit code, clear out the list of queued actions so nothing else will run
-								if(!bResult && bStopCompilationAfterErrors)
+								if (!bResult && bStopCompilationAfterErrors)
 								{
 									QueuedActions.Clear();
 								}
@@ -288,14 +286,6 @@ namespace UnrealBuildTool
 		/// <param name="CompletedEvent">Event to set once an event is complete</param>
 		static void ExecuteAction(ManagedProcessGroup ProcessGroup, BuildAction Action, List<BuildAction> CompletedActions, AutoResetEvent CompletedEvent)
 		{
-			if (Action.Inner.bShouldOutputStatusDescription && !String.IsNullOrEmpty(Action.Inner.StatusDescription))
-			{
-				Action.LogLines.Add(Action.Inner.StatusDescription);
-			}
-
-			// BEGIN TEMPORARY TO CATCH PVS-STUDIO ISSUES
-			int LogLineCount = Action.LogLines.Count;
-			// END TEMPORARY
 			try
 			{
 				using (ManagedProcess Process = new ManagedProcess(ProcessGroup, Action.Inner.CommandPath.FullName, Action.Inner.CommandArguments, Action.Inner.WorkingDirectory.FullName, null, null, ProcessPriorityClass.BelowNormal))
@@ -303,17 +293,14 @@ namespace UnrealBuildTool
 					Action.LogLines.AddRange(Process.ReadAllLines());
 					Action.ExitCode = Process.ExitCode;
 				}
-				// BEGIN TEMPORARY TO CATCH PVS-STUDIO ISSUES
-				Action.GotOutput = Action.LogLines.Count > LogLineCount;
-				// END TEMPORARY
 			}
-			catch(Exception Ex)
+			catch (Exception Ex)
 			{
 				Log.WriteException(Ex, null);
 				Action.ExitCode = 1;
 			}
 
-			lock(CompletedActions)
+			lock (CompletedActions)
 			{
 				CompletedActions.Add(Action);
 			}
@@ -328,9 +315,9 @@ namespace UnrealBuildTool
 		/// <param name="VisitedActions">Set of visited actions</param>
 		private static void RecursiveIncDependents(BuildAction Action, HashSet<BuildAction> VisitedActions)
 		{
-			foreach(BuildAction Dependency in Action.Dependants)
+			foreach (BuildAction Dependency in Action.Dependants)
 			{
-				if(!VisitedActions.Contains(Action))
+				if (!VisitedActions.Contains(Action))
 				{
 					VisitedActions.Add(Action);
 					Dependency.TotalDependantCount++;
diff --git a/Engine/Source/Programs/UnrealBuildTool/System/ActionGraph.cs b/Engine/Source/Programs/UnrealBuildTool/System/ActionGraph.cs
--- a/Engine/Source/Programs/UnrealBuildTool/System/ActionGraph.cs
+++ b/Engine/Source/Programs/UnrealBuildTool/System/ActionGraph.cs
@@ -337,19 +337,15 @@ namespace UnrealBuildTool
 				}
 				else if (BuildConfiguration.bAllowFASTBuild && FASTBuild.IsAvailable())
 				{
-					Executor = new FASTBuild();
+					Executor = new FASTBuild(BuildConfiguration.MaxParallelActions);
 				}
 				else if(BuildConfiguration.bAllowSNDBS && SNDBS.IsAvailable())
 				{
 					Executor = new SNDBS();
 				}
-				else if(BuildConfiguration.bAllowParallelExecutor && ParallelExecutor.IsAvailable())
-				{
-					Executor = new ParallelExecutor(BuildConfiguration.MaxParallelActions);
-				}
 				else
 				{
-					Executor = new LocalExecutor(BuildConfiguration.MaxParallelActions);
+					Executor = new ParallelExecutor(BuildConfiguration.MaxParallelActions);
 				}
 
 				// Execute the build
