[Backout] - CL16642253
Original CL Desc
-----------------------------------------------------------------
UnrealBuildTool: Remove LocalExecutor and use ParallelExecutor for all host platforms
Update ManagedProcess to capture stderr if using C# Process objects

diff --git a/Engine/Source/Programs/Shared/EpicGames.Core/ManagedProcess.cs b/Engine/Source/Programs/Shared/EpicGames.Core/ManagedProcess.cs
--- a/Engine/Source/Programs/Shared/EpicGames.Core/ManagedProcess.cs
+++ b/Engine/Source/Programs/Shared/EpicGames.Core/ManagedProcess.cs
@@ -600,14 +600,6 @@ namespace EpicGames.Core
 		/// <param name="Flags">Flags for the new process</param>
 		private void CreateManagedProcessPortable(ManagedProcessGroup? Group, string FileName, string CommandLine, string? WorkingDirectory, IReadOnlyDictionary<string, string>? Environment, ProcessPriorityClass Priority, ManagedProcessFlags Flags)
 		{
-			// TODO: Process Arguments follow windows conventions in .NET Core
-			// Which means single quotes ' are not considered quotes.
-			// see [REDACTED_PROFILE_URL]/runtime/issues/29857
-			// also see UE-102580
-			// for rules see https://docs.microsoft.com/en-us/cpp/cpp/main-function-command-line-args
-			CommandLine = CommandLine.Replace('\'', '\"');
-
-
 			// for non-Windows platforms
 			FrameworkProcess = new Process();
 			FrameworkProcess.StartInfo.FileName = FileName;
@@ -615,7 +607,6 @@ namespace EpicGames.Core
 			FrameworkProcess.StartInfo.WorkingDirectory = WorkingDirectory;
 			FrameworkProcess.StartInfo.RedirectStandardInput = true;
 			FrameworkProcess.StartInfo.RedirectStandardOutput = true;
-			FrameworkProcess.StartInfo.RedirectStandardError = true;
 			FrameworkProcess.StartInfo.UseShellExecute = false;
 			FrameworkProcess.StartInfo.CreateNoWindow = true;
 
@@ -640,9 +631,6 @@ namespace EpicGames.Core
 			Id = FrameworkProcess.Id;
 			StdIn = FrameworkProcess.StandardInput.BaseStream;
 			StdOut = FrameworkProcess.StandardOutput.BaseStream;
-			StdOutText = FrameworkProcess.StandardOutput;
-			StdErr = FrameworkProcess.StandardError.BaseStream;
-			StdErrText = FrameworkProcess.StandardError;
 		}
 
 		/// <summary>
@@ -696,23 +684,18 @@ namespace EpicGames.Core
 		public int Read(byte[] Buffer, int Offset, int Count)
 		{
 			// Fill the buffer, reentering managed code every 20ms to allow thread abort exceptions to be thrown
-			Task<int> ReadTask = StdOut!.ReadAsync(Buffer, Offset, Count);
-			while(!ReadTask.Wait(20))
+			Task<int> ReadTask;
+			if(FrameworkProcess == null)
 			{
-				// Spin through managed code to allow things like ThreadAbortExceptions to be thrown.
+				ReadTask = StdOut!.ReadAsync(Buffer, Offset, Count);
 			}
-			if (ReadTask.Result != 0)
+			else
 			{
-				return ReadTask.Result;
+				ReadTask = FrameworkProcess.StandardOutput.BaseStream.ReadAsync(Buffer, Offset, Count);
 			}
-
-			if (StdOut != StdErr)
+			while(!ReadTask.Wait(20))
 			{
-				ReadTask = StdErr!.ReadAsync(Buffer, Offset, Count);
-				while (!ReadTask.Wait(20))
-				{
-					// Spin through managed code to allow things like ThreadAbortExceptions to be thrown.
-				}
+				// Spin through managed code to allow things like ThreadAbortExceptions to be thrown.
 			}
 			return ReadTask.Result;
 		}
@@ -726,17 +709,14 @@ namespace EpicGames.Core
 		/// <returns>Number of bytes read</returns>
 		public async Task<int> ReadAsync(byte[] Buffer, int Offset, int Count, CancellationToken CancellationToken)
 		{
-			Task<int> ReadTask = StdOut!.ReadAsync(Buffer, Offset, Count, CancellationToken);
-			if (StdOut != StdErr)
+			if (FrameworkProcess == null)
 			{
-				await ReadTask;
-				if (ReadTask.Result == 0)
-				{
-					ReadTask = StdErr!.ReadAsync(Buffer, Offset, Count, CancellationToken);
-				}
+				return await StdOut!.ReadAsync(Buffer, Offset, Count, CancellationToken);
+			}
+			else
+			{
+				return await FrameworkProcess.StandardOutput.BaseStream.ReadAsync(Buffer, Offset, Count, CancellationToken);
 			}
-
-			return await ReadTask;
 		}
 
 		/// <summary>
@@ -747,16 +727,14 @@ namespace EpicGames.Core
 		/// <returns></returns>
 		public Task CopyToAsync(Stream OutputStream, CancellationToken CancellationToken)
 		{
-			if (StdOut == StdErr)
+			if(FrameworkProcess == null)
 			{
 				return StdOut!.CopyToAsync(OutputStream, CancellationToken);
 			}
-
-			return new Task(async () =>
+			else
 			{
-				await StdOut!.CopyToAsync(OutputStream, CancellationToken);
-				await StdErr!.CopyToAsync(OutputStream, CancellationToken);
-			});
+				return FrameworkProcess.StandardOutput.BaseStream.CopyToAsync(OutputStream, CancellationToken);
+			}
 		}
 
 		/// <summary>
@@ -783,23 +761,9 @@ namespace EpicGames.Core
 			using (CancellationTokenRegistration Registration = CancellationToken.Register(() => TaskCompletionSource.SetResult(false)))
 			{
 				byte[] Buffer = new byte[BufferSize];
-				for (;;)
-				{
-					Task<int> ReadTask = StdOut!.ReadAsync(Buffer, 0, BufferSize, CancellationToken);
-
-					Task CompletedTask = await Task.WhenAny(ReadTask, TaskCompletionSource.Task);
-					CancellationToken.ThrowIfCancellationRequested();
-
-					int Bytes = await ReadTask;
-					if (Bytes == 0)
-					{
-						break;
-					}
-					await WriteOutputAsync(Buffer, 0, Bytes, CancellationToken);
-				}
-				for (;;)
+				for (; ; )
 				{
-					Task<int> ReadTask = StdErr!.ReadAsync(Buffer, 0, BufferSize, CancellationToken);
+					Task<int> ReadTask = StdOut.ReadAsync(Buffer, 0, BufferSize, CancellationToken);
 
 					Task CompletedTask = await Task.WhenAny(ReadTask, TaskCompletionSource.Task);
 					CancellationToken.ThrowIfCancellationRequested();
@@ -867,7 +831,6 @@ namespace EpicGames.Core
 				Array.Copy(Buffer, LastStartIdx, Buffer, 0, Buffer.Length - LastStartIdx);
 				NumBytesInBuffer -= LastStartIdx;
 			}
-			WaitForExit();
 			return OutputLines;
 		}
 
@@ -877,17 +840,14 @@ namespace EpicGames.Core
 		/// <returns>New line</returns>
 		public async Task<string?> ReadLineAsync()
 		{
-			Task<string?> ReadLineTask = StdOutText!.ReadLineAsync();
-			if (StdOutText != StdErrText)
+			if (FrameworkProcess == null)
 			{
-				await ReadLineTask;
-				if (ReadLineTask.Result == null)
-				{
-					ReadLineTask = StdErrText!.ReadLineAsync();
-				}
+				return await StdOutText!.ReadLineAsync();
+			}
+			else
+			{
+				return await FrameworkProcess.StandardOutput.ReadLineAsync();
 			}
-
-			return await ReadLineTask;
 		}
 
 		/// <summary>
@@ -901,7 +861,15 @@ namespace EpicGames.Core
 			try
 			{
 				// Busy wait for the ReadLine call to finish, so we can get interrupted by thread abort exceptions
-				Task<string?> ReadLineTask = StdOutText!.ReadLineAsync();
+				Task<string?> ReadLineTask;
+				if(FrameworkProcess == null)
+				{
+					ReadLineTask = StdOutText!.ReadLineAsync();
+				}
+				else
+				{
+					ReadLineTask = FrameworkProcess.StandardOutput.ReadLineAsync();
+				}
 				for (;;)
 				{
 					const int MillisecondsTimeout = 20;
@@ -910,37 +878,16 @@ namespace EpicGames.Core
 							: ReadLineTask.Wait(MillisecondsTimeout))
 					{
 						Line = ReadLineTask.IsCompleted ? ReadLineTask.Result : null;
-						if (Line != null)
-						{
-							return true;
-						}
-						break;
-					}
-				}
-
-				if (StdOutText != StdErrText)
-				{
-					ReadLineTask = StdErrText!.ReadLineAsync();
-					for (;;)
-					{
-						const int MillisecondsTimeout = 20;
-						if (Token.HasValue
-								? ReadLineTask.Wait(MillisecondsTimeout, Token.Value)
-								: ReadLineTask.Wait(MillisecondsTimeout))
-						{
-							Line = ReadLineTask.IsCompleted ? ReadLineTask.Result : null;
-							return Line != null;
-						}
+						return Line != null;
 					}
 				}
 			}
 			catch (OperationCanceledException)
 			{
 				// If the cancel token is signalled, just return false.
+				Line = null;
+				return false;
 			}
-
-			Line = null;
-			return false;
 		}
 
 		/// <summary>
diff --git a/Engine/Source/Programs/UnrealBuildTool/Configuration/BuildConfiguration.cs b/Engine/Source/Programs/UnrealBuildTool/Configuration/BuildConfiguration.cs
--- a/Engine/Source/Programs/UnrealBuildTool/Configuration/BuildConfiguration.cs
+++ b/Engine/Source/Programs/UnrealBuildTool/Configuration/BuildConfiguration.cs
@@ -87,6 +87,12 @@ namespace UnrealBuildTool
 		[CommandLine("-NoUBTMakefiles", Value = "false")]
 		public bool bUseUBTMakefiles = true;
 
+		/// <summary>
+		/// Whether to allow using parallel executor on Windows.
+		/// </summary>
+		[XmlConfigFile]
+		public bool bAllowParallelExecutor = true;
+
 		/// <summary>
 		/// Number of actions that can be executed in parallel. If 0 then code will pick a default based
 		/// on the number of cores and memory available. Applies to the ParallelExecutor, HybridExecutor, and LocalExecutor
diff --git a/Engine/Source/Programs/UnrealBuildTool/Executors/Experimental/FASTBuild.cs b/Engine/Source/Programs/UnrealBuildTool/Executors/Experimental/FASTBuild.cs
--- a/Engine/Source/Programs/UnrealBuildTool/Executors/Experimental/FASTBuild.cs
+++ b/Engine/Source/Programs/UnrealBuildTool/Executors/Experimental/FASTBuild.cs
@@ -70,10 +70,6 @@ namespace UnrealBuildTool
 
 	class FASTBuild : ActionExecutor
 	{
-		/// <summary>
-		/// Executor to use for local actions
-		/// </summary>
-		ActionExecutor LocalExecutor;
 
 		public readonly static string DefaultExecutableBasePath	= Path.Combine(UnrealBuild.EngineDirectory.FullName, "Extras", "ThirdPartyNotUE", "FASTBuild");
 
@@ -154,16 +150,6 @@ namespace UnrealBuildTool
 
 		//////////////////////////////////////////
 
-		/// <summary>
-		/// Constructor
-		/// </summary>
-		public FASTBuild(int MaxLocalActions)
-		{
-			XmlConfig.ApplyTo(this);
-
-			this.LocalExecutor = new ParallelExecutor(MaxLocalActions);
-		}
-
 		public override string Name
 		{
 			get { return "FASTBuild"; }
@@ -451,7 +437,7 @@ namespace UnrealBuildTool
 
 			if (PreCompileActions.Any())
 			{
-				bool bResult = LocalExecutor.ExecuteActions(PreCompileActions);
+				bool bResult = new LocalExecutor().ExecuteActions(PreCompileActions);
 
 				if (!bResult)
 					return false;
@@ -480,7 +466,7 @@ namespace UnrealBuildTool
 
 			if (PostCompileActions.Any())
 			{
-				bool bResult = LocalExecutor.ExecuteActions(PostCompileActions);
+				bool bResult = new LocalExecutor().ExecuteActions(PostCompileActions);
 
 				if (!bResult)
 					return false;
diff --git a/Engine/Source/Programs/UnrealBuildTool/Executors/LocalExecutor.cs b/Engine/Source/Programs/UnrealBuildTool/Executors/LocalExecutor.cs
new file mode 100644
--- /dev/null
+++ b/Engine/Source/Programs/UnrealBuildTool/Executors/LocalExecutor.cs
@@ -0,0 +1,403 @@
+// Copyright Epic Games, Inc. All Rights Reserved.
+
+using System;
+using System.Collections.Generic;
+using System.Text;
+using System.Text.RegularExpressions;
+using System.Diagnostics;
+using System.IO;
+using System.Reflection;
+using System.Threading;
+using System.Linq;
+using EpicGames.Core;
+
+namespace UnrealBuildTool
+{
+	class ActionThread
+	{
+		/// <summary>
+		/// Cache the exit code from the command so that the executor can report errors
+		/// </summary>
+		public int ExitCode = 0;
+
+		/// <summary>
+		/// Set to true only when the local or RPC action is complete
+		/// </summary>
+		public bool bComplete = false;
+
+		/// <summary>
+		/// Cache the action that this thread is managing
+		/// </summary>
+		LinkedAction Action;
+
+		/// <summary>
+		/// For reporting status to the user
+		/// </summary>
+		int JobNumber;
+		int TotalJobs;
+
+		/// <summary>
+		/// Constructor, takes the action to process
+		/// </summary>
+		public ActionThread(LinkedAction InAction, int InJobNumber, int InTotalJobs)
+		{
+			Action = InAction;
+			JobNumber = InJobNumber;
+			TotalJobs = InTotalJobs;
+		}
+
+		/// <summary>
+		/// Used when debuging Actions outputs all action return values to debug out
+		/// </summary>
+		/// <param name="sender"> Sending object</param>
+		/// <param name="e">  Event arguments (In this case, the line of string output)</param>
+		protected void ActionDebugOutput(object sender, DataReceivedEventArgs e)
+		{
+			string Output = e.Data;
+			if (Output == null)
+			{
+				return;
+			}
+
+			Log.TraceInformation(Output);
+		}
+
+
+		/// <summary>
+		/// The actual function to run in a thread. This is potentially long and blocking
+		/// </summary>
+		private void ThreadFunc()
+		{
+			string? Args = Action.CommandArguments;
+			// TODO: Process Arguments follow windows conventions in .NET Core
+			// Which means single quotes ' are not considered quotes.
+			// see [REDACTED_PROFILE_URL]/runtime/issues/29857
+			// also see UE-102580
+			// for rules see https://docs.microsoft.com/en-us/cpp/cpp/main-function-command-line-args
+			Args = Args?.Replace('\'', '\"');
+
+			// Create the action's process.
+			ProcessStartInfo ActionStartInfo = new ProcessStartInfo();
+			ActionStartInfo.WorkingDirectory = Action.WorkingDirectory.FullName;
+			ActionStartInfo.FileName = Action.CommandPath.FullName;
+			ActionStartInfo.Arguments = Args;
+			ActionStartInfo.UseShellExecute = false;
+			ActionStartInfo.RedirectStandardInput = false;
+			ActionStartInfo.RedirectStandardOutput = false;
+			ActionStartInfo.RedirectStandardError = false;
+
+			// Log command-line used to execute task if debug info printing is enabled.
+			Log.TraceVerbose("Executing: {0} {1}", ActionStartInfo.FileName, ActionStartInfo.Arguments);
+
+			// Log summary if wanted.
+			if (Action.bShouldOutputStatusDescription)
+			{
+				string CommandDescription = Action.CommandDescription != null ? Action.CommandDescription : Path.GetFileName(ActionStartInfo.FileName);
+				if (string.IsNullOrEmpty(CommandDescription))
+				{
+					Log.TraceInformation(Action.StatusDescription);
+				}
+				else
+				{
+					Log.TraceInformation("[{0}/{1}] {2} {3}", JobNumber, TotalJobs, CommandDescription, Action.StatusDescription);
+				}
+			}
+
+			// Try to launch the action's process, and produce a friendly error message if it fails.
+			Process? ActionProcess = null;
+			try
+			{
+				try
+				{
+					ActionProcess = new Process();
+					ActionProcess.StartInfo = ActionStartInfo;
+					ActionStartInfo.RedirectStandardOutput = true;
+					ActionStartInfo.RedirectStandardError = true;
+					ActionProcess.OutputDataReceived += new DataReceivedEventHandler(ActionDebugOutput);
+					ActionProcess.ErrorDataReceived += new DataReceivedEventHandler(ActionDebugOutput);
+					ActionProcess.Start();
+
+					ActionProcess.BeginOutputReadLine();
+					ActionProcess.BeginErrorReadLine();
+				}
+				catch (Exception ex)
+				{
+					Log.TraceError("Failed to start local process for action: {0} {1}", Action.CommandPath, Action.CommandArguments);
+					Log.WriteException(ex, null);
+					ExitCode = 1;
+					bComplete = true;
+					return;
+				}
+
+				// wait for process to start
+				// NOTE: this may or may not be necessary; seems to depend on whether the system UBT is running on start the process in a timely manner.
+				int checkIterations = 0;
+				bool haveConfiguredProcess = false;
+				do
+				{
+					if (ActionProcess.HasExited)
+					{
+						if (haveConfiguredProcess == false)
+							Debug.WriteLine("Process for action exited before able to configure!");
+						break;
+					}
+
+					if (!haveConfiguredProcess)
+					{
+						try
+						{
+							ActionProcess.PriorityClass = ProcessPriorityClass.BelowNormal;
+							haveConfiguredProcess = true;
+						}
+						catch (Exception)
+						{
+						}
+						break;
+					}
+
+					Thread.Sleep(10);
+
+					checkIterations++;
+				} while (checkIterations < 100);
+				if (checkIterations == 100)
+				{
+					throw new BuildException("Failed to configure local process for action: {0} {1}", Action.CommandPath, Action.CommandArguments);
+				}
+
+				// block until it's complete
+				// @todo iosmerge: UBT had started looking at:	if (Utils.IsValidProcess(Process))
+				//    do we need to check that in the thread model?
+				ActionProcess.WaitForExit();
+
+				// capture exit code
+				ExitCode = ActionProcess.ExitCode;
+			}
+			finally
+			{
+				// As the process has finished now, free its resources. On non-Windows platforms, processes depend 
+				// on POSIX/BSD threading and these are limited per application. Disposing the Process releases 
+				// these thread resources.
+				if (ActionProcess != null)
+					ActionProcess.Close();
+			}
+
+			// we are done!!
+			bComplete = true;
+		}
+
+		/// <summary>
+		/// Starts a thread and runs the action in that thread
+		/// </summary>
+		public void Run()
+		{
+			Thread T = new Thread(ThreadFunc);
+			T.Start();
+		}
+	};
+
+	class LocalExecutor : ActionExecutor
+	{
+		/// <summary>
+		/// Maximum processor count for local execution. 
+		/// </summary>
+		[XmlConfigFile]
+		int MaxProcessorCount = int.MaxValue;
+
+		/// <summary>
+		/// Processor count multiplier for local execution. Can be below 1 to reserve CPU for other tasks.
+		/// When using the local executor (not XGE), run a single action on each CPU core. Note that you can set this to a larger value
+		/// to get slightly faster build times in many cases, but your computer's responsiveness during compiling may be much worse.
+		/// This value is ignored if the CPU does not support hyper-threading.
+		/// </summary>
+		[XmlConfigFile]
+		double ProcessorCountMultiplier = 1.0;
+
+		/// <summary>
+		/// Free memory per action in bytes, used to limit the number of parallel actions if the machine is memory starved.
+		/// Set to 0 to disable free memory checking.
+		/// </summary>
+		[XmlConfigFile]
+		double MemoryPerActionBytes = 1.5 * 1024 * 1024 * 1024;
+
+		/// <summary>
+		/// How many processes that will be executed in parallel
+		/// </summary>
+		int NumParallelProcesses;
+
+		/// <summary>
+		/// Constructor
+		/// </summary>
+		/// <param name="MaxLocalActions">How many actions to execute in parallel. When 0 a default will be chosen based on system resources</param>
+		public LocalExecutor(int MaxLocalActions=0)
+		{
+			XmlConfig.ApplyTo(this);
+
+			// if specified this caps how many processors we can use
+			if (MaxLocalActions > 0)
+			{
+				NumParallelProcesses = MaxLocalActions;
+			}
+			else
+			{
+				// Figure out how many processors to use
+				NumParallelProcesses = Utils.GetMaxActionsToExecuteInParallel(MaxProcessorCount, ProcessorCountMultiplier, Convert.ToInt64(MemoryPerActionBytes));
+			}
+		}
+
+		public override string Name
+		{
+			get { return "Local"; }
+		}
+
+		/// <summary>
+		/// Executes the specified actions locally.
+		/// </summary>
+		/// <returns>True if all the tasks successfully executed, or false if any of them failed.</returns>
+		public override bool ExecuteActions(List<LinkedAction> Actions)
+		{
+			// Time to sleep after each iteration of the loop in order to not busy wait.
+			const float LoopSleepTime = 0.1f;
+
+			// The number of actions to execute in parallel is trying to keep the CPU busy enough in presence of I/O stalls.
+			Log.TraceInformation("Performing {0} actions ({1} in parallel)", Actions.Count, NumParallelProcesses);
+
+			Dictionary<LinkedAction, ActionThread?> ActionThreadDictionary = new Dictionary<LinkedAction, ActionThread?>();
+			int JobNumber = 1;
+			using (ProgressWriter ProgressWriter = new ProgressWriter("Compiling C++ source code...", false))
+			{
+				int ProgressValue = 0;
+				while (true)
+				{
+					// Count the number of pending and still executing actions.
+					int NumUnexecutedActions = 0;
+					int NumExecutingActions = 0;
+					foreach (LinkedAction Action in Actions)
+					{
+						ActionThread? ActionThread = null;
+						bool bFoundActionProcess = ActionThreadDictionary.TryGetValue(Action, out ActionThread);
+						if (bFoundActionProcess == false)
+						{
+							NumUnexecutedActions++;
+						}
+						else if (ActionThread != null)
+						{
+							if (ActionThread.bComplete == false)
+							{
+								NumUnexecutedActions++;
+								NumExecutingActions++;
+							}
+						}
+					}
+
+					// Update the current progress
+					int NewProgressValue = Actions.Count + 1 - NumUnexecutedActions;
+					if (ProgressValue != NewProgressValue)
+					{
+						ProgressWriter.Write(ProgressValue, Actions.Count + 1);
+						ProgressValue = NewProgressValue;
+					}
+
+					// If there aren't any pending actions left, we're done executing.
+					if (NumUnexecutedActions == 0)
+					{
+						break;
+					}
+
+					// If there are fewer actions executing than the maximum, look for pending actions that don't have any outdated
+					// prerequisites.
+					foreach (LinkedAction Action in Actions)
+					{
+						ActionThread? ActionProcess = null;
+						bool bFoundActionProcess = ActionThreadDictionary.TryGetValue(Action, out ActionProcess);
+						if (bFoundActionProcess == false)
+						{
+							if (NumExecutingActions < Math.Max(1, NumParallelProcesses))
+							{
+								// Determine whether there are any prerequisites of the action that are outdated.
+								bool bHasOutdatedPrerequisites = false;
+								bool bHasFailedPrerequisites = false;
+								foreach (LinkedAction PrerequisiteAction in Action.PrerequisiteActions)
+								{
+									if (Actions.Contains(PrerequisiteAction))
+									{
+										ActionThread? PrerequisiteProcess = null;
+										bool bFoundPrerequisiteProcess = ActionThreadDictionary.TryGetValue(PrerequisiteAction, out PrerequisiteProcess);
+										if (bFoundPrerequisiteProcess == true)
+										{
+											if (PrerequisiteProcess == null)
+											{
+												bHasFailedPrerequisites = true;
+											}
+											else if (PrerequisiteProcess.bComplete == false)
+											{
+												bHasOutdatedPrerequisites = true;
+											}
+											else if (PrerequisiteProcess.ExitCode != 0)
+											{
+												bHasFailedPrerequisites = true;
+											}
+										}
+										else
+										{
+											bHasOutdatedPrerequisites = true;
+										}
+									}
+								}
+
+								// If there are any failed prerequisites of this action, don't execute it.
+								if (bHasFailedPrerequisites)
+								{
+									// Add a null entry in the dictionary for this action.
+									ActionThreadDictionary.Add(Action, null);
+								}
+								// If there aren't any outdated prerequisites of this action, execute it.
+								else if (!bHasOutdatedPrerequisites)
+								{
+									ActionThread ActionThread = new ActionThread(Action, JobNumber, Actions.Count);
+									JobNumber++;
+
+									try
+									{
+										ActionThread.Run();
+									}
+									catch (Exception ex)
+									{
+										throw new BuildException(ex, "Failed to start thread for action: {0} {1}\r\n{2}", Action.CommandPath, Action.CommandArguments, ex.ToString());
+									}
+
+									ActionThreadDictionary.Add(Action, ActionThread);
+
+									NumExecutingActions++;
+								}
+							}
+						}
+					}
+
+					System.Threading.Thread.Sleep(TimeSpan.FromSeconds(LoopSleepTime));
+				}
+			}
+
+			// Check whether any of the tasks failed and log action stats if wanted.
+			bool bSuccess = true;
+			foreach (KeyValuePair<LinkedAction, ActionThread?> ActionProcess in ActionThreadDictionary)
+			{
+				LinkedAction Action = ActionProcess.Key;
+				ActionThread? ActionThread = ActionProcess.Value;
+
+				// Check for pending actions, preemptive failure
+				if (ActionThread == null)
+				{
+					bSuccess = false;
+					continue;
+				}
+				// Check for executed action but general failure
+				if (ActionThread.ExitCode != 0)
+				{
+					bSuccess = false;
+				}
+			}
+
+			return bSuccess;
+		}
+	};
+}
diff --git a/Engine/Source/Programs/UnrealBuildTool/Executors/ParallelExecutor.cs b/Engine/Source/Programs/UnrealBuildTool/Executors/ParallelExecutor.cs
--- a/Engine/Source/Programs/UnrealBuildTool/Executors/ParallelExecutor.cs
+++ b/Engine/Source/Programs/UnrealBuildTool/Executors/ParallelExecutor.cs
@@ -1,11 +1,13 @@
 // Copyright Epic Games, Inc. All Rights Reserved.
 
-using EpicGames.Core;
 using System;
 using System.Collections.Generic;
 using System.Diagnostics;
 using System.Linq;
+using System.Text;
 using System.Threading;
+using System.Threading.Tasks;
+using EpicGames.Core;
 
 namespace UnrealBuildTool
 {
@@ -29,6 +31,10 @@ namespace UnrealBuildTool
 			public List<string> LogLines = new List<string>();
 			public int ExitCode = -1;
 
+			// BEGIN TEMPORARY TO CATCH PVS-STUDIO ISSUES
+			public bool GotOutput = false;
+			// END TEMPORARY
+
 			public BuildAction(LinkedAction Inner)
 			{
 				this.Inner = Inner;
@@ -107,7 +113,7 @@ namespace UnrealBuildTool
 		/// <returns>True if the parallel executor can be used</returns>
 		public static bool IsAvailable()
 		{
-			return true;
+			return RuntimePlatform.IsWindows;
 		}
 
 		/// <summary>
@@ -116,12 +122,11 @@ namespace UnrealBuildTool
 		/// <returns>True if all the tasks successfully executed, or false if any of them failed.</returns>
 		public override bool ExecuteActions(List<LinkedAction> InputActions)
 		{
-			int ActualNumParallelProcesses = Math.Min(InputActions.Count, NumParallelProcesses);
-			Log.TraceInformation("Building {0} {1} with {2} {3}...", InputActions.Count, (InputActions.Count == 1) ? "action" : "actions", ActualNumParallelProcesses, (ActualNumParallelProcesses == 1) ? "process" : "processes");
+			Log.TraceInformation("Building {0} {1} with {2} {3}...", InputActions.Count, (InputActions.Count == 1) ? "action" : "actions", NumParallelProcesses, (NumParallelProcesses == 1)? "process" : "processes");
 
 			// Create actions with all our internal metadata
 			List<BuildAction> Actions = new List<BuildAction>();
-			for (int Idx = 0; Idx < InputActions.Count; Idx++)
+			for(int Idx = 0; Idx < InputActions.Count; Idx++)
 			{
 				BuildAction Action = new BuildAction(InputActions[Idx]);
 				Action.SortIndex = Idx;
@@ -138,10 +143,10 @@ namespace UnrealBuildTool
 			Dictionary<LinkedAction, BuildAction> LinkedActionToBuildAction = Actions.ToDictionary(x => x.Inner, x => x);
 			foreach (BuildAction Action in Actions)
 			{
-				foreach (LinkedAction PrerequisiteAction in Action.Inner.PrerequisiteActions)
+				foreach(LinkedAction PrerequisiteAction in Action.Inner.PrerequisiteActions)
 				{
 					BuildAction? Dependency;
-					if (LinkedActionToBuildAction.TryGetValue(PrerequisiteAction, out Dependency))
+					if(LinkedActionToBuildAction.TryGetValue(PrerequisiteAction, out Dependency))
 					{
 						Action.Dependencies.Add(Dependency);
 						Dependency.Dependants.Add(Action);
@@ -151,7 +156,7 @@ namespace UnrealBuildTool
 
 			// Figure out the recursive dependency count
 			HashSet<BuildAction> VisitedActions = new HashSet<BuildAction>();
-			foreach (BuildAction Action in Actions)
+			foreach(BuildAction Action in Actions)
 			{
 				Action.MissingDependencyCount = Action.Dependencies.Count;
 				RecursiveIncDependents(Action, VisitedActions);
@@ -159,9 +164,9 @@ namespace UnrealBuildTool
 
 			// Create the list of things to process
 			List<BuildAction> QueuedActions = new List<BuildAction>();
-			foreach (BuildAction Action in Actions)
+			foreach(BuildAction Action in Actions)
 			{
-				if (Action.MissingDependencyCount == 0)
+				if(Action.MissingDependencyCount == 0)
 				{
 					QueuedActions.Add(Action);
 				}
@@ -172,30 +177,29 @@ namespace UnrealBuildTool
 			{
 				// Create a job object for all the child processes
 				bool bResult = true;
-				Dictionary<BuildAction, Thread> ExecutingActions = new Dictionary<BuildAction, Thread>();
+				Dictionary<BuildAction, Thread> ExecutingActions = new Dictionary<BuildAction,Thread>();
 				List<BuildAction> CompletedActions = new List<BuildAction>();
 
-				using (ManagedProcessGroup ProcessGroup = new ManagedProcessGroup())
+				using(ManagedProcessGroup ProcessGroup = new ManagedProcessGroup())
 				{
-					using (AutoResetEvent CompletedEvent = new AutoResetEvent(false))
+					using(AutoResetEvent CompletedEvent = new AutoResetEvent(false))
 					{
 						int NumCompletedActions = 0;
 						using (ProgressWriter ProgressWriter = new ProgressWriter("Compiling C++ source code...", false))
 						{
-							while (QueuedActions.Count > 0 || ExecutingActions.Count > 0)
+							while(QueuedActions.Count > 0 || ExecutingActions.Count > 0)
 							{
 								// Sort the actions by the number of things dependent on them
-								QueuedActions.Sort((A, B) => (A.TotalDependantCount == B.TotalDependantCount) ? (B.SortIndex - A.SortIndex) : (B.TotalDependantCount - A.TotalDependantCount));
+								QueuedActions.Sort((A, B) => (A.TotalDependantCount == B.TotalDependantCount)? (B.SortIndex - A.SortIndex) : (B.TotalDependantCount - A.TotalDependantCount));
 
 								// Create threads up to the maximum number of actions
-								while (ExecutingActions.Count < ActualNumParallelProcesses && QueuedActions.Count > 0)
+								while(ExecutingActions.Count < NumParallelProcesses && QueuedActions.Count > 0)
 								{
 									BuildAction Action = QueuedActions[QueuedActions.Count - 1];
 									QueuedActions.RemoveAt(QueuedActions.Count - 1);
 
 									Thread ExecutingThread = new Thread(() => { ExecuteAction(ProcessGroup, Action, CompletedActions, CompletedEvent); });
-									string Description = $"{(Action.Inner.CommandDescription != null ? Action.Inner.CommandDescription : Action.Inner.CommandPath.GetFileName())} {Action.Inner.StatusDescription}".Trim();
-									ExecutingThread.Name = String.Format("Build:{0}", Description);
+									ExecutingThread.Name = String.Format("Build:{0}", Action.Inner.StatusDescription);
 									ExecutingThread.Start();
 
 									ExecutingActions.Add(Action, ExecutingThread);
@@ -205,9 +209,9 @@ namespace UnrealBuildTool
 								CompletedEvent.WaitOne();
 
 								// Wait for something to finish and flush it to the log
-								lock (CompletedActions)
+								lock(CompletedActions)
 								{
-									foreach (BuildAction CompletedAction in CompletedActions)
+									foreach(BuildAction CompletedAction in CompletedActions)
 									{
 										// Join the thread
 										Thread CompletedThread = ExecutingActions[CompletedAction];
@@ -218,26 +222,23 @@ namespace UnrealBuildTool
 										NumCompletedActions++;
 										ProgressWriter.Write(NumCompletedActions, InputActions.Count);
 
-										string Description = $"{(CompletedAction.Inner.CommandDescription != null ? CompletedAction.Inner.CommandDescription : CompletedAction.Inner.CommandPath.GetFileName())} {CompletedAction.Inner.StatusDescription}".Trim();
-
 										// Write it to the log
-										if (CompletedAction.Inner.bShouldOutputStatusDescription)
-										{
-											Log.TraceInformation("[{0}/{1}] {2}", NumCompletedActions, InputActions.Count, Description);
-										}
-
-										foreach (string Line in CompletedAction.LogLines)
+										if(CompletedAction.LogLines.Count > 0)
 										{
-											Log.TraceInformation(Line);
+											Log.TraceInformation("[{0}/{1}] {2}", NumCompletedActions, InputActions.Count, CompletedAction.LogLines[0]);
+											for(int LineIdx = 1; LineIdx < CompletedAction.LogLines.Count; LineIdx++)
+											{
+												Log.TraceInformation("{0}", CompletedAction.LogLines[LineIdx]);
+											}
 										}
 
 										// Check the exit code
-										if (CompletedAction.ExitCode == 0)
+										if(CompletedAction.ExitCode == 0)
 										{
 											// Mark all the dependents as done
-											foreach (BuildAction DependantAction in CompletedAction.Dependants)
+											foreach(BuildAction DependantAction in CompletedAction.Dependants)
 											{
-												if (--DependantAction.MissingDependencyCount == 0)
+												if(--DependantAction.MissingDependencyCount == 0)
 												{
 													QueuedActions.Add(DependantAction);
 												}
@@ -246,15 +247,16 @@ namespace UnrealBuildTool
 										else
 										{
 											// BEGIN TEMPORARY TO CATCH PVS-STUDIO ISSUES
-											if (CompletedAction.LogLines.Count == 0)
+											if (!CompletedAction.GotOutput)
 											{
-												Log.TraceInformation("[{0}/{1}] {2} - Error but no output", NumCompletedActions, InputActions.Count, Description);
-												Log.TraceInformation("[{0}/{1}] {2} - {3} {4} {5} {6}", NumCompletedActions, InputActions.Count, Description, CompletedAction.ExitCode,
+												string StatusDescription = CompletedAction.LogLines.Count > 0 ? CompletedAction.LogLines[0] : "";
+												Log.TraceInformation("[{0}/{1}] {2} - Error but no output", NumCompletedActions, InputActions.Count, StatusDescription);
+												Log.TraceInformation("[{0}/{1}] {2} - {3} {4} {5} {6}", NumCompletedActions, InputActions.Count, StatusDescription, CompletedAction.ExitCode, 
 													CompletedAction.Inner.WorkingDirectory, CompletedAction.Inner.CommandPath, CompletedAction.Inner.CommandArguments);
 											}
 											// END TEMPORARY
 											// Update the exit code if it's not already set
-											if (bResult && CompletedAction.ExitCode != 0)
+											if(bResult && CompletedAction.ExitCode != 0)
 											{
 												bResult = false;
 											}
@@ -264,7 +266,7 @@ namespace UnrealBuildTool
 								}
 
 								// If we've already got a non-zero exit code, clear out the list of queued actions so nothing else will run
-								if (!bResult && bStopCompilationAfterErrors)
+								if(!bResult && bStopCompilationAfterErrors)
 								{
 									QueuedActions.Clear();
 								}
@@ -286,6 +288,14 @@ namespace UnrealBuildTool
 		/// <param name="CompletedEvent">Event to set once an event is complete</param>
 		static void ExecuteAction(ManagedProcessGroup ProcessGroup, BuildAction Action, List<BuildAction> CompletedActions, AutoResetEvent CompletedEvent)
 		{
+			if (Action.Inner.bShouldOutputStatusDescription && !String.IsNullOrEmpty(Action.Inner.StatusDescription))
+			{
+				Action.LogLines.Add(Action.Inner.StatusDescription);
+			}
+
+			// BEGIN TEMPORARY TO CATCH PVS-STUDIO ISSUES
+			int LogLineCount = Action.LogLines.Count;
+			// END TEMPORARY
 			try
 			{
 				using (ManagedProcess Process = new ManagedProcess(ProcessGroup, Action.Inner.CommandPath.FullName, Action.Inner.CommandArguments, Action.Inner.WorkingDirectory.FullName, null, null, ProcessPriorityClass.BelowNormal))
@@ -293,14 +303,17 @@ namespace UnrealBuildTool
 					Action.LogLines.AddRange(Process.ReadAllLines());
 					Action.ExitCode = Process.ExitCode;
 				}
+				// BEGIN TEMPORARY TO CATCH PVS-STUDIO ISSUES
+				Action.GotOutput = Action.LogLines.Count > LogLineCount;
+				// END TEMPORARY
 			}
-			catch (Exception Ex)
+			catch(Exception Ex)
 			{
 				Log.WriteException(Ex, null);
 				Action.ExitCode = 1;
 			}
 
-			lock (CompletedActions)
+			lock(CompletedActions)
 			{
 				CompletedActions.Add(Action);
 			}
@@ -315,9 +328,9 @@ namespace UnrealBuildTool
 		/// <param name="VisitedActions">Set of visited actions</param>
 		private static void RecursiveIncDependents(BuildAction Action, HashSet<BuildAction> VisitedActions)
 		{
-			foreach (BuildAction Dependency in Action.Dependants)
+			foreach(BuildAction Dependency in Action.Dependants)
 			{
-				if (!VisitedActions.Contains(Action))
+				if(!VisitedActions.Contains(Action))
 				{
 					VisitedActions.Add(Action);
 					Dependency.TotalDependantCount++;
diff --git a/Engine/Source/Programs/UnrealBuildTool/System/ActionGraph.cs b/Engine/Source/Programs/UnrealBuildTool/System/ActionGraph.cs
--- a/Engine/Source/Programs/UnrealBuildTool/System/ActionGraph.cs
+++ b/Engine/Source/Programs/UnrealBuildTool/System/ActionGraph.cs
@@ -337,16 +337,20 @@ namespace UnrealBuildTool
 				}
 				else if (BuildConfiguration.bAllowFASTBuild && FASTBuild.IsAvailable())
 				{
-					Executor = new FASTBuild(BuildConfiguration.MaxParallelActions);
+					Executor = new FASTBuild();
 				}
 				else if(BuildConfiguration.bAllowSNDBS && SNDBS.IsAvailable())
 				{
 					Executor = new SNDBS();
 				}
-				else
+				else if(BuildConfiguration.bAllowParallelExecutor && ParallelExecutor.IsAvailable())
 				{
 					Executor = new ParallelExecutor(BuildConfiguration.MaxParallelActions);
 				}
+				else
+				{
+					Executor = new LocalExecutor(BuildConfiguration.MaxParallelActions);
+				}
 
 				// Execute the build
 				Stopwatch Timer = Stopwatch.StartNew();
