UnrealBuildFiles.csproj

Generate UnrealBuildFiles.csproj, a project that contains references to all of the Build.cs and Target.cs files in the solution, and add it to the generated solution.

With this project in the solution, Intellisense is available when editing these files, and the project may be compiled in Visual Studio, providing representative results of what UBT will do with the files (UBT will re-compile the files when it needs to)

Using Visual Studio [REDACTED_PHONE], adding this project increased memory usage by ~20MB after load (from 540MB to 560MB) for a workspace containing Engine, and two sample games.

diff --git a/Engine/Source/Programs/UnrealBuildTool/ProjectFiles/ProjectFileGenerator.cs b/Engine/Source/Programs/UnrealBuildTool/ProjectFiles/ProjectFileGenerator.cs
--- a/Engine/Source/Programs/UnrealBuildTool/ProjectFiles/ProjectFileGenerator.cs
+++ b/Engine/Source/Programs/UnrealBuildTool/ProjectFiles/ProjectFileGenerator.cs
@@ -560,6 +560,64 @@ namespace UnrealBuildTool
 			}
 		}
 
+		private void AddModuleFilesProject(List<FileReference> AllModuleFiles, List<FileReference> AllTargetFiles)
+		{
+			DirectoryReference ProjectFilesDirectory = DirectoryReference.Combine(UnrealBuild.EngineDirectory, "Intermediate", "ProjectFiles");
+			FileReference ModuleFilesProjectLocation = FileReference.Combine(ProjectFilesDirectory, "UnrealBuildFiles.csproj");
+
+			{
+				using FileStream Stream = FileReference.Open(ModuleFilesProjectLocation, FileMode.Create, FileAccess.Write, FileShare.Read);
+				using StreamWriter Writer = new StreamWriter(Stream, Encoding.UTF8);
+
+				Writer.WriteLine("<Project Sdk=\"Microsoft.NET.Sdk\">");
+
+				Writer.WriteLine("  <PropertyGroup>");
+				Writer.WriteLine("    <TargetFramework>netcoreapp3.1</TargetFramework>");
+				Writer.WriteLine("    <Configurations>Debug;Release;Development</Configurations>"); // VCSharpProject requires at least Debug & Development configurations
+				Writer.WriteLine("    <EnableDefaultItems>false</EnableDefaultItems>"); // Without this, the contents of Engine/Intermediate/ProjectFiles will be automatically included in the project
+				Writer.WriteLine("  </PropertyGroup>");
+				
+				Writer.WriteLine("  <ItemGroup>");
+				Writer.WriteLine("    <ProjectReference Include=\"../../Source/Programs/DotNETCommon/BuildUtilities/BuildUtilities.csproj\" />");
+				Writer.WriteLine("    <ProjectReference Include=\"../../Source/Programs/UnrealBuildTool/UnrealBuildTool.csproj\" />");
+				Writer.WriteLine("  </ItemGroup>");
+
+				Writer.WriteLine("  <ItemGroup>");
+				foreach (FileReference ModuleFile in AllModuleFiles)
+				{
+					string CsprojRelativePath = ModuleFile.MakeRelativeTo(ProjectFilesDirectory);
+					string ProjectFolder = ModuleFile.MakeRelativeTo(UnrealBuild.EngineDirectory);
+					if (!ModuleFile.IsUnderDirectory(UnrealBuild.EngineDirectory))
+					{
+						while(ProjectFolder.StartsWith($"..{Path.DirectorySeparatorChar}"))
+						{
+							ProjectFolder = ProjectFolder[3..];
+						}
+					}
+					Writer.WriteLine($"  <Compile Include=\"{CsprojRelativePath}\"><Link>{ProjectFolder}</Link></Compile>");
+				}
+				foreach (FileReference ModuleFile in AllTargetFiles)
+				{
+					string CsprojRelativePath = ModuleFile.MakeRelativeTo(ProjectFilesDirectory);
+					string ProjectFolder = ModuleFile.MakeRelativeTo(UnrealBuild.EngineDirectory);
+					if (!ModuleFile.IsUnderDirectory(UnrealBuild.EngineDirectory))
+					{
+						while(ProjectFolder.StartsWith($"..{Path.DirectorySeparatorChar}"))
+						{
+							ProjectFolder = ProjectFolder[3..];
+						}
+					}
+					Writer.WriteLine($"  <Compile Include=\"{CsprojRelativePath}\"><Link>{ProjectFolder}</Link></Compile>");
+				}
+				Writer.WriteLine("  </ItemGroup>");
+
+				Writer.WriteLine("</Project>");
+			}
+
+			VCSharpProjectFile ModuleFilesProject = new VCSharpProjectFile(ModuleFilesProjectLocation);
+			AddExistingProjectFile(ModuleFilesProject, bForceDevelopmentConfiguration: false);
+		}
+
 		private static void DiscoverCSharpProgramProjectsRecursively(DirectoryReference SearchFolder, List<FileReference> FoundProjects)
 		{
 			// Scan all the files in this directory
@@ -986,6 +1044,9 @@ namespace UnrealBuildTool
 
 					// Discover C# programs which should additionally be included in the solution.
 					DiscoverCSharpProgramProjects(AllGameProjects, ProgramsFolder);
+
+					// Add projects for Build.cs and Target.cs files
+					AddModuleFilesProject(AllModuleFiles, AllTargetFiles);
 				}
 
 
