UnrealBuildTool/UnrealBuildFiles.csproj:

Create project in its own directory to avoid confusing other builds with UnrealBuildFiles' intermediate files

diff --git a/Engine/Source/Programs/UnrealBuildTool/ProjectFiles/ProjectFileGenerator.cs b/Engine/Source/Programs/UnrealBuildTool/ProjectFiles/ProjectFileGenerator.cs
--- a/Engine/Source/Programs/UnrealBuildTool/ProjectFiles/ProjectFileGenerator.cs
+++ b/Engine/Source/Programs/UnrealBuildTool/ProjectFiles/ProjectFileGenerator.cs
@@ -562,7 +562,8 @@ namespace UnrealBuildTool
 
 		private void AddModuleFilesProject(List<FileReference> AllModuleFiles, List<FileReference> AllTargetFiles)
 		{
-			DirectoryReference ProjectFilesDirectory = DirectoryReference.Combine(UnrealBuild.EngineDirectory, "Intermediate", "ProjectFiles");
+			DirectoryReference ProjectFilesDirectory = DirectoryReference.Combine(UnrealBuild.EngineDirectory, "Intermediate", "Build", "UnrealBuildFiles");
+			DirectoryReference.CreateDirectory(ProjectFilesDirectory);
 			FileReference ModuleFilesProjectLocation = FileReference.Combine(ProjectFilesDirectory, "UnrealBuildFiles.csproj");
 
 			{
@@ -578,8 +579,8 @@ namespace UnrealBuildTool
 				Writer.WriteLine("  </PropertyGroup>");
 				
 				Writer.WriteLine("  <ItemGroup>");
-				Writer.WriteLine("    <ProjectReference Include=\"../../Source/Programs/DotNETCommon/BuildUtilities/BuildUtilities.csproj\" />");
-				Writer.WriteLine("    <ProjectReference Include=\"../../Source/Programs/UnrealBuildTool/UnrealBuildTool.csproj\" />");
+				Writer.WriteLine("    <ProjectReference Include=\"../../../Source/Programs/DotNETCommon/BuildUtilities/BuildUtilities.csproj\" />");
+				Writer.WriteLine("    <ProjectReference Include=\"../../../Source/Programs/UnrealBuildTool/UnrealBuildTool.csproj\" />");
 				Writer.WriteLine("  </ItemGroup>");
 
 				Writer.WriteLine("  <ItemGroup>");
