ProjectFileGenerator.cs:

Remove redundant list copy

#trivial

diff --git a/Engine/Source/Programs/UnrealBuildTool/ProjectFiles/ProjectFileGenerator.cs b/Engine/Source/Programs/UnrealBuildTool/ProjectFiles/ProjectFileGenerator.cs
--- a/Engine/Source/Programs/UnrealBuildTool/ProjectFiles/ProjectFileGenerator.cs
+++ b/Engine/Source/Programs/UnrealBuildTool/ProjectFiles/ProjectFileGenerator.cs
@@ -1462,15 +1462,8 @@ namespace UnrealBuildTool
 		/// <returns>Filtered list of module files</returns>
 		protected List<FileReference> DiscoverModules(List<FileReference> AllGameProjects)
 		{
-			List<FileReference> AllModuleFiles = new List<FileReference>();
-
 			// Locate all modules (*.Build.cs files)
-			List<FileReference> FoundModuleFiles = RulesCompiler.FindAllRulesSourceFiles(RulesCompiler.RulesFileType.Module, GameFolders: AllGameProjects.Select(x => x.Directory).ToList(), ForeignPlugins: null, AdditionalSearchPaths: null);
-			foreach (FileReference BuildFileName in FoundModuleFiles)
-			{
-				AllModuleFiles.Add(BuildFileName);
-			}
-			return AllModuleFiles;
+			return RulesCompiler.FindAllRulesSourceFiles(RulesCompiler.RulesFileType.Module, GameFolders: AllGameProjects.Select(x => x.Directory).ToList(), ForeignPlugins: null, AdditionalSearchPaths: null);
 		}
 
 		/// <summary>
