fix for always-enabled default child plugin whitelists.

#b ben.marsh
#rnx

diff --git a/Engine/Source/Programs/UnrealBuildTool/Configuration/UEBuildTarget.cs b/Engine/Source/Programs/UnrealBuildTool/Configuration/UEBuildTarget.cs
--- a/Engine/Source/Programs/UnrealBuildTool/Configuration/UEBuildTarget.cs
+++ b/Engine/Source/Programs/UnrealBuildTool/Configuration/UEBuildTarget.cs
@@ -3118,6 +3118,11 @@ namespace UnrealBuildTool
 						PluginReference.bHasExplicitPlatforms = Plugin.Descriptor.bHasExplicitPlatforms;
 						PluginReference.bOptional = true;
 
+						if (PluginReference.bHasExplicitPlatforms)
+						{
+							PluginReference.WhitelistPlatforms = PluginReference.SupportedTargetPlatforms; //synthesize whitelist if it must be explicit
+						}
+
 						AddPlugin(PluginReference, "default plugins", ExcludeFolders, NameToInstance, NameToInfo);
 					}
 				}
