UnrealBuildTool:

Per-project rules assembly csprojen.

Previously: all .build.cs and .target.cs were being grouped into a single project, which resulted in duplicate symbols when trying to build the project.

Now: multiple assemblies csprojs are generated, one for each Engine & Game project. Also, includes preprocessor defines.

diff --git a/Engine/Source/Programs/UnrealBuildTool/ProjectFiles/ProjectFileGenerator.cs b/Engine/Source/Programs/UnrealBuildTool/ProjectFiles/ProjectFileGenerator.cs
--- a/Engine/Source/Programs/UnrealBuildTool/ProjectFiles/ProjectFileGenerator.cs
+++ b/Engine/Source/Programs/UnrealBuildTool/ProjectFiles/ProjectFileGenerator.cs
@@ -560,11 +560,11 @@ namespace UnrealBuildTool
 			}
 		}
 
-		private void AddModuleFilesProject(List<FileReference> AllModuleFiles, List<FileReference> AllTargetFiles)
+		private VCSharpProjectFile CreateRulesAssemblyProject(RulesAssembly RulesAssembly, DirectoryReference FSPathBase)
 		{
-			DirectoryReference ProjectFilesDirectory = DirectoryReference.Combine(Unreal.EngineDirectory, "Intermediate", "Build", "UnrealBuildFiles");
+			DirectoryReference ProjectFilesDirectory = DirectoryReference.Combine(FSPathBase, "Intermediate/Build/BuildRulesProjects", RulesAssembly.GetSimpleAssemblyName());
 			DirectoryReference.CreateDirectory(ProjectFilesDirectory);
-			FileReference ModuleFilesProjectLocation = FileReference.Combine(ProjectFilesDirectory, "UnrealBuildFiles.csproj");
+			FileReference ModuleFilesProjectLocation = FileReference.Combine(ProjectFilesDirectory, RulesAssembly.GetSimpleAssemblyName() + ".csproj");
 
 			{
 				using FileStream Stream = FileReference.Open(ModuleFilesProjectLocation, FileMode.Create, FileAccess.Write, FileShare.Read);
@@ -574,40 +574,25 @@ namespace UnrealBuildTool
 
 				Writer.WriteLine("  <PropertyGroup>");
 				Writer.WriteLine("    <TargetFramework>netcoreapp3.1</TargetFramework>");
+				Writer.WriteLine("    <AppendTargetFrameworkToOutputPath>false</AppendTargetFrameworkToOutputPath>"); // Shorten intermediate filepath slightly
 				Writer.WriteLine("    <Configurations>Debug;Release;Development</Configurations>"); // VCSharpProject requires at least Debug & Development configurations
-				Writer.WriteLine("    <EnableDefaultItems>false</EnableDefaultItems>"); // Without this, the contents of Engine/Intermediate/ProjectFiles will be automatically included in the project
+				Writer.WriteLine("    <DefineConstants>$(DefineConstants);" + String.Join(';', RulesAssembly.PreprocessorDefines) + "</DefineConstants>");
 				Writer.WriteLine("  </PropertyGroup>");
 				
 				Writer.WriteLine("  <ItemGroup>");
-				Writer.WriteLine("    <ProjectReference Include=\"../../../Source/Programs/DotNETCommon/BuildUtilities/BuildUtilities.csproj\" />");
-				Writer.WriteLine("    <ProjectReference Include=\"../../../Source/Programs/UnrealBuildTool/UnrealBuildTool.csproj\" />");
+				Writer.WriteLine("    <ProjectReference Include=\"" + 
+					FileReference.Combine(Unreal.EngineDirectory, "Source/Programs/DotNETCommon/BuildUtilities/BuildUtilities.csproj").MakeRelativeTo(ProjectFilesDirectory) + 
+					"\"><Private>false</Private></ProjectReference>"); // <Private>false</Private> to suppress copying of dependencies into output directory
+				Writer.WriteLine("    <ProjectReference Include=\"" + 
+					FileReference.Combine(Unreal.EngineDirectory, "Source", "Programs", "UnrealBuildTool", "UnrealBuildTool.csproj").MakeRelativeTo(ProjectFilesDirectory) + 
+					"\"><Private>false</Private></ProjectReference>");
 				Writer.WriteLine("  </ItemGroup>");
 
 				Writer.WriteLine("  <ItemGroup>");
-				foreach (FileReference ModuleFile in AllModuleFiles)
+				foreach (FileReference ModuleFile in RulesAssembly.AssemblySourceFiles)
 				{
 					string CsprojRelativePath = ModuleFile.MakeRelativeTo(ProjectFilesDirectory);
-					string ProjectFolder = ModuleFile.MakeRelativeTo(Unreal.EngineDirectory);
-					if (!ModuleFile.IsUnderDirectory(Unreal.EngineDirectory))
-					{
-						while(ProjectFolder.StartsWith($"..{Path.DirectorySeparatorChar}"))
-						{
-							ProjectFolder = ProjectFolder[3..];
-						}
-					}
-					Writer.WriteLine($"  <Compile Include=\"{CsprojRelativePath}\"><Link>{ProjectFolder}</Link></Compile>");
-				}
-				foreach (FileReference ModuleFile in AllTargetFiles)
-				{
-					string CsprojRelativePath = ModuleFile.MakeRelativeTo(ProjectFilesDirectory);
-					string ProjectFolder = ModuleFile.MakeRelativeTo(Unreal.EngineDirectory);
-					if (!ModuleFile.IsUnderDirectory(Unreal.EngineDirectory))
-					{
-						while(ProjectFolder.StartsWith($"..{Path.DirectorySeparatorChar}"))
-						{
-							ProjectFolder = ProjectFolder[3..];
-						}
-					}
+					string ProjectFolder = ModuleFile.MakeRelativeTo(FSPathBase);
 					Writer.WriteLine($"  <Compile Include=\"{CsprojRelativePath}\"><Link>{ProjectFolder}</Link></Compile>");
 				}
 				Writer.WriteLine("  </ItemGroup>");
@@ -616,7 +601,7 @@ namespace UnrealBuildTool
 			}
 
 			VCSharpProjectFile ModuleFilesProject = new VCSharpProjectFile(ModuleFilesProjectLocation);
-			AddExistingProjectFile(ModuleFilesProject, bForceDevelopmentConfiguration: false);
+			return ModuleFilesProject;
 		}
 
 		private static void DiscoverCSharpProgramProjectsRecursively(DirectoryReference SearchFolder, List<FileReference> FoundProjects)
@@ -856,9 +841,10 @@ namespace UnrealBuildTool
 			List<ProjectFile> GameProjects;
 			List<ProjectFile> ModProjects;
 			Dictionary<FileReference, ProjectFile> ProgramProjects;
+			Dictionary<RulesAssembly, DirectoryReference> RulesAssemblies;
 			{
 				// Setup buildable projects for all targets
-				AddProjectsForAllTargets(PlatformProjectGenerators, AllGameProjects, AllTargetFiles, Arguments, out EngineProject, out GameProjects, out ProgramProjects);
+				AddProjectsForAllTargets(PlatformProjectGenerators, AllGameProjects, AllTargetFiles, Arguments, out EngineProject, out GameProjects, out ProgramProjects, out RulesAssemblies);
 
 				// Add projects for mods
 				AddProjectsForMods(GameProjects, out ModProjects);
@@ -1046,8 +1032,12 @@ namespace UnrealBuildTool
 					// Discover C# programs which should additionally be included in the solution.
 					DiscoverCSharpProgramProjects(AllGameProjects, ProgramsFolder);
 
-					// Add projects for Build.cs and Target.cs files
-					AddModuleFilesProject(AllModuleFiles, AllTargetFiles);
+					foreach (KeyValuePair<RulesAssembly, DirectoryReference> RulesAssemblyEntry in RulesAssemblies)
+					{
+						VCSharpProjectFile RulesProject = CreateRulesAssemblyProject(RulesAssemblyEntry.Key, RulesAssemblyEntry.Value);
+						AddExistingProjectFile(RulesProject, bForceDevelopmentConfiguration: false);
+						RootFolder.AddSubFolder("Rules").ChildProjects.Add(RulesProject);
+					}
 				}
 
 
@@ -2236,6 +2226,7 @@ namespace UnrealBuildTool
 		/// <param name="EngineProject">The engine project we created</param>
 		/// <param name="GameProjects">Map of game folder name to all of the game projects we created</param>
 		/// <param name="ProgramProjects">Map of program names to all of the program projects we created</param>
+		/// <param name="RulesAssemblies">Map of RuleAssemblies to their base folders</param>
 		private void AddProjectsForAllTargets(
 			PlatformProjectGeneratorCollection PlatformProjectGenerators,
 			List<FileReference> AllGames,
@@ -2243,12 +2234,14 @@ namespace UnrealBuildTool
 			String[] Arguments,
 			out ProjectFile EngineProject,
 			out List<ProjectFile> GameProjects,
-			out Dictionary<FileReference, ProjectFile> ProgramProjects)
+			out Dictionary<FileReference, ProjectFile> ProgramProjects,
+			out Dictionary<RulesAssembly, DirectoryReference> RulesAssemblies)
 		{
 			// As we're creating project files, we'll also keep track of whether we created an "engine" project and return that if we have one
 			EngineProject = null;
 			GameProjects = new List<ProjectFile>();
 			ProgramProjects = new Dictionary<FileReference, ProjectFile>();
+			RulesAssemblies = new Dictionary<RulesAssembly, DirectoryReference>();
 
 			// Separate the .target.cs files that are platform extension specializations, per target name. These will be added alongside their base target.cs
 			Dictionary<string, List<FileReference>> AllSubTargetFilesPerTarget = new Dictionary<string, List<FileReference>>();
@@ -2305,10 +2298,21 @@ namespace UnrealBuildTool
 					if (CheckProjectFile == null)
 					{
 						RulesAssembly = RulesCompiler.CreateEngineRulesAssembly(false, false);
+					
+						// Record the Engine assembly, and any parent assemblies (varies e.g. Rules, ProgramRules, MarketplaceRules)
+						if (!RulesAssemblies.ContainsKey(RulesAssembly))
+						{
+							RulesAssembly NextAssembly = RulesAssembly;
+							do
+							{
+								RulesAssemblies.Add(NextAssembly, Unreal.EngineDirectory);
+							} while ((NextAssembly = NextAssembly.Parent) != null);
+						}
 					}
 					else
 					{
 						RulesAssembly = RulesCompiler.CreateProjectRulesAssembly(CheckProjectFile, false, false);
+						// Recording of game rule assemblies happens after BaseFolder has been computed
 					}
 
 					// Create target rules for all of the platforms and configuration combinations that we want to enable support for.
@@ -2413,6 +2417,8 @@ namespace UnrealBuildTool
 					else
 					{
 						BaseFolder = GameFolder;
+
+						RulesAssemblies.TryAdd(RulesAssembly, BaseFolder);
 					}
 
 					bool bProjectAlreadyExisted;
diff --git a/Engine/Source/Programs/UnrealBuildTool/System/RulesAssembly.cs b/Engine/Source/Programs/UnrealBuildTool/System/RulesAssembly.cs
--- a/Engine/Source/Programs/UnrealBuildTool/System/RulesAssembly.cs
+++ b/Engine/Source/Programs/UnrealBuildTool/System/RulesAssembly.cs
@@ -29,6 +29,19 @@ namespace UnrealBuildTool
 		/// </summary>
 		private Assembly? CompiledAssembly;
 
+		/// <summary>
+		/// Returns the simple name of the assembly e.g. "UE5ProgramRules"
+		/// </summary>
+		/// <returns></returns>
+		public string? GetSimpleAssemblyName()
+		{
+			if (CompiledAssembly != null)
+			{
+				return CompiledAssembly.GetName().Name;
+			}
+			else return null;
+		}
+
 		/// <summary>
 		/// The base directories for this assembly
 		/// </summary>
@@ -73,7 +86,17 @@ namespace UnrealBuildTool
 		/// <summary>
 		/// The parent rules assembly that this assembly inherits. Game assemblies inherit the engine assembly, and the engine assembly inherits nothing.
 		/// </summary>
-		private RulesAssembly? Parent;
+		public RulesAssembly? Parent { get; }
+
+		/// <summary>
+		/// The set of files that were compiled to create this assembly
+		/// </summary>
+		public HashSet<FileReference>? AssemblySourceFiles { get; }
+
+		/// <summary>
+		/// Any preprocessor defines that were set when this assembly was created
+		/// </summary>
+		public List<string>? PreprocessorDefines { get;  }
 
 		/// <summary>
 		/// Constructor. Compiles a rules assembly from the given source files.
@@ -101,14 +124,14 @@ namespace UnrealBuildTool
 			this.Parent = Parent;
 
 			// Find all the source files
-			HashSet<FileReference> AssemblySourceFiles = new HashSet<FileReference>();
+			AssemblySourceFiles = new HashSet<FileReference>();
 			AssemblySourceFiles.UnionWith(ModuleFileToContext.Keys);
 			AssemblySourceFiles.UnionWith(TargetFiles);
 
 			// Compile the assembly
 			if (AssemblySourceFiles.Count > 0)
 			{
-				List<string> PreprocessorDefines = GetPreprocessorDefinitions();
+				PreprocessorDefines = GetPreprocessorDefinitions();
 				CompiledAssembly = DynamicCompilation.CompileAndLoadAssembly(AssemblyFileName, AssemblySourceFiles, PreprocessorDefines: PreprocessorDefines, DoNotCompile: bSkipCompile);
 			}
 
