UnrealBuildTool: added CppCompileEnvironment.AnalyzeStackSizeWarning to replace hardcoded value of 80K (now defaults to 300K).
Our stack size defaults to 4.7M on Windows, so allowing functions to use more stack space in favor of going through the heap.

diff --git a/Engine/Source/Programs/UnrealBuildTool/Platform/Windows/VCToolChain.cs b/Engine/Source/Programs/UnrealBuildTool/Platform/Windows/VCToolChain.cs
--- a/Engine/Source/Programs/UnrealBuildTool/Platform/Windows/VCToolChain.cs
+++ b/Engine/Source/Programs/UnrealBuildTool/Platform/Windows/VCToolChain.cs
@@ -229,9 +229,9 @@ namespace UnrealBuildTool
 				// Don't cause analyze warnings to be errors
 				Arguments.Add("/analyze:WX-");
 
-				// Report functions that use a LOT of stack space.  You can lower this value if you
+				// Report functions that use a LOT of stack space. You can lower this value if you
 				// want more aggressive checking for functions that use a lot of stack memory.
-				Arguments.Add("/analyze:stacksize81940");
+				Arguments.Add("/analyze:stacksize" + CompileEnvironment.AnalyzeStackSizeWarning);
 
 				// Don't bother generating code, only analyze code (may report fewer warnings though.)
 				//Arguments.Add("/analyze:only");
diff --git a/Engine/Source/Programs/UnrealBuildTool/System/CppCompileEnvironment.cs b/Engine/Source/Programs/UnrealBuildTool/System/CppCompileEnvironment.cs
--- a/Engine/Source/Programs/UnrealBuildTool/System/CppCompileEnvironment.cs
+++ b/Engine/Source/Programs/UnrealBuildTool/System/CppCompileEnvironment.cs
@@ -398,6 +398,11 @@ namespace UnrealBuildTool
 		/// </summary>
 		public CppStandardVersion CppStandard = CppStandardVersion.Default;
 
+		/// <summary>
+		/// The amount of the stack usage to report static analysis warnings.
+		/// </summary>
+		public int AnalyzeStackSizeWarning = 300000;
+
 		/// <summary>
 		/// Default constructor.
 		/// </summary>
