Fix AArch64 -> Arm64

diff --git a/Engine/Source/Programs/UnrealBuildTool/Configuration/ModuleRules.cs b/Engine/Source/Programs/UnrealBuildTool/Configuration/ModuleRules.cs
--- a/Engine/Source/Programs/UnrealBuildTool/Configuration/ModuleRules.cs
+++ b/Engine/Source/Programs/UnrealBuildTool/Configuration/ModuleRules.cs
@@ -1491,7 +1491,7 @@ namespace UnrealBuildTool
 			{
 				return Target.Platform == UnrealTargetPlatform.Win64 ||
 					Target.Platform == UnrealTargetPlatform.Linux ||
-					Target.Platform == UnrealTargetPlatform.LinuxAArch64;
+					Target.Platform == UnrealTargetPlatform.LinuxArm64;
 			}
 		}
 
@@ -1527,7 +1527,7 @@ namespace UnrealBuildTool
 				Architecture = "x86_64";
 				Platform = "unknown-linux-gnu";
 			}
-			else if (Target.Platform == UnrealTargetPlatform.LinuxAArch64)
+			else if (Target.Platform == UnrealTargetPlatform.LinuxArm64)
 			{
 				Architecture = "aarch64";
 				Platform = "unknown-linux-gnueabi";
@@ -1563,7 +1563,7 @@ namespace UnrealBuildTool
 			{
 				LibraryExtension = ".lib";
 			}
-			else if (Target.Platform == UnrealTargetPlatform.Linux || Target.Platform == UnrealTargetPlatform.LinuxAArch64)
+			else if (Target.Platform == UnrealTargetPlatform.Linux || Target.Platform == UnrealTargetPlatform.LinuxArm64)
 			{
 				LibraryExtension = ".a";
 			}
@@ -1571,7 +1571,7 @@ namespace UnrealBuildTool
 			foreach (string Library in Libraries)
 			{
 				string LibraryPath = Path.Combine(VcPackageRoot, "lib", $"{Library}{LibraryExtension}");
-				if ((Target.Platform == UnrealTargetPlatform.Linux || Target.Platform == UnrealTargetPlatform.LinuxAArch64) && !Library.StartsWith("lib"))
+				if ((Target.Platform == UnrealTargetPlatform.Linux || Target.Platform == UnrealTargetPlatform.LinuxArm64) && !Library.StartsWith("lib"))
 				{
 					LibraryPath = Path.Combine(VcPackageRoot, "lib", $"lib{Library}{LibraryExtension}");
 				}
