Remove debug logging

diff --git a/Engine/Source/Programs/UnrealBuildTool/Configuration/UEBuildTarget.cs b/Engine/Source/Programs/UnrealBuildTool/Configuration/UEBuildTarget.cs
--- a/Engine/Source/Programs/UnrealBuildTool/Configuration/UEBuildTarget.cs
+++ b/Engine/Source/Programs/UnrealBuildTool/Configuration/UEBuildTarget.cs
@@ -1064,9 +1064,6 @@ namespace UnrealBuildTool
 
 			bCompileMonolithic = (Rules.LinkType == TargetLinkType.Monolithic);
 
-			// make UnrealEditor-ASan
-			// make UnrealEditor
-
 			// Set the build environment
 			bUseSharedBuildEnvironment = (Rules.BuildEnvironment == TargetBuildEnvironment.Shared);
 			if (bUseSharedBuildEnvironment)
@@ -1091,8 +1088,6 @@ namespace UnrealBuildTool
 				ProjectDirectory = Unreal.EngineDirectory;
 			}
 
-			System.Console.WriteLine("THE REAL Name: {0} {1} -- {2}", TargetName, AppName, bUseSharedBuildEnvironment);
-
 			// Build the project intermediate directory
 			ProjectIntermediateDirectory = DirectoryReference.Combine(ProjectDirectory, PlatformIntermediateFolder, TargetName, Configuration.ToString());
 
diff --git a/Engine/Source/Programs/UnrealBuildTool/Platform/Linux/UEBuildLinux.cs b/Engine/Source/Programs/UnrealBuildTool/Platform/Linux/UEBuildLinux.cs
--- a/Engine/Source/Programs/UnrealBuildTool/Platform/Linux/UEBuildLinux.cs
+++ b/Engine/Source/Programs/UnrealBuildTool/Platform/Linux/UEBuildLinux.cs
@@ -223,12 +223,6 @@ namespace UnrealBuildTool
 				{
 					Target.Name = Target.Name + "-" + SanitizerSuffix;
 				}
-
-				System.Console.WriteLine("New - Name: {0}", Target.Name);
-			}
-			else
-			{
-				System.Console.WriteLine("Cannot Name: {0}", Target.Name);
 			}
 
 			if (Target.bAllowLTCG && Target.LinkType != TargetLinkType.Monolithic)
