UnrealBuildTool: Action Executor that uses async tasks instead of background threads (currently disabled)

diff --git a/Engine/Source/Programs/UnrealBuildTool/Configuration/BuildConfiguration.cs b/Engine/Source/Programs/UnrealBuildTool/Configuration/BuildConfiguration.cs
--- a/Engine/Source/Programs/UnrealBuildTool/Configuration/BuildConfiguration.cs
+++ b/Engine/Source/Programs/UnrealBuildTool/Configuration/BuildConfiguration.cs
@@ -62,6 +62,12 @@ namespace UnrealBuildTool
 		[XmlConfigFile]
 		public bool bAllowSNDBS = true;
 
+		/// <summary>
+		/// Whether the experimental async TaskExecutor may be used.
+		/// </summary>
+		[XmlConfigFile]
+		public bool bAllowTaskExecutor = false;
+
 		/// <summary>
 		/// Enables support for very fast iterative builds by caching target data. Turning this on causes Unreal Build Tool to emit
 		/// 'UBT Makefiles' for targets when they are built the first time. Subsequent builds will load these Makefiles and begin
diff --git a/Engine/Source/Programs/UnrealBuildTool/Executors/Experimental/TaskExecutor.cs b/Engine/Source/Programs/UnrealBuildTool/Executors/Experimental/TaskExecutor.cs
new file mode 100644
--- /dev/null
+++ b/Engine/Source/Programs/UnrealBuildTool/Executors/Experimental/TaskExecutor.cs
@@ -0,0 +1,264 @@
+// Copyright Epic Games, Inc. All Rights Reserved.
+
+using EpicGames.Core;
+using System;
+using System.Collections.Generic;
+using System.Diagnostics;
+using System.IO;
+using System.Linq;
+using System.Threading;
+using System.Threading.Tasks;
+
+namespace UnrealBuildTool
+{
+	/// <summary>
+	/// This executor is similar to LocalExecutor, but uses p/invoke on Windows to ensure that child processes are started at a lower priority and are terminated when the parent process terminates.
+	/// </summary>
+	class TaskExecutor : ActionExecutor
+	{
+		/// <summary>
+		/// Maximum processor count for local execution. 
+		/// </summary>
+		[XmlConfigFile]
+		public static int MaxProcessorCount = int.MaxValue;
+
+		/// <summary>
+		/// Processor count multiplier for local execution. Can be below 1 to reserve CPU for other tasks.
+		/// When using the local executor (not XGE), run a single action on each CPU core. Note that you can set this to a larger value
+		/// to get slightly faster build times in many cases, but your computer's responsiveness during compiling may be much worse.
+		/// This value is ignored if the CPU does not support hyper-threading.
+		/// </summary>
+		[XmlConfigFile]
+		public static double ProcessorCountMultiplier = 1.0;
+
+		/// <summary>
+		/// Free memory per action in bytes, used to limit the number of parallel actions if the machine is memory starved.
+		/// Set to 0 to disable free memory checking.
+		/// </summary>
+		[XmlConfigFile]
+		static double MemoryPerActionBytes = 1.5 * 1024 * 1024 * 1024;
+
+		/// <summary>
+		/// When enabled, will stop compiling targets after a compile error occurs.
+		/// </summary>
+		[XmlConfigFile]
+		bool bStopCompilationAfterErrors = false;
+
+		/// <summary>
+		/// How many processes that will be executed in parallel
+		/// </summary>
+		public int NumParallelProcesses { get; private set; }
+
+		private static char[] LineEndingSplit = new char[] { '\n', '\r' };
+
+		public static int GetDefaultNumParallelProcesses()
+		{
+			return Utils.GetMaxActionsToExecuteInParallel(MaxProcessorCount, ProcessorCountMultiplier, Convert.ToInt64(MemoryPerActionBytes));
+		}
+
+		/// <summary>
+		/// Constructor
+		/// </summary>
+		/// <param name="MaxLocalActions">How many actions to execute in parallel</param>
+		public TaskExecutor(int MaxLocalActions)
+		{
+			XmlConfig.ApplyTo(this);
+
+			// if specified this caps how many processors we can use
+			if (MaxLocalActions > 0)
+			{
+				NumParallelProcesses = MaxLocalActions;
+			}
+			else
+			{
+				// Figure out how many processors to use
+				NumParallelProcesses = GetDefaultNumParallelProcesses();
+			}
+		}
+
+		/// <summary>
+		/// Returns the name of this executor
+		/// </summary>
+		public override string Name
+		{
+			get { return "Task"; }
+		}
+
+		/// <summary>
+		/// Checks whether the parallel executor can be used
+		/// </summary>
+		/// <returns>True if the parallel executor can be used</returns>
+		public static bool IsAvailable()
+		{
+			return true;
+		}
+
+		class ExecuteResults
+		{
+			public List<string> LogLines { get; private set; }
+			public int ExitCode { get; private set; }
+
+			public ExecuteResults(List<string> LogLines, int ExitCode)
+			{
+				this.LogLines = LogLines;
+				this.ExitCode = ExitCode;
+			}
+		}
+
+
+		/// <summary>
+		/// Executes the specified actions locally.
+		/// </summary>
+		/// <returns>True if all the tasks successfully executed, or false if any of them failed.</returns>
+		public override bool ExecuteActions(List<LinkedAction> InputActions)
+		{
+			int NumCompletedActions = 0;
+			int TotalActions = InputActions.Count;
+			int ActualNumParallelProcesses = Math.Min(TotalActions, NumParallelProcesses);
+
+			using ManagedProcessGroup ProcessGroup = new ManagedProcessGroup();
+			using SemaphoreSlim MaxProcessSemaphore = new SemaphoreSlim(ActualNumParallelProcesses, ActualNumParallelProcesses);
+			using ProgressWriter ProgressWriter = new ProgressWriter("Compiling C++ source code...", false);
+
+			Log.TraceInformation("Building {0} {1} with {2} {3}...", TotalActions, (TotalActions == 1) ? "action" : "actions", ActualNumParallelProcesses, (ActualNumParallelProcesses == 1) ? "process" : "processes");
+
+			Dictionary<LinkedAction, Task<ExecuteResults>> Tasks = new Dictionary<LinkedAction, Task<ExecuteResults>>();
+			List<Task> AllTasks = new List<Task>();
+
+			using LogIndentScope Indent = new LogIndentScope("  ");
+
+			CancellationTokenSource CancellationTokenSource = new CancellationTokenSource();
+			CancellationToken CancellationToken = CancellationTokenSource.Token;
+
+			// Create a task for every action
+			foreach (LinkedAction Action in InputActions)
+			{
+				Task<ExecuteResults> ExecuteTask = ExecuteAction(Action, Tasks, ProcessGroup, MaxProcessSemaphore, CancellationToken);
+				Task LogTask = ExecuteTask.ContinueWith(antecedent => LogCompletedAction(Action, antecedent, CancellationTokenSource, ProgressWriter, TotalActions, ref NumCompletedActions));
+
+				Tasks.Add(Action, ExecuteTask);
+				AllTasks.Add(ExecuteTask);
+				AllTasks.Add(LogTask);
+			}
+
+			// Wait for all tasks to complete
+			Task.WaitAll(AllTasks.ToArray());
+
+			// Return if all tasks succeeded
+			return Tasks.Values.All(x => x.Result.ExitCode == 0);
+		}
+
+		private async Task<ExecuteResults> ExecuteAction(LinkedAction Action, Dictionary<LinkedAction, Task<ExecuteResults>> Tasks, ManagedProcessGroup ProcessGroup, SemaphoreSlim MaxProcessSemaphore, CancellationToken CancellationToken)
+		{
+			Task? SemaphoreTask = null;
+			try
+			{
+				while (true)
+				{
+					if (Action.PrerequisiteActions.All(x => Tasks.ContainsKey(x)))
+					{
+						break;
+					}
+					await Task.Delay(100);
+					CancellationToken.ThrowIfCancellationRequested();
+				}
+
+				await Task.WhenAll(Action.PrerequisiteActions.Select(x => Tasks[x]).ToArray());
+				CancellationToken.ThrowIfCancellationRequested();
+
+				SemaphoreTask = MaxProcessSemaphore.WaitAsync(CancellationToken);
+				await SemaphoreTask;
+
+				string Description = string.Empty;
+				List<string> LogLines;
+				int ExitCode;
+				using (ManagedProcess Process = new ManagedProcess(ProcessGroup, Action.CommandPath.FullName, Action.CommandArguments, Action.WorkingDirectory.FullName, null, null, ProcessPriorityClass.BelowNormal))
+				{
+					MemoryStream StdOutStream = new MemoryStream();
+					await Process.CopyToAsync(StdOutStream, CancellationToken);
+					CancellationToken.ThrowIfCancellationRequested();
+					LogLines = Console.OutputEncoding.GetString(StdOutStream.GetBuffer(), 0, Convert.ToInt32(StdOutStream.Length)).Split(LineEndingSplit, StringSplitOptions.RemoveEmptyEntries).ToList();
+					ExitCode = Process.ExitCode;
+				}
+
+				return new ExecuteResults(LogLines, ExitCode);
+			}
+			catch (OperationCanceledException)
+			{
+				return new ExecuteResults(new List<string>(), int.MaxValue);
+			}
+			catch (Exception Ex)
+			{
+				Log.WriteException(Ex, null);
+				return new ExecuteResults(new List<string>(), int.MaxValue);
+			}
+			finally
+			{
+				if (SemaphoreTask != null && SemaphoreTask.Status == TaskStatus.RanToCompletion)
+				{
+					MaxProcessSemaphore.Release();
+				}
+			}
+		}
+
+		private void LogCompletedAction(LinkedAction Action, Task<ExecuteResults> ExecuteTask, CancellationTokenSource CancellationTokenSource, ProgressWriter ProgressWriter, int TotalActions, ref int NumCompletedActions)
+		{
+			List<string> LogLines = new List<string>();
+			int ExitCode = int.MaxValue;
+			if (ExecuteTask.Status == TaskStatus.RanToCompletion)
+			{
+				LogLines = ExecuteTask.Result.LogLines;
+				ExitCode = ExecuteTask.Result.ExitCode;
+			}
+
+			// Write it to the log
+			string Description = string.Empty;
+			if (Action.bShouldOutputStatusDescription || LogLines.Count == 0)
+			{
+				Description = $"{(Action.CommandDescription != null ? Action.CommandDescription : Action.CommandPath.GetFileNameWithoutExtension())} {Action.StatusDescription}".Trim();
+			}
+			else if (LogLines.Count > 0)
+			{
+				Description = $"{(Action.CommandDescription != null ? Action.CommandDescription : Action.CommandPath.GetFileNameWithoutExtension())} {LogLines[0]}".Trim();
+			}
+
+			int CompletedActions;
+			CompletedActions = Interlocked.Increment(ref NumCompletedActions);
+			ProgressWriter.Write(CompletedActions, TotalActions);
+
+			// Cancelled
+			if (ExitCode == int.MaxValue)
+			{
+				Log.TraceInformation("[{0}/{1}] {2} cancelled", CompletedActions, TotalActions, Description);
+				return;
+			}
+
+			Log.TraceInformation("[{0}/{1}] {2}", CompletedActions, TotalActions, Description);
+			foreach (string Line in LogLines.Skip(Action.bShouldOutputStatusDescription ? 0 : 1))
+			{
+				Log.TraceInformation(Line);
+			}
+
+			if (ExitCode != 0)
+			{
+				// BEGIN TEMPORARY TO CATCH PVS-STUDIO ISSUES
+				if (LogLines.Count == 0)
+				{
+					lock (ProgressWriter)
+					{
+						Log.TraceInformation("[{0}/{1}] {2} - Error but no output", NumCompletedActions, TotalActions, Description);
+						Log.TraceInformation("[{0}/{1}] {2} - {3} {4} {5} {6}", NumCompletedActions, TotalActions, Description, ExitCode,
+							Action.WorkingDirectory, Action.CommandPath, Action.CommandArguments);
+					}
+				}
+				// END TEMPORARY
+
+				// Cancel other tasks
+				if (bStopCompilationAfterErrors)
+				{
+					CancellationTokenSource.Cancel();
+				}
+			}
+		}
+	}
+}
diff --git a/Engine/Source/Programs/UnrealBuildTool/System/ActionGraph.cs b/Engine/Source/Programs/UnrealBuildTool/System/ActionGraph.cs
--- a/Engine/Source/Programs/UnrealBuildTool/System/ActionGraph.cs
+++ b/Engine/Source/Programs/UnrealBuildTool/System/ActionGraph.cs
@@ -343,6 +343,10 @@ namespace UnrealBuildTool
 				{
 					Executor = new SNDBS();
 				}
+				else if (BuildConfiguration.bAllowTaskExecutor && TaskExecutor.IsAvailable())
+				{
+					Executor = new TaskExecutor(BuildConfiguration.MaxParallelActions);
+				}
 				else
 				{
 					Executor = new ParallelExecutor(BuildConfiguration.MaxParallelActions);
