UBT - Started seeing "UnrealBuildTool : error : Unhandled exception: System.IO.DirectoryNotFoundException: Could not find a part of the path", turned out to be an issue with checking a folder to make sure it exists before it attempts to add files from that path.

diff --git a/Engine/Source/Programs/UnrealBuildTool/Configuration/UEBuildModuleCPP.cs b/Engine/Source/Programs/UnrealBuildTool/Configuration/UEBuildModuleCPP.cs
--- a/Engine/Source/Programs/UnrealBuildTool/Configuration/UEBuildModuleCPP.cs
+++ b/Engine/Source/Programs/UnrealBuildTool/Configuration/UEBuildModuleCPP.cs
@@ -494,7 +494,10 @@ namespace UnrealBuildTool
 				List<string> GeneratedFiles = new List<string>();
 				foreach (string GeneratedDir in GeneratedCppDirectories)
 				{
-					GeneratedFiles.AddRange(Directory.GetFiles(GeneratedDir, "*.gen.cpp"));
+					if (Directory.Exists(GeneratedDir))
+					{
+						GeneratedFiles.AddRange(Directory.GetFiles(GeneratedDir, "*.gen.cpp"));
+					}
 				}
 
 				if(GeneratedFiles.Count > 0)
