Update valid range version maximum to allow for Xcode 13 beta.

#rnx

diff --git a/Engine/Source/Programs/UnrealBuildTool/Platform/Mac/ApplePlatformSDK.cs b/Engine/Source/Programs/UnrealBuildTool/Platform/Mac/ApplePlatformSDK.cs
--- a/Engine/Source/Programs/UnrealBuildTool/Platform/Mac/ApplePlatformSDK.cs
+++ b/Engine/Source/Programs/UnrealBuildTool/Platform/Mac/ApplePlatformSDK.cs
@@ -24,13 +24,13 @@ namespace UnrealBuildTool
 			if (RuntimePlatform.IsMac)
 			{
 				MinVersion = "11.0.0";
-				MaxVersion = "12.9.9";
+				MaxVersion = "13.9.9";
 			}
 			else
 			{
 				// @todo turnkey: these are MobileDevice .dll versions in Windows - to get the iTunes app version (12.3.4.1 etc) would need to hunt down the .exe
 				MinVersion = "1100.0.0.0";
-				MaxVersion = "1299.0";
+				MaxVersion = "1399.0";
 			}
 		}
 
