Add support for PDB page size override, introduced in VS2019 16.11 Preview 1.

This is used when dealing with PDB files that are over 4GB.

Use -PdbPageSize=<value> on the UBT command line, or add this to BuildConfiguration.xml:

  <WindowsPlatform>
    <PDBPageSize>32768</PDBPageSize>
  </WindowsPlatform>

Expected values are 4096 (default), 8192, 16384, 32768

Note to select an installed preview build of VS, you need to specify

  <WindowsPlatform>
    <CompilerVersion>Preview</CompilerVersion>
  </WindowsPlatform>

diff --git a/Engine/Source/Programs/UnrealBuildTool/Platform/Windows/UEBuildWindows.cs b/Engine/Source/Programs/UnrealBuildTool/Platform/Windows/UEBuildWindows.cs
--- a/Engine/Source/Programs/UnrealBuildTool/Platform/Windows/UEBuildWindows.cs
+++ b/Engine/Source/Programs/UnrealBuildTool/Platform/Windows/UEBuildWindows.cs
@@ -250,6 +250,13 @@ namespace UnrealBuildTool
 		[XmlConfigFile(Category = "WindowsPlatform")]
 		public int PCHMemoryAllocationFactor = 0;
 
+		/// <summary>
+		/// Override the PDB page size. Used when the PDB file grows over 4GB. Expected values are 4096 (default), 8192, 16384, 32768
+		/// </summary>
+		[XmlConfigFile(Category = "WindowsPlatform")]
+		[CommandLine("-PDBPageSize")]
+		public int PDBPageSize = 0;
+
 		/// <summary>
 		/// True if we allow using addresses larger than 2GB on 32 bit builds
 		/// </summary>
@@ -539,6 +546,11 @@ namespace UnrealBuildTool
 			get { return Inner.PCHMemoryAllocationFactor; }
 		}
 
+		public int PDBPageSize
+		{
+			get { return Inner.PDBPageSize; }
+		}
+
 		public bool bBuildLargeAddressAwareBinary
 		{
 			get { return Inner.bBuildLargeAddressAwareBinary; }
diff --git a/Engine/Source/Programs/UnrealBuildTool/Platform/Windows/VCToolChain.cs b/Engine/Source/Programs/UnrealBuildTool/Platform/Windows/VCToolChain.cs
--- a/Engine/Source/Programs/UnrealBuildTool/Platform/Windows/VCToolChain.cs
+++ b/Engine/Source/Programs/UnrealBuildTool/Platform/Windows/VCToolChain.cs
@@ -970,6 +970,18 @@ namespace UnrealBuildTool
 				Arguments.Add("/INCREMENTAL:NO");
 			}
 
+
+			// Override the PDB page size, if requested (VS2019 16.11 Preview 1 & above)
+			if (Target.WindowsPlatform.PDBPageSize > 0 && Target.WindowsPlatform.Compiler >= WindowsCompiler.VisualStudio2019)
+			{
+				Version CompilerVersion;
+				if (Version.TryParse(Target.WindowsPlatform.CompilerVersion, out CompilerVersion) && CompilerVersion >= new Version("14.29.30130"))
+				{
+					Arguments.Add($"/PDBPAGESIZE:{Target.WindowsPlatform.PDBPageSize}");
+				}
+			}
+
+
 			// Disable
 			//LINK : warning LNK4199: /DELAYLOAD:nvtt_64.dll ignored; no imports found from nvtt_64.dll
 			// type warning as we leverage the DelayLoad option to put third-party DLLs into a
