UnrealBuildTool: Use /sourceDependencies to generate dependency list instead of cl-filter if the the msvc compiler version is at least 14.27

diff --git a/Engine/Source/Programs/UnrealBuildTool/Platform/HoloLens/HoloLensToolChain.cs b/Engine/Source/Programs/UnrealBuildTool/Platform/HoloLens/HoloLensToolChain.cs
--- a/Engine/Source/Programs/UnrealBuildTool/Platform/HoloLens/HoloLensToolChain.cs
+++ b/Engine/Source/Programs/UnrealBuildTool/Platform/HoloLens/HoloLensToolChain.cs
@@ -698,18 +698,27 @@ namespace UnrealBuildTool
 				
 				if (CompileEnvironment.bGenerateDependenciesFile)
 				{
-					List<string> CommandArguments = new List<string>();
-					CompileAction.DependencyListFile = FileItem.GetItemByFileReference(FileReference.Combine(OutputDir, String.Format("{0}.txt", SourceFile.Location.GetFileName())));
-					CompileAction.ProducedItems.Add(CompileAction.DependencyListFile);
-					CommandArguments.Add(String.Format("-dependencies={0}", Utils.MakePathSafeToUseWithCommandLine(CompileAction.DependencyListFile.Location)));
-
-					CommandArguments.Add(String.Format("-compiler={0}", Utils.MakePathSafeToUseWithCommandLine(CompileAction.CommandPath)));
-					CommandArguments.Add("--");
-					CommandArguments.Add(Utils.MakePathSafeToUseWithCommandLine(CompileAction.CommandPath));
-					CommandArguments.Add(CompileAction.CommandArguments);
-					CommandArguments.Add("/showIncludes");
-					CompileAction.CommandArguments = string.Join(" ", CommandArguments);
-					CompileAction.CommandPath = FileReference.Combine(Unreal.EngineDirectory, "Build", "Windows", "cl-filter", "cl-filter.exe");
+					if (EnvVars.ToolChainVersion >= VersionNumber.Parse("14.27"))
+					{
+						CompileAction.DependencyListFile = FileItem.GetItemByFileReference(FileReference.Combine(OutputDir, String.Format("{0}.json", SourceFile.Location.GetFileName())));
+						CompileAction.ProducedItems.Add(CompileAction.DependencyListFile);
+						CompileAction.CommandArguments = string.Join(" ", CompileAction.CommandArguments, "/sourceDependencies", $"\"{CompileAction.DependencyListFile.AbsolutePath}\"");
+					}
+					else
+					{
+						List<string> CommandArguments = new List<string>();
+						CompileAction.DependencyListFile = FileItem.GetItemByFileReference(FileReference.Combine(OutputDir, String.Format("{0}.txt", SourceFile.Location.GetFileName())));
+						CompileAction.ProducedItems.Add(CompileAction.DependencyListFile);
+						CommandArguments.Add(String.Format("-dependencies={0}", Utils.MakePathSafeToUseWithCommandLine(CompileAction.DependencyListFile.Location)));
+
+						CommandArguments.Add(String.Format("-compiler={0}", Utils.MakePathSafeToUseWithCommandLine(CompileAction.CommandPath)));
+						CommandArguments.Add("--");
+						CommandArguments.Add(Utils.MakePathSafeToUseWithCommandLine(CompileAction.CommandPath));
+						CommandArguments.Add(CompileAction.CommandArguments);
+						CommandArguments.Add("/showIncludes");
+						CompileAction.CommandArguments = string.Join(" ", CommandArguments);
+						CompileAction.CommandPath = FileReference.Combine(Unreal.EngineDirectory, "Build", "Windows", "cl-filter", "cl-filter.exe");
+					}
 				}
 
 				if (CompileEnvironment.PrecompiledHeaderAction == PrecompiledHeaderAction.Create)
diff --git a/Engine/Source/Programs/UnrealBuildTool/Platform/Windows/VCToolChain.cs b/Engine/Source/Programs/UnrealBuildTool/Platform/Windows/VCToolChain.cs
--- a/Engine/Source/Programs/UnrealBuildTool/Platform/Windows/VCToolChain.cs
+++ b/Engine/Source/Programs/UnrealBuildTool/Platform/Windows/VCToolChain.cs
@@ -1234,13 +1234,21 @@ namespace UnrealBuildTool
 					CompileAction.ActionType = ActionType.CompileModuleInterface;
 					CompileAction.AdditionalPrerequisiteItems.Add(IfcDepsFile); // Force the dependencies file into the action graph
 					CompileAction.AdditionalProducedItems.Add(IfcFile);
-					CompileAction.DependencyListFile = FileItem.GetItemByFileReference(FileReference.Combine(OutputDir, String.Format("{0}.txt", SourceFile.Location.GetFileName())));
 					CompileAction.CompiledModuleInterfaceFile = IfcFile;
 				}
-				else if (CompileEnvironment.bGenerateDependenciesFile)
+				
+				if (CompileEnvironment.bGenerateDependenciesFile)
 				{
-					CompileAction.DependencyListFile = FileItem.GetItemByFileReference(FileReference.Combine(OutputDir, String.Format("{0}.txt", SourceFile.Location.GetFileName())));
-					CompileAction.bShowIncludes = Target.WindowsPlatform.bShowIncludes;
+					if (EnvVars.ToolChainVersion >= VersionNumber.Parse("14.27"))
+					{
+						CompileAction.DependencyListFile = FileItem.GetItemByFileReference(FileReference.Combine(OutputDir, String.Format("{0}.json", SourceFile.Location.GetFileName())));
+						CompileAction.Arguments.Add($"/sourceDependencies \"{CompileAction.DependencyListFile.AbsolutePath}\"");
+					}
+					else
+					{
+						CompileAction.DependencyListFile = FileItem.GetItemByFileReference(FileReference.Combine(OutputDir, String.Format("{0}.txt", SourceFile.Location.GetFileName())));
+						CompileAction.bShowIncludes = Target.WindowsPlatform.bShowIncludes;
+					}
 				}
 
 				if (Target.bPrintToolChainTimingInfo || Target.WindowsPlatform.bCompilerTrace)
diff --git a/Engine/Source/Programs/UnrealBuildTool/System/CppDependencyCache.cs b/Engine/Source/Programs/UnrealBuildTool/System/CppDependencyCache.cs
--- a/Engine/Source/Programs/UnrealBuildTool/System/CppDependencyCache.cs
+++ b/Engine/Source/Programs/UnrealBuildTool/System/CppDependencyCache.cs
@@ -452,7 +452,7 @@ namespace UnrealBuildTool
 					}
 				}
 
-				return new DependencyInfo(InputFile.LastWriteTimeUtc.Ticks, ProducedModule, ImportedModules.ToList(), Files);
+				return new DependencyInfo(InputFile.LastWriteTimeUtc.Ticks, ProducedModule, ImportedModules?.ToList(), Files);
 			}
 			else
 			{
