UnrealBuildTool: Move /sourceDependencies arg to GetCompilerArguments()

#rnx

diff --git a/Engine/Source/Programs/UnrealBuildTool/Platform/Windows/VCCompileAction.cs b/Engine/Source/Programs/UnrealBuildTool/Platform/Windows/VCCompileAction.cs
--- a/Engine/Source/Programs/UnrealBuildTool/Platform/Windows/VCCompileAction.cs
+++ b/Engine/Source/Programs/UnrealBuildTool/Platform/Windows/VCCompileAction.cs
@@ -439,6 +439,11 @@ namespace UnrealBuildTool
 				Arguments.Add(String.Format("/Fo\"{0}\"", ObjectFile.AbsolutePath));
 			}
 
+			if (DependencyListFile != null && DependencyListFile.HasExtension(".json"))
+			{
+				Arguments.Add(String.Format("/sourceDependencies\"{0}\"", DependencyListFile.AbsolutePath));
+			}
+
 			Arguments.AddRange(this.Arguments);
 			return Arguments;
 		}
diff --git a/Engine/Source/Programs/UnrealBuildTool/Platform/Windows/VCToolChain.cs b/Engine/Source/Programs/UnrealBuildTool/Platform/Windows/VCToolChain.cs
--- a/Engine/Source/Programs/UnrealBuildTool/Platform/Windows/VCToolChain.cs
+++ b/Engine/Source/Programs/UnrealBuildTool/Platform/Windows/VCToolChain.cs
@@ -1242,7 +1242,6 @@ namespace UnrealBuildTool
 					if (EnvVars.ToolChainVersion >= VersionNumber.Parse("14.27"))
 					{
 						CompileAction.DependencyListFile = FileItem.GetItemByFileReference(FileReference.Combine(OutputDir, String.Format("{0}.json", SourceFile.Location.GetFileName())));
-						CompileAction.Arguments.Add($"/sourceDependencies \"{CompileAction.DependencyListFile.AbsolutePath}\"");
 					}
 					else
 					{
