UnrealBuildTool: /sourceDependencies isn't valid for clang

#rnx

diff --git a/Engine/Source/Programs/UnrealBuildTool/Platform/HoloLens/HoloLensToolChain.cs b/Engine/Source/Programs/UnrealBuildTool/Platform/HoloLens/HoloLensToolChain.cs
--- a/Engine/Source/Programs/UnrealBuildTool/Platform/HoloLens/HoloLensToolChain.cs
+++ b/Engine/Source/Programs/UnrealBuildTool/Platform/HoloLens/HoloLensToolChain.cs
@@ -698,7 +698,7 @@ namespace UnrealBuildTool
 				
 				if (CompileEnvironment.bGenerateDependenciesFile)
 				{
-					if (EnvVars.ToolChainVersion >= VersionNumber.Parse("14.27"))
+					if (EnvVars.ToolChainVersion >= VersionNumber.Parse("14.27") && Target.WindowsPlatform.Compiler >= WindowsCompiler.VisualStudio2019)
 					{
 						CompileAction.DependencyListFile = FileItem.GetItemByFileReference(FileReference.Combine(OutputDir, String.Format("{0}.json", SourceFile.Location.GetFileName())));
 						CompileAction.ProducedItems.Add(CompileAction.DependencyListFile);
diff --git a/Engine/Source/Programs/UnrealBuildTool/Platform/Windows/VCToolChain.cs b/Engine/Source/Programs/UnrealBuildTool/Platform/Windows/VCToolChain.cs
--- a/Engine/Source/Programs/UnrealBuildTool/Platform/Windows/VCToolChain.cs
+++ b/Engine/Source/Programs/UnrealBuildTool/Platform/Windows/VCToolChain.cs
@@ -1239,7 +1239,7 @@ namespace UnrealBuildTool
 				
 				if (CompileEnvironment.bGenerateDependenciesFile)
 				{
-					if (EnvVars.ToolChainVersion >= VersionNumber.Parse("14.27"))
+					if (EnvVars.ToolChainVersion >= VersionNumber.Parse("14.27") && Target.WindowsPlatform.Compiler >= WindowsCompiler.VisualStudio2019)
 					{
 						CompileAction.DependencyListFile = FileItem.GetItemByFileReference(FileReference.Combine(OutputDir, String.Format("{0}.json", SourceFile.Location.GetFileName())));
 					}
