UEBuildModule:

From UDN 00331538, record PublicPreBuildLibraries in PublicLibraries, not PublicSystemLibraries.

diff --git a/Engine/Source/Programs/UnrealBuildTool/Configuration/UEBuildModule.cs b/Engine/Source/Programs/UnrealBuildTool/Configuration/UEBuildModule.cs
--- a/Engine/Source/Programs/UnrealBuildTool/Configuration/UEBuildModule.cs
+++ b/Engine/Source/Programs/UnrealBuildTool/Configuration/UEBuildModule.cs
@@ -201,7 +201,7 @@ namespace UnrealBuildTool
 				else if (PublicPreBuildLibraries.Contains(LibraryName))
 				{
 					Log.TraceLog("Library '{0}' was not resolvable to a file when used in Module '{1}'.  Be sure to add either a TargetRules.PreBuildSteps entry or a TargetRules.PreBuildTargets entry to assure it is built for your target.", LibraryName, Name);
-					PublicSystemLibraries.Add(LibraryName);
+					PublicLibraries.Add(Library.Location);
 				}
 				else
 				{
