Properly filter out all 32-bit Android include and libraries

#android
#ue5

diff --git a/Engine/Source/Programs/UnrealBuildTool/Platform/Android/AndroidToolChain.cs b/Engine/Source/Programs/UnrealBuildTool/Platform/Android/AndroidToolChain.cs
--- a/Engine/Source/Programs/UnrealBuildTool/Platform/Android/AndroidToolChain.cs
+++ b/Engine/Source/Programs/UnrealBuildTool/Platform/Android/AndroidToolChain.cs
@@ -125,22 +125,30 @@ namespace UnrealBuildTool
 		// the "-android" suffix paths here are vcpkg triplets for the android platform
 		static private Dictionary<string, string[]> AllArchNames = new Dictionary<string, string[]> {
 			{ "-arm64", new string[] { "arm64", "arm64-v8a", "arm64-android" } },
-			{ "-x64",   new string[] { "x64", "x86_64", "x64-android" } },
+			{ "-x64",   new string[] { "x64", "x86_64", "x64-android" } }
+		};
+
+		// architecture paths to use for filtering include and lib paths
+		static private Dictionary<string, string[]> AllFilterArchNames = new Dictionary<string, string[]> {
+			{ "-armv7", new string[] { "armv7", "armeabi-v7a", "arm-android" } },
+			{ "-arm64", new string[] { "arm64", "arm64-v8a", "arm64-android" } },
+			{ "-x86",   new string[] { "x86", "x86-android" } },
+			{ "-x64",   new string[] { "x64", "x86_64", "x64-android" } }
 		};
 
 		static private Dictionary<string, string[]> LibrariesToSkip = new Dictionary<string, string[]> {
 			{ "-arm64", new string[] { "nvToolsExt", "nvToolsExtStub", "vorbisenc", } },
-			{ "-x64",   new string[] { "nvToolsExt", "nvToolsExtStub", "oculus", "OVRPlugin", "vrapi", "ovrkernel", "systemutils", "openglloader", "ovrplatformloader", "gpg", "vorbisenc", } },
+			{ "-x64",   new string[] { "nvToolsExt", "nvToolsExtStub", "oculus", "OVRPlugin", "vrapi", "ovrkernel", "systemutils", "openglloader", "ovrplatformloader", "gpg", "vorbisenc", } }
 		};
 
 		static private Dictionary<string, string[]> ModulesToSkip = new Dictionary<string, string[]> {
 			{ "-arm64", new string[] {  } },
-			{ "-x64",   new string[] { "OnlineSubsystemOculus", "OculusHMD", "OculusMR", "OnlineSubsystemGooglePlay" } },
+			{ "-x64",   new string[] { "OnlineSubsystemOculus", "OculusHMD", "OculusMR", "OnlineSubsystemGooglePlay" } }
 		};
 
 		static private Dictionary<string, string[]> GeneratedModulesToSkip = new Dictionary<string, string[]> {
 			{ "-arm64", new string[] {  } },
-			{ "-x64",   new string[] { "OculusEntitlementCallbackProxy", "OculusCreateSessionCallbackProxy", "OculusFindSessionsCallbackProxy", "OculusIdentityCallbackProxy", "OculusNetConnection", "OculusNetDriver", "OnlineSubsystemOculus_init" } },
+			{ "-x64",   new string[] { "OculusEntitlementCallbackProxy", "OculusCreateSessionCallbackProxy", "OculusFindSessionsCallbackProxy", "OculusIdentityCallbackProxy", "OculusNetConnection", "OculusNetDriver", "OnlineSubsystemOculus_init" } }
 		};
 
 		public string NDKToolchainVersion;
@@ -1018,7 +1026,7 @@ namespace UnrealBuildTool
 			Dir = Dir.Replace("\\", "/").ToLowerInvariant();
 
 			// look for other architectures in the Dir path, and fail if it finds it
-			foreach (KeyValuePair<string, string[]> Pair in AllArchNames)
+			foreach (KeyValuePair<string, string[]> Pair in AllFilterArchNames)
 			{
 				if (Pair.Key != Arch)
 				{
