LiveCoding: Support updating actions that are invoking cl.exe or clang-cl.exe

#rnx

diff --git a/Engine/Source/Programs/UnrealBuildTool/System/HotReload.cs b/Engine/Source/Programs/UnrealBuildTool/System/HotReload.cs
--- a/Engine/Source/Programs/UnrealBuildTool/System/HotReload.cs
+++ b/Engine/Source/Programs/UnrealBuildTool/System/HotReload.cs
@@ -802,48 +802,53 @@ namespace UnrealBuildTool
 			{
 				if(Action.ActionType == ActionType.Compile)
 				{
-					if(!Action.CommandPath.GetFileName().Equals("cl-filter.exe", StringComparison.OrdinalIgnoreCase))
+					if (!Action.CommandPath.GetFileName().Equals("cl-filter.exe", StringComparison.OrdinalIgnoreCase)
+						&& !Action.CommandPath.GetFileName().Equals("cl.exe", StringComparison.OrdinalIgnoreCase)
+						&& !Action.CommandPath.GetFileName().Equals("clang-cl.exe", StringComparison.OrdinalIgnoreCase))
 					{
 						throw new BuildException("Unable to patch action graph - unexpected executable in compile action ({0})", Action.CommandPath);
 					}
 
 					List<string> Arguments = Utils.ParseArgumentList(Action.CommandArguments);
 
-					// Find the index of the cl-filter argument delimiter
-					int DelimiterIdx = Arguments.IndexOf("--");
-					if(DelimiterIdx == -1)
-					{
-						throw new BuildException("Unable to patch action graph - missing '--' delimiter to cl-filter");
-					}
-
-					// Fix the dependencies path
-					const string DependenciesPrefix = "-dependencies=";
+					Action NewAction = new Action(Action.Inner);
+					Action.Inner = NewAction;
 
-					int DependenciesIdx = 0;
-					for(;;DependenciesIdx++)
+					int DelimiterIdx = -1;
+					if (Action.CommandPath.GetFileName().Equals("cl-filter.exe", StringComparison.OrdinalIgnoreCase))
 					{
-						if(DependenciesIdx == DelimiterIdx)
-						{
-							throw new BuildException("Unable to patch action graph - missing '{0}' argument to cl-filter", DependenciesPrefix);
-						}
-						else if(Arguments[DependenciesIdx].StartsWith(DependenciesPrefix, StringComparison.OrdinalIgnoreCase))
+						// Find the index of the cl-filter argument delimiter
+						DelimiterIdx = Arguments.IndexOf("--");
+						if (DelimiterIdx == -1)
 						{
-							break;
+							throw new BuildException("Unable to patch action graph - missing '--' delimiter to cl-filter");
 						}
-					}
 
-					Action NewAction = new Action(Action.Inner);
-					Action.Inner = NewAction;
+						// Fix the dependencies path
+						const string DependenciesPrefix = "-dependencies=";
 
-					FileReference OldDependenciesFile = new FileReference(Arguments[DependenciesIdx].Substring(DependenciesPrefix.Length));
-					FileItem OldDependenciesFileItem = Action.ProducedItems.First(x => x.Location == OldDependenciesFile);
-					NewAction.ProducedItems.Remove(OldDependenciesFileItem);
+						int DependenciesIdx = 0;
+						for (; ; DependenciesIdx++)
+						{
+							if (DependenciesIdx == DelimiterIdx)
+							{
+								throw new BuildException("Unable to patch action graph - missing '{0}' argument to cl-filter", DependenciesPrefix);
+							}
+							else if (Arguments[DependenciesIdx].StartsWith(DependenciesPrefix, StringComparison.OrdinalIgnoreCase))
+							{
+								break;
+							}
+						}
 
-					FileReference NewDependenciesFile = OldDependenciesFile.ChangeExtension(".lc.response");
-					FileItem NewDependenciesFileItem = FileItem.GetItemByFileReference(NewDependenciesFile);
-					NewAction.ProducedItems.Add(NewDependenciesFileItem);
+						FileReference OldDependenciesFile = new FileReference(Arguments[DependenciesIdx].Substring(DependenciesPrefix.Length));
+						FileItem OldDependenciesFileItem = Action.ProducedItems.First(x => x.Location == OldDependenciesFile);
+						NewAction.PrerequisiteItems.Remove(OldDependenciesFileItem);
 
-					Arguments[DependenciesIdx] = DependenciesPrefix + NewDependenciesFile.FullName;
+						FileReference NewDependenciesFile = OldDependenciesFile.ChangeExtension(".lc.response");
+						FileItem NewDependenciesFileItem = FileItem.GetItemByFileReference(NewDependenciesFile);
+						NewAction.PrerequisiteItems.Add(NewDependenciesFileItem);
+						Arguments[DependenciesIdx] = DependenciesPrefix + NewDependenciesFile.FullName;
+					}
 
 					// Fix the response file
 					int ResponseFileIdx = DelimiterIdx + 1;
@@ -851,7 +856,7 @@ namespace UnrealBuildTool
 					{
 						if (ResponseFileIdx == Arguments.Count)
 						{
-							throw new BuildException("Unable to patch action graph - missing response file argument to cl-filter");
+							throw new BuildException($"Unable to patch action graph - missing response file argument to {Action.CommandPath.GetFileName()}");
 						}
 						else if (Arguments[ResponseFileIdx].StartsWith("@", StringComparison.Ordinal))
 						{
@@ -859,7 +864,7 @@ namespace UnrealBuildTool
 						}
 					}
 
-					FileReference OldResponseFile = new FileReference(Arguments[ResponseFileIdx].Substring(1));
+					FileReference OldResponseFile = new FileReference(Arguments[ResponseFileIdx].Substring(1).Trim('\"'));
 					FileReference NewResponseFile = new FileReference(OldResponseFile.FullName + ".lc");
 
 					const string OutputFilePrefix = "/Fo";
@@ -870,7 +875,7 @@ namespace UnrealBuildTool
 						string ResponseLine = ResponseLines[Idx];
 						if(ResponseLine.StartsWith(OutputFilePrefix, StringComparison.Ordinal))
 						{
-							FileReference OldOutputFile = new FileReference(ResponseLine.Substring(3).Trim('\"'));
+							FileReference OldOutputFile = new FileReference(ResponseLine.Substring(OutputFilePrefix.Length).Trim('\"'));
 							FileItem OldOutputFileItem = Action.ProducedItems.First(x => x.Location == OldOutputFile);
 							NewAction.ProducedItems.Remove(OldOutputFileItem);
 
@@ -884,6 +889,31 @@ namespace UnrealBuildTool
 							break;
 						}
 					}
+
+					// Update dependency file path for cl or clang-cl which is in the response file
+					if (Action.CommandPath.GetFileName().Equals("cl.exe", StringComparison.OrdinalIgnoreCase) ||
+						Action.CommandPath.GetFileName().Equals("clang-cl.exe", StringComparison.OrdinalIgnoreCase))
+					{
+						string SourceDependencyPrefix = Action.CommandPath.GetFileName().Equals("cl.exe", StringComparison.OrdinalIgnoreCase) ? "/sourceDependencies" : "/clang:-MD /clang:-MF";
+						string NewExtension = Action.CommandPath.GetFileName().Equals("cl.exe", StringComparison.OrdinalIgnoreCase) ? ".lc.json" : ".lc.d";
+						for (int Idx = 0; Idx < ResponseLines.Length; Idx++)
+						{
+							string ResponseLine = ResponseLines[Idx];
+							if (ResponseLine.StartsWith(SourceDependencyPrefix, StringComparison.Ordinal))
+							{
+								FileReference OldSourceDependencyFile = new FileReference(ResponseLine.Substring(SourceDependencyPrefix.Length).Trim().Trim('\"'));
+								FileItem OldSourceDependencyFileItem = Action.ProducedItems.First(x => x.Location == OldSourceDependencyFile);
+								NewAction.PrerequisiteItems.Remove(OldSourceDependencyFileItem);
+
+								FileReference NewSourceDependencyFile = OldSourceDependencyFile.ChangeExtension(NewExtension);
+								FileItem NewSourceDependencyFileItem = FileItem.GetItemByFileReference(NewSourceDependencyFile);
+								NewAction.PrerequisiteItems.Add(NewSourceDependencyFileItem);
+
+								ResponseLines[Idx] = SourceDependencyPrefix + "\"" + NewSourceDependencyFile.FullName + "\"";
+								break;
+							}
+						}
+					}
 					FileReference.WriteAllLines(NewResponseFile, ResponseLines);
 
 					Arguments[ResponseFileIdx] = "@" + NewResponseFile.FullName;
