UnrealBuildTool: Clean up VCToolChain functions that add command line args to be more consistently used, allowing more paths that can be relative to reduce command line length
Verified relative pathing works fine for msvc clang builds
Use .d files for dependencies when compiling with windows clang, instead of piping through cl-filter.exe
Use relative pathing for generated unity files

#rnx

diff --git a/Engine/Source/Programs/UnrealBuildTool/Platform/HoloLens/HoloLensToolChain.cs b/Engine/Source/Programs/UnrealBuildTool/Platform/HoloLens/HoloLensToolChain.cs
--- a/Engine/Source/Programs/UnrealBuildTool/Platform/HoloLens/HoloLensToolChain.cs
+++ b/Engine/Source/Programs/UnrealBuildTool/Platform/HoloLens/HoloLensToolChain.cs
@@ -704,6 +704,12 @@ namespace UnrealBuildTool
 						CompileAction.ProducedItems.Add(CompileAction.DependencyListFile);
 						CompileAction.CommandArguments = string.Join(" ", CompileAction.CommandArguments, "/sourceDependencies", $"\"{CompileAction.DependencyListFile.AbsolutePath}\"");
 					}
+					else if (Target.WindowsPlatform.Compiler == WindowsCompiler.Clang)
+					{
+						CompileAction.DependencyListFile = FileItem.GetItemByFileReference(FileReference.Combine(OutputDir, String.Format("{0}.d", SourceFile.Location.GetFileName())));
+						CompileAction.ProducedItems.Add(CompileAction.DependencyListFile);
+						CompileAction.CommandArguments = string.Join(" ", CompileAction.CommandArguments, "/clang:-MD", $"/clang:-MF\"{CompileAction.DependencyListFile.AbsolutePath}\"");
+					}
 					else
 					{
 						List<string> CommandArguments = new List<string>();
diff --git a/Engine/Source/Programs/UnrealBuildTool/Platform/Windows/VCCompileAction.cs b/Engine/Source/Programs/UnrealBuildTool/Platform/Windows/VCCompileAction.cs
--- a/Engine/Source/Programs/UnrealBuildTool/Platform/Windows/VCCompileAction.cs
+++ b/Engine/Source/Programs/UnrealBuildTool/Platform/Windows/VCCompileAction.cs
@@ -219,7 +219,7 @@ namespace UnrealBuildTool
 		/// <summary>
 		/// Whether to use cl-filter
 		/// </summary>
-		bool UsingClFilter => DependencyListFile != null && !DependencyListFile.HasExtension(".json");
+		bool UsingClFilter => DependencyListFile != null && !DependencyListFile.HasExtension(".json") && !DependencyListFile.HasExtension(".d");
 
 		/// <inheritdoc/>
 		FileReference IExternalAction.CommandPath
@@ -388,10 +388,9 @@ namespace UnrealBuildTool
 		public List<string> GetCompilerArguments()
 		{
 			List<string> Arguments = new List<string>();
-
 			if (SourceFile != null)
 			{
-				Arguments.Add(Utils.MakePathSafeToUseWithCommandLine(SourceFile.FullName));
+				VCToolChain.AddSourceFile(Arguments, SourceFile, CompilerType, PreprocessedFile != null);
 			}
 
 			foreach (DirectoryReference IncludePath in IncludePaths)
@@ -413,35 +412,37 @@ namespace UnrealBuildTool
 
 			foreach (FileItem ForceIncludeFile in ForceIncludeFiles)
 			{
-				Arguments.Add(String.Format("/FI\"{0}\"", ForceIncludeFile.Location));
+				VCToolChain.AddForceIncludeFile(Arguments, ForceIncludeFile, CompilerType, PreprocessedFile != null);
 			}
 
 			if (CreatePchFile != null)
 			{
-				Arguments.Add(String.Format("/Yc\"{0}\"", PchThroughHeaderFile.FullName));
-				Arguments.Add(String.Format("/Fp\"{0}\"", CreatePchFile.FullName));
+				VCToolChain.AddCreatePchFile(Arguments, PchThroughHeaderFile, CreatePchFile, CompilerType, PreprocessedFile != null);
 			}
+
 			if (UsingPchFile != null && CompilerType != WindowsCompiler.Clang)
 			{
-				Arguments.Add(String.Format("/Yu\"{0}\"", PchThroughHeaderFile.FullName));
-				Arguments.Add(String.Format("/Fp\"{0}\"", UsingPchFile.FullName));
+				VCToolChain.AddUsingPchFile(Arguments, PchThroughHeaderFile, UsingPchFile, CompilerType, PreprocessedFile != null);
 			}
 
 			if (PreprocessedFile != null)
 			{
-				Arguments.Add("/P"); // Preprocess
-				Arguments.Add("/C"); // Preserve comments when preprocessing
-				Arguments.Add(String.Format("/Fi\"{0}\"", PreprocessedFile)); // Preprocess to a file
+				VCToolChain.AddPreprocessedFile(Arguments, PreprocessedFile, CompilerType, PreprocessedFile != null);
 			}
 
 			if (ObjectFile != null)
 			{
-				Arguments.Add(String.Format("/Fo\"{0}\"", ObjectFile.AbsolutePath));
+				VCToolChain.AddObjectFile(Arguments, ObjectFile, CompilerType, PreprocessedFile != null);
 			}
 
 			if (DependencyListFile != null && DependencyListFile.HasExtension(".json"))
 			{
-				Arguments.Add(String.Format("/sourceDependencies\"{0}\"", DependencyListFile.AbsolutePath));
+				VCToolChain.AddSourceDependenciesFile(Arguments, DependencyListFile, CompilerType, PreprocessedFile != null);
+			}
+
+			if (DependencyListFile != null && DependencyListFile.HasExtension(".d"))
+			{
+				VCToolChain.AddSourceDependsFile(Arguments, DependencyListFile, CompilerType, PreprocessedFile != null);
 			}
 
 			Arguments.AddRange(this.Arguments);
@@ -456,18 +457,21 @@ namespace UnrealBuildTool
 			}
 			else
 			{
-				return String.Format("@\"{0}\"", ResponseFile.Location);
+				string ResponseFileString = VCToolChain.NormalizeCommandLinePath(ResponseFile, CompilerType, PreprocessedFile != null);
+				return String.Format("@{0}", Utils.MakePathSafeToUseWithCommandLine(ResponseFileString));
 			}
 		}
 
 		string GetClFilterArguments()
 		{
 			List<string> Arguments = new List<string>();
-			Arguments.Add(String.Format("-dependencies={0}", Utils.MakePathSafeToUseWithCommandLine(DependencyListFile.Location)));
+			string DependencyListFileString = VCToolChain.NormalizeCommandLinePath(DependencyListFile, CompilerType, PreprocessedFile != null);
+			Arguments.Add(String.Format("-dependencies={0}", Utils.MakePathSafeToUseWithCommandLine(DependencyListFileString)));
 
 			if (TimingFile != null)
 			{
-				Arguments.Add(String.Format("-timing={0}", Utils.MakePathSafeToUseWithCommandLine(TimingFile.Location)));
+				string TimingFileString = VCToolChain.NormalizeCommandLinePath(TimingFile, CompilerType, PreprocessedFile != null);
+				Arguments.Add(String.Format("-timing={0}", Utils.MakePathSafeToUseWithCommandLine(TimingFileString)));
 			}
 			if (bShowIncludes)
 			{
diff --git a/Engine/Source/Programs/UnrealBuildTool/Platform/Windows/VCToolChain.cs b/Engine/Source/Programs/UnrealBuildTool/Platform/Windows/VCToolChain.cs
--- a/Engine/Source/Programs/UnrealBuildTool/Platform/Windows/VCToolChain.cs
+++ b/Engine/Source/Programs/UnrealBuildTool/Platform/Windows/VCToolChain.cs
@@ -116,43 +116,97 @@ namespace UnrealBuildTool
 			}
 		}
 
-
-		public static void AddIncludePath(List<string> Arguments, DirectoryReference IncludePath, WindowsCompiler Compiler, bool bPreprocessOnly)
+		public static string NormalizeCommandLinePath(FileSystemReference Reference, WindowsCompiler Compiler, bool bPreprocessOnly)
 		{
 			// Try to use a relative path to shorten command line length. Always need the full path when preprocessing because the output file will be in a different place, where include paths cannot be relative.
-			string IncludePathString;
-			if(IncludePath.IsUnderDirectory(Unreal.RootDirectory) && Compiler != WindowsCompiler.Clang && !bPreprocessOnly)
+			if (Reference.IsUnderDirectory(Unreal.RootDirectory) && !bPreprocessOnly)
 			{
-				IncludePathString = IncludePath.MakeRelativeTo(UnrealBuildTool.EngineSourceDirectory);
-			}
-			else
-			{
-				IncludePathString = IncludePath.FullName;
+				return Reference.MakeRelativeTo(UnrealBuildTool.EngineSourceDirectory);
 			}
 
-			// If the value has a space in it and isn't wrapped in quotes, do that now. Make sure it doesn't include a trailing slash, because that will escape the closing quote.
-			if (IncludePathString.Contains(" "))
-			{
-				IncludePathString = "\"" + IncludePathString + "\"";
-			}
+			return Reference.FullName;
+		}
 
-			Arguments.Add("/I " + IncludePathString);
+		public static string NormalizeCommandLinePath(FileItem Item, WindowsCompiler Compiler, bool bPreprocessOnly)
+		{
+			return NormalizeCommandLinePath(Item.Location, Compiler, bPreprocessOnly);
 		}
 
-		public static void AddSystemIncludePath(List<string> Arguments, DirectoryReference IncludePath, WindowsCompiler Compiler, bool bPreprocessOnly)
+		public static void AddSourceFile(List<string> Arguments, FileItem SourceFile, WindowsCompiler Compiler, bool bPreprocessOnly)
+		{
+			string SourceFileString = NormalizeCommandLinePath(SourceFile, Compiler, bPreprocessOnly);
+			Arguments.Add(Utils.MakePathSafeToUseWithCommandLine(SourceFileString));
+		}
+
+		public static void AddIncludePath(List<string> Arguments, DirectoryReference IncludePath, WindowsCompiler Compiler, bool bPreprocessOnly, bool bSystemInclude = false)
 		{
-			if (Compiler == WindowsCompiler.Clang)
+			// Try to use a relative path to shorten command line length. Always need the full path when preprocessing because the output file will be in a different place, where include paths cannot be relative.
+			string IncludePathString = NormalizeCommandLinePath(IncludePath, Compiler, bPreprocessOnly);
+
+			if (Compiler == WindowsCompiler.Clang && bSystemInclude)
 			{
 				// Clang has special treatment for system headers; only system include directories are searched when include directives use angle brackets,
 				// and warnings are disabled to allow compiler toolchains to be upgraded separately.
-				Arguments.Add(String.Format("/imsvc \"{0}\"", IncludePath));
+				Arguments.Add("/imsvc " + Utils.MakePathSafeToUseWithCommandLine(IncludePathString));
 			}
 			else
 			{
-				AddIncludePath(Arguments, IncludePath, Compiler, bPreprocessOnly);
+				Arguments.Add("/I " + Utils.MakePathSafeToUseWithCommandLine(IncludePathString));
 			}
 		}
 
+		public static void AddSystemIncludePath(List<string> Arguments, DirectoryReference IncludePath, WindowsCompiler Compiler, bool bPreprocessOnly)
+		{
+			AddIncludePath(Arguments, IncludePath, Compiler, bPreprocessOnly, true);
+		}
+
+		public static void AddForceIncludeFile(List<string> Arguments, FileItem ForceIncludeFile, WindowsCompiler Compiler, bool bPreprocessOnly)
+		{
+			string ForceIncludeFileString = NormalizeCommandLinePath(ForceIncludeFile, Compiler, bPreprocessOnly);
+			Arguments.Add($"/FI\"{ForceIncludeFileString}\"");
+		}
+
+		public static void AddCreatePchFile(List<string> Arguments, FileItem PchThroughHeaderFile, FileItem CreatePchFile, WindowsCompiler Compiler, bool bPreprocessOnly)
+		{
+			string PchThroughHeaderFilePath = NormalizeCommandLinePath(PchThroughHeaderFile, Compiler, bPreprocessOnly);
+			string CreatePchFilePath = NormalizeCommandLinePath(CreatePchFile, Compiler, bPreprocessOnly);
+			Arguments.Add($"/Yc\"{PchThroughHeaderFilePath}\"");
+			Arguments.Add($"/Fp\"{CreatePchFilePath}\"");
+		}
+
+		public static void AddUsingPchFile(List<string> Arguments, FileItem PchThroughHeaderFile, FileItem UsingPchFile, WindowsCompiler Compiler, bool bPreprocessOnly)
+		{
+			string PchThroughHeaderFilePath = NormalizeCommandLinePath(PchThroughHeaderFile, Compiler, bPreprocessOnly);
+			string UsingPchFilePath = NormalizeCommandLinePath(UsingPchFile, Compiler, bPreprocessOnly);
+			Arguments.Add($"/Yu\"{PchThroughHeaderFilePath}\"");
+			Arguments.Add($"/Fp\"{UsingPchFilePath}\"");
+		}
+
+		public static void AddPreprocessedFile(List<string> Arguments, FileItem PreprocessedFile, WindowsCompiler Compiler, bool bPreprocessOnly)
+		{
+			string PreprocessedFileString = NormalizeCommandLinePath(PreprocessedFile, Compiler, bPreprocessOnly);
+			Arguments.Add("/P"); // Preprocess
+			Arguments.Add("/C"); // Preserve comments when preprocessing
+			Arguments.Add($"/Fi\"{PreprocessedFileString}\""); // Preprocess to a file
+		}
+
+		public static void AddObjectFile(List<string> Arguments, FileItem ObjectFile, WindowsCompiler Compiler, bool bPreprocessOnly)
+		{
+			string ObjectFileString = NormalizeCommandLinePath(ObjectFile, Compiler, bPreprocessOnly);
+			Arguments.Add($"/Fo\"{ObjectFileString}\"");
+		}
+
+		public static void AddSourceDependenciesFile(List<string> Arguments, FileItem SourceDependenciesFile, WindowsCompiler Compiler, bool bPreprocessOnly)
+		{
+			string SourceDependenciesFileString = NormalizeCommandLinePath(SourceDependenciesFile, Compiler, bPreprocessOnly);
+			Arguments.Add("/sourceDependencies " + Utils.MakePathSafeToUseWithCommandLine(SourceDependenciesFileString));
+		}
+
+		public static void AddSourceDependsFile(List<string> Arguments, FileItem SourceDependsFile, WindowsCompiler Compiler, bool bPreprocessOnly)
+		{
+			string SourceDependsFileString = NormalizeCommandLinePath(SourceDependsFile, Compiler, bPreprocessOnly);
+			Arguments.Add($"/clang:-MD /clang:-MF\"{SourceDependsFileString}\"");
+		}
 
 		protected virtual void AppendCLArguments_Global(CppCompileEnvironment CompileEnvironment, List<string> Arguments)
 		{
@@ -1092,7 +1146,8 @@ namespace UnrealBuildTool
 				if (CompileEnvironment.PrecompiledHeaderAction == PrecompiledHeaderAction.Create)
 				{
 					// Generate a CPP File that just includes the precompiled header.
-					string PchCppFile = string.Format("// Compiler: {0}\n#include \"{1}\"\r\n", EnvVars.CompilerVersion, CompileEnvironment.PrecompiledHeaderIncludeFilename.FullName.Replace('\\', '/'));
+					string PrecompiledHeaderIncludeFilenameString = NormalizeCommandLinePath(CompileEnvironment.PrecompiledHeaderIncludeFilename, Target.WindowsPlatform.Compiler, CompileEnvironment.bPreprocessOnly);
+					string PchCppFile = string.Format("// Compiler: {0}\n#include \"{1}\"\r\n", EnvVars.CompilerVersion, PrecompiledHeaderIncludeFilenameString.Replace('\\', '/'));
 					CompileAction.SourceFile = FileItem.GetItemByFileReference(CompileEnvironment.PrecompiledHeaderIncludeFilename.ChangeExtension(".cpp"));
 					Graph.CreateIntermediateTextFile(CompileAction.SourceFile, PchCppFile, StringComparison.Ordinal);
 
@@ -1243,6 +1298,10 @@ namespace UnrealBuildTool
 					{
 						CompileAction.DependencyListFile = FileItem.GetItemByFileReference(FileReference.Combine(OutputDir, String.Format("{0}.json", SourceFile.Location.GetFileName())));
 					}
+					else if (Target.WindowsPlatform.Compiler == WindowsCompiler.Clang)
+					{
+						CompileAction.DependencyListFile = FileItem.GetItemByFileReference(FileReference.Combine(OutputDir, String.Format("{0}.d", SourceFile.Location.GetFileName())));
+					}
 					else
 					{
 						CompileAction.DependencyListFile = FileItem.GetItemByFileReference(FileReference.Combine(OutputDir, String.Format("{0}.txt", SourceFile.Location.GetFileName())));
diff --git a/Engine/Source/Programs/UnrealBuildTool/System/Unity.cs b/Engine/Source/Programs/UnrealBuildTool/System/Unity.cs
--- a/Engine/Source/Programs/UnrealBuildTool/System/Unity.cs
+++ b/Engine/Source/Programs/UnrealBuildTool/System/Unity.cs
@@ -238,12 +238,6 @@ namespace UnrealBuildTool
 
 				OutputUnityCPPWriter.WriteLine("// This file is automatically generated at compile-time to include some subset of the user-created cpp files.");
 
-				// Add source files to the unity file
-				foreach (FileItem CPPFile in UnityFile.Files)
-				{
-					OutputUnityCPPWriter.WriteLine("#include \"{0}\"", CPPFile.AbsolutePath.Replace('\\', '/'));
-				}
-
 				// Determine unity file path name
 				string UnityCPPFileName;
 				if (AllUnityFiles.Count > 1)
@@ -256,6 +250,17 @@ namespace UnrealBuildTool
 				}
 				FileReference UnityCPPFilePath = FileReference.Combine(IntermediateDirectory, UnityCPPFileName);
 
+				// Add source files to the unity file
+				foreach (FileItem CPPFile in UnityFile.Files)
+				{
+					string CPPFileString = CPPFile.AbsolutePath;
+					if (CPPFile.Location.IsUnderDirectory(Unreal.RootDirectory))
+					{
+						CPPFileString = CPPFile.Location.MakeRelativeTo(UnrealBuildTool.EngineSourceDirectory);
+					}
+					OutputUnityCPPWriter.WriteLine("#include \"{0}\"", CPPFileString.Replace('\\', '/'));
+				}
+
 				// Write the unity file to the intermediate folder.
 				FileItem UnityCPPFile = Graph.CreateIntermediateTextFile(UnityCPPFilePath, OutputUnityCPPWriter.ToString(), StringComparison.Ordinal);
 				NewCPPFiles.Add(UnityCPPFile);
