Fixed IOS Shipping build linker error with spaces in project path in Launcher builds

diff --git a/Engine/Source/Programs/UnrealBuildTool/Platform/IOS/IOSToolChain.cs b/Engine/Source/Programs/UnrealBuildTool/Platform/IOS/IOSToolChain.cs
--- a/Engine/Source/Programs/UnrealBuildTool/Platform/IOS/IOSToolChain.cs
+++ b/Engine/Source/Programs/UnrealBuildTool/Platform/IOS/IOSToolChain.cs
@@ -530,7 +530,7 @@ namespace UnrealBuildTool
 				FileItem OutputFile = FileItem.GetItemByFileReference(LinkEnvironment.OutputFilePath);
 
 				Result += " -fembed-bitcode -Xlinker -bitcode_verify -Xlinker -bitcode_hide_symbols -Xlinker -bitcode_symbol_map ";
-				Result += " -Xlinker \"" + Path.GetDirectoryName(OutputFile.AbsolutePath) + "\"";
+				Result += " -Xlinker \\\"" + Path.GetDirectoryName(OutputFile.AbsolutePath) + "\\\"";
 			}
 
 			Result += " -dead_strip";
