[UBT] - add functionality to honor bPreprocessOnly (PS) and publish the output file path to the log (used by unrealVS)

diff --git a/Engine/Source/Programs/UnrealBuildTool/Platform/Windows/VCCompileAction.cs b/Engine/Source/Programs/UnrealBuildTool/Platform/Windows/VCCompileAction.cs
--- a/Engine/Source/Programs/UnrealBuildTool/Platform/Windows/VCCompileAction.cs
+++ b/Engine/Source/Programs/UnrealBuildTool/Platform/Windows/VCCompileAction.cs
@@ -428,6 +428,9 @@ namespace UnrealBuildTool
 			if (PreprocessedFile != null)
 			{
 				VCToolChain.AddPreprocessedFile(Arguments, PreprocessedFile, CompilerType, PreprocessedFile != null);
+
+				// this is parsed by external tools wishing to open this file directly.
+				Log.TraceInformation("PreProcessPath: " + PreprocessedFile);
 			}
 
 			if (ObjectFile != null)
