[UBT] - add functionality to honor bPreprocessOnly (PS) and publish the output file path to the log (used by unrealVS) - moved this into the *new* location for these settings

diff --git a/Engine/Source/Programs/UnrealBuildTool/Platform/Windows/VCCompileAction.cs b/Engine/Source/Programs/UnrealBuildTool/Platform/Windows/VCCompileAction.cs
--- a/Engine/Source/Programs/UnrealBuildTool/Platform/Windows/VCCompileAction.cs
+++ b/Engine/Source/Programs/UnrealBuildTool/Platform/Windows/VCCompileAction.cs
@@ -428,9 +428,6 @@ namespace UnrealBuildTool
 			if (PreprocessedFile != null)
 			{
 				VCToolChain.AddPreprocessedFile(Arguments, PreprocessedFile, CompilerType, PreprocessedFile != null);
-
-				// this is parsed by external tools wishing to open this file directly.
-				Log.TraceInformation("PreProcessPath: " + PreprocessedFile);
 			}
 
 			if (ObjectFile != null)
diff --git a/Engine/Source/Programs/UnrealBuildTool/Platform/Windows/VCToolChain.cs b/Engine/Source/Programs/UnrealBuildTool/Platform/Windows/VCToolChain.cs
--- a/Engine/Source/Programs/UnrealBuildTool/Platform/Windows/VCToolChain.cs
+++ b/Engine/Source/Programs/UnrealBuildTool/Platform/Windows/VCToolChain.cs
@@ -188,6 +188,9 @@ namespace UnrealBuildTool
 			Arguments.Add("/P"); // Preprocess
 			Arguments.Add("/C"); // Preserve comments when preprocessing
 			Arguments.Add($"/Fi\"{PreprocessedFileString}\""); // Preprocess to a file
+
+			// this is parsed by external tools wishing to open this file directly.
+			Log.TraceInformation("PreProcessPath: " + PreprocessedFile);
 		}
 
 		public static void AddObjectFile(List<string> Arguments, FileItem ObjectFile, WindowsCompiler Compiler, bool bPreprocessOnly)
