UnrealBuildTool: Fix bShowIncludes not setting the compile flag when /sourceDependencies is used to track dependencies. Also support bShowIncludes for Windows Clang.

#rnx

diff --git a/Engine/Source/Programs/UnrealBuildTool/Platform/Windows/VCToolChain.cs b/Engine/Source/Programs/UnrealBuildTool/Platform/Windows/VCToolChain.cs
--- a/Engine/Source/Programs/UnrealBuildTool/Platform/Windows/VCToolChain.cs
+++ b/Engine/Source/Programs/UnrealBuildTool/Platform/Windows/VCToolChain.cs
@@ -300,6 +300,19 @@ namespace UnrealBuildTool
 			// Enable intrinsic functions.
 			Arguments.Add("/Oi");
 
+			// Trace includes
+			if (Target.WindowsPlatform.bShowIncludes)
+			{
+				if (Target.WindowsPlatform.Compiler == WindowsCompiler.Clang)
+				{
+					Arguments.Add("/clang:--trace-includes");
+				}
+				else if (Target.WindowsPlatform.Compiler >= WindowsCompiler.VisualStudio2017 )
+				{
+					Arguments.Add("/showIncludes");
+				}
+			}
+
 
 			if (Target.WindowsPlatform.Compiler == WindowsCompiler.Clang && Target.Platform.IsInGroup(UnrealPlatformGroup.Windows))
 			{
