UnrealBuildTool: Print absolute paths in diagnostics

#rnx

diff --git a/Engine/Source/Programs/UnrealBuildTool/Platform/Windows/VCToolChain.cs b/Engine/Source/Programs/UnrealBuildTool/Platform/Windows/VCToolChain.cs
--- a/Engine/Source/Programs/UnrealBuildTool/Platform/Windows/VCToolChain.cs
+++ b/Engine/Source/Programs/UnrealBuildTool/Platform/Windows/VCToolChain.cs
@@ -313,6 +313,15 @@ namespace UnrealBuildTool
 				}
 			}
 
+			// Print absolute paths in diagnostics
+			if (Target.WindowsPlatform.Compiler == WindowsCompiler.Clang)
+			{
+				Arguments.Add("-fdiagnostics-absolute-paths");
+			}
+			else if (Target.WindowsPlatform.Compiler >= WindowsCompiler.VisualStudio2017)
+			{
+				Arguments.Add("/FC");
+			}
 
 			if (Target.WindowsPlatform.Compiler == WindowsCompiler.Clang && Target.Platform.IsInGroup(UnrealPlatformGroup.Windows))
 			{
