UnrealBuildTool: Use ConcurrentQueue in ParallelExecutor instead of lock

#rnx

diff --git a/Engine/Source/Programs/UnrealBuildTool/Executors/ParallelExecutor.cs b/Engine/Source/Programs/UnrealBuildTool/Executors/ParallelExecutor.cs
--- a/Engine/Source/Programs/UnrealBuildTool/Executors/ParallelExecutor.cs
+++ b/Engine/Source/Programs/UnrealBuildTool/Executors/ParallelExecutor.cs
@@ -2,6 +2,7 @@
 
 using EpicGames.Core;
 using System;
+using System.Collections.Concurrent;
 using System.Collections.Generic;
 using System.Diagnostics;
 using System.Linq;
@@ -184,7 +185,7 @@ namespace UnrealBuildTool
 				// Create a job object for all the child processes
 				bool bResult = true;
 				Dictionary<BuildAction, Thread> ExecutingActions = new Dictionary<BuildAction, Thread>();
-				List<BuildAction> CompletedActions = new List<BuildAction>();
+				ConcurrentQueue<BuildAction> CompletedActions = new ConcurrentQueue<BuildAction>();
 
 				TimeSpan TotalProcessingTime;
 				using (ManagedProcessGroup ProcessGroup = new ManagedProcessGroup())
@@ -206,8 +207,7 @@ namespace UnrealBuildTool
 									QueuedActions.RemoveAt(QueuedActions.Count - 1);
 
 									Thread ExecutingThread = new Thread(() => { ExecuteAction(ProcessGroup, Action, CompletedActions, CompletedEvent); });
-									string Description = $"{(Action.Inner.CommandDescription != null ? Action.Inner.CommandDescription : Action.Inner.CommandPath.GetFileName())} {Action.Inner.StatusDescription}".Trim();
-									ExecutingThread.Name = String.Format("Build:{0}", Description);
+									ExecutingThread.Name = $"Build:{(Action.Inner.CommandDescription != null ? Action.Inner.CommandDescription : Action.Inner.CommandPath.GetFileName())} {Action.Inner.StatusDescription}".Trim();
 									ExecutingThread.Start();
 
 									ExecutingActions.Add(Action, ExecutingThread);
@@ -217,76 +217,73 @@ namespace UnrealBuildTool
 								CompletedEvent.WaitOne();
 
 								// Wait for something to finish and flush it to the log
-								lock (CompletedActions)
+								while (CompletedActions.TryDequeue(out BuildAction? CompletedAction))
 								{
-									foreach (BuildAction CompletedAction in CompletedActions)
-									{
-										// Join the thread
-										Thread CompletedThread = ExecutingActions[CompletedAction];
-										CompletedThread.Join();
-										ExecutingActions.Remove(CompletedAction);
+									// Join the thread
+									Thread CompletedThread = ExecutingActions[CompletedAction];
+									CompletedThread.Join();
+									ExecutingActions.Remove(CompletedAction);
 
-										// Update the progress
-										NumCompletedActions++;
-										ProgressWriter.Write(NumCompletedActions, InputActions.Count);
+									// Update the progress
+									NumCompletedActions++;
+									ProgressWriter.Write(NumCompletedActions, InputActions.Count);
 
-										string Description = string.Empty;
+									string Description = string.Empty;
 
-										// Write it to the log
-										if (CompletedAction.Inner.bShouldOutputStatusDescription || CompletedAction.LogLines.Count == 0)
-										{
-											Description = $"{(CompletedAction.Inner.CommandDescription != null ? CompletedAction.Inner.CommandDescription : CompletedAction.Inner.CommandPath.GetFileNameWithoutExtension())} {CompletedAction.Inner.StatusDescription}".Trim();
-										}
-										else if (CompletedAction.LogLines.Count > 0)
-										{
-											Description = $"{(CompletedAction.Inner.CommandDescription != null ? CompletedAction.Inner.CommandDescription : CompletedAction.Inner.CommandPath.GetFileNameWithoutExtension())} {CompletedAction.LogLines[0]}".Trim();
-										}
+									// Write it to the log
+									if (CompletedAction.Inner.bShouldOutputStatusDescription || CompletedAction.LogLines.Count == 0)
+									{
+										Description = $"{(CompletedAction.Inner.CommandDescription != null ? CompletedAction.Inner.CommandDescription : CompletedAction.Inner.CommandPath.GetFileNameWithoutExtension())} {CompletedAction.Inner.StatusDescription}".Trim();
+									}
+									else if (CompletedAction.LogLines.Count > 0)
+									{
+										Description = $"{(CompletedAction.Inner.CommandDescription != null ? CompletedAction.Inner.CommandDescription : CompletedAction.Inner.CommandPath.GetFileNameWithoutExtension())} {CompletedAction.LogLines[0]}".Trim();
+									}
 
-										Log.TraceInformation("[{0}/{1}] {2}", NumCompletedActions, InputActions.Count, Description);
-										foreach (string Line in CompletedAction.LogLines.Skip(CompletedAction.Inner.bShouldOutputStatusDescription ? 0 : 1))
-										{
-											Log.TraceInformation(Line);
-										}
+									Log.TraceInformation("[{0}/{1}] {2}", NumCompletedActions, InputActions.Count, Description);
+									foreach (string Line in CompletedAction.LogLines.Skip(CompletedAction.Inner.bShouldOutputStatusDescription ? 0 : 1))
+									{
+										Log.TraceInformation(Line);
+									}
 
-										AllCompletedActions.Add(CompletedAction);
+									AllCompletedActions.Add(CompletedAction);
 
-										// Check the exit code
-										if (CompletedAction.ExitCode == 0)
+									// Check the exit code
+									if (CompletedAction.ExitCode == 0)
+									{
+										// Mark all the dependents as done
+										foreach (BuildAction DependantAction in CompletedAction.Dependants)
 										{
-											// Mark all the dependents as done
-											foreach (BuildAction DependantAction in CompletedAction.Dependants)
+											if (--DependantAction.MissingDependencyCount == 0)
 											{
-												if (--DependantAction.MissingDependencyCount == 0)
-												{
-													QueuedActions.Add(DependantAction);
-												}
+												QueuedActions.Add(DependantAction);
 											}
 										}
-										else
+									}
+									else
+									{
+										// BEGIN TEMPORARY TO CATCH PVS-STUDIO ISSUES
+										if (CompletedAction.LogLines.Count == 0)
 										{
-											// BEGIN TEMPORARY TO CATCH PVS-STUDIO ISSUES
-											if (CompletedAction.LogLines.Count == 0)
-											{
-												Log.TraceInformation("[{0}/{1}] {2} - Error but no output", NumCompletedActions, InputActions.Count, Description);
-												Log.TraceInformation("[{0}/{1}] {2} - {3} {4} {5} {6}", NumCompletedActions, InputActions.Count, Description, CompletedAction.ExitCode,
-													CompletedAction.Inner.WorkingDirectory, CompletedAction.Inner.CommandPath, CompletedAction.Inner.CommandArguments);
-											}
-											// END TEMPORARY
-											// Update the exit code if it's not already set
-											if (bResult && CompletedAction.ExitCode != 0)
-											{
-												bResult = false;
-											}
+											Log.TraceInformation("[{0}/{1}] {2} - Error but no output", NumCompletedActions, InputActions.Count, Description);
+											Log.TraceInformation("[{0}/{1}] {2} - {3} {4} {5} {6}", NumCompletedActions, InputActions.Count, Description, CompletedAction.ExitCode,
+												CompletedAction.Inner.WorkingDirectory, CompletedAction.Inner.CommandPath, CompletedAction.Inner.CommandArguments);
+										}
+										// END TEMPORARY
+										// Update the exit code if it's not already set
+										if (bResult && CompletedAction.ExitCode != 0)
+										{
+											bResult = false;
 										}
 									}
-									CompletedActions.Clear();
 								}
+								CompletedActions.Clear();
+							}
 
-								// If we've already got a non-zero exit code, clear out the list of queued actions so nothing else will run
-								if (!bResult && bStopCompilationAfterErrors)
-								{
-									QueuedActions.Clear();
-								}
+							// If we've already got a non-zero exit code, clear out the list of queued actions so nothing else will run
+							if (!bResult && bStopCompilationAfterErrors)
+							{
+								QueuedActions.Clear();
 							}
 						}
 					}
@@ -334,7 +331,7 @@ namespace UnrealBuildTool
 		/// <param name="Action">The action to execute</param>
 		/// <param name="CompletedActions">On completion, the list to add the completed action to</param>
 		/// <param name="CompletedEvent">Event to set once an event is complete</param>
-		static void ExecuteAction(ManagedProcessGroup ProcessGroup, BuildAction Action, List<BuildAction> CompletedActions, AutoResetEvent CompletedEvent)
+		static void ExecuteAction(ManagedProcessGroup ProcessGroup, BuildAction Action, ConcurrentQueue<BuildAction> CompletedActions, AutoResetEvent CompletedEvent)
 		{
 			try
 			{
@@ -352,10 +349,7 @@ namespace UnrealBuildTool
 				Action.ExitCode = 1;
 			}
 
-			lock (CompletedActions)
-			{
-				CompletedActions.Add(Action);
-			}
+			CompletedActions.Enqueue(Action);
 
 			CompletedEvent.Set();
 		}
