Update UE5/Main to v19 toolchain

diff --git a/Engine/Build/BatchFiles/Linux/SetupToolchain.sh b/Engine/Build/BatchFiles/Linux/SetupToolchain.sh
--- a/Engine/Build/BatchFiles/Linux/SetupToolchain.sh
+++ b/Engine/Build/BatchFiles/Linux/SetupToolchain.sh
@@ -4,7 +4,7 @@
 # put ourselves into Engine directory (two up from location of this script)
 pushd "`dirname "$0"`/../../.." > /dev/null
 
-TOOLCHAIN_VERSION=v18_clang-11.0.1-centos7
+TOOLCHAIN_VERSION=v19_clang-11.0.1-centos7
 TOOLCHAIN_ARCHIVE=$TOOLCHAIN_VERSION.tar.gz
 
 TOOLCHAIN_URL=http://cdn.unrealengine.com/Toolchain_Linux/native-linux-$TOOLCHAIN_ARCHIVE
diff --git a/Engine/Programs/UnrealGameSync/UnrealGameSync.ini b/Engine/Programs/UnrealGameSync/UnrealGameSync.ini
--- a/Engine/Programs/UnrealGameSync/UnrealGameSync.ini
+++ b/Engine/Programs/UnrealGameSync/UnrealGameSync.ini
@@ -37,7 +37,7 @@ SdkInstallerDir=
 +SdkInfo=(Category="Android", Description="NDK r21", Browse="$(SdkInstallerDir)\\Android")
 +SdkInfo=(Category="Windows", Description="Visual Studio 2019 v16.4")
 +SdkInfo=(Category="Windows", Description="Windows SDK 10.0.16299.0")
-+SdkInfo=(Category="Linux", Description="v17_clang-10.0.1-centos7")
++SdkInfo=(Category="Linux", Description="v19_clang-11.0.1-centos7")
 
 [//UE5/Release-5.0-EarlyAccess/UE4Games.uprojectdirs]
 Message=:alert:  Branch is in HARDCORE LOCKDOWN for UE5 EA2 testing. Use JIRA FixVersion '5.0-Early Access 2' for tracking.
diff --git a/Engine/Source/Programs/UnrealBuildTool/Platform/Linux/LinuxPlatformSDK.cs b/Engine/Source/Programs/UnrealBuildTool/Platform/Linux/LinuxPlatformSDK.cs
--- a/Engine/Source/Programs/UnrealBuildTool/Platform/Linux/LinuxPlatformSDK.cs
+++ b/Engine/Source/Programs/UnrealBuildTool/Platform/Linux/LinuxPlatformSDK.cs
@@ -15,14 +15,14 @@ namespace UnrealBuildTool
 	{
 		public override string GetMainVersion()
 		{
-			return "v18_clang-11.0.1-centos7";
+			return "v19_clang-11.0.1-centos7";
 		}
 
 		public override void GetValidVersionRange(out string MinVersion, out string MaxVersion)
 		{
 			// all that matters is the number after the v, according to TryConvertVersionToInt()
 			MinVersion = "v10_clang-5.0.0-centos7";
-			MaxVersion = "v18_clang-11.0.1-centos7";
+			MaxVersion = "v19_clang-11.0.1-centos7";
 		}
 
 		public override void GetValidSoftwareVersionRange(out string MinVersion, out string MaxVersion)
