UBT: cl.exe can't handle response filepaths longer than 260 characters, and relative paths can push it over the limit when otherwise it would have been fine.

Hopefully I can get Microsoft to fix this.

#rnx

diff --git a/Engine/Source/Programs/UnrealBuildTool/Platform/Windows/VCCompileAction.cs b/Engine/Source/Programs/UnrealBuildTool/Platform/Windows/VCCompileAction.cs
--- a/Engine/Source/Programs/UnrealBuildTool/Platform/Windows/VCCompileAction.cs
+++ b/Engine/Source/Programs/UnrealBuildTool/Platform/Windows/VCCompileAction.cs
@@ -461,6 +461,12 @@ namespace UnrealBuildTool
 			else
 			{
 				string ResponseFileString = VCToolChain.NormalizeCommandLinePath(ResponseFile, CompilerType, PreprocessedFile != null);
+
+				// cl.exe can't handle response files with a path longer than 260 characters, and relative paths can push it over the limit
+				if (!System.IO.Path.IsPathRooted(ResponseFileString) && System.IO.Path.Combine(WorkingDirectory.FullName, ResponseFileString).Length > 260)
+				{
+					ResponseFileString = ResponseFile.FullName;
+				}
 				return String.Format("@{0}", Utils.MakePathSafeToUseWithCommandLine(ResponseFileString));
 			}
 		}
