UBT: Include PCH settings in Visual Studio project files for Intellisense parser.

diff --git a/Engine/Source/Programs/UnrealBuildTool/Configuration/UEBuildModuleCPP.cs b/Engine/Source/Programs/UnrealBuildTool/Configuration/UEBuildModuleCPP.cs
--- a/Engine/Source/Programs/UnrealBuildTool/Configuration/UEBuildModuleCPP.cs
+++ b/Engine/Source/Programs/UnrealBuildTool/Configuration/UEBuildModuleCPP.cs
@@ -669,7 +669,7 @@ namespace UnrealBuildTool
 			CPPOutput Output;
 			if (ToolChain == null)
 			{
-				Output = new CPPOutput();
+				Output = GetDefaultPchOutput(WrapperLocation);
 			}
 			else
 			{
@@ -709,7 +709,7 @@ namespace UnrealBuildTool
 				CPPOutput Output;
 				if (ToolChain == null)
 				{
-					Output = new CPPOutput();
+					Output = GetDefaultPchOutput(WrapperLocation);
 				}
 				else
 				{
@@ -721,6 +721,18 @@ namespace UnrealBuildTool
 			return Instance;
 		}
 
+		/// <summary>
+		/// Gets a default PCH output file when using intellisense. Saves having to construct every toolchain.
+		/// </summary>
+		/// <param name="PchIncludeFilename"></param>
+		/// <returns></returns>
+		static CPPOutput GetDefaultPchOutput(FileReference PchIncludeFilename)
+		{
+			CPPOutput Output = new CPPOutput();
+			Output.PrecompiledHeaderFile = FileItem.GetItemByFileReference(PchIncludeFilename.ChangeExtension(".pch"));
+			return Output;
+		}
+
 		/// <summary>
 		/// Determines if a module compile environment is compatible with the given shared PCH compile environment
 		/// </summary>
diff --git a/Engine/Source/Programs/UnrealBuildTool/ProjectFiles/Project.cs b/Engine/Source/Programs/UnrealBuildTool/ProjectFiles/Project.cs
--- a/Engine/Source/Programs/UnrealBuildTool/ProjectFiles/Project.cs
+++ b/Engine/Source/Programs/UnrealBuildTool/ProjectFiles/Project.cs
@@ -327,8 +327,21 @@ namespace UnrealBuildTool
 
 			foreach (DirectoryReference BaseDir in Module.ModuleDirectories)
 			{
-				AddIntelliSenseIncludePaths(BaseDir, BaseDirToSystemIncludePaths, CompileEnvironment.SystemIncludePaths);
-				AddIntelliSenseIncludePaths(BaseDir, BaseDirToUserIncludePaths, CompileEnvironment.UserIncludePaths);
+				BuildEnvironment BuildEnvironment;
+				if (!BaseDirToBuildEnvironment.TryGetValue(BaseDir, out BuildEnvironment))
+				{
+					BuildEnvironment = new BuildEnvironment();
+					BaseDirToBuildEnvironment.Add(BaseDir, BuildEnvironment);
+				}
+
+				if (CompileEnvironment.PrecompiledHeaderAction == PrecompiledHeaderAction.Include)
+				{
+					BuildEnvironment.PchFile = CompileEnvironment.PrecompiledHeaderFile?.Location;
+					BuildEnvironment.PchIncludeFile = CompileEnvironment.PrecompiledHeaderIncludeFilename;
+				}
+
+				AddIntelliSenseIncludePaths(BuildEnvironment.SystemIncludePaths, CompileEnvironment.SystemIncludePaths);
+				AddIntelliSenseIncludePaths(BuildEnvironment.UserIncludePaths, CompileEnvironment.UserIncludePaths);
 			}
 
 			SetIntelliSenseCppVersion(Module.Rules.CppStandard);
@@ -379,23 +392,6 @@ namespace UnrealBuildTool
 			}
 		}
 
-		/// <summary>
-		/// Adds all of the specified include paths to this VCProject's list of include paths for all modules in the project
-		/// </summary>
-		/// <param name="BaseDir">The base directory to which the include paths apply</param>
-		/// <param name="BaseDirToCollection">Map of base directory to include paths</param>
-		/// <param name="NewIncludePaths">List of include paths to add</param>
-		public void AddIntelliSenseIncludePaths(DirectoryReference BaseDir, Dictionary<DirectoryReference, IncludePathsCollection> BaseDirToCollection, IEnumerable<DirectoryReference> NewIncludePaths)
-		{
-			IncludePathsCollection ModuleUserPaths;
-			if (!BaseDirToCollection.TryGetValue(BaseDir, out ModuleUserPaths))
-			{
-				ModuleUserPaths = new IncludePathsCollection();
-				BaseDirToCollection.Add(BaseDir, ModuleUserPaths);
-			}
-			AddIntelliSenseIncludePaths(ModuleUserPaths, NewIncludePaths);
-		}
-
 		/// <summary>
 		/// Adds all of the specified include paths to this VCProject's list of include paths for all modules in the project
 		/// </summary>
@@ -569,6 +565,17 @@ namespace UnrealBuildTool
 			public HashSet<DirectoryReference> AbsolutePaths = new HashSet<DirectoryReference>();
 		}
 
+		/// <summary>
+		/// Build environment for a particular module
+		/// </summary>
+		public class BuildEnvironment
+		{
+			public FileReference PchFile;
+			public FileReference PchIncludeFile;
+			public IncludePathsCollection UserIncludePaths = new IncludePathsCollection();
+			public IncludePathsCollection SystemIncludePaths = new IncludePathsCollection();
+		}
+
 		/// <summary>
 		/// Merged list of include paths for the project
 		/// </summary>
@@ -582,12 +589,7 @@ namespace UnrealBuildTool
 		/// <summary>
 		/// Map of base directory to user include paths
 		/// </summary>
-		protected Dictionary<DirectoryReference, IncludePathsCollection> BaseDirToUserIncludePaths = new Dictionary<DirectoryReference, IncludePathsCollection>();
-
-		/// <summary>
-		/// Map of base directory to system include paths
-		/// </summary>
-		protected Dictionary<DirectoryReference, IncludePathsCollection> BaseDirToSystemIncludePaths = new Dictionary<DirectoryReference, IncludePathsCollection>();
+		protected Dictionary<DirectoryReference, BuildEnvironment> BaseDirToBuildEnvironment = new Dictionary<DirectoryReference, BuildEnvironment>();
 
 		/// <summary>
 		/// Legacy accessor for user search paths
diff --git a/Engine/Source/Programs/UnrealBuildTool/ProjectFiles/VisualStudio/VCProject.cs b/Engine/Source/Programs/UnrealBuildTool/ProjectFiles/VisualStudio/VCProject.cs
--- a/Engine/Source/Programs/UnrealBuildTool/ProjectFiles/VisualStudio/VCProject.cs
+++ b/Engine/Source/Programs/UnrealBuildTool/ProjectFiles/VisualStudio/VCProject.cs
@@ -771,29 +771,36 @@ namespace UnrealBuildTool
 			return false;
 		}
 
+		bool TryGetBuildEnvironment(DirectoryReference BaseDir, out BuildEnvironment OutBuildEnvironment)
+		{
+			for (DirectoryReference CurrentDir = BaseDir; CurrentDir != null; CurrentDir = CurrentDir.ParentDirectory)
+			{
+				BuildEnvironment BuildEnvironment;
+				if (BaseDirToBuildEnvironment.TryGetValue(CurrentDir, out BuildEnvironment))
+				{
+					OutBuildEnvironment = BuildEnvironment;
+					return true;
+				}
+			}
+
+			OutBuildEnvironment = null;
+			return false;
+		}
+
 		/// <summary>
 		/// Append a list of include paths to a property list
 		/// </summary>
 		/// <param name="Builder">String builder for the property value</param>
-		/// <param name="BaseDir">Directory containing the source file</param>
-		/// <param name="BaseDirToIncludePaths">Map of base directory to set of include paths</param>
+		/// <param name="Collection">Collection of include paths</param>
 		/// <param name="IgnorePaths">Set of paths to ignore</param>
-		void AppendIncludePaths(StringBuilder Builder, DirectoryReference BaseDir, Dictionary<DirectoryReference, IncludePathsCollection> BaseDirToIncludePaths, HashSet<DirectoryReference> IgnorePaths)
+		void AppendIncludePaths(StringBuilder Builder, IncludePathsCollection Collection, HashSet<DirectoryReference> IgnorePaths)
 		{
-			for (DirectoryReference CurrentDir = BaseDir; CurrentDir != null; CurrentDir = CurrentDir.ParentDirectory)
+			foreach (DirectoryReference IncludePath in Collection.AbsolutePaths)
 			{
-				IncludePathsCollection Collection;
-				if (BaseDirToIncludePaths.TryGetValue(CurrentDir, out Collection))
+				if (!IgnorePaths.Contains(IncludePath) && !IncludePathIsFilteredOut(IncludePath))
 				{
-					foreach (DirectoryReference IncludePath in Collection.AbsolutePaths)
-					{
-						if (!IgnorePaths.Contains(IncludePath) && !IncludePathIsFilteredOut(IncludePath))
-						{
-							Builder.Append(NormalizeProjectPath(IncludePath.FullName));
-							Builder.Append(';');
-						}
-					}
-					break;
+					Builder.Append(NormalizeProjectPath(IncludePath.FullName));
+					Builder.Append(';');
 				}
 			}
 		}
@@ -831,10 +838,10 @@ namespace UnrealBuildTool
 				{
 					for (DirectoryReference CurrentDir = Pair.Key; CurrentDir != null; CurrentDir = CurrentDir.ParentDirectory)
 					{
-						IncludePathsCollection IncludePaths;
-						if (BaseDirToUserIncludePaths.TryGetValue(CurrentDir, out IncludePaths))
+						BuildEnvironment BuildEnvironment;
+						if (BaseDirToBuildEnvironment.TryGetValue(CurrentDir, out BuildEnvironment))
 						{
-							foreach (DirectoryReference IncludePath in IncludePaths.AbsolutePaths)
+							foreach (DirectoryReference IncludePath in BuildEnvironment.UserIncludePaths.AbsolutePaths)
 							{
 								int Count;
 								IncludePathToCount.TryGetValue(IncludePath, out Count);
@@ -1194,20 +1201,27 @@ namespace UnrealBuildTool
 
 
 						// Find the include search paths
-						string IncludeSearchPaths;
-						if (!DirectoryToIncludeSearchPaths.TryGetValue(Directory, out IncludeSearchPaths))
+						VCProjectFileContent.AppendLine("    <{0} Include=\"{1}\">", VCFileType, EscapeFileName(AliasedFile.FileSystemPath));
+						if (TryGetBuildEnvironment(Directory, out BuildEnvironment BuildEnvironment))
 						{
-							StringBuilder Builder = new StringBuilder();
-							AppendIncludePaths(Builder, Directory, BaseDirToUserIncludePaths, SharedIncludeSearchPathsSet);
-							AppendIncludePaths(Builder, Directory, BaseDirToSystemIncludePaths, SharedIncludeSearchPathsSet);
-							IncludeSearchPaths = Builder.ToString();
+							string IncludeSearchPaths = String.Empty;
+							if (!DirectoryToIncludeSearchPaths.TryGetValue(Directory, out IncludeSearchPaths))
+							{
+								StringBuilder Builder = new StringBuilder();
+								AppendIncludePaths(Builder, BuildEnvironment.UserIncludePaths, SharedIncludeSearchPathsSet);
+								AppendIncludePaths(Builder, BuildEnvironment.SystemIncludePaths, SharedIncludeSearchPathsSet);
+								IncludeSearchPaths = Builder.ToString();
 
-							DirectoryToIncludeSearchPaths.Add(Directory, IncludeSearchPaths);
-						}
+								DirectoryToIncludeSearchPaths.Add(Directory, IncludeSearchPaths);
+							}
+							VCProjectFileContent.AppendLine("      <AdditionalIncludeDirectories>$(NMakeIncludeSearchPath);{0}</AdditionalIncludeDirectories>", IncludeSearchPaths);
+							VCProjectFileContent.AppendLine("      <ForcedIncludeFiles>$(NMakeForcedIncludes);{0}</ForcedIncludeFiles>", ForceIncludePaths);
 
-						VCProjectFileContent.AppendLine("    <{0} Include=\"{1}\">", VCFileType, EscapeFileName(AliasedFile.FileSystemPath));
-						VCProjectFileContent.AppendLine("      <AdditionalIncludeDirectories>$(NMakeIncludeSearchPath);{0}</AdditionalIncludeDirectories>", IncludeSearchPaths);
-						VCProjectFileContent.AppendLine("      <ForcedIncludeFiles>$(NMakeForcedIncludes);{0}</ForcedIncludeFiles>", ForceIncludePaths);
+							if (BuildEnvironment.PchFile != null && BuildEnvironment.PchIncludeFile != null)
+							{
+								VCProjectFileContent.AppendLine("      <AdditionalOptions>$(AdditionalOptions);/Yu&quot;{0}&quot; /Fp&quot;{1}&quot;</AdditionalOptions>", BuildEnvironment.PchFile, BuildEnvironment.PchIncludeFile);
+							}
+						}
 						VCProjectFileContent.AppendLine("    </{0}>", VCFileType);
 					}
 
