UBT: Fix arguments for Intellisense PCH generation.

diff --git a/Engine/Source/Programs/UnrealBuildTool/Configuration/UEBuildModuleCPP.cs b/Engine/Source/Programs/UnrealBuildTool/Configuration/UEBuildModuleCPP.cs
--- a/Engine/Source/Programs/UnrealBuildTool/Configuration/UEBuildModuleCPP.cs
+++ b/Engine/Source/Programs/UnrealBuildTool/Configuration/UEBuildModuleCPP.cs
@@ -669,7 +669,7 @@ namespace UnrealBuildTool
 			CPPOutput Output;
 			if (ToolChain == null)
 			{
-				Output = GetDefaultPchOutput(WrapperLocation);
+				Output = new CPPOutput();
 			}
 			else
 			{
@@ -709,7 +709,7 @@ namespace UnrealBuildTool
 				CPPOutput Output;
 				if (ToolChain == null)
 				{
-					Output = GetDefaultPchOutput(WrapperLocation);
+					Output = new CPPOutput();
 				}
 				else
 				{
@@ -721,18 +721,6 @@ namespace UnrealBuildTool
 			return Instance;
 		}
 
-		/// <summary>
-		/// Gets a default PCH output file when using intellisense. Saves having to construct every toolchain.
-		/// </summary>
-		/// <param name="PchIncludeFilename"></param>
-		/// <returns></returns>
-		static CPPOutput GetDefaultPchOutput(FileReference PchIncludeFilename)
-		{
-			CPPOutput Output = new CPPOutput();
-			Output.PrecompiledHeaderFile = FileItem.GetItemByFileReference(PchIncludeFilename.ChangeExtension(".pch"));
-			return Output;
-		}
-
 		/// <summary>
 		/// Determines if a module compile environment is compatible with the given shared PCH compile environment
 		/// </summary>
diff --git a/Engine/Source/Programs/UnrealBuildTool/ProjectFiles/Project.cs b/Engine/Source/Programs/UnrealBuildTool/ProjectFiles/Project.cs
--- a/Engine/Source/Programs/UnrealBuildTool/ProjectFiles/Project.cs
+++ b/Engine/Source/Programs/UnrealBuildTool/ProjectFiles/Project.cs
@@ -336,7 +336,6 @@ namespace UnrealBuildTool
 
 				if (CompileEnvironment.PrecompiledHeaderAction == PrecompiledHeaderAction.Include)
 				{
-					BuildEnvironment.PchFile = CompileEnvironment.PrecompiledHeaderFile?.Location;
 					BuildEnvironment.PchIncludeFile = CompileEnvironment.PrecompiledHeaderIncludeFilename;
 				}
 
@@ -570,7 +569,6 @@ namespace UnrealBuildTool
 		/// </summary>
 		public class BuildEnvironment
 		{
-			public FileReference PchFile;
 			public FileReference PchIncludeFile;
 			public IncludePathsCollection UserIncludePaths = new IncludePathsCollection();
 			public IncludePathsCollection SystemIncludePaths = new IncludePathsCollection();
diff --git a/Engine/Source/Programs/UnrealBuildTool/ProjectFiles/VisualStudio/VCProject.cs b/Engine/Source/Programs/UnrealBuildTool/ProjectFiles/VisualStudio/VCProject.cs
--- a/Engine/Source/Programs/UnrealBuildTool/ProjectFiles/VisualStudio/VCProject.cs
+++ b/Engine/Source/Programs/UnrealBuildTool/ProjectFiles/VisualStudio/VCProject.cs
@@ -1217,9 +1217,9 @@ namespace UnrealBuildTool
 							VCProjectFileContent.AppendLine("      <AdditionalIncludeDirectories>$(NMakeIncludeSearchPath);{0}</AdditionalIncludeDirectories>", IncludeSearchPaths);
 							VCProjectFileContent.AppendLine("      <ForcedIncludeFiles>$(NMakeForcedIncludes);{0}</ForcedIncludeFiles>", ForceIncludePaths);
 
-							if (BuildEnvironment.PchFile != null && BuildEnvironment.PchIncludeFile != null)
+							if (BuildEnvironment.PchIncludeFile != null)
 							{
-								VCProjectFileContent.AppendLine("      <AdditionalOptions>$(AdditionalOptions);/Yu&quot;{0}&quot; /Fp&quot;{1}&quot;</AdditionalOptions>", BuildEnvironment.PchFile, BuildEnvironment.PchIncludeFile);
+								VCProjectFileContent.AppendLine("      <AdditionalOptions>$(AdditionalOptions) /Yu&quot;{0}&quot;</AdditionalOptions>", BuildEnvironment.PchIncludeFile);
 							}
 						}
 						VCProjectFileContent.AppendLine("    </{0}>", VCFileType);
