PR #8195: Fix Rider project file generator on Linux platform. (Contributed by [REDACTED_PERSON])

diff --git a/Engine/Source/Programs/UnrealBuildTool/ProjectFiles/Rider/RiderProjectFile.cs b/Engine/Source/Programs/UnrealBuildTool/ProjectFiles/Rider/RiderProjectFile.cs
--- a/Engine/Source/Programs/UnrealBuildTool/ProjectFiles/Rider/RiderProjectFile.cs
+++ b/Engine/Source/Programs/UnrealBuildTool/ProjectFiles/Rider/RiderProjectFile.cs
@@ -523,7 +523,7 @@ namespace UnrealBuildTool
 					throw new ArgumentException("Wrong Target.Architecture: {0}", Target.Architecture);
 				}
 
-				string PlatformSdkVersionString = UEBuildPlatformSDK.GetSDKForPlatform(Target.ToString()).GetMainVersion();
+				string PlatformSdkVersionString = UEBuildPlatformSDK.GetSDKForPlatform(BuildPlatform.GetPlatformName()).GetInstalledSDKVersion();
 				var Version = GetLinuxToolchainVersionFromFullString(PlatformSdkVersionString);
 
 				var InternalSdkPath = UEBuildPlatform.GetSDK(UnrealTargetPlatform.Linux).GetInternalSDKPath();
