UnrealBuildTool: Write relative paths in generated pch comments if under the unreal root

#rnx

diff --git a/Engine/Source/Programs/UnrealBuildTool/Configuration/UEBuildModuleCPP.cs b/Engine/Source/Programs/UnrealBuildTool/Configuration/UEBuildModuleCPP.cs
--- a/Engine/Source/Programs/UnrealBuildTool/Configuration/UEBuildModuleCPP.cs
+++ b/Engine/Source/Programs/UnrealBuildTool/Configuration/UEBuildModuleCPP.cs
@@ -1000,7 +1000,13 @@ namespace UnrealBuildTool
 				StringBuilder WrapperContents = new StringBuilder();
 				using (StringWriter Writer = new StringWriter(WrapperContents))
 				{
-					Writer.WriteLine("// Dedicated PCH for {0}", File.AbsolutePath);
+					string FileString = File.AbsolutePath;
+					if (File.Location.IsUnderDirectory(Unreal.RootDirectory))
+					{
+						FileString = File.Location.MakeRelativeTo(UnrealBuildTool.EngineSourceDirectory);
+					}
+					FileString = FileString.Replace('\\', '/');
+					Writer.WriteLine("// Dedicated PCH for {0}", FileString);
 					Writer.WriteLine();
 					WriteDefinitions(CompileEnvironment.Definitions, Writer);
 					Writer.WriteLine();
@@ -1178,9 +1184,15 @@ namespace UnrealBuildTool
 			StringBuilder WrapperContents = new StringBuilder();
 			using (StringWriter Writer = new StringWriter(WrapperContents))
 			{
-				Writer.WriteLine("// PCH for {0}", IncludedFile.AbsolutePath);
+				string IncludeFileString = IncludedFile.AbsolutePath;
+				if (IncludedFile.Location.IsUnderDirectory(Unreal.RootDirectory))
+				{
+					IncludeFileString = IncludedFile.Location.MakeRelativeTo(UnrealBuildTool.EngineSourceDirectory);
+				}
+				IncludeFileString = IncludeFileString.Replace('\\', '/');
+				Writer.WriteLine("// PCH for {0}", IncludeFileString);
 				WriteDefinitions(Definitions, Writer);
-				Writer.WriteLine("#include \"{0}\"", IncludedFile.AbsolutePath.Replace('\\', '/'));
+				Writer.WriteLine("#include \"{0}\"", IncludeFileString);
 			}
 
 			// Create the item
