PR #8219: Update compiler version to 143 in VS Project Generator (Contributed by [REDACTED_PERSON])

diff --git a/Engine/Source/Programs/UnrealBuildTool/ProjectFiles/VisualStudio/VCProjectFileGenerator.cs b/Engine/Source/Programs/UnrealBuildTool/ProjectFiles/VisualStudio/VCProjectFileGenerator.cs
--- a/Engine/Source/Programs/UnrealBuildTool/ProjectFiles/VisualStudio/VCProjectFileGenerator.cs
+++ b/Engine/Source/Programs/UnrealBuildTool/ProjectFiles/VisualStudio/VCProjectFileGenerator.cs
@@ -242,8 +242,9 @@ namespace UnrealBuildTool
 				case VCProjectFileFormat.VisualStudio2017:
 					return "v141";
 				case VCProjectFileFormat.VisualStudio2019:
-				case VCProjectFileFormat.VisualStudio2022:
 					return "v142";
+				case VCProjectFileFormat.VisualStudio2022:
+					return "v143";
 
 			}
 			return string.Empty;
