Modify vscode project file generation for linux so that UE types are visualized properly

PR #4509

diff --git a/Engine/Source/Programs/UnrealBuildTool/ProjectFiles/VisualStudioCode/VSCodeProjectFileGenerator.cs b/Engine/Source/Programs/UnrealBuildTool/ProjectFiles/VisualStudioCode/VSCodeProjectFileGenerator.cs
--- a/Engine/Source/Programs/UnrealBuildTool/ProjectFiles/VisualStudioCode/VSCodeProjectFileGenerator.cs
+++ b/Engine/Source/Programs/UnrealBuildTool/ProjectFiles/VisualStudioCode/VSCodeProjectFileGenerator.cs
@@ -1282,9 +1282,15 @@ namespace UnrealBuildTool
 						OutFile.AddField("type", "cppvsdbg");
 						OutFile.AddField("visualizerFile", MakeUnquotedPathString(FileReference.Combine(UE4ProjectRoot, "Engine", "Extras", "VisualStudioDebugging", "Unreal.natvis"), EPathType.Absolute));
 					}
-					else
+					else if (HostPlatform == UnrealTargetPlatform.Linux)
 					{
 						OutFile.AddField("type", "cppdbg");
+						OutFile.AddField("visualizerFile", MakeUnquotedPathString(FileReference.Combine(UE4ProjectRoot, "Engine", "Extras", "VisualStudioDebugging", "Unreal.natvis"), EPathType.Absolute));
+						OutFile.AddField("showDisplayString", true);
+					}
+					else
+					{
+						OutFile.AddField("type", "lldb");
 					}
 				}
 				OutFile.EndObject();
