Disable clang color if using a build machine

diff --git a/Engine/Source/Programs/UnrealBuildTool/Platform/Linux/UEBuildLinux.cs b/Engine/Source/Programs/UnrealBuildTool/Platform/Linux/UEBuildLinux.cs
--- a/Engine/Source/Programs/UnrealBuildTool/Platform/Linux/UEBuildLinux.cs
+++ b/Engine/Source/Programs/UnrealBuildTool/Platform/Linux/UEBuildLinux.cs
@@ -579,6 +579,12 @@ namespace UnrealBuildTool
 				Options |= LinuxToolChainOptions.DisableSplitDebugInfoWithObjCopy;
 			}
 
+			// Disable color logging if we are on a build machine
+			if (Environment.GetEnvironmentVariable("IsBuildMachine") == "1")
+			{
+				Log.ColorConsoleOutput = false;
+			}
+
 			return new LinuxToolChain(Target.Architecture, SDK, Target.LinuxPlatform.bPreservePSYM, Options);
 		}
 
