UnrealBuildTool.csproj:

Minor cleanups

#trivial

diff --git a/Engine/Source/Programs/UnrealBuildTool/UnrealBuildTool.csproj b/Engine/Source/Programs/UnrealBuildTool/UnrealBuildTool.csproj
--- a/Engine/Source/Programs/UnrealBuildTool/UnrealBuildTool.csproj
+++ b/Engine/Source/Programs/UnrealBuildTool/UnrealBuildTool.csproj
@@ -1,4 +1,4 @@
-﻿<Project Sdk="Microsoft.NET.Sdk">
+<Project Sdk="Microsoft.NET.Sdk">
   <PropertyGroup>
     <TargetFramework>netcoreapp3.1</TargetFramework>
     <Configuration Condition=" '$(Configuration)' == '' ">Development</Configuration>
@@ -44,15 +44,11 @@
     <AllowUnsafeBlocks>false</AllowUnsafeBlocks>
     <Prefer32Bit>false</Prefer32Bit>
 	<OutputPath>..\..\..\Binaries\DotNET\UnrealBuildTool</OutputPath>
-    <DocumentationFile>..\..\..\Binaries\DotNET\UnrealBuildTool.xml</DocumentationFile>
-    <NoWarn>
-    </NoWarn>
-    <!-- The mono compiler used to ship on Mac is only compatible with version 6 -->
-    <LangVersion>6</LangVersion>
+    <DocumentationFile>..\..\..\Binaries\DotNET\UnrealBuildTool\UnrealBuildTool.xml</DocumentationFile>
   </PropertyGroup>
   <PropertyGroup Condition=" '$(Configuration)|$(Platform)' == 'Debug|AnyCPU' ">
     <OutputPath>..\..\..\Binaries\DotNET\UnrealBuildTool</OutputPath>
-    <DocumentationFile>..\..\..\Binaries\DotNET\UnrealBuildTool\UnrealBuildTool.XML</DocumentationFile>
+    <DocumentationFile>..\..\..\Binaries\DotNET\UnrealBuildTool\UnrealBuildTool.xml</DocumentationFile>
     <DebugSymbols>true</DebugSymbols>
     <DebugType>full</DebugType>
   </PropertyGroup>
@@ -73,8 +69,6 @@
     </When>
   </Choose>
 
-
-
   <ItemGroup>
     <Compile Include="..\DotNETCommon\MetaData.cs">
       <Link>Properties\MetaData.cs</Link>
@@ -120,4 +114,4 @@
     <!-- Required for Hololens -->
     <PackageReference Include="System.Reflection.MetadataLoadContext" Version="4.7.2" /> 
   </ItemGroup>
-</Project>
\ No newline at end of file
+</Project>
