UnrealBuildTool: Warn if a module is referenced with incorrect text case

diff --git a/Engine/Source/Programs/UnrealBuildTool/Configuration/UEBuildTarget.cs b/Engine/Source/Programs/UnrealBuildTool/Configuration/UEBuildTarget.cs
--- a/Engine/Source/Programs/UnrealBuildTool/Configuration/UEBuildTarget.cs
+++ b/Engine/Source/Programs/UnrealBuildTool/Configuration/UEBuildTarget.cs
@@ -4024,6 +4024,12 @@ namespace UnrealBuildTool
 				Module = InstantiateModule(RulesObject, GeneratedCodeDirectory);
 				Modules.Add(Module.Name, Module);
 			}
+
+			// Warn if the module reference has incorrect text case
+			if (!Module.Name.Equals(ModuleName) && Module.Name.Equals(ModuleName, StringComparison.InvariantCultureIgnoreCase))
+			{
+				Log.TraceWarning("Module '{0}' (referenced via {1}) has incorrect text case. Did you mean '{2}'?", ModuleName, ReferenceChain, Module.Name);
+			}
 			return Module;
 		}
 
