UnrealBuildTool: Distinguish Internal from Public API

Support "Internal" include directory that will only be added to a Referencing module if the referenced module has the same scope. Intent is to to "hide" APIs from external modules that do not need to include headers that are intended for internal module use only.

diff --git a/Engine/Source/Programs/UnrealBuildTool/Configuration/ModuleRules.cs b/Engine/Source/Programs/UnrealBuildTool/Configuration/ModuleRules.cs
--- a/Engine/Source/Programs/UnrealBuildTool/Configuration/ModuleRules.cs
+++ b/Engine/Source/Programs/UnrealBuildTool/Configuration/ModuleRules.cs
@@ -910,6 +910,11 @@ namespace UnrealBuildTool
 		/// </summary>
 		public List<string> PublicIncludePaths = new List<string>();
 
+		/// <summary>
+		/// (This setting is currently not need as we discover all files from the 'Internal' folder) List of all paths to include files that are exposed to other internal modules
+		/// </summary>
+		public List<string> InternalncludePaths = new List<string>();
+
 		/// <summary>
 		/// List of all paths to this module's internal include files, not exposed to other modules (at least one include to the 'Private' path, more if we want to avoid relative paths)
 		/// </summary>
diff --git a/Engine/Source/Programs/UnrealBuildTool/Configuration/UEBuildModule.cs b/Engine/Source/Programs/UnrealBuildTool/Configuration/UEBuildModule.cs
--- a/Engine/Source/Programs/UnrealBuildTool/Configuration/UEBuildModule.cs
+++ b/Engine/Source/Programs/UnrealBuildTool/Configuration/UEBuildModule.cs
@@ -77,6 +77,11 @@ namespace UnrealBuildTool
 		/// </summary>
 		public readonly HashSet<DirectoryReference> PublicIncludePaths;
 
+		/// <summary>
+		/// Set of all internal include paths
+		/// </summary>
+		public readonly HashSet<DirectoryReference> InternalIncludePaths;
+
 		/// <summary>
 		/// Nested public include paths which used to be added automatically, but are now only added for modules with bNestedPublicIncludePaths set.
 		/// </summary>
@@ -182,6 +187,7 @@ namespace UnrealBuildTool
 			HashSet<string> PublicPreBuildLibraries = HashSetFromOptionalEnumerableStringParameter(Rules.PublicPreBuildLibraries);
 			PublicDefinitions = HashSetFromOptionalEnumerableStringParameter(Rules.PublicDefinitions);
 			PublicIncludePaths = CreateDirectoryHashSet(Rules.PublicIncludePaths);
+			InternalIncludePaths = CreateDirectoryHashSet(Rules.InternalncludePaths);
 			PublicSystemIncludePaths = CreateDirectoryHashSet(Rules.PublicSystemIncludePaths);
 			PublicSystemLibraryPaths = CreateDirectoryHashSet(Rules.PublicSystemLibraryPaths);
 			HashSet<string> PublicAdditionalLibraries = HashSetFromOptionalEnumerableStringParameter(Rules.PublicAdditionalLibraries.Union(PublicPreBuildLibraries));
@@ -425,7 +431,11 @@ namespace UnrealBuildTool
 			{
 				Directories.Add(PublicIncludePath);
 			}
-			foreach(DirectoryReference PrivateIncludePath in PrivateIncludePaths)
+			foreach (DirectoryReference InternalIncludePath in InternalIncludePaths)
+			{
+				Directories.Add(InternalIncludePath);
+			}
+			foreach (DirectoryReference PrivateIncludePath in PrivateIncludePaths)
 			{
 				Directories.Add(PrivateIncludePath);
 			}
@@ -518,6 +528,7 @@ namespace UnrealBuildTool
 		/// Sets up the environment for compiling any module that includes the public interface of this module.
 		/// </summary>
 		public virtual void AddModuleToCompileEnvironment(
+			UEBuildModule? SourceModule,
 			UEBuildBinary? SourceBinary,
 			HashSet<DirectoryReference> IncludePaths,
 			HashSet<DirectoryReference> SystemIncludePaths,
@@ -533,15 +544,21 @@ namespace UnrealBuildTool
 
 			// Add this module's public include paths and definitions.
 			AddIncludePaths(IncludePaths, PublicIncludePaths);
-			if(bLegacyPublicIncludePaths)
+			if (bLegacyPublicIncludePaths)
 			{
 				AddIncludePaths(IncludePaths, LegacyPublicIncludePaths);
 			}
 			SystemIncludePaths.UnionWith(PublicSystemIncludePaths);
 			Definitions.AddRange(PublicDefinitions);
 
+			// Add this module's internal include paths, only if the scope contains the same as the SourceModule's scope
+			if (SourceModule != null && Rules.Context.Scope.Contains(SourceModule.Rules.Context.Scope))
+			{
+				AddIncludePaths(IncludePaths, InternalIncludePaths);
+			}
+
 			// Add the import or export declaration for the module
-			if(Rules.Type == ModuleRules.ModuleType.CPlusPlus)
+			if (Rules.Type == ModuleRules.ModuleType.CPlusPlus)
 			{
 				if(Rules.Target.LinkType == TargetLinkType.Monolithic)
 				{
@@ -611,7 +628,7 @@ namespace UnrealBuildTool
 			// Now set up the compile environment for the modules in the original order that we encountered them
 			foreach (UEBuildModule Module in ModuleToIncludePathsOnlyFlag.Keys)
 			{
-				Module.AddModuleToCompileEnvironment(Binary, IncludePaths, SystemIncludePaths, ModuleInterfacePaths, Definitions, AdditionalFrameworks, AdditionalPrerequisites, bWithLegacyPublicIncludePaths);
+				Module.AddModuleToCompileEnvironment(this, Binary, IncludePaths, SystemIncludePaths, ModuleInterfacePaths, Definitions, AdditionalFrameworks, AdditionalPrerequisites, bWithLegacyPublicIncludePaths);
 			}
 		}
 
@@ -1078,6 +1095,7 @@ namespace UnrealBuildTool
 
 			ExportJsonStringArray(Writer, "PublicSystemIncludePaths", PublicSystemIncludePaths.Select(x => x.FullName));
 			ExportJsonStringArray(Writer, "PublicIncludePaths", PublicIncludePaths.Select(x => x.FullName));
+			ExportJsonStringArray(Writer, "InternalIncludePaths", PublicIncludePaths.Select(x => x.FullName));
 			ExportJsonStringArray(Writer, "PrivateIncludePaths", PrivateIncludePaths.Select(x => x.FullName));
 			ExportJsonStringArray(Writer, "PublicLibraries", PublicLibraries.Select(x => x.FullName));
 			ExportJsonStringArray(Writer, "PublicSystemLibraries", PublicSystemLibraries);
diff --git a/Engine/Source/Programs/UnrealBuildTool/Configuration/UEBuildModuleCPP.cs b/Engine/Source/Programs/UnrealBuildTool/Configuration/UEBuildModuleCPP.cs
--- a/Engine/Source/Programs/UnrealBuildTool/Configuration/UEBuildModuleCPP.cs
+++ b/Engine/Source/Programs/UnrealBuildTool/Configuration/UEBuildModuleCPP.cs
@@ -245,6 +245,13 @@ namespace UnrealBuildTool
 					EnumerateLegacyIncludePaths(DirectoryItem.GetItemByDirectoryReference(PublicDirectory), ExcludeNames, LegacyPublicIncludePaths);
 				}
 
+				// Add the 'internal' directory, if it exists
+				DirectoryReference InternalDirectory = DirectoryReference.Combine(ModuleDir, "Internal");
+				if (DirectoryLookupCache.DirectoryExists(InternalDirectory))
+				{
+					InternalIncludePaths.Add(InternalDirectory);
+				}
+
 				// Add the base private directory for this module
 				DirectoryReference PrivateDirectory = DirectoryReference.Combine(ModuleDir, "Private");
 				if (DirectoryLookupCache.DirectoryExists(PrivateDirectory))
@@ -284,6 +291,7 @@ namespace UnrealBuildTool
 		/// Sets up the environment for compiling any module that includes the public interface of this module.
 		/// </summary>
 		public override void AddModuleToCompileEnvironment(
+			UEBuildModule? SourceModule,
 			UEBuildBinary? SourceBinary,
 			HashSet<DirectoryReference> IncludePaths,
 			HashSet<DirectoryReference> SystemIncludePaths,
@@ -314,7 +322,7 @@ namespace UnrealBuildTool
 
 			ModuleInterfacePaths.Add(UEToolChain.GetModuleInterfaceDir(IntermediateDirectory));
 
-			base.AddModuleToCompileEnvironment(SourceBinary, IncludePaths, SystemIncludePaths, ModuleInterfacePaths, Definitions, AdditionalFrameworks, AdditionalPrerequisites, bLegacyPublicIncludePaths);
+			base.AddModuleToCompileEnvironment(SourceModule, SourceBinary, IncludePaths, SystemIncludePaths, ModuleInterfacePaths, Definitions, AdditionalFrameworks, AdditionalPrerequisites, bLegacyPublicIncludePaths);
 		}
 
 		// UEBuildModule interface.
@@ -1444,7 +1452,7 @@ namespace UnrealBuildTool
 			// Now set up the compile environment for the modules in the original order that we encountered them
 			foreach (UEBuildModule Module in ModuleToIncludePathsOnlyFlag.Keys)
 			{
-				Module.AddModuleToCompileEnvironment(null, CompileEnvironment.UserIncludePaths, CompileEnvironment.SystemIncludePaths, CompileEnvironment.ModuleInterfacePaths, CompileEnvironment.Definitions, CompileEnvironment.AdditionalFrameworks, CompileEnvironment.AdditionalPrerequisites, Rules.bLegacyPublicIncludePaths);
+				Module.AddModuleToCompileEnvironment(this, null, CompileEnvironment.UserIncludePaths, CompileEnvironment.SystemIncludePaths, CompileEnvironment.ModuleInterfacePaths, CompileEnvironment.Definitions, CompileEnvironment.AdditionalFrameworks, CompileEnvironment.AdditionalPrerequisites, Rules.bLegacyPublicIncludePaths);
 			}
 			return CompileEnvironment;
 		}
diff --git a/Engine/Source/Programs/UnrealBuildTool/Configuration/UEBuildTarget.cs b/Engine/Source/Programs/UnrealBuildTool/Configuration/UEBuildTarget.cs
--- a/Engine/Source/Programs/UnrealBuildTool/Configuration/UEBuildTarget.cs
+++ b/Engine/Source/Programs/UnrealBuildTool/Configuration/UEBuildTarget.cs
@@ -2231,7 +2231,7 @@ namespace UnrealBuildTool
 			List<string> Definitions = new List<string>(GlobalCompileEnvironment.Definitions);
 			foreach(UEBuildModule Module in Binary.Modules)
 			{
-				Module.AddModuleToCompileEnvironment(null, new HashSet<DirectoryReference>(), new HashSet<DirectoryReference>(), new HashSet<DirectoryReference>(), Definitions, new List<UEBuildFramework>(), new List<FileItem>(), false);
+				Module.AddModuleToCompileEnvironment(null, null, new HashSet<DirectoryReference>(), new HashSet<DirectoryReference>(), new HashSet<DirectoryReference>(), Definitions, new List<UEBuildFramework>(), new List<FileItem>(), false);
 			}
 
 			// Write the header
@@ -3845,6 +3845,7 @@ namespace UnrealBuildTool
 			// Make sure include paths don't end in trailing slashes. This can result in enclosing quotes being escaped when passed to command line tools.
 			RemoveTrailingSlashes(RulesObject.PublicIncludePaths);
 			RemoveTrailingSlashes(RulesObject.PublicSystemIncludePaths);
+			RemoveTrailingSlashes(RulesObject.InternalncludePaths);
 			RemoveTrailingSlashes(RulesObject.PrivateIncludePaths);
 			RemoveTrailingSlashes(RulesObject.PublicSystemLibraryPaths);
 
@@ -3943,6 +3944,7 @@ namespace UnrealBuildTool
 					if (RulesObject.File.IsUnderDirectory(ProjectSourceDirectoryName))
 					{
 						RulesObject.PublicIncludePaths = CombinePathList(ProjectSourceDirectoryName, RulesObject.PublicIncludePaths);
+						RulesObject.InternalncludePaths = CombinePathList(ProjectSourceDirectoryName, RulesObject.InternalncludePaths);
 						RulesObject.PrivateIncludePaths = CombinePathList(ProjectSourceDirectoryName, RulesObject.PrivateIncludePaths);
 						RulesObject.PublicSystemLibraryPaths = CombinePathList(ProjectSourceDirectoryName, RulesObject.PublicSystemLibraryPaths);
 					}
diff --git a/Engine/Source/Programs/UnrealBuildTool/Platform/HoloLens/UEBuildHoloLens.cs b/Engine/Source/Programs/UnrealBuildTool/Platform/HoloLens/UEBuildHoloLens.cs
--- a/Engine/Source/Programs/UnrealBuildTool/Platform/HoloLens/UEBuildHoloLens.cs
+++ b/Engine/Source/Programs/UnrealBuildTool/Platform/HoloLens/UEBuildHoloLens.cs
@@ -509,6 +509,7 @@ namespace UnrealBuildTool
 				// The HoloLensSDK module handles proper paths and libs for HoloLens.
 				// However, the D3D11RHI module will include the DX11 module.
 				Rules.PublicIncludePaths.Clear();
+				Rules.InternalncludePaths.Clear();
 				Rules.PublicSystemLibraryPaths.Clear();
 				Rules.PublicSystemLibraries.Clear();
 				Rules.PublicAdditionalLibraries.Clear();
diff --git a/Engine/Source/Programs/UnrealBuildTool/ProjectFiles/Rider/RiderProjectFile.cs b/Engine/Source/Programs/UnrealBuildTool/ProjectFiles/Rider/RiderProjectFile.cs
--- a/Engine/Source/Programs/UnrealBuildTool/ProjectFiles/Rider/RiderProjectFile.cs
+++ b/Engine/Source/Programs/UnrealBuildTool/ProjectFiles/Rider/RiderProjectFile.cs
@@ -284,7 +284,8 @@ namespace UnrealBuildTool
 
 			ExportJsonStringArray(Writer, "PublicSystemIncludePaths", Module.PublicSystemIncludePaths.Select(x => x.FullName ));
 			ExportJsonStringArray(Writer, "PublicIncludePaths", Module.PublicIncludePaths.Select(x => x.FullName ));
-			
+			ExportJsonStringArray(Writer, "InternalIncludePaths", Module.InternalIncludePaths.Select(x => x.FullName));
+
 			ExportJsonStringArray(Writer, "LegacyPublicIncludePaths", Module.LegacyPublicIncludePaths.Select(x => x.FullName ));
 			
 			ExportJsonStringArray(Writer, "PrivateIncludePaths", Module.PrivateIncludePaths.Select(x => x.FullName));
diff --git a/Engine/Source/Programs/UnrealBuildTool/ProjectFiles/VisualStudioCode/VSCodeProjectFileGenerator.cs b/Engine/Source/Programs/UnrealBuildTool/ProjectFiles/VisualStudioCode/VSCodeProjectFileGenerator.cs
--- a/Engine/Source/Programs/UnrealBuildTool/ProjectFiles/VisualStudioCode/VSCodeProjectFileGenerator.cs
+++ b/Engine/Source/Programs/UnrealBuildTool/ProjectFiles/VisualStudioCode/VSCodeProjectFileGenerator.cs
@@ -1588,7 +1588,7 @@ namespace UnrealBuildTool
 
 		private void GetExcludePathsCPP(DirectoryReference BaseDir, List<string> PathsToExclude)
 		{
-			string[] DirWhiteList = { "Binaries", "Build", "Config", "Plugins", "Source", "Private", "Public", "Classes", "Resources" };
+			string[] DirWhiteList = { "Binaries", "Build", "Config", "Plugins", "Source", "Private", "Public", "Internal", "Classes", "Resources" };
 			foreach (DirectoryReference SubDir in DirectoryReference.EnumerateDirectories(BaseDir, "*", SearchOption.TopDirectoryOnly))
 			{
 				if (Array.Find(DirWhiteList, Dir => Dir == SubDir.GetDirectoryName()) == null)
