vcpkg: Include paths should be added to PublicSystemIncludePaths, not PublicIncludePaths

#rnx

diff --git a/Engine/Source/Programs/UnrealBuildTool/Configuration/ModuleRules.cs b/Engine/Source/Programs/UnrealBuildTool/Configuration/ModuleRules.cs
--- a/Engine/Source/Programs/UnrealBuildTool/Configuration/ModuleRules.cs
+++ b/Engine/Source/Programs/UnrealBuildTool/Configuration/ModuleRules.cs
@@ -1601,7 +1601,7 @@ namespace UnrealBuildTool
 					throw new DirectoryNotFoundException(IncludePath);
 				}
 
-				PublicIncludePaths.Add(Path.Combine(VcPackageRoot, "include"));
+				PublicSystemIncludePaths.Add(Path.Combine(VcPackageRoot, "include"));
 			}
 		}
 	}
