GDK:  Build fix for packages 1/2

- swapped order of Test and Shipping entries in UnrealTargetConfiguration so that they are roughly in increasing optimization order

diff --git a/Engine/Source/Programs/UnrealBuildTool/Configuration/UEBuildTarget.cs b/Engine/Source/Programs/UnrealBuildTool/Configuration/UEBuildTarget.cs
--- a/Engine/Source/Programs/UnrealBuildTool/Configuration/UEBuildTarget.cs
+++ b/Engine/Source/Programs/UnrealBuildTool/Configuration/UEBuildTarget.cs
@@ -563,7 +563,7 @@ namespace UnrealBuildTool
 	}
 
 	/// <summary>
-	/// The type of configuration a target can be built for
+	/// The type of configuration a target can be built for.  Roughly order by optimization level. 
 	/// </summary>
 	public enum UnrealTargetConfiguration
 	{
@@ -588,14 +588,14 @@ namespace UnrealBuildTool
 		Development,
 
 		/// <summary>
-		/// Shipping configuration
+		/// Test configuration
 		/// </summary>
-		Shipping,
+		Test,
 
 		/// <summary>
-		/// Test configuration
+		/// Shipping configuration
 		/// </summary>
-		Test,
+		Shipping,
 	}
 
 	/// <summary>
