Backing out CL 16857053 to see if this fixes the hangs we've been seeing during Mac and Linux compiles.

#rnx

diff --git a/Engine/Source/Programs/UnrealBuildTool/Executors/ParallelExecutor.cs b/Engine/Source/Programs/UnrealBuildTool/Executors/ParallelExecutor.cs
--- a/Engine/Source/Programs/UnrealBuildTool/Executors/ParallelExecutor.cs
+++ b/Engine/Source/Programs/UnrealBuildTool/Executors/ParallelExecutor.cs
@@ -2,7 +2,6 @@
 
 using EpicGames.Core;
 using System;
-using System.Collections.Concurrent;
 using System.Collections.Generic;
 using System.Diagnostics;
 using System.Linq;
@@ -185,7 +184,7 @@ namespace UnrealBuildTool
 				// Create a job object for all the child processes
 				bool bResult = true;
 				Dictionary<BuildAction, Thread> ExecutingActions = new Dictionary<BuildAction, Thread>();
-				ConcurrentQueue<BuildAction> CompletedActions = new ConcurrentQueue<BuildAction>();
+				List<BuildAction> CompletedActions = new List<BuildAction>();
 
 				TimeSpan TotalProcessingTime;
 				using (ManagedProcessGroup ProcessGroup = new ManagedProcessGroup())
@@ -207,7 +206,8 @@ namespace UnrealBuildTool
 									QueuedActions.RemoveAt(QueuedActions.Count - 1);
 
 									Thread ExecutingThread = new Thread(() => { ExecuteAction(ProcessGroup, Action, CompletedActions, CompletedEvent); });
-									ExecutingThread.Name = $"Build:{(Action.Inner.CommandDescription != null ? Action.Inner.CommandDescription : Action.Inner.CommandPath.GetFileName())} {Action.Inner.StatusDescription}".Trim();
+									string Description = $"{(Action.Inner.CommandDescription != null ? Action.Inner.CommandDescription : Action.Inner.CommandPath.GetFileName())} {Action.Inner.StatusDescription}".Trim();
+									ExecutingThread.Name = String.Format("Build:{0}", Description);
 									ExecutingThread.Start();
 
 									ExecutingActions.Add(Action, ExecutingThread);
@@ -217,73 +217,76 @@ namespace UnrealBuildTool
 								CompletedEvent.WaitOne();
 
 								// Wait for something to finish and flush it to the log
-								while (CompletedActions.TryDequeue(out BuildAction? CompletedAction))
+								lock (CompletedActions)
 								{
-									// Join the thread
-									Thread CompletedThread = ExecutingActions[CompletedAction];
-									CompletedThread.Join();
-									ExecutingActions.Remove(CompletedAction);
+									foreach (BuildAction CompletedAction in CompletedActions)
+									{
+										// Join the thread
+										Thread CompletedThread = ExecutingActions[CompletedAction];
+										CompletedThread.Join();
+										ExecutingActions.Remove(CompletedAction);
 
-									// Update the progress
-									NumCompletedActions++;
-									ProgressWriter.Write(NumCompletedActions, InputActions.Count);
+										// Update the progress
+										NumCompletedActions++;
+										ProgressWriter.Write(NumCompletedActions, InputActions.Count);
 
-									string Description = string.Empty;
+										string Description = string.Empty;
 
-									// Write it to the log
-									if (CompletedAction.Inner.bShouldOutputStatusDescription || CompletedAction.LogLines.Count == 0)
-									{
-										Description = $"{(CompletedAction.Inner.CommandDescription != null ? CompletedAction.Inner.CommandDescription : CompletedAction.Inner.CommandPath.GetFileNameWithoutExtension())} {CompletedAction.Inner.StatusDescription}".Trim();
-									}
-									else if (CompletedAction.LogLines.Count > 0)
-									{
-										Description = $"{(CompletedAction.Inner.CommandDescription != null ? CompletedAction.Inner.CommandDescription : CompletedAction.Inner.CommandPath.GetFileNameWithoutExtension())} {CompletedAction.LogLines[0]}".Trim();
-									}
+										// Write it to the log
+										if (CompletedAction.Inner.bShouldOutputStatusDescription || CompletedAction.LogLines.Count == 0)
+										{
+											Description = $"{(CompletedAction.Inner.CommandDescription != null ? CompletedAction.Inner.CommandDescription : CompletedAction.Inner.CommandPath.GetFileNameWithoutExtension())} {CompletedAction.Inner.StatusDescription}".Trim();
+										}
+										else if (CompletedAction.LogLines.Count > 0)
+										{
+											Description = $"{(CompletedAction.Inner.CommandDescription != null ? CompletedAction.Inner.CommandDescription : CompletedAction.Inner.CommandPath.GetFileNameWithoutExtension())} {CompletedAction.LogLines[0]}".Trim();
+										}
 
-									Log.TraceInformation("[{0}/{1}] {2}", NumCompletedActions, InputActions.Count, Description);
-									foreach (string Line in CompletedAction.LogLines.Skip(CompletedAction.Inner.bShouldOutputStatusDescription ? 0 : 1))
-									{
-										Log.TraceInformation(Line);
-									}
+										Log.TraceInformation("[{0}/{1}] {2}", NumCompletedActions, InputActions.Count, Description);
+										foreach (string Line in CompletedAction.LogLines.Skip(CompletedAction.Inner.bShouldOutputStatusDescription ? 0 : 1))
+										{
+											Log.TraceInformation(Line);
+										}
 
-									AllCompletedActions.Add(CompletedAction);
+										AllCompletedActions.Add(CompletedAction);
 
-									// Check the exit code
-									if (CompletedAction.ExitCode == 0)
-									{
-										// Mark all the dependents as done
-										foreach (BuildAction DependantAction in CompletedAction.Dependants)
+										// Check the exit code
+										if (CompletedAction.ExitCode == 0)
 										{
-											if (--DependantAction.MissingDependencyCount == 0)
+											// Mark all the dependents as done
+											foreach (BuildAction DependantAction in CompletedAction.Dependants)
 											{
-												QueuedActions.Add(DependantAction);
+												if (--DependantAction.MissingDependencyCount == 0)
+												{
+													QueuedActions.Add(DependantAction);
+												}
 											}
 										}
-									}
-									else
-									{
-										// BEGIN TEMPORARY TO CATCH PVS-STUDIO ISSUES
-										if (CompletedAction.LogLines.Count == 0)
-										{
-											Log.TraceInformation("[{0}/{1}] {2} - Error but no output", NumCompletedActions, InputActions.Count, Description);
-											Log.TraceInformation("[{0}/{1}] {2} - {3} {4} {5} {6}", NumCompletedActions, InputActions.Count, Description, CompletedAction.ExitCode,
-												CompletedAction.Inner.WorkingDirectory, CompletedAction.Inner.CommandPath, CompletedAction.Inner.CommandArguments);
-										}
-										// END TEMPORARY
-										// Update the exit code if it's not already set
-										if (bResult && CompletedAction.ExitCode != 0)
+										else
 										{
-											bResult = false;
+											// BEGIN TEMPORARY TO CATCH PVS-STUDIO ISSUES
+											if (CompletedAction.LogLines.Count == 0)
+											{
+												Log.TraceInformation("[{0}/{1}] {2} - Error but no output", NumCompletedActions, InputActions.Count, Description);
+												Log.TraceInformation("[{0}/{1}] {2} - {3} {4} {5} {6}", NumCompletedActions, InputActions.Count, Description, CompletedAction.ExitCode,
+													CompletedAction.Inner.WorkingDirectory, CompletedAction.Inner.CommandPath, CompletedAction.Inner.CommandArguments);
+											}
+											// END TEMPORARY
+											// Update the exit code if it's not already set
+											if (bResult && CompletedAction.ExitCode != 0)
+											{
+												bResult = false;
+											}
 										}
 									}
+									CompletedActions.Clear();
 								}
-								CompletedActions.Clear();
-							}
 
-							// If we've already got a non-zero exit code, clear out the list of queued actions so nothing else will run
-							if (!bResult && bStopCompilationAfterErrors)
-							{
-								QueuedActions.Clear();
+								// If we've already got a non-zero exit code, clear out the list of queued actions so nothing else will run
+								if (!bResult && bStopCompilationAfterErrors)
+								{
+									QueuedActions.Clear();
+								}
 							}
 						}
 					}
@@ -331,7 +334,7 @@ namespace UnrealBuildTool
 		/// <param name="Action">The action to execute</param>
 		/// <param name="CompletedActions">On completion, the list to add the completed action to</param>
 		/// <param name="CompletedEvent">Event to set once an event is complete</param>
-		static void ExecuteAction(ManagedProcessGroup ProcessGroup, BuildAction Action, ConcurrentQueue<BuildAction> CompletedActions, AutoResetEvent CompletedEvent)
+		static void ExecuteAction(ManagedProcessGroup ProcessGroup, BuildAction Action, List<BuildAction> CompletedActions, AutoResetEvent CompletedEvent)
 		{
 			try
 			{
@@ -349,7 +352,10 @@ namespace UnrealBuildTool
 				Action.ExitCode = 1;
 			}
 
-			CompletedActions.Enqueue(Action);
+			lock (CompletedActions)
+			{
+				CompletedActions.Add(Action);
+			}
 
 			CompletedEvent.Set();
 		}
