RulesCompiler: Add a batched FindAllRulesFiles() override

Batched function allows greater overlap when scanning the filesystem.
AutomationTool startup time is reduced by ~150ms on my machine

diff --git a/Engine/Source/Programs/UnrealBuildTool/System/RulesCompiler.cs b/Engine/Source/Programs/UnrealBuildTool/System/RulesCompiler.cs
--- a/Engine/Source/Programs/UnrealBuildTool/System/RulesCompiler.cs
+++ b/Engine/Source/Programs/UnrealBuildTool/System/RulesCompiler.cs
@@ -141,21 +141,7 @@ namespace UnrealBuildTool
 				}
 			}
 
-			// Iterate over all the folders to check
-			List<FileReference> SourceFiles = new List<FileReference>();
-			HashSet<FileReference> UniqueSourceFiles = new HashSet<FileReference>();
-			foreach (DirectoryReference Folder in Folders)
-			{
-				IReadOnlyList<FileReference> SourceFilesForFolder = FindAllRulesFiles(Folder, RulesFileType);
-				foreach (FileReference SourceFile in SourceFilesForFolder)
-				{
-					if (UniqueSourceFiles.Add(SourceFile))
-					{
-						SourceFiles.Add(SourceFile);
-					}
-				}
-			}
-			return SourceFiles;
+			return FindAllRulesFiles(Folders, RulesFileType);
 		}
 
 		/// <summary>
@@ -213,39 +199,64 @@ namespace UnrealBuildTool
 		/// <returns>List of rules files of the given type</returns>
 		private static IReadOnlyList<FileReference> FindAllRulesFiles(DirectoryReference Directory, RulesFileType Type)
 		{
-			// Check to see if we've already cached source files for this folder
-			RulesFileCache? Cache;
-			if (!RootFolderToRulesFileCache.TryGetValue(Directory, out Cache))
+			return FindAllRulesFiles(new List<DirectoryReference> { Directory }, Type);
+		}
+
+		/// <summary>
+		/// Finds all the rules of the given type under a given directories
+		/// </summary>
+		/// <param name="Directories">Directories to search</param>
+		/// <param name="Type">Type of rules to return</param>
+		/// <returns>List of rules files of the given type</returns>
+		private static List<FileReference> FindAllRulesFiles(IEnumerable<DirectoryReference> Directories, RulesFileType Type)
+		{
+			List<(DirectoryReference, RulesFileCache)> Caches = new List<(DirectoryReference, RulesFileCache)>(Directories.Count());
+			using (ThreadPoolWorkQueue Queue = new ThreadPoolWorkQueue())
 			{
-				Cache = new RulesFileCache();
-				using(ThreadPoolWorkQueue Queue = new ThreadPoolWorkQueue())
+				foreach (DirectoryReference Directory in Directories)
 				{
-					DirectoryItem BaseDirectory = DirectoryItem.GetItemByDirectoryReference(Directory);
-					Queue.Enqueue(() => FindAllRulesFilesRecursively(BaseDirectory, Cache, Queue));
+					RulesFileCache? Cache;
+					if (!RootFolderToRulesFileCache.TryGetValue(Directory, out Cache))
+					{
+						Cache = new RulesFileCache();
+						Queue.Enqueue(() => FindAllRulesFilesRecursively(DirectoryItem.GetItemByDirectoryReference(Directory), Cache, Queue));
+					}
+					Caches.Add((Directory, Cache));
 				}
-				Cache.ModuleRules.Sort((A, B) => A.FullName.CompareTo(B.FullName));
-				Cache.TargetRules.Sort((A, B) => A.FullName.CompareTo(B.FullName));
-				Cache.AutomationModules.Sort((A, B) => A.FullName.CompareTo(B.FullName));
-				RootFolderToRulesFileCache[Directory] = Cache;
 			}
 
-			// Get the list of files of the type we're looking for
-			if (Type == RulesCompiler.RulesFileType.Module)
-			{
-				return Cache.ModuleRules;
-			}
-			else if (Type == RulesCompiler.RulesFileType.Target)
+			List<FileReference> Files = new List<FileReference>();
+			
+			foreach ((DirectoryReference Directory, RulesFileCache Cache) in Caches)
 			{
-				return Cache.TargetRules;
-			}
-			else if (Type == RulesCompiler.RulesFileType.AutomationModule)
-			{
-				return Cache.AutomationModules;
-			}
-			else
-			{
-				throw new BuildException("Unhandled rules type: {0}", Type);
+				if (!RootFolderToRulesFileCache.ContainsKey(Directory))
+				{
+					Cache.ModuleRules.Sort((A, B) => A.FullName.CompareTo(B.FullName));
+					Cache.TargetRules.Sort((A, B) => A.FullName.CompareTo(B.FullName));
+					Cache.AutomationModules.Sort((A, B) => A.FullName.CompareTo(B.FullName));
+					RootFolderToRulesFileCache[Directory] = Cache;
+				}
+
+				// Get the list of files of the type we're looking for
+				if (Type == RulesFileType.Module)
+				{
+					Files.AddRange(Cache.ModuleRules);
+				}
+				else if (Type == RulesFileType.Target)
+				{
+					Files.AddRange(Cache.TargetRules);
+				}
+				else if (Type == RulesFileType.AutomationModule)
+				{
+					Files.AddRange(Cache.AutomationModules);
+				}
+				else
+				{
+					throw new Exception($"Unhandled rules type: {Type}");
+				}
 			}
+
+			return Files;
 		}
 
 		/// <summary>
