UnrealBuildTool: QMakefileGenerator.cs

Use StringBuilders to significantly reduce runtime. More could be done, but this seems to address the slowest part.

#trivial

diff --git a/Engine/Source/Programs/UnrealBuildTool/ProjectFiles/QMake/QMakefileGenerator.cs b/Engine/Source/Programs/UnrealBuildTool/ProjectFiles/QMake/QMakefileGenerator.cs
--- a/Engine/Source/Programs/UnrealBuildTool/ProjectFiles/QMake/QMakefileGenerator.cs
+++ b/Engine/Source/Programs/UnrealBuildTool/ProjectFiles/QMake/QMakefileGenerator.cs
@@ -249,9 +249,9 @@ namespace UnrealBuildTool
 
 			string QMakeSectionEnd = " \n\n";
 
-			string QMakeSourceFilesList = "SOURCES += \\ \n";
-			string QMakeHeaderFilesList = "HEADERS += \\ \n";
-			string QMakeConfigFilesList = "OTHER_FILES += \\ \n";
+			StringBuilder QMakeSourceFilesListBuilder = new StringBuilder("SOURCES += \\ \n");
+			StringBuilder QMakeHeaderFilesListBuilder = new StringBuilder("HEADERS += \\ \n");
+			StringBuilder QMakeConfigFilesListBuilder = new StringBuilder("OTHER_FILES += \\ \n");
 			string QMakeTargetList = "QMAKE_EXTRA_TARGETS += \\ \n";
 
 			if (!String.IsNullOrEmpty(GameProjectName))
@@ -295,17 +295,17 @@ namespace UnrealBuildTool
 						{
 							if (!SourceFileRelativeToRoot.StartsWith("..") && !Path.IsPathRooted(SourceFileRelativeToRoot))
 							{
-								QMakeSourceFilesList += ("\t\"" + "$$unrealRootPath/Engine/" + SourceFileRelativeToRoot + "\" \\\n");
+								QMakeSourceFilesListBuilder.Append("\t\"" + "$$unrealRootPath/Engine/" + SourceFileRelativeToRoot + "\" \\\n");
 							}
 							else
 							{
 								if (String.IsNullOrEmpty(GameProjectName))
 								{
-									QMakeSourceFilesList += ("\t\"" + SourceFileRelativeToRoot.Substring(3) + "\" \\\n");
+									QMakeSourceFilesListBuilder.Append("\t\"" + SourceFileRelativeToRoot.Substring(3) + "\" \\\n");
 								}
 								else
 								{
-									QMakeSourceFilesList += ("\t\"$$" + GameProjectName + "RootPath/" + Utils.MakePathRelativeTo(CurSourceFile.FullName, GameProjectPath) + "\" \\\n");
+									QMakeSourceFilesListBuilder.Append("\t\"$$" + GameProjectName + "RootPath/" + Utils.MakePathRelativeTo(CurSourceFile.FullName, GameProjectPath) + "\" \\\n");
 								}
 							}
 						}
@@ -314,18 +314,18 @@ namespace UnrealBuildTool
 							if (!SourceFileRelativeToRoot.StartsWith("..") && !Path.IsPathRooted(SourceFileRelativeToRoot))
 							{
 								// SourceFileRelativeToRoot = "Engine/" + SourceFileRelativeToRoot;
-								QMakeHeaderFilesList += ("\t\"" + "$$unrealRootPath/Engine/" + SourceFileRelativeToRoot + "\" \\\n");
+								QMakeHeaderFilesListBuilder.Append("\t\"" + "$$unrealRootPath/Engine/" + SourceFileRelativeToRoot + "\" \\\n");
 							}
 							else
 							{
 								if (String.IsNullOrEmpty(GameProjectName))
 								{
 									// SourceFileRelativeToRoot = SourceFileRelativeToRoot.Substring (3);
-									QMakeHeaderFilesList += ("\t\"" + SourceFileRelativeToRoot.Substring(3) + "\" \\\n");
+									QMakeHeaderFilesListBuilder.Append("\t\"" + SourceFileRelativeToRoot.Substring(3) + "\" \\\n");
 								}
 								else
 								{
-									QMakeHeaderFilesList += ("\t\"$$" + GameProjectName + "RootPath/" + Utils.MakePathRelativeTo(CurSourceFile.FullName, GameProjectPath) + "\" \\\n");
+									QMakeHeaderFilesListBuilder.Append("\t\"$$" + GameProjectName + "RootPath/" + Utils.MakePathRelativeTo(CurSourceFile.FullName, GameProjectPath) + "\" \\\n");
 								}
 							}
 						}
@@ -334,7 +334,7 @@ namespace UnrealBuildTool
 							if (!SourceFileRelativeToRoot.StartsWith("..") && !Path.IsPathRooted(SourceFileRelativeToRoot))
 							{
 								// SourceFileRelativeToRoot = "Engine/" + SourceFileRelativeToRoot;
-								QMakeConfigFilesList += ("\t\"" + "$$unrealRootPath/Engine/" + SourceFileRelativeToRoot + "\" \\\n");
+								QMakeConfigFilesListBuilder.Append("\t\"" + "$$unrealRootPath/Engine/" + SourceFileRelativeToRoot + "\" \\\n");
 
 							}
 							else
@@ -342,11 +342,11 @@ namespace UnrealBuildTool
 								if (String.IsNullOrEmpty(GameProjectName))
 								{
 									// SourceFileRelativeToRoot = SourceFileRelativeToRoot.Substring (3);
-									QMakeConfigFilesList += ("\t\"" + SourceFileRelativeToRoot.Substring(3) + "\" \\\n");
+									QMakeConfigFilesListBuilder.Append("\t\"" + SourceFileRelativeToRoot.Substring(3) + "\" \\\n");
 								}
 								else
 								{
-									QMakeConfigFilesList += ("\t\"$$" + GameProjectName + "RootPath/" + Utils.MakePathRelativeTo(CurSourceFile.FullName, GameProjectPath) + "\" \\\n");
+									QMakeConfigFilesListBuilder.Append("\t\"$$" + GameProjectName + "RootPath/" + Utils.MakePathRelativeTo(CurSourceFile.FullName, GameProjectPath) + "\" \\\n");
 								}
 							}
 						}
@@ -356,14 +356,14 @@ namespace UnrealBuildTool
 			}
 
 			// Add section end to section strings;
-			QMakeSourceFilesList += QMakeSectionEnd;
-			QMakeHeaderFilesList += QMakeSectionEnd;
-			QMakeConfigFilesList += QMakeSectionEnd;
+			QMakeSourceFilesListBuilder.Append(QMakeSectionEnd);
+			QMakeHeaderFilesListBuilder.Append(QMakeSectionEnd);
+			QMakeConfigFilesListBuilder.Append(QMakeSectionEnd);
 
 			// Append sections to the QMakeLists.txt file
-			QMakeSourcePriFileContent.Append(QMakeSourceFilesList);
-			QMakeHeaderPriFileContent.Append(QMakeHeaderFilesList);
-			QMakeConfigPriFileContent.Append(QMakeConfigFilesList);
+			QMakeSourcePriFileContent.Append(QMakeSourceFilesListBuilder);
+			QMakeHeaderPriFileContent.Append(QMakeHeaderFilesListBuilder);
+			QMakeConfigPriFileContent.Append(QMakeConfigFilesListBuilder);
 
 			string QMakeProjectCmdArg = "";
 
