UnrealBuildTool:

Add a -Help option that prints descriptions of global options.
(Tool mode options are not available - support for those may be added in a future CL)

Example output:

Global options:
  -Help               :  Display this help.
  -Verbose            :  Increase output verbosity
  -VeryVerbose        :  Increase output verbosity more
  -Log                :  Specify a log file location instead of the default Engine/Programs/UnrealBuildTool/Log.txt
  -Timestamps         :  Include timestamps in the log
  -FromMsBuild        :  Format messages for msbuild
  -Progress           :  Write progress messages in a format that can be parsed by other programs
  -NoMutex            :  Allow more than one instance of the program to run at once
  -WaitMutex          :  Wait for another instance to finish and then start, rather than aborting immediately
  -RemoteIni          :  Remote tool ini directory
  -Mode=              :  Select tool mode. One of the following (default tool mode is "Build"):
                           AggregateParsedTimingInfo, Build, Clean, Deploy, Execute, GenerateClangDatabase, GenerateProjectFiles,
                           IOSPostBuildSync, JsonExport, ParseMsvcTimingInfo, PVSGather, QueryTargets, SetupPlatforms,
                           ValidatePlatforms, WriteDocumentation, WriteMetadata
  -Clean              :  Clean build products. Equivalent to -Mode=Clean
  -ProjectFiles       :  Generate project files based on IDE preference. Equivalent to -Mode=GenerateProjectFiles
  -ProjectFileFormat= :  Generate project files in specified format. May be used multiple times.
  -Makefile           :  Generate Linux Makefile
  -CMakefile          :  Generate project files for CMake
  -QMakefile          :  Generate project files for QMake
  -KDevelopfile       :  Generate project files for KDevelop
  -CodeliteFiles      :  Generate project files for Codelite
  -XCodeProjectFiles  :  Generate project files for XCode
  -EddieProjectFiles  :  Generate project files for Eddie
  -VSCode             :  Generate project files for Visual Studio Code
  -VSMac              :  Generate project files for Visual Studio Mac
  -CLion              :  Generate project files for CLion
  -Rider              :  Generate project files for Rider

diff --git a/Engine/Source/Programs/Shared/EpicGames.Core/CommandLineArguments.cs b/Engine/Source/Programs/Shared/EpicGames.Core/CommandLineArguments.cs
--- a/Engine/Source/Programs/Shared/EpicGames.Core/CommandLineArguments.cs
+++ b/Engine/Source/Programs/Shared/EpicGames.Core/CommandLineArguments.cs
@@ -1045,6 +1045,24 @@ namespace EpicGames.Core
 			}
 		}
 
+		/// <summary>
+		/// Checks to see if any arguments are used
+		/// </summary>
+		/// <returns></returns>
+		public bool AreAnyArgumentsUsed()
+		{
+			return UsedArguments.Cast<bool>().Count(B => B) != 0;
+		}
+		
+		/// <summary>
+		/// Count the number of value (non-flag) arguments on the command line
+		/// </summary>
+		/// <returns></returns>
+		public int CountValueArguments()
+		{
+			return FlagArguments.Cast<bool>().Count(B => !B);
+		}
+
 		/// <summary>
 		/// Checks that a given string is a valid argument prefix
 		/// </summary>
diff --git a/Engine/Source/Programs/Shared/EpicGames.Core/CommandLineAttribute.cs b/Engine/Source/Programs/Shared/EpicGames.Core/CommandLineAttribute.cs
--- a/Engine/Source/Programs/Shared/EpicGames.Core/CommandLineAttribute.cs
+++ b/Engine/Source/Programs/Shared/EpicGames.Core/CommandLineAttribute.cs
@@ -35,6 +35,11 @@ namespace EpicGames.Core
 		/// </summary>
 		public char ListSeparator = '\0';
 
+		/// <summary>
+		/// Description of the operation.
+		/// </summary>
+		public string? Description;
+
 		/// <summary>
 		/// Constructor
 		/// </summary>
diff --git a/Engine/Source/Programs/UnrealBuildTool/UnrealBuildTool.cs b/Engine/Source/Programs/UnrealBuildTool/UnrealBuildTool.cs
--- a/Engine/Source/Programs/UnrealBuildTool/UnrealBuildTool.cs
+++ b/Engine/Source/Programs/UnrealBuildTool/UnrealBuildTool.cs
@@ -302,73 +302,82 @@ namespace UnrealBuildTool
 		/// </summary>
 		class GlobalOptions
 		{
+			/// <summary>
+			/// User asked for help
+			/// </summary>
+			[CommandLine(Prefix = "-Help", Description = "Display this help.")]
+			[CommandLine(Prefix = "-h")]
+			[CommandLine(Prefix = "--help")]
+			public bool bGetHelp = false;
+			
 			/// <summary>
 			/// The amount of detail to write to the log
 			/// </summary>
-			[CommandLine(Prefix = "-Verbose", Value ="Verbose")]
-			[CommandLine(Prefix = "-VeryVerbose", Value ="VeryVerbose")]
+			[CommandLine(Prefix = "-Verbose", Value ="Verbose", Description = "Increase output verbosity")]
+			[CommandLine(Prefix = "-VeryVerbose", Value ="VeryVerbose", Description = "Increase output verbosity more")]
 			public LogEventType LogOutputLevel = LogEventType.Log;
 
 			/// <summary>
 			/// Specifies the path to a log file to write. Note that the default mode (eg. building, generating project files) will create a log file by default if this not specified.
 			/// </summary>
-			[CommandLine(Prefix = "-Log")]
+			[CommandLine(Prefix = "-Log", Description = "Specify a log file location instead of the default Engine/Programs/UnrealBuildTool/Log.txt")]
 			public FileReference? LogFileName = null;
 
 			/// <summary>
 			/// Whether to include timestamps in the log
 			/// </summary>
-			[CommandLine(Prefix = "-Timestamps")]
+			[CommandLine(Prefix = "-Timestamps", Description = "Include timestamps in the log")]
 			public bool bLogTimestamps = false;
 
 			/// <summary>
 			/// Whether to format messages in MsBuild format
 			/// </summary>
-			[CommandLine(Prefix = "-FromMsBuild")]
+			[CommandLine(Prefix = "-FromMsBuild", Description = "Format messages for msbuild")]
 			public bool bLogFromMsBuild = false;
 
 			/// <summary>
 			/// Whether to write progress markup in a format that can be parsed by other programs
 			/// </summary>
-			[CommandLine(Prefix = "-Progress")]
+			[CommandLine(Prefix = "-Progress", Description = "Write progress messages in a format that can be parsed by other programs")]
 			public bool bWriteProgressMarkup = false;
 
 			/// <summary>
 			/// Whether to ignore the mutex
 			/// </summary>
-			[CommandLine(Prefix = "-NoMutex")]
+			[CommandLine(Prefix = "-NoMutex", Description = "Allow more than one instance of the program to run at once")]
 			public bool bNoMutex = false;
 
 			/// <summary>
 			/// Whether to wait for the mutex rather than aborting immediately
 			/// </summary>
-			[CommandLine(Prefix = "-WaitMutex")]
+			[CommandLine(Prefix = "-WaitMutex", Description = "Wait for another instance to finish and then start, rather than aborting immediately")]
 			public bool bWaitMutex = false;
 
 			/// <summary>
-			/// Whether to wait for the mutex rather than aborting immediately
 			/// </summary>
-			[CommandLine(Prefix = "-RemoteIni")]
+			[CommandLine(Prefix = "-RemoteIni", Description = "Remote tool ini directory")]
 			public string RemoteIni = "";
 
 			/// <summary>
 			/// The mode to execute
 			/// </summary>
-			[CommandLine]
-			[CommandLine("-Clean", Value="Clean")]
-			[CommandLine("-ProjectFiles", Value="GenerateProjectFiles")]
-			[CommandLine("-ProjectFileFormat=", Value="GenerateProjectFiles")]
-			[CommandLine("-Makefile", Value="GenerateProjectFiles")]
-			[CommandLine("-CMakefile", Value="GenerateProjectFiles")]
-			[CommandLine("-QMakefile", Value="GenerateProjectFiles")]
-			[CommandLine("-KDevelopfile", Value="GenerateProjectFiles")]
-			[CommandLine("-CodeliteFiles", Value="GenerateProjectFiles")]
-			[CommandLine("-XCodeProjectFiles", Value="GenerateProjectFiles")]
-			[CommandLine("-EdditProjectFiles", Value="GenerateProjectFiles")]
-			[CommandLine("-VSCode", Value="GenerateProjectFiles")]
-			[CommandLine("-VSMac", Value="GenerateProjectFiles")]
-			[CommandLine("-CLion", Value="GenerateProjectFiles")]
-			[CommandLine("-Rider", Value="GenerateProjectFiles")]
+			[CommandLine("-Mode=")] // description handling is special-cased in PrintUsage()
+
+			[CommandLine("-Clean", Value="Clean", Description = "Clean build products. Equivalent to -Mode=Clean")]
+
+			[CommandLine("-ProjectFiles", Value="GenerateProjectFiles", Description = "Generate project files based on IDE preference. Equivalent to -Mode=GenerateProjectFiles")]
+			[CommandLine("-ProjectFileFormat=", Value="GenerateProjectFiles", Description = "Generate project files in specified format. May be used multiple times.")]
+			[CommandLine("-Makefile", Value="GenerateProjectFiles", Description = "Generate Linux Makefile")]
+			[CommandLine("-CMakefile", Value="GenerateProjectFiles", Description = "Generate project files for CMake")]
+			[CommandLine("-QMakefile", Value="GenerateProjectFiles", Description = "Generate project files for QMake")]
+			[CommandLine("-KDevelopfile", Value="GenerateProjectFiles", Description = "Generate project files for KDevelop")]
+			[CommandLine("-CodeliteFiles", Value="GenerateProjectFiles", Description = "Generate project files for Codelite")]
+			[CommandLine("-XCodeProjectFiles", Value="GenerateProjectFiles", Description = "Generate project files for XCode")]
+			[CommandLine("-EddieProjectFiles", Value="GenerateProjectFiles", Description = "Generate project files for Eddie")]
+			[CommandLine("-VSCode", Value="GenerateProjectFiles", Description = "Generate project files for Visual Studio Code")]
+			[CommandLine("-VSMac", Value="GenerateProjectFiles", Description = "Generate project files for Visual Studio Mac")]
+			[CommandLine("-CLion", Value="GenerateProjectFiles", Description = "Generate project files for CLion")]
+			[CommandLine("-Rider", Value="GenerateProjectFiles", Description = "Generate project files for Rider")]
 			#if __VPROJECT_AVAILABLE__
 				[CommandLine("-VProject", Value = "GenerateProjectFiles")]
 			#endif
@@ -388,6 +397,79 @@ namespace UnrealBuildTool
 			}
 		}
 
+		/// <summary>
+		/// Get all the valid Modes
+		/// </summary>
+		/// <returns></returns>
+		private static Dictionary<string, Type> GetModes()
+		{
+			Dictionary<string, Type> ModeNameToType = new Dictionary<string, Type>(StringComparer.OrdinalIgnoreCase);
+			foreach (Type Type in Assembly.GetExecutingAssembly().GetTypes())
+			{
+				if (Type.IsClass && !Type.IsAbstract && Type.IsSubclassOf(typeof(ToolMode)))
+				{
+					ToolModeAttribute? Attribute = Type.GetCustomAttribute<ToolModeAttribute>();
+					if (Attribute == null)
+					{
+						throw new BuildException("Class '{0}' should have a ToolModeAttribute", Type.Name);
+					}
+					ModeNameToType.Add(Attribute.Name, Type);
+				}
+			}
+			return ModeNameToType;
+		}
+		public static readonly Dictionary<string, Type> ModeNameToType = GetModes();
+
+		/// <summary>
+		/// Print (incomplete) usage information
+		/// </summary>
+		private static void PrintUsage()
+		{
+			Console.WriteLine("Global options:");
+			int LongestPrefix = 0;
+			foreach (FieldInfo Info in typeof(GlobalOptions).GetFields())
+			{
+				foreach (CommandLineAttribute Att in Info.GetCustomAttributes<CommandLineAttribute>())
+				{
+					if (Att.Prefix != null && Att.Description != null)
+					{
+						LongestPrefix = Att.Prefix.Length > LongestPrefix ? Att.Prefix.Length : LongestPrefix;
+					}
+				}
+			}
+
+			foreach (FieldInfo Info in typeof(GlobalOptions).GetFields())
+			{
+				foreach (CommandLineAttribute Att in Info.GetCustomAttributes<CommandLineAttribute>())
+				{
+					if (Att.Prefix != null && Att.Description != null)
+					{
+						Console.WriteLine($"  {Att.Prefix.PadRight(LongestPrefix)} :  {Att.Description}");
+					}
+
+					// special case for Mode
+					if (String.Equals(Att.Prefix, "-Mode="))
+					{
+						Console.WriteLine($"  {Att.Prefix!.PadRight(LongestPrefix)} :  Select tool mode. One of the following (default tool mode is \"Build\"):");
+						string Indent = "".PadRight(LongestPrefix + 8);
+						string Line = Indent;
+						IOrderedEnumerable<string> SortedModeNames = ModeNameToType.Keys.ToList().OrderBy(Name => Name);
+						foreach (string ModeName in SortedModeNames.SkipLast(1))
+						{
+							Line += $"{ModeName}, ";
+							if (Line.Length > 110)
+							{
+								Console.WriteLine(Line);
+								Line = Indent;
+							}
+						}
+						Line += SortedModeNames.Last();
+						Console.WriteLine(Line);
+					}
+				}
+			}
+		}
+
 		/// <summary>
 		/// Main entry point. Parses any global options and initializes the logging system, then invokes the appropriate command.
 		/// </summary>
@@ -407,6 +489,25 @@ namespace UnrealBuildTool
 				// Parse the global options
 				GlobalOptions Options = new GlobalOptions(Arguments);
 
+				Console.WriteLine(Arguments.CountValueArguments());
+
+				if (
+					// Print usage if there are zero arguments provided
+					ArgumentsArray.Length == 0 
+
+					// Print usage if the user asks for help
+					|| Options.bGetHelp 
+
+					// Print usage if none of the arguments provided were recognized as part of GlobalOptions and there fewer than three value arguments:
+					// For the default Build mode, at least three value arguments are required: target, platform, and configuration
+					// This is an imperfect test of a valid command line - the Build ToolMode should provide further feedback if the command line is not valid
+					|| (!Arguments.AreAnyArgumentsUsed() && Arguments.CountValueArguments() < 3)
+					)
+				{
+					PrintUsage();
+					return Options.bGetHelp ? 0 : 1;
+				}
+				
 				// Configure the log system
 				Log.OutputLevel = Options.LogOutputLevel;
 				Log.IncludeTimestamps = Options.bLogTimestamps;
@@ -439,21 +540,6 @@ namespace UnrealBuildTool
 				Type? ModeType = typeof(BuildMode);
 				if(Options.Mode != null)
 				{
-					// Find all the valid modes
-					Dictionary<string, Type> ModeNameToType = new Dictionary<string, Type>(StringComparer.OrdinalIgnoreCase);
-					foreach(Type Type in Assembly.GetExecutingAssembly().GetTypes())
-					{
-						if(Type.IsClass && !Type.IsAbstract && Type.IsSubclassOf(typeof(ToolMode)))
-						{
-							ToolModeAttribute? Attribute = Type.GetCustomAttribute<ToolModeAttribute>();
-							if(Attribute == null)
-							{
-								throw new BuildException("Class '{0}' should have a ToolModeAttribute", Type.Name);
-							}
-							ModeNameToType.Add(Attribute.Name, Type);
-						}
-					}
-
 					// Try to get the correct mode
 					if(!ModeNameToType.TryGetValue(Options.Mode, out ModeType))
 					{
