UnrealBuildTool: Add Internal headers to UHTModuleInfo

#rnx

diff --git a/Engine/Source/Programs/UnrealBuildTool/System/ExternalExecution.cs b/Engine/Source/Programs/UnrealBuildTool/System/ExternalExecution.cs
--- a/Engine/Source/Programs/UnrealBuildTool/System/ExternalExecution.cs
+++ b/Engine/Source/Programs/UnrealBuildTool/System/ExternalExecution.cs
@@ -221,6 +221,11 @@ namespace UnrealBuildTool
 		/// </summary>
 		public List<FileItem> PublicUObjectHeaders;
 
+		/// <summary>
+		/// Internal headers with UObjects
+		/// </summary>
+		public List<FileItem> InternalUObjectHeaders;
+
 		/// <summary>
 		/// Private headers with UObjects
 		/// </summary>
@@ -255,6 +260,7 @@ namespace UnrealBuildTool
 			this.OverrideModuleType = OverrideType.ToString();
 			this.PublicUObjectClassesHeaders = new List<FileItem>();
 			this.PublicUObjectHeaders = new List<FileItem>();
+			this.InternalUObjectHeaders = new List<FileItem>();
 			this.PrivateUObjectHeaders = new List<FileItem>();
 			this.GeneratedCodeDirectory = GeneratedCodeDirectory;
 			this.GeneratedCodeVersion = GeneratedCodeVersion;
@@ -270,6 +276,7 @@ namespace UnrealBuildTool
 			OverrideModuleType = Reader.ReadString();
 			PublicUObjectClassesHeaders = Reader.ReadList(() => Reader.ReadFileItem());
 			PublicUObjectHeaders = Reader.ReadList(() => Reader.ReadFileItem());
+			InternalUObjectHeaders = Reader.ReadList(() => Reader.ReadFileItem());
 			PrivateUObjectHeaders = Reader.ReadList(() => Reader.ReadFileItem());
 			GeneratedCPPFilenameBase = Reader.ReadString();
 			GeneratedCodeDirectory = Reader.ReadDirectoryItem();
@@ -286,6 +293,7 @@ namespace UnrealBuildTool
 			Writer.WriteString(OverrideModuleType);
 			Writer.WriteList(PublicUObjectClassesHeaders, Item => Writer.WriteFileItem(Item));
 			Writer.WriteList(PublicUObjectHeaders, Item => Writer.WriteFileItem(Item));
+			Writer.WriteList(InternalUObjectHeaders, Item => Writer.WriteFileItem(Item));
 			Writer.WriteList(PrivateUObjectHeaders, Item => Writer.WriteFileItem(Item));
 			Writer.WriteString(GeneratedCPPFilenameBase);
 			Writer.WriteDirectoryItem(GeneratedCodeDirectory);
@@ -338,6 +346,7 @@ namespace UnrealBuildTool
 			public string OutputDirectory { get; set; }
 			public List<string> ClassesHeaders { get; set; }
 			public List<string> PublicHeaders { get; set; }
+			public List<string> InternalHeaders { get; set; }
 			public List<string> PrivateHeaders { get; set; }
 			public string? GeneratedCPPFilenameBase { get; set; }
 			public bool SaveExportedHeaders { get; set; }
@@ -354,6 +363,7 @@ namespace UnrealBuildTool
 				OutputDirectory = Path.GetDirectoryName(Info.GeneratedCPPFilenameBase)!;
 				ClassesHeaders = Info.PublicUObjectClassesHeaders.Select((Header) => Header.AbsolutePath).ToList();
 				PublicHeaders = Info.PublicUObjectHeaders.Select((Header) => Header.AbsolutePath).ToList();
+				InternalHeaders = Info.InternalUObjectHeaders.Select((Header) => Header.AbsolutePath).ToList();
 				PrivateHeaders = Info.PrivateUObjectHeaders.Select((Header) => Header.AbsolutePath).ToList();
 				GeneratedCPPFilenameBase = Info.GeneratedCPPFilenameBase;
 				SaveExportedHeaders = !Info.bIsReadOnly;
@@ -608,7 +618,7 @@ namespace UnrealBuildTool
 			{
 				UEBuildModuleCPP Module = ModulesSortedByType[Idx];
 				UHTModuleInfo Info = ModuleInfoArray[Idx];
-				if (Info.PublicUObjectClassesHeaders.Count > 0 || Info.PrivateUObjectHeaders.Count > 0 || Info.PublicUObjectHeaders.Count > 0)
+				if (Info.PublicUObjectClassesHeaders.Count > 0 || Info.PrivateUObjectHeaders.Count > 0 || Info.PublicUObjectHeaders.Count > 0 || Info.InternalUObjectHeaders.Count > 0)
 				{
 					// Set a flag indicating that we need to add the generated headers directory
 					Module.bAddGeneratedCodeIncludePath = true;
@@ -638,8 +648,9 @@ namespace UnrealBuildTool
 					DirectoryItem ModuleDirectoryItem = DirectoryItem.GetItemByDirectoryReference(Module.ModuleDirectory);
 
 					List<FileItem> ReflectedHeaderFiles = new List<FileItem>();
-					ReflectedHeaderFiles.AddRange(Info.PublicUObjectHeaders);
 					ReflectedHeaderFiles.AddRange(Info.PublicUObjectClassesHeaders);
+					ReflectedHeaderFiles.AddRange(Info.PublicUObjectHeaders);
+					ReflectedHeaderFiles.AddRange(Info.InternalUObjectHeaders);
 					ReflectedHeaderFiles.AddRange(Info.PrivateUObjectHeaders);
 					UObjectModuleHeaders.Add(new UHTModuleHeaderInfo(ModuleDirectoryItem, ReflectedHeaderFiles, Module.Rules.bUsePrecompiled));
 				}
@@ -694,6 +705,11 @@ namespace UnrealBuildTool
 							ModuleInfo.PublicUObjectHeaders.Add(HeaderFile);
 							bFoundHeaderLocation = true;
 						}
+						else if (HeaderFile.Location.IsUnderDirectory(DirectoryReference.Combine(ModuleDirectory, "Internal")))
+						{
+							ModuleInfo.InternalUObjectHeaders.Add(HeaderFile);
+							bFoundHeaderLocation = true;
+						}
 					}
 					if (!bFoundHeaderLocation)
 					{
@@ -888,6 +904,7 @@ namespace UnrealBuildTool
 				List<FileItem> AllUObjectHeaders = new List<FileItem>();
 				AllUObjectHeaders.AddRange(Module.PublicUObjectClassesHeaders);
 				AllUObjectHeaders.AddRange(Module.PublicUObjectHeaders);
+				AllUObjectHeaders.AddRange(Module.InternalUObjectHeaders);
 				AllUObjectHeaders.AddRange(Module.PrivateUObjectHeaders);
 
 				// Load up the old timestamp file and check to see if anything has changed
@@ -1008,6 +1025,7 @@ namespace UnrealBuildTool
 								List<string> AllUObjectFiles = new List<string>();
 								AllUObjectFiles.AddRange(Module.PublicUObjectClassesHeaders.ConvertAll(Item => Item.AbsolutePath));
 								AllUObjectFiles.AddRange(Module.PublicUObjectHeaders.ConvertAll(Item => Item.AbsolutePath));
+								AllUObjectFiles.AddRange(Module.InternalUObjectHeaders.ConvertAll(Item => Item.AbsolutePath));
 								AllUObjectFiles.AddRange(Module.PrivateUObjectHeaders.ConvertAll(Item => Item.AbsolutePath));
 								FileReference.WriteAllLines(TimestampFile, AllUObjectFiles);
 							}
