UnrealBuildTool:

Attempts to identify invalid arguments & print help were insufficiently clever.

(UnrealBuildTool "-Target=UnrealEditor Win64 Development" would show help rather than build)

#trivial

diff --git a/Engine/Source/Programs/UnrealBuildTool/UnrealBuildTool.cs b/Engine/Source/Programs/UnrealBuildTool/UnrealBuildTool.cs
--- a/Engine/Source/Programs/UnrealBuildTool/UnrealBuildTool.cs
+++ b/Engine/Source/Programs/UnrealBuildTool/UnrealBuildTool.cs
@@ -489,19 +489,12 @@ namespace UnrealBuildTool
 				// Parse the global options
 				GlobalOptions Options = new GlobalOptions(Arguments);
 
-				Console.WriteLine(Arguments.CountValueArguments());
-
 				if (
 					// Print usage if there are zero arguments provided
 					ArgumentsArray.Length == 0 
 
 					// Print usage if the user asks for help
 					|| Options.bGetHelp 
-
-					// Print usage if none of the arguments provided were recognized as part of GlobalOptions and there fewer than three value arguments:
-					// For the default Build mode, at least three value arguments are required: target, platform, and configuration
-					// This is an imperfect test of a valid command line - the Build ToolMode should provide further feedback if the command line is not valid
-					|| (!Arguments.AreAnyArgumentsUsed() && Arguments.CountValueArguments() < 3)
 					)
 				{
 					PrintUsage();
