UE4 -> UE

diff --git a/Engine/Source/Programs/UnrealBuildTool/Configuration/TargetRules.cs b/Engine/Source/Programs/UnrealBuildTool/Configuration/TargetRules.cs
--- a/Engine/Source/Programs/UnrealBuildTool/Configuration/TargetRules.cs
+++ b/Engine/Source/Programs/UnrealBuildTool/Configuration/TargetRules.cs
@@ -1882,7 +1882,7 @@ namespace UnrealBuildTool
 			if(DefaultBuildSettings < BuildSettingsVersion.V2)
 			{
 				Diagnostics.Add("[Upgrade]");
-				Diagnostics.Add("[Upgrade] Using backward-compatible build settings. The latest version of UE4 sets the following values by default, which may require code changes:");
+				Diagnostics.Add("[Upgrade] Using backward-compatible build settings. The latest version of UE sets the following values by default, which may require code changes:");
 
 				List<Tuple<string, string>> ModifiedSettings = new List<Tuple<string, string>>();
 				if(DefaultBuildSettings < BuildSettingsVersion.V2)
