UnrealBuildTool: Cancel TaskExecutor action if any PrerequisiteActions fail, as it will also most likely fail

#rnx

diff --git a/Engine/Source/Programs/UnrealBuildTool/Executors/Experimental/TaskExecutor.cs b/Engine/Source/Programs/UnrealBuildTool/Executors/Experimental/TaskExecutor.cs
--- a/Engine/Source/Programs/UnrealBuildTool/Executors/Experimental/TaskExecutor.cs
+++ b/Engine/Source/Programs/UnrealBuildTool/Executors/Experimental/TaskExecutor.cs
@@ -208,7 +208,14 @@ namespace UnrealBuildTool
 					CancellationToken.ThrowIfCancellationRequested();
 				}
 
-				await Task.WhenAll(Action.PrerequisiteActions.Select(x => Tasks[x]).ToArray());
+				ExecuteResults[] Results = await Task.WhenAll(Action.PrerequisiteActions.Select(x => Tasks[x]).ToArray());
+
+				// Cancel this task if any PrerequisiteActions fail (or were cancelled)
+				if (Results.Any(x => x.ExitCode != 0))
+				{
+					throw new OperationCanceledException();
+				}
+
 				CancellationToken.ThrowIfCancellationRequested();
 
 				SemaphoreTask = MaxProcessSemaphore.WaitAsync(CancellationToken);
