UnrealBuildTool: Additional cleanup in TaskExecutor. Enforce output always printing in correct order

#rnx

diff --git a/Engine/Source/Programs/UnrealBuildTool/Executors/Experimental/TaskExecutor.cs b/Engine/Source/Programs/UnrealBuildTool/Executors/Experimental/TaskExecutor.cs
--- a/Engine/Source/Programs/UnrealBuildTool/Executors/Experimental/TaskExecutor.cs
+++ b/Engine/Source/Programs/UnrealBuildTool/Executors/Experimental/TaskExecutor.cs
@@ -12,7 +12,7 @@ using System.Threading.Tasks;
 namespace UnrealBuildTool
 {
 	/// <summary>
-	/// This executor is similar to LocalExecutor, but uses p/invoke on Windows to ensure that child processes are started at a lower priority and are terminated when the parent process terminates.
+	/// This executor is similar to ParallelExecutor, but async Tasks to process the action graph
 	/// </summary>
 	class TaskExecutor : ActionExecutor
 	{
@@ -91,15 +91,15 @@ namespace UnrealBuildTool
 		}
 
 		/// <summary>
-		/// Checks whether the parallel executor can be used
+		/// Checks whether the task executor can be used
 		/// </summary>
-		/// <returns>True if the parallel executor can be used</returns>
+		/// <returns>True if the task executor can be used</returns>
 		public static bool IsAvailable()
 		{
 			return true;
 		}
 
-		class ExecuteResults
+		private class ExecuteResults
 		{
 			public List<string> LogLines { get; private set; }
 			public int ExitCode { get; private set; }
@@ -149,7 +149,7 @@ namespace UnrealBuildTool
 			foreach (LinkedAction Action in InputActions)
 			{
 				Task<ExecuteResults> ExecuteTask = ExecuteAction(Action, Tasks, ProcessGroup, MaxProcessSemaphore, CancellationToken);
-				Task LogTask = ExecuteTask.ContinueWith(antecedent => LogCompletedAction(Action, antecedent, CancellationTokenSource, ProgressWriter, TotalActions, ref NumCompletedActions));
+				Task LogTask = ExecuteTask.ContinueWith(antecedent => LogCompletedAction(Action, antecedent, CancellationTokenSource, ProgressWriter, TotalActions, ref NumCompletedActions), CancellationToken);
 
 				Tasks.Add(Action, ExecuteTask);
 				AllTasks.Add(ExecuteTask);
@@ -198,16 +198,14 @@ namespace UnrealBuildTool
 			Task? SemaphoreTask = null;
 			try
 			{
-				while (true)
+				// Wait for Tasks list to be populated with any PrerequisiteActions
+				while (!Action.PrerequisiteActions.All(x => Tasks.ContainsKey(x)))
 				{
-					if (Action.PrerequisiteActions.All(x => Tasks.ContainsKey(x)))
-					{
-						break;
-					}
-					await Task.Delay(100);
+					await Task.Delay(100, CancellationToken);
 					CancellationToken.ThrowIfCancellationRequested();
 				}
 
+				// Wait for all PrerequisiteActions to complete
 				ExecuteResults[] Results = await Task.WhenAll(Action.PrerequisiteActions.Select(x => Tasks[x]).ToArray());
 
 				// Cancel this task if any PrerequisiteActions fail (or were cancelled)
@@ -218,26 +216,24 @@ namespace UnrealBuildTool
 
 				CancellationToken.ThrowIfCancellationRequested();
 
+				// Limit the number of concurrent processes that will run in parallel
 				SemaphoreTask = MaxProcessSemaphore.WaitAsync(CancellationToken);
 				await SemaphoreTask;
 
-				string Description = string.Empty;
-				List<string> LogLines;
-				int ExitCode;
-				TimeSpan ExecutionTime;
-				TimeSpan ProcessorTime;
+				CancellationToken.ThrowIfCancellationRequested();
+
 				using (ManagedProcess Process = new ManagedProcess(ProcessGroup, Action.CommandPath.FullName, Action.CommandArguments, Action.WorkingDirectory.FullName, null, null, ProcessPriorityClass.BelowNormal))
 				{
 					MemoryStream StdOutStream = new MemoryStream();
 					await Process.CopyToAsync(StdOutStream, CancellationToken);
 					CancellationToken.ThrowIfCancellationRequested();
-					LogLines = Console.OutputEncoding.GetString(StdOutStream.GetBuffer(), 0, Convert.ToInt32(StdOutStream.Length)).Split(LineEndingSplit, StringSplitOptions.RemoveEmptyEntries).ToList();
-					ExitCode = Process.ExitCode;
-					ProcessorTime = Process.TotalProcessorTime;
-					ExecutionTime = Process.ExitTime - Process.StartTime;
-				}
 
-				return new ExecuteResults(LogLines, ExitCode, ExecutionTime, ProcessorTime);
+					List<string> LogLines = Console.OutputEncoding.GetString(StdOutStream.GetBuffer(), 0, Convert.ToInt32(StdOutStream.Length)).Split(LineEndingSplit, StringSplitOptions.RemoveEmptyEntries).ToList();
+					int ExitCode = Process.ExitCode;
+					TimeSpan ProcessorTime = Process.TotalProcessorTime;
+					TimeSpan ExecutionTime = Process.ExitTime - Process.StartTime;
+					return new ExecuteResults(LogLines, ExitCode, ExecutionTime, ProcessorTime);
+				}
 			}
 			catch (OperationCanceledException)
 			{
@@ -278,41 +274,41 @@ namespace UnrealBuildTool
 				Description = $"{(Action.CommandDescription != null ? Action.CommandDescription : Action.CommandPath.GetFileNameWithoutExtension())} {LogLines[0]}".Trim();
 			}
 
-			int CompletedActions;
-			CompletedActions = Interlocked.Increment(ref NumCompletedActions);
-			ProgressWriter.Write(CompletedActions, TotalActions);
-
-			// Cancelled
-			if (ExitCode == int.MaxValue)
+			lock (ProgressWriter)
 			{
-				Log.TraceInformation("[{0}/{1}] {2} cancelled", CompletedActions, TotalActions, Description);
-				return;
-			}
+				int CompletedActions;
+				CompletedActions = Interlocked.Increment(ref NumCompletedActions);
+				ProgressWriter.Write(CompletedActions, TotalActions);
 
-			Log.TraceInformation("[{0}/{1}] {2}", CompletedActions, TotalActions, Description);
-			foreach (string Line in LogLines.Skip(Action.bShouldOutputStatusDescription ? 0 : 1))
-			{
-				Log.TraceInformation(Line);
-			}
+				// Cancelled
+				if (ExitCode == int.MaxValue)
+				{
+					Log.TraceInformation("[{0}/{1}] {2} cancelled", CompletedActions, TotalActions, Description);
+					return;
+				}
 
-			if (ExitCode != 0)
-			{
-				// BEGIN TEMPORARY TO CATCH PVS-STUDIO ISSUES
-				if (LogLines.Count == 0)
+				Log.TraceInformation("[{0}/{1}] {2}", CompletedActions, TotalActions, Description);
+				foreach (string Line in LogLines.Skip(Action.bShouldOutputStatusDescription ? 0 : 1))
+				{
+					Log.TraceInformation(Line);
+				}
+
+				if (ExitCode != 0)
 				{
-					lock (ProgressWriter)
+					// BEGIN TEMPORARY TO CATCH PVS-STUDIO ISSUES
+					if (LogLines.Count == 0)
 					{
 						Log.TraceInformation("[{0}/{1}] {2} - Error but no output", NumCompletedActions, TotalActions, Description);
 						Log.TraceInformation("[{0}/{1}] {2} - {3} {4} {5} {6}", NumCompletedActions, TotalActions, Description, ExitCode,
 							Action.WorkingDirectory, Action.CommandPath, Action.CommandArguments);
 					}
-				}
-				// END TEMPORARY
+					// END TEMPORARY
 
-				// Cancel other tasks
-				if (bStopCompilationAfterErrors)
-				{
-					CancellationTokenSource.Cancel();
+					// Cancel all other pending tasks
+					if (bStopCompilationAfterErrors)
+					{
+						CancellationTokenSource.Cancel();
+					}
 				}
 			}
 		}
