UnrealBuildTool: Clean up some intellisense suggestions, format documents, remove and sort usings, etc.

#rnx

diff --git a/Engine/Source/Programs/Shared/EpicGames.Core/AssemblyUtils.cs b/Engine/Source/Programs/Shared/EpicGames.Core/AssemblyUtils.cs
--- a/Engine/Source/Programs/Shared/EpicGames.Core/AssemblyUtils.cs
+++ b/Engine/Source/Programs/Shared/EpicGames.Core/AssemblyUtils.cs
@@ -13,49 +13,49 @@ namespace EpicGames.Core
 {
 	public static class AssemblyUtils
 	{
-        /// <summary>
-        /// Gets the original location (path and filename) of an assembly.
-        /// This method is using Assembly.CodeBase property to properly resolve original
-        /// assembly path in case shadow copying is enabled.
-        /// </summary>
-        /// <returns>Absolute path and filename to the assembly.</returns>
-        public static string GetOriginalLocation(this Assembly ThisAssembly)
-        {
-            return new Uri(ThisAssembly.CodeBase).LocalPath;
-        }
-    
-        /// <summary>
-        /// Version info of the executable which runs this code.
-        /// </summary>
-        public static FileVersionInfo ExecutableVersion
-        {
-            get
-            {
-                return FileVersionInfo.GetVersionInfo(Assembly.GetEntryAssembly().GetOriginalLocation());
-            }
-        }
+		/// <summary>
+		/// Gets the original location (path and filename) of an assembly.
+		/// This method is using Assembly.CodeBase property to properly resolve original
+		/// assembly path in case shadow copying is enabled.
+		/// </summary>
+		/// <returns>Absolute path and filename to the assembly.</returns>
+		public static string GetOriginalLocation(this Assembly ThisAssembly)
+		{
+			return new Uri(ThisAssembly.CodeBase).LocalPath;
+		}
+
+		/// <summary>
+		/// Version info of the executable which runs this code.
+		/// </summary>
+		public static FileVersionInfo ExecutableVersion
+		{
+			get
+			{
+				return FileVersionInfo.GetVersionInfo(Assembly.GetEntryAssembly().GetOriginalLocation());
+			}
+		}
 
-        /// <summary>
-        /// Installs an assembly resolver. Mostly used to get shared assemblies that we don't want copied around to various output locations as happens when "Copy Local" is set to true
-        /// for an assembly reference (which is the default).
-        /// </summary>
-        public static void InstallAssemblyResolver(string PathToBinariesDotNET)
-        {
+		/// <summary>
+		/// Installs an assembly resolver. Mostly used to get shared assemblies that we don't want copied around to various output locations as happens when "Copy Local" is set to true
+		/// for an assembly reference (which is the default).
+		/// </summary>
+		public static void InstallAssemblyResolver(string PathToBinariesDotNET)
+		{
 			AppDomain.CurrentDomain.AssemblyResolve += (sender, args) =>
-            {
-			    // Name is fully qualified assembly definition - e.g. "p4dn, Version=1.0.0.0, Culture=neutral, PublicKeyToken=ff968dc1933aba6f"
-			    string AssemblyName = args.Name.Split(',')[0];
+			{
+				// Name is fully qualified assembly definition - e.g. "p4dn, Version=1.0.0.0, Culture=neutral, PublicKeyToken=ff968dc1933aba6f"
+				string AssemblyName = args.Name.Split(',')[0];
 
 				return (
-                    from KnownAssemblyName in new[] { "SwarmAgent.exe", "../ThirdParty/Ionic/Ionic.Zip.Reduced.dll", "../ThirdParty/Newtonsoft/NewtonSoft.Json.dll" }
-                    where AssemblyName.Equals(Path.GetFileNameWithoutExtension(KnownAssemblyName), StringComparison.InvariantCultureIgnoreCase)
-                    let ResolvedAssemblyFilename = Path.Combine(PathToBinariesDotNET, KnownAssemblyName)
-                    // check if the file exists first. If we just try to load it, we correctly throw an exception, but it's a generic
-                    // FileNotFoundException, which is not informative. Better to return null.
-                    select File.Exists(ResolvedAssemblyFilename) ? Assembly.LoadFile(ResolvedAssemblyFilename) : null
-                    ).FirstOrDefault();
+					from KnownAssemblyName in new[] { "SwarmAgent.exe", "../ThirdParty/Ionic/Ionic.Zip.Reduced.dll", "../ThirdParty/Newtonsoft/NewtonSoft.Json.dll" }
+					where AssemblyName.Equals(Path.GetFileNameWithoutExtension(KnownAssemblyName), StringComparison.InvariantCultureIgnoreCase)
+					let ResolvedAssemblyFilename = Path.Combine(PathToBinariesDotNET, KnownAssemblyName)
+					// check if the file exists first. If we just try to load it, we correctly throw an exception, but it's a generic
+					// FileNotFoundException, which is not informative. Better to return null.
+					select File.Exists(ResolvedAssemblyFilename) ? Assembly.LoadFile(ResolvedAssemblyFilename) : null
+					).FirstOrDefault();
 			};
-        }
+		}
 
 		/// <summary>
 		/// Installs an assembly resolver, which will load *any* assembly which exists recursively within the supplied folder.
@@ -64,13 +64,12 @@ namespace EpicGames.Core
 		public static void InstallRecursiveAssemblyResolver(string RootDirectory)
 		{
 			RefreshAssemblyCache(RootDirectory);
-			
+
 			AppDomain.CurrentDomain.AssemblyResolve += (sender, args) =>
 			{
 				// Name is fully qualified assembly definition - e.g. "p4dn, Version=1.0.0.0, Culture=neutral, PublicKeyToken=ff968dc1933aba6f"
 				string AssemblyName = args.Name.Split(',')[0];
-				string AssemblyLocation;
-				if (AssemblyLocationCache.TryGetValue(AssemblyName, out AssemblyLocation))
+				if (AssemblyLocationCache.TryGetValue(AssemblyName, out string AssemblyLocation))
 				{
 					// We have this assembly in our folder.					
 					if (File.Exists(AssemblyLocation))
@@ -128,9 +127,9 @@ namespace EpicGames.Core
 		}
 
 		// Map of assembly name to path on disk
-		private static Dictionary<string, string> AssemblyLocationCache = new Dictionary<string, string>(StringComparer.OrdinalIgnoreCase);
+		private static readonly Dictionary<string, string> AssemblyLocationCache = new Dictionary<string, string>(StringComparer.OrdinalIgnoreCase);
 		// Track last modified date of each assembly, so we can ensure we always reference the latest one in the case of stale assemblies on disk.
-		private static Dictionary<string, DateTime> AssemblyWriteTimes = new Dictionary<string, DateTime>(StringComparer.OrdinalIgnoreCase);
+		private static readonly Dictionary<string, DateTime> AssemblyWriteTimes = new Dictionary<string, DateTime>(StringComparer.OrdinalIgnoreCase);
 
-    }
+	}
 }
diff --git a/Engine/Source/Programs/Shared/EpicGames.Core/BinaryArchive.cs b/Engine/Source/Programs/Shared/EpicGames.Core/BinaryArchive.cs
--- a/Engine/Source/Programs/Shared/EpicGames.Core/BinaryArchive.cs
+++ b/Engine/Source/Programs/Shared/EpicGames.Core/BinaryArchive.cs
@@ -6,7 +6,6 @@ using System.Collections.Generic;
 using System.Linq.Expressions;
 using System.Reflection;
 using System.Reflection.Emit;
-using System.Text;
 
 #nullable enable
 
@@ -77,7 +76,7 @@ namespace EpicGames.Core
 			}
 		}
 
-		static ConcurrentDictionary<Type, TypeSerializer> TypeToSerializerInfo = new ConcurrentDictionary<Type, TypeSerializer>()
+		static readonly ConcurrentDictionary<Type, TypeSerializer> TypeToSerializerInfo = new ConcurrentDictionary<Type, TypeSerializer>()
 		{
 			[typeof(byte)] = TypeSerializer.Create(Reader => Reader.ReadByte(), (Writer, Value) => Writer.WriteByte(Value)),
 			[typeof(sbyte)] = TypeSerializer.Create(Reader => Reader.ReadSignedByte(), (Writer, Value) => Writer.WriteSignedByte(Value)),
@@ -90,8 +89,8 @@ namespace EpicGames.Core
 			[typeof(string)] = TypeSerializer.Create(Reader => Reader.ReadString(), (Writer, Value) => Writer.WriteString(Value))
 		};
 
-		static ConcurrentDictionary<Type, ContentHash> TypeToDigest = new ConcurrentDictionary<Type, ContentHash>();
-		static ConcurrentDictionary<ContentHash, Type> DigestToType = new ConcurrentDictionary<ContentHash, Type>();
+		static readonly ConcurrentDictionary<Type, ContentHash> TypeToDigest = new ConcurrentDictionary<Type, ContentHash>();
+		static readonly ConcurrentDictionary<ContentHash, Type> DigestToType = new ConcurrentDictionary<ContentHash, Type>();
 
 		static MethodInfo GetMethodInfo(Expression<Action> Expr)
 		{
@@ -108,14 +107,14 @@ namespace EpicGames.Core
 			return ((MethodCallExpression)Expr.Body).Method.GetGenericMethodDefinition();
 		}
 
-		static MethodInfo ReadBoolMethodInfo = GetMethodInfo<BinaryArchiveReader>(x => x.ReadBool());
-		static MethodInfo ReadIntMethodInfo = GetMethodInfo<BinaryArchiveReader>(x => x.ReadInt());
-		static MethodInfo ReadArrayMethodInfo = GetGenericMethodInfo<BinaryArchiveReader, Func<int>>((x, y) => x.ReadArray(y));
-		static MethodInfo ReadPolymorphicObjectMethodInfo = GetMethodInfo(() => ReadObject(null!));
-		static MethodInfo WriteBoolMethodInfo = GetMethodInfo<BinaryArchiveWriter>(x => x.WriteBool(false));
-		static MethodInfo WriteIntMethodInfo = GetMethodInfo<BinaryArchiveWriter>(x => x.WriteInt(0));
-		static MethodInfo WriteArrayMethodInfo = GetGenericMethodInfo<BinaryArchiveWriter, Action<int>>((x, y) => x.WriteArray(null!, y));
-		static MethodInfo WritePolymorphicObjectMethodInfo = GetMethodInfo(() => WritePolymorphicObject(null!, null!));
+		static readonly MethodInfo ReadBoolMethodInfo = GetMethodInfo<BinaryArchiveReader>(x => x.ReadBool());
+		static readonly MethodInfo ReadIntMethodInfo = GetMethodInfo<BinaryArchiveReader>(x => x.ReadInt());
+		static readonly MethodInfo ReadArrayMethodInfo = GetGenericMethodInfo<BinaryArchiveReader, Func<int>>((x, y) => x.ReadArray(y));
+		static readonly MethodInfo ReadPolymorphicObjectMethodInfo = GetMethodInfo(() => ReadObject(null!));
+		static readonly MethodInfo WriteBoolMethodInfo = GetMethodInfo<BinaryArchiveWriter>(x => x.WriteBool(false));
+		static readonly MethodInfo WriteIntMethodInfo = GetMethodInfo<BinaryArchiveWriter>(x => x.WriteInt(0));
+		static readonly MethodInfo WriteArrayMethodInfo = GetGenericMethodInfo<BinaryArchiveWriter, Action<int>>((x, y) => x.WriteArray(null!, y));
+		static readonly MethodInfo WritePolymorphicObjectMethodInfo = GetMethodInfo(() => WritePolymorphicObject(null!, null!));
 
 		/// <summary>
 		/// Reads an arbitrary value from the given archive.
@@ -238,8 +237,7 @@ namespace EpicGames.Core
 			Type SerializerKey = ConverterType ?? Type;
 
 			// Get the serializer info
-			TypeSerializer? SerializerInfo;
-			if (!TypeToSerializerInfo.TryGetValue(SerializerKey, out SerializerInfo))
+			if (!TypeToSerializerInfo.TryGetValue(SerializerKey, out TypeSerializer? SerializerInfo))
 			{
 				lock (TypeToSerializerInfo)
 				{
@@ -290,8 +288,7 @@ namespace EpicGames.Core
 			}
 
 			// Get the default converter type
-			Type? DefaultConverterType;
-			if(BinaryConverter.TryGetConverterType(Type, out DefaultConverterType))
+			if (BinaryConverter.TryGetConverterType(Type, out Type? DefaultConverterType))
 			{
 				return FindOrAddSerializerInfo(Type, DefaultConverterType);
 			}
@@ -380,10 +377,10 @@ namespace EpicGames.Core
 					WritePropertyMethod = FindOrAddSerializerInfo(Property.PropertyType, Converter?.Type).WriteMethodInfo;
 				}
 
-				Generator.Emit(OpCodes.Ldarg_0);											// [0] = Arg0 (BinaryArchiveWriter)
-				Generator.Emit(OpCodes.Ldarg_1);											// [1] = Arg1 (Instance)
-				Generator.EmitCall(OpCodes.Call, Property.GetGetMethod(true)!, null);		// [1] = GetMethod(Instance)
-				Generator.EmitCall(OpCodes.Call, WritePropertyMethod, null);		        // [EMPTY] Write(BinaryArchiveWriter, GetMethod(Instance));
+				Generator.Emit(OpCodes.Ldarg_0);                                            // [0] = Arg0 (BinaryArchiveWriter)
+				Generator.Emit(OpCodes.Ldarg_1);                                            // [1] = Arg1 (Instance)
+				Generator.EmitCall(OpCodes.Call, Property.GetGetMethod(true)!, null);       // [1] = GetMethod(Instance)
+				Generator.EmitCall(OpCodes.Call, WritePropertyMethod, null);                // [EMPTY] Write(BinaryArchiveWriter, GetMethod(Instance));
 			}
 
 			// Return
@@ -426,7 +423,7 @@ namespace EpicGames.Core
 				Generator.Emit(OpCodes.Dup);                                                // [1] = new Type()
 
 				Generator.Emit(OpCodes.Ldarg_0);                                            // [2] = Arg1 (Reader)
-				Generator.Emit(OpCodes.Call, ReadPropertyMethod);	                        // [2] = Reader.ReadMethod() or ReadMethod(Reader)
+				Generator.Emit(OpCodes.Call, ReadPropertyMethod);                           // [2] = Reader.ReadMethod() or ReadMethod(Reader)
 
 				Generator.Emit(OpCodes.Call, Property.GetSetMethod(true)!);                 // SetMethod(new Type(), ReadMethod(Reader))
 			}
diff --git a/Engine/Source/Programs/Shared/EpicGames.Core/BinaryArchiveAttributes.cs b/Engine/Source/Programs/Shared/EpicGames.Core/BinaryArchiveAttributes.cs
--- a/Engine/Source/Programs/Shared/EpicGames.Core/BinaryArchiveAttributes.cs
+++ b/Engine/Source/Programs/Shared/EpicGames.Core/BinaryArchiveAttributes.cs
@@ -1,8 +1,6 @@
 // Copyright Epic Games, Inc. All Rights Reserved.
 
 using System;
-using System.Collections.Generic;
-using System.Text;
 
 namespace EpicGames.Core
 {
diff --git a/Engine/Source/Programs/Shared/EpicGames.Core/BinaryArchiveReader.cs b/Engine/Source/Programs/Shared/EpicGames.Core/BinaryArchiveReader.cs
--- a/Engine/Source/Programs/Shared/EpicGames.Core/BinaryArchiveReader.cs
+++ b/Engine/Source/Programs/Shared/EpicGames.Core/BinaryArchiveReader.cs
@@ -4,9 +4,7 @@ using System;
 using System.Collections.Generic;
 using System.Diagnostics.CodeAnalysis;
 using System.IO;
-using System.Linq;
 using System.Text;
-using System.Threading.Tasks;
 
 #nullable disable
 
@@ -32,11 +30,11 @@ namespace EpicGames.Core
 		/// Current position within the buffer
 		/// </summary>
 		int BufferPos;
-		
+
 		/// <summary>
 		/// List of previously serialized objects
 		/// </summary>
-		List<object> Objects = new List<object>();
+		readonly List<object> Objects = new List<object>();
 
 		/// <summary>
 		/// Constructor
@@ -71,7 +69,7 @@ namespace EpicGames.Core
 		/// </summary>
 		public void Dispose()
 		{
-			if(Stream != null)
+			if (Stream != null)
 			{
 				Stream.Dispose();
 				Stream = null;
@@ -176,7 +174,7 @@ namespace EpicGames.Core
 		public string ReadString()
 		{
 			byte[] Bytes = ReadByteArray();
-			if(Bytes == null)
+			if (Bytes == null)
 			{
 				return null;
 			}
@@ -221,7 +219,7 @@ namespace EpicGames.Core
 		private T[] ReadPrimitiveArray<T>(int ElementSize) where T : struct
 		{
 			int Length = ReadInt();
-			if(Length < 0)
+			if (Length < 0)
 			{
 				return null;
 			}
@@ -295,14 +293,14 @@ namespace EpicGames.Core
 		public T[] ReadArray<T>(Func<T> ReadElement)
 		{
 			int Count = ReadInt();
-			if(Count < 0)
+			if (Count < 0)
 			{
 				return null;
 			}
 			else
 			{
 				T[] Result = new T[Count];
-				for(int Idx = 0; Idx < Count; Idx++)
+				for (int Idx = 0; Idx < Count; Idx++)
 				{
 					Result[Idx] = ReadElement();
 				}
@@ -319,14 +317,14 @@ namespace EpicGames.Core
 		public List<T> ReadList<T>(Func<T> ReadElement)
 		{
 			int Count = ReadInt();
-			if(Count < 0)
+			if (Count < 0)
 			{
 				return null;
 			}
 			else
 			{
 				List<T> Result = new List<T>(Count);
-				for(int Idx = 0; Idx < Count; Idx++)
+				for (int Idx = 0; Idx < Count; Idx++)
 				{
 					Result.Add(ReadElement());
 				}
@@ -343,14 +341,14 @@ namespace EpicGames.Core
 		public HashSet<T> ReadHashSet<T>(Func<T> ReadElement)
 		{
 			int Count = ReadInt();
-			if(Count < 0)
+			if (Count < 0)
 			{
 				return null;
 			}
 			else
 			{
 				HashSet<T> Result = new HashSet<T>();
-				for(int Idx = 0; Idx < Count; Idx++)
+				for (int Idx = 0; Idx < Count; Idx++)
 				{
 					Result.Add(ReadElement());
 				}
@@ -368,14 +366,14 @@ namespace EpicGames.Core
 		public HashSet<T> ReadHashSet<T>(Func<T> ReadElement, IEqualityComparer<T> Comparer)
 		{
 			int Count = ReadInt();
-			if(Count < 0)
+			if (Count < 0)
 			{
 				return null;
 			}
 			else
 			{
 				HashSet<T> Result = new HashSet<T>(Comparer);
-				for(int Idx = 0; Idx < Count; Idx++)
+				for (int Idx = 0; Idx < Count; Idx++)
 				{
 					Result.Add(ReadElement());
 				}
@@ -394,14 +392,14 @@ namespace EpicGames.Core
 		public Dictionary<K, V> ReadDictionary<K, V>(Func<K> ReadKey, Func<V> ReadValue)
 		{
 			int Count = ReadInt();
-			if(Count < 0)
+			if (Count < 0)
 			{
 				return null;
 			}
 			else
 			{
 				Dictionary<K, V> Result = new Dictionary<K, V>(Count);
-				for(int Idx = 0; Idx < Count; Idx++)
+				for (int Idx = 0; Idx < Count; Idx++)
 				{
 					Result.Add(ReadKey(), ReadValue());
 				}
@@ -421,14 +419,14 @@ namespace EpicGames.Core
 		public Dictionary<K, V> ReadDictionary<K, V>(Func<K> ReadKey, Func<V> ReadValue, IEqualityComparer<K> Comparer)
 		{
 			int Count = ReadInt();
-			if(Count < 0)
+			if (Count < 0)
 			{
 				return null;
 			}
 			else
 			{
 				Dictionary<K, V> Result = new Dictionary<K, V>(Count, Comparer);
-				for(int Idx = 0; Idx < Count; Idx++)
+				for (int Idx = 0; Idx < Count; Idx++)
 				{
 					Result.Add(ReadKey(), ReadValue());
 				}
@@ -443,7 +441,7 @@ namespace EpicGames.Core
 		/// <param name="ReadValue">Delegate used to read a value</param>
 		public Nullable<T> ReadNullable<T>(Func<T> ReadValue) where T : struct
 		{
-			if(ReadBool())
+			if (ReadBool())
 			{
 				return new Nullable<T>(ReadValue());
 			}
@@ -461,7 +459,7 @@ namespace EpicGames.Core
 		/// <returns>The object instance</returns>
 		public T ReadOptionalObject<T>(Func<T> Read) where T : class
 		{
-			if(ReadBool())
+			if (ReadBool())
 			{
 				return Read();
 			}
@@ -483,13 +481,13 @@ namespace EpicGames.Core
 		public T ReadObjectReference<T>(Func<T> CreateObject, Action<T> ReadObject) where T : class
 		{
 			int Index = ReadInt();
-			if(Index < 0)
+			if (Index < 0)
 			{
 				return null;
 			}
 			else
 			{
-				if(Index == Objects.Count)
+				if (Index == Objects.Count)
 				{
 					T Object = CreateObject();
 					Objects.Add(Object);
@@ -509,19 +507,19 @@ namespace EpicGames.Core
 		public T ReadObjectReference<T>(Func<T> ReadObject) where T : class
 		{
 			int Index = ReadInt();
-			if(Index < 0)
+			if (Index < 0)
 			{
 				return null;
 			}
 			else
 			{
 				// Temporarily add the reader to the object list, so we can detect invalid recursive references. 
-				if(Index == Objects.Count)
+				if (Index == Objects.Count)
 				{
 					Objects.Add(null);
 					Objects[Index] = ReadObject();
 				}
-				if(Objects[Index] == null)
+				if (Objects[Index] == null)
 				{
 					throw new InvalidOperationException("Attempt to serialize reference to object recursively.");
 				}
diff --git a/Engine/Source/Programs/Shared/EpicGames.Core/BinaryArchiveWriter.cs b/Engine/Source/Programs/Shared/EpicGames.Core/BinaryArchiveWriter.cs
--- a/Engine/Source/Programs/Shared/EpicGames.Core/BinaryArchiveWriter.cs
+++ b/Engine/Source/Programs/Shared/EpicGames.Core/BinaryArchiveWriter.cs
@@ -3,10 +3,8 @@
 using System;
 using System.Collections.Generic;
 using System.IO;
-using System.Linq;
 using System.Runtime.CompilerServices;
 using System.Text;
-using System.Threading.Tasks;
 
 #nullable disable
 
@@ -36,7 +34,7 @@ namespace EpicGames.Core
 		/// <summary>
 		/// Instance of the ReferenceComparer class which can be shared by all archive writers
 		/// </summary>
-		static ReferenceComparer ReferenceComparerInstance = new ReferenceComparer();
+		static readonly ReferenceComparer ReferenceComparerInstance = new ReferenceComparer();
 
 		/// <summary>
 		/// The output stream being written to
@@ -56,7 +54,7 @@ namespace EpicGames.Core
 		/// <summary>
 		/// Map of object instance to unique id
 		/// </summary>
-		Dictionary<object, int> ObjectToUniqueId = new Dictionary<object, int>(ReferenceComparerInstance);
+		readonly Dictionary<object, int> ObjectToUniqueId = new Dictionary<object, int>(ReferenceComparerInstance);
 
 		/// <summary>
 		/// Constructor
@@ -84,7 +82,7 @@ namespace EpicGames.Core
 		{
 			Flush();
 
-			if(Stream != null)
+			if (Stream != null)
 			{
 				Stream.Dispose();
 				Stream = null;
@@ -96,7 +94,7 @@ namespace EpicGames.Core
 		/// </summary>
 		public void Flush()
 		{
-			if(BufferPos > 0)
+			if (BufferPos > 0)
 			{
 				Stream.Write(Buffer, 0, BufferPos);
 				BufferPos = 0;
@@ -109,10 +107,10 @@ namespace EpicGames.Core
 		/// <param name="NumBytes">Minimum amount of space required in the output buffer</param>
 		private void EnsureSpace(int NumBytes)
 		{
-			if(BufferPos + NumBytes > Buffer.Length)
+			if (BufferPos + NumBytes > Buffer.Length)
 			{
 				Flush();
-				if(NumBytes > Buffer.Length)
+				if (NumBytes > Buffer.Length)
 				{
 					Buffer = new byte[NumBytes];
 				}
@@ -125,7 +123,7 @@ namespace EpicGames.Core
 		/// <param name="Value">Value to write</param>
 		public void WriteBool(bool Value)
 		{
-			WriteByte(Value? (byte)1 : (byte)0);
+			WriteByte(Value ? (byte)1 : (byte)0);
 		}
 
 		/// <summary>
@@ -234,7 +232,7 @@ namespace EpicGames.Core
 		public void WriteString(string Value)
 		{
 			byte[] Bytes;
-			if(Value == null)
+			if (Value == null)
 			{
 				Bytes = null;
 			}
@@ -278,7 +276,7 @@ namespace EpicGames.Core
 		/// <param name="Data">Data to write. May be null.</param>
 		private void WritePrimitiveArray<T>(T[] Data, int ElementSize) where T : struct
 		{
-			if(Data == null)
+			if (Data == null)
 			{
 				WriteInt(-1);
 			}
@@ -332,9 +330,9 @@ namespace EpicGames.Core
 		/// <param name="Size">Size of the data, in bytes</param>
 		private void WriteBulkData(Array Data, int Size)
 		{
-			if(Size > 0)
+			if (Size > 0)
 			{
-				for(int Pos = 0;; )
+				for (int Pos = 0; ;)
 				{
 					int CopySize = Math.Min(Size - Pos, Buffer.Length - BufferPos);
 
@@ -342,7 +340,7 @@ namespace EpicGames.Core
 					BufferPos += CopySize;
 					Pos += CopySize;
 
-					if(Pos == Size)
+					if (Pos == Size)
 					{
 						break;
 					}
@@ -360,14 +358,14 @@ namespace EpicGames.Core
 		/// <param name="WriteElement">Writes an individual element to the archive</param>
 		public void WriteArray<T>(T[] Items, Action<T> WriteElement)
 		{
-			if(Items == null)
+			if (Items == null)
 			{
 				WriteInt(-1);
 			}
 			else
 			{
 				WriteInt(Items.Length);
-				for(int Idx = 0; Idx < Items.Length; Idx++)
+				for (int Idx = 0; Idx < Items.Length; Idx++)
 				{
 					WriteElement(Items[Idx]);
 				}
@@ -382,14 +380,14 @@ namespace EpicGames.Core
 		/// <param name="WriteElement">Writes an individual element to the archive</param>
 		public void WriteList<T>(IReadOnlyList<T> Items, Action<T> WriteElement)
 		{
-			if(Items == null)
+			if (Items == null)
 			{
 				WriteInt(-1);
 			}
 			else
 			{
 				WriteInt(Items.Count);
-				for(int Idx = 0; Idx < Items.Count; Idx++)
+				for (int Idx = 0; Idx < Items.Count; Idx++)
 				{
 					WriteElement(Items[Idx]);
 				}
@@ -404,14 +402,14 @@ namespace EpicGames.Core
 		/// <param name="WriteElement">Delegate used to read a single element</param>
 		public void WriteHashSet<T>(HashSet<T> Set, Action<T> WriteElement)
 		{
-			if(Set == null)
+			if (Set == null)
 			{
 				WriteInt(-1);
 			}
 			else
 			{
 				WriteInt(Set.Count);
-				foreach(T Element in Set)
+				foreach (T Element in Set)
 				{
 					WriteElement(Element);
 				}
@@ -428,14 +426,14 @@ namespace EpicGames.Core
 		/// <param name="WriteValue">Delegate used to read a single value</param>
 		public void WriteDictionary<K, V>(IDictionary<K, V> Dictionary, Action<K> WriteKey, Action<V> WriteValue)
 		{
-			if(Dictionary == null)
+			if (Dictionary == null)
 			{
 				WriteInt(-1);
 			}
 			else
 			{
 				WriteInt(Dictionary.Count);
-				foreach(KeyValuePair<K, V> Pair in Dictionary)
+				foreach (KeyValuePair<K, V> Pair in Dictionary)
 				{
 					WriteKey(Pair.Key);
 					WriteValue(Pair.Value);
@@ -451,7 +449,7 @@ namespace EpicGames.Core
 		/// <param name="WriteValue">Delegate used to write a value</param>
 		public void WriteNullable<T>(Nullable<T> Item, Action<T> WriteValue) where T : struct
 		{
-			if(Item.HasValue)
+			if (Item.HasValue)
 			{
 				WriteBool(true);
 				WriteValue(Item.Value);
@@ -470,7 +468,7 @@ namespace EpicGames.Core
 		/// <param name="WriteObject">Delegate used to write the object</param>
 		public void WriteOptionalObject<T>(T Object, Action WriteObject) where T : class
 		{
-			if(Object == null)
+			if (Object == null)
 			{
 				WriteBool(false);
 			}
@@ -496,14 +494,13 @@ namespace EpicGames.Core
 		/// <param name="WriteObject">Delegate used to write the object</param>
 		public void WriteObjectReference(object Object, Action WriteObject)
 		{
-			if(Object == null)
+			if (Object == null)
 			{
 				WriteInt(-1);
 			}
 			else
 			{
-				int Index;
-				if(ObjectToUniqueId.TryGetValue(Object, out Index))
+				if (ObjectToUniqueId.TryGetValue(Object, out int Index))
 				{
 					WriteInt(Index);
 				}
diff --git a/Engine/Source/Programs/Shared/EpicGames.Core/BinaryConverter.cs b/Engine/Source/Programs/Shared/EpicGames.Core/BinaryConverter.cs
--- a/Engine/Source/Programs/Shared/EpicGames.Core/BinaryConverter.cs
+++ b/Engine/Source/Programs/Shared/EpicGames.Core/BinaryConverter.cs
@@ -1,11 +1,8 @@
 // Copyright Epic Games, Inc. All Rights Reserved.
 
-using EpicGames.Core;
 using System;
 using System.Collections.Concurrent;
-using System.Collections.Generic;
 using System.Reflection;
-using System.Text;
 
 #nullable enable
 
@@ -51,7 +48,7 @@ namespace EpicGames.Core
 		/// <summary>
 		/// Map from type to the converter type
 		/// </summary>
-		static ConcurrentDictionary<Type, Type> TypeToConverterType = new ConcurrentDictionary<Type, Type>();
+		static readonly ConcurrentDictionary<Type, Type> TypeToConverterType = new ConcurrentDictionary<Type, Type>();
 
 		/// <summary>
 		/// Explicitly register the converter for a type. If Type is a generic type, the converter should also be a generic type with the same type arguments.
@@ -74,8 +71,7 @@ namespace EpicGames.Core
 		/// <returns>True if a converter was found</returns>
 		public static bool TryGetConverterType(Type Type, out Type? ConverterType)
 		{
-			Type? CustomConverterType;
-			if (TypeToConverterType.TryGetValue(Type, out CustomConverterType))
+			if (TypeToConverterType.TryGetValue(Type, out Type? CustomConverterType))
 			{
 				ConverterType = CustomConverterType;
 				return true;
diff --git a/Engine/Source/Programs/Shared/EpicGames.Core/BinaryFormatterUtils.cs b/Engine/Source/Programs/Shared/EpicGames.Core/BinaryFormatterUtils.cs
--- a/Engine/Source/Programs/Shared/EpicGames.Core/BinaryFormatterUtils.cs
+++ b/Engine/Source/Programs/Shared/EpicGames.Core/BinaryFormatterUtils.cs
@@ -1,12 +1,7 @@
 // Copyright Epic Games, Inc. All Rights Reserved.
 
-using System;
-using System.Collections.Generic;
 using System.IO;
-using System.Linq;
 using System.Runtime.Serialization.Formatters.Binary;
-using System.Text;
-using System.Threading.Tasks;
 
 namespace EpicGames.Core
 {
@@ -22,11 +17,9 @@ namespace EpicGames.Core
 		/// <returns>Instance of the object that was read from disk</returns>
 		public static T Load<T>(FileReference Location)
 		{
-			using (FileStream Stream = new FileStream(Location.FullName, FileMode.Open, FileAccess.Read))
-			{
-				BinaryFormatter Formatter = new BinaryFormatter();
-				return (T)Formatter.Deserialize(Stream);
-			}
+			using FileStream Stream = new FileStream(Location.FullName, FileMode.Open, FileAccess.Read);
+			BinaryFormatter Formatter = new BinaryFormatter();
+			return (T)Formatter.Deserialize(Stream);
 		}
 
 		/// <summary>
@@ -37,11 +30,9 @@ namespace EpicGames.Core
 		public static void Save(FileReference Location, object Object)
 		{
 			DirectoryReference.CreateDirectory(Location.Directory);
-			using (FileStream Stream = new FileStream(Location.FullName, FileMode.Create, FileAccess.Write))
-			{
-				BinaryFormatter Formatter = new BinaryFormatter();
-				Formatter.Serialize(Stream, Object);
-			}
+			using FileStream Stream = new FileStream(Location.FullName, FileMode.Create, FileAccess.Write);
+			BinaryFormatter Formatter = new BinaryFormatter();
+			Formatter.Serialize(Stream, Object);
 		}
 
 		/// <summary>
@@ -52,7 +43,7 @@ namespace EpicGames.Core
 		public static void SaveIfDifferent(FileReference Location, object Object)
 		{
 			byte[] Contents;
-			using(MemoryStream Stream = new MemoryStream())
+			using (MemoryStream Stream = new MemoryStream())
 			{
 				BinaryFormatter Formatter = new BinaryFormatter();
 				Formatter.Serialize(Stream, Object);
diff --git a/Engine/Source/Programs/Shared/EpicGames.Core/BinaryReaderExtensions.cs b/Engine/Source/Programs/Shared/EpicGames.Core/BinaryReaderExtensions.cs
--- a/Engine/Source/Programs/Shared/EpicGames.Core/BinaryReaderExtensions.cs
+++ b/Engine/Source/Programs/Shared/EpicGames.Core/BinaryReaderExtensions.cs
@@ -3,9 +3,6 @@
 using System;
 using System.Collections.Generic;
 using System.IO;
-using System.Linq;
-using System.Text;
-using System.Threading.Tasks;
 
 #nullable disable
 
@@ -55,13 +52,13 @@ namespace EpicGames.Core
 		public static T[] ReadArray<T>(this BinaryReader Reader, Func<T> ReadElement)
 		{
 			int NumItems = Reader.ReadInt32();
-			if(NumItems < 0)
+			if (NumItems < 0)
 			{
 				return null;
 			}
 
 			T[] Items = new T[NumItems];
-			for(int Idx = 0; Idx < NumItems; Idx++)
+			for (int Idx = 0; Idx < NumItems; Idx++)
 			{
 				Items[Idx] = ReadElement();
 			}
@@ -100,13 +97,13 @@ namespace EpicGames.Core
 		public static List<T> ReadList<T>(this BinaryReader Reader, Func<T> ReadElement)
 		{
 			int NumItems = Reader.ReadInt32();
-			if(NumItems < 0)
+			if (NumItems < 0)
 			{
 				return null;
 			}
 
 			List<T> Items = new List<T>(NumItems);
-			for(int Idx = 0; Idx < NumItems; Idx++)
+			for (int Idx = 0; Idx < NumItems; Idx++)
 			{
 				Items.Add(ReadElement());
 			}
@@ -124,13 +121,13 @@ namespace EpicGames.Core
 		public static Dictionary<K, V> ReadDictionary<K, V>(this BinaryReader Reader, Func<K> ReadKey, Func<V> ReadValue)
 		{
 			int NumItems = Reader.ReadInt32();
-			if(NumItems < 0)
+			if (NumItems < 0)
 			{
 				return null;
 			}
 
 			Dictionary<K, V> Items = new Dictionary<K, V>(NumItems);
-			for(int Idx = 0; Idx < NumItems; Idx++)
+			for (int Idx = 0; Idx < NumItems; Idx++)
 			{
 				K Key = ReadKey();
 				V Value = ReadValue();
@@ -148,7 +145,7 @@ namespace EpicGames.Core
 		/// <returns>Object instance or null</returns>
 		public static T ReadNullable<T>(this BinaryReader Reader, Func<T> ReadItem) where T : class
 		{
-			if(Reader.ReadBoolean())
+			if (Reader.ReadBoolean())
 			{
 				return ReadItem();
 			}
@@ -166,40 +163,40 @@ namespace EpicGames.Core
 		/// <returns>The value read from the stream</returns>
 		public static object ReadObject(this BinaryReader Reader, Type ObjectType)
 		{
-			if(ObjectType == typeof(string))
+			if (ObjectType == typeof(string))
 			{
 				return Reader.ReadString();
 			}
-			else if(ObjectType == typeof(bool))
+			else if (ObjectType == typeof(bool))
 			{
 				return Reader.ReadBoolean();
 			}
-			else if(ObjectType == typeof(int))
+			else if (ObjectType == typeof(int))
 			{
 				return Reader.ReadInt32();
 			}
-			else if(ObjectType == typeof(float))
+			else if (ObjectType == typeof(float))
 			{
 				return Reader.ReadSingle();
 			}
-			else if(ObjectType == typeof(double))
+			else if (ObjectType == typeof(double))
 			{
 				return Reader.ReadDouble();
 			}
-			else if(ObjectType == typeof(string[]))
+			else if (ObjectType == typeof(string[]))
 			{
 				return Reader.ReadStringArray();
 			}
-			else if(ObjectType == typeof(bool?))
+			else if (ObjectType == typeof(bool?))
 			{
 				int Value = Reader.ReadInt32();
-				return (Value == -1)? (bool?)null : (Value == 0)? (bool?)false : (bool?)true;
+				return (Value == -1) ? (bool?)null : (Value == 0) ? (bool?)false : (bool?)true;
 			}
 			else if (ObjectType == typeof(FileReference))
 			{
 				return Reader.ReadFileReference();
 			}
-			else if(ObjectType.IsEnum)
+			else if (ObjectType.IsEnum)
 			{
 				return Enum.ToObject(ObjectType, Reader.ReadInt32());
 			}
diff --git a/Engine/Source/Programs/Shared/EpicGames.Core/BinaryWriterExtensions.cs b/Engine/Source/Programs/Shared/EpicGames.Core/BinaryWriterExtensions.cs
--- a/Engine/Source/Programs/Shared/EpicGames.Core/BinaryWriterExtensions.cs
+++ b/Engine/Source/Programs/Shared/EpicGames.Core/BinaryWriterExtensions.cs
@@ -3,9 +3,6 @@
 using System;
 using System.Collections.Generic;
 using System.IO;
-using System.Linq;
-using System.Text;
-using System.Threading.Tasks;
 
 #nullable disable
 
@@ -66,14 +63,14 @@ namespace EpicGames.Core
 		/// <param name="WriteElement">Delegate to call to serialize each element</param>
 		public static void Write<T>(this BinaryWriter Writer, T[] Items, Action<T> WriteElement)
 		{
-			if(Items == null)
+			if (Items == null)
 			{
 				Writer.Write(-1);
 			}
 			else
 			{
 				Writer.Write(Items.Length);
-				for(int Idx = 0; Idx < Items.Length; Idx++)
+				for (int Idx = 0; Idx < Items.Length; Idx++)
 				{
 					WriteElement(Items[Idx]);
 				}
@@ -136,14 +133,14 @@ namespace EpicGames.Core
 		/// <returns>Dictionary of objects, as serialized. May be null.</returns>
 		public static void Write<K, V>(this BinaryWriter Writer, Dictionary<K, V> Items, Action<K> WriteKey, Action<V> WriteValue)
 		{
-			if(Items == null)
+			if (Items == null)
 			{
 				Writer.Write(-1);
 			}
 			else
 			{
 				Writer.Write(Items.Count);
-				foreach(KeyValuePair<K, V> Item in Items)
+				foreach (KeyValuePair<K, V> Item in Items)
 				{
 					WriteKey(Item.Key);
 					WriteValue(Item.Value);
@@ -160,7 +157,7 @@ namespace EpicGames.Core
 		/// <returns>Object instance or null</returns>
 		public static void WriteNullable<T>(this BinaryWriter Writer, T Item, Action WriteItem) where T : class
 		{
-			if(Item == null)
+			if (Item == null)
 			{
 				Writer.Write(false);
 			}
@@ -179,38 +176,38 @@ namespace EpicGames.Core
 		/// <param name="Value">The value to output</param>
 		public static void Write(this BinaryWriter Writer, Type FieldType, object Value)
 		{
-			if(FieldType == typeof(string))
+			if (FieldType == typeof(string))
 			{
 				Writer.Write((string)Value);
 			}
-			else if(FieldType == typeof(bool))
+			else if (FieldType == typeof(bool))
 			{
 				Writer.Write((bool)Value);
 			}
-			else if(FieldType == typeof(int))
+			else if (FieldType == typeof(int))
 			{
 				Writer.Write((int)Value);
 			}
-			else if(FieldType == typeof(float))
+			else if (FieldType == typeof(float))
 			{
 				Writer.Write((float)Value);
 			}
-			else if(FieldType == typeof(double))
+			else if (FieldType == typeof(double))
 			{
 				Writer.Write((double)Value);
 			}
-			else if(FieldType.IsEnum)
+			else if (FieldType.IsEnum)
 			{
 				Writer.Write((int)Value);
 			}
-			else if(FieldType == typeof(string[]))
+			else if (FieldType == typeof(string[]))
 			{
 				Writer.Write((string[])Value);
 			}
-			else if(FieldType == typeof(bool?))
+			else if (FieldType == typeof(bool?))
 			{
 				bool? NullableValue = (bool?)Value;
-				Writer.Write(NullableValue.HasValue? NullableValue.Value? 1 : 0 : -1);
+				Writer.Write(NullableValue.HasValue ? NullableValue.Value ? 1 : 0 : -1);
 			}
 			else if (FieldType == typeof(FileReference))
 			{
diff --git a/Engine/Source/Programs/Shared/EpicGames.Core/Boxed.cs b/Engine/Source/Programs/Shared/EpicGames.Core/Boxed.cs
--- a/Engine/Source/Programs/Shared/EpicGames.Core/Boxed.cs
+++ b/Engine/Source/Programs/Shared/EpicGames.Core/Boxed.cs
@@ -1,9 +1,5 @@
 // Copyright Epic Games, Inc. All Rights Reserved.
 
-using System;
-using System.Collections.Generic;
-using System.Text;
-
 namespace EpicGames.Core
 {
 	/// <summary>
diff --git a/Engine/Source/Programs/Shared/EpicGames.Core/CancellationTask.cs b/Engine/Source/Programs/Shared/EpicGames.Core/CancellationTask.cs
--- a/Engine/Source/Programs/Shared/EpicGames.Core/CancellationTask.cs
+++ b/Engine/Source/Programs/Shared/EpicGames.Core/CancellationTask.cs
@@ -1,8 +1,6 @@
 // Copyright Epic Games, Inc. All Rights Reserved.
 
 using System;
-using System.Collections.Generic;
-using System.Text;
 using System.Threading;
 using System.Threading.Tasks;
 
@@ -16,12 +14,12 @@ namespace EpicGames.Core
 		/// <summary>
 		/// Completion source for the task
 		/// </summary>
-		TaskCompletionSource<bool> CompletionSource;
+		readonly TaskCompletionSource<bool> CompletionSource;
 
 		/// <summary>
 		/// Registration handle with the cancellation token
 		/// </summary>
-		CancellationTokenRegistration Registration;
+		readonly CancellationTokenRegistration Registration;
 
 		/// <summary>
 		/// Constructor
diff --git a/Engine/Source/Programs/UnrealBuildTool/Executors/Experimental/TaskExecutor.cs b/Engine/Source/Programs/UnrealBuildTool/Executors/Experimental/TaskExecutor.cs
--- a/Engine/Source/Programs/UnrealBuildTool/Executors/Experimental/TaskExecutor.cs
+++ b/Engine/Source/Programs/UnrealBuildTool/Executors/Experimental/TaskExecutor.cs
@@ -20,7 +20,7 @@ namespace UnrealBuildTool
 		/// Maximum processor count for local execution. 
 		/// </summary>
 		[XmlConfigFile]
-		public static int MaxProcessorCount = int.MaxValue;
+		private static int MaxProcessorCount = int.MaxValue;
 
 		/// <summary>
 		/// Processor count multiplier for local execution. Can be below 1 to reserve CPU for other tasks.
@@ -29,33 +29,33 @@ namespace UnrealBuildTool
 		/// This value is ignored if the CPU does not support hyper-threading.
 		/// </summary>
 		[XmlConfigFile]
-		public static double ProcessorCountMultiplier = 1.0;
+		private static double ProcessorCountMultiplier = 1.0;
 
 		/// <summary>
 		/// Free memory per action in bytes, used to limit the number of parallel actions if the machine is memory starved.
 		/// Set to 0 to disable free memory checking.
 		/// </summary>
 		[XmlConfigFile]
-		static double MemoryPerActionBytes = 1.5 * 1024 * 1024 * 1024;
+		private static readonly double MemoryPerActionBytes = 1.5 * 1024 * 1024 * 1024;
 
 		/// <summary>
 		/// When enabled, will stop compiling targets after a compile error occurs.
 		/// </summary>
 		[XmlConfigFile]
-		bool bStopCompilationAfterErrors = false;
+		private static readonly bool bStopCompilationAfterErrors = false;
 
 		/// <summary>
 		/// Whether to show compilation times along with worst offenders or not.
 		/// </summary>
 		[XmlConfigFile]
-		bool bShowCompilationTimes = false;
+		private static readonly bool bShowCompilationTimes = false;
 
 		/// <summary>
 		/// How many processes that will be executed in parallel
 		/// </summary>
 		public int NumParallelProcesses { get; private set; }
 
-		private static char[] LineEndingSplit = new char[] { '\n', '\r' };
+		private static readonly char[] LineEndingSplit = new char[] { '\n', '\r' };
 
 		public static int GetDefaultNumParallelProcesses()
 		{
@@ -157,7 +157,7 @@ namespace UnrealBuildTool
 			}
 
 			// Wait for all tasks to complete
-			Task.WaitAll(AllTasks.ToArray());
+			Task.WaitAll(AllTasks.ToArray(), CancellationToken);
 
 			if (bShowCompilationTimes)
 			{
@@ -174,7 +174,7 @@ namespace UnrealBuildTool
 					Log.TraceInformation("");
 					foreach (var Pair in Tasks.OrderByDescending(x => x.Value.Result.ExecutionTime).Take(20))
 					{
-						string Description = $"{(Pair.Key.Inner.CommandDescription != null ? Pair.Key.Inner.CommandDescription : Pair.Key.Inner.CommandPath.GetFileNameWithoutExtension())} {Pair.Key.Inner.StatusDescription}".Trim();
+						string Description = $"{(Pair.Key.Inner.CommandDescription ?? Pair.Key.Inner.CommandPath.GetFileNameWithoutExtension())} {Pair.Key.Inner.StatusDescription}".Trim();
 						if (Pair.Value.Result.ProcessorTime.Ticks > 0)
 						{
 							Log.TraceInformation("{0} [ Wall Time {1:0.00} s / CPU Time {2:0.00} s ]", Description, Pair.Value.Result.ExecutionTime.TotalSeconds, Pair.Value.Result.ProcessorTime.TotalSeconds);
@@ -222,18 +222,17 @@ namespace UnrealBuildTool
 
 				CancellationToken.ThrowIfCancellationRequested();
 
-				using (ManagedProcess Process = new ManagedProcess(ProcessGroup, Action.CommandPath.FullName, Action.CommandArguments, Action.WorkingDirectory.FullName, null, null, ProcessPriorityClass.BelowNormal))
-				{
-					MemoryStream StdOutStream = new MemoryStream();
-					await Process.CopyToAsync(StdOutStream, CancellationToken);
-					CancellationToken.ThrowIfCancellationRequested();
+				using ManagedProcess Process = new ManagedProcess(ProcessGroup, Action.CommandPath.FullName, Action.CommandArguments, Action.WorkingDirectory.FullName, null, null, ProcessPriorityClass.BelowNormal);
 
-					List<string> LogLines = Console.OutputEncoding.GetString(StdOutStream.GetBuffer(), 0, Convert.ToInt32(StdOutStream.Length)).Split(LineEndingSplit, StringSplitOptions.RemoveEmptyEntries).ToList();
-					int ExitCode = Process.ExitCode;
-					TimeSpan ProcessorTime = Process.TotalProcessorTime;
-					TimeSpan ExecutionTime = Process.ExitTime - Process.StartTime;
-					return new ExecuteResults(LogLines, ExitCode, ExecutionTime, ProcessorTime);
-				}
+				MemoryStream StdOutStream = new MemoryStream();
+				await Process.CopyToAsync(StdOutStream, CancellationToken);
+				CancellationToken.ThrowIfCancellationRequested();
+
+				List<string> LogLines = Console.OutputEncoding.GetString(StdOutStream.GetBuffer(), 0, Convert.ToInt32(StdOutStream.Length)).Split(LineEndingSplit, StringSplitOptions.RemoveEmptyEntries).ToList();
+				int ExitCode = Process.ExitCode;
+				TimeSpan ProcessorTime = Process.TotalProcessorTime;
+				TimeSpan ExecutionTime = Process.ExitTime - Process.StartTime;
+				return new ExecuteResults(LogLines, ExitCode, ExecutionTime, ProcessorTime);
 			}
 			catch (OperationCanceledException)
 			{
@@ -267,11 +266,11 @@ namespace UnrealBuildTool
 			string Description = string.Empty;
 			if (Action.bShouldOutputStatusDescription || LogLines.Count == 0)
 			{
-				Description = $"{(Action.CommandDescription != null ? Action.CommandDescription : Action.CommandPath.GetFileNameWithoutExtension())} {Action.StatusDescription}".Trim();
+				Description = $"{(Action.CommandDescription ?? Action.CommandPath.GetFileNameWithoutExtension())} {Action.StatusDescription}".Trim();
 			}
 			else if (LogLines.Count > 0)
 			{
-				Description = $"{(Action.CommandDescription != null ? Action.CommandDescription : Action.CommandPath.GetFileNameWithoutExtension())} {LogLines[0]}".Trim();
+				Description = $"{(Action.CommandDescription ?? Action.CommandPath.GetFileNameWithoutExtension())} {LogLines[0]}".Trim();
 			}
 
 			lock (ProgressWriter)
diff --git a/Engine/Source/Programs/UnrealBuildTool/System/ActionExecutor.cs b/Engine/Source/Programs/UnrealBuildTool/System/ActionExecutor.cs
--- a/Engine/Source/Programs/UnrealBuildTool/System/ActionExecutor.cs
+++ b/Engine/Source/Programs/UnrealBuildTool/System/ActionExecutor.cs
@@ -1,10 +1,6 @@
 // Copyright Epic Games, Inc. All Rights Reserved.
 
-using System;
 using System.Collections.Generic;
-using System.Linq;
-using System.Text;
-using System.Threading.Tasks;
 
 namespace UnrealBuildTool
 {
diff --git a/Engine/Source/Programs/UnrealBuildTool/System/ActionGraph.cs b/Engine/Source/Programs/UnrealBuildTool/System/ActionGraph.cs
--- a/Engine/Source/Programs/UnrealBuildTool/System/ActionGraph.cs
+++ b/Engine/Source/Programs/UnrealBuildTool/System/ActionGraph.cs
@@ -1,16 +1,15 @@
 // Copyright Epic Games, Inc. All Rights Reserved.
 
+using EpicGames.Core;
 using System;
 using System.Collections.Generic;
-using System.IO;
 using System.Diagnostics;
+using System.IO;
 using System.Linq;
-using System.Runtime.Serialization;
-using EpicGames.Core;
 using System.Text;
-using System.Threading.Tasks;
 using System.Text.Json;
 using System.Text.Json.Serialization;
+using System.Threading.Tasks;
 using UnrealBuildBase;
 
 namespace UnrealBuildTool
@@ -58,8 +57,7 @@ namespace UnrealBuildTool
 				Action.PrerequisiteActions = new HashSet<LinkedAction>();
 				foreach(FileItem PrerequisiteItem in Action.PrerequisiteItems)
 				{
-					LinkedAction? PrerequisiteAction;
-					if(ItemToProducingAction.TryGetValue(PrerequisiteItem, out PrerequisiteAction))
+					if (ItemToProducingAction.TryGetValue(PrerequisiteItem, out LinkedAction? PrerequisiteAction))
 					{
 						Action.PrerequisiteActions.Add(PrerequisiteAction);
 					}
@@ -83,8 +81,7 @@ namespace UnrealBuildTool
 			{
 				foreach(FileItem ProducedItem in Action.ProducedItems)
 				{
-					IExternalAction? ExistingAction;
-					if(ItemToProducingAction.TryGetValue(ProducedItem, out ExistingAction))
+					if (ItemToProducingAction.TryGetValue(ProducedItem, out IExternalAction? ExistingAction))
 					{
 						bResult &= CheckForConflicts(ExistingAction, Action);
 					}
@@ -428,8 +425,7 @@ namespace UnrealBuildTool
 						bool bActionOnlyDependsOnNonCyclicalActions = true;
 						foreach (FileItem PrerequisiteItem in Action.PrerequisiteItems)
 						{
-							LinkedAction? ProducingAction;
-							if (ItemToProducingAction.TryGetValue(PrerequisiteItem, out ProducingAction))
+							if (ItemToProducingAction.TryGetValue(PrerequisiteItem, out LinkedAction? ProducingAction))
 							{
 								if (!ActionIsNonCyclical.ContainsKey(ProducingAction))
 								{
@@ -589,7 +585,6 @@ namespace UnrealBuildTool
 			}
 
 			// Determine the last time the action was run based on the write times of its produced files.
-			string? LatestUpdatedProducedItemName = null;
 			DateTimeOffset LastExecutionTimeUtc = DateTimeOffset.MaxValue;
 			foreach (FileItem ProducedItem in RootAction.ProducedItems)
 			{
@@ -619,7 +614,6 @@ namespace UnrealBuildTool
 					if (ProducedItem.LastWriteTimeUtc < LastExecutionTimeUtc)
 					{
 						LastExecutionTimeUtc = ProducedItem.LastWriteTimeUtc;
-						LatestUpdatedProducedItemName = ProducedItem.AbsolutePath;
 					}
 				}
 				else
@@ -682,17 +676,16 @@ namespace UnrealBuildTool
 			// Check the dependency list
 			if(!bIsOutdated && RootAction.DependencyListFile != null)
 			{
-				List<FileItem>? DependencyFiles;
-				if(!CppDependencies.TryGetDependencies(RootAction.DependencyListFile, out DependencyFiles))
+				if (!CppDependencies.TryGetDependencies(RootAction.DependencyListFile, out List<FileItem>? DependencyFiles))
 				{
 					Log.TraceLog("{0}: Missing dependency list file \"{1}\"", RootAction.StatusDescription, RootAction.DependencyListFile);
 					bIsOutdated = true;
 				}
 				else
 				{
-					foreach(FileItem DependencyFile in DependencyFiles)
+					foreach (FileItem DependencyFile in DependencyFiles)
 					{
-						if(!DependencyFile.Exists || DependencyFile.LastWriteTimeUtc > LastExecutionTimeUtc)
+						if (!DependencyFile.Exists || DependencyFile.LastWriteTimeUtc > LastExecutionTimeUtc)
 						{
 							Log.TraceLog(
 								"{0}: Dependency {1} is newer than the last execution of the action: {2} vs {3}",
@@ -775,7 +768,7 @@ namespace UnrealBuildTool
 						Dependencies.Add(Action.DependencyListFile);
 					}
 				}
-				Parallel.ForEach(Dependencies, File => { List<FileItem>? Temp; CppDependencies.TryGetDependencies(File, out Temp); });
+				Parallel.ForEach(Dependencies, File => { CppDependencies.TryGetDependencies(File, out _); });
 			}
 
 			using(Timeline.ScopeEvent("Cache outdated actions"))
@@ -856,37 +849,35 @@ namespace UnrealBuildTool
 		public static void ExportJson(IEnumerable<LinkedAction> Actions, FileReference OutputFile)
 		{
 			DirectoryReference.CreateDirectory(OutputFile.Directory);
-			using (JsonWriter Writer = new JsonWriter(OutputFile))
-			{
-				Writer.WriteObjectStart();
+			using JsonWriter Writer = new JsonWriter(OutputFile);
+			Writer.WriteObjectStart();
 
-				Writer.WriteObjectStart("Environment");
-				foreach (object? Object in Environment.GetEnvironmentVariables())
+			Writer.WriteObjectStart("Environment");
+			foreach (object? Object in Environment.GetEnvironmentVariables())
+			{
+				System.Collections.DictionaryEntry Pair = (System.Collections.DictionaryEntry)Object!;
+				if (!UnrealBuildTool.InitialEnvironment!.Contains(Pair.Key) || (string)(UnrealBuildTool.InitialEnvironment![Pair.Key]!) != (string)(Pair.Value!))
 				{
-					System.Collections.DictionaryEntry Pair = (System.Collections.DictionaryEntry)Object!;
-					if (!UnrealBuildTool.InitialEnvironment!.Contains(Pair.Key) || (string)(UnrealBuildTool.InitialEnvironment![Pair.Key]!) != (string)(Pair.Value!))
-					{
-						Writer.WriteValue((string)Pair!.Key, (string)Pair.Value!);
-					}
+					Writer.WriteValue((string)Pair!.Key, (string)Pair.Value!);
 				}
-				Writer.WriteObjectEnd();
+			}
+			Writer.WriteObjectEnd();
 
-				Dictionary<LinkedAction, int> ActionToId = new Dictionary<LinkedAction, int>();
-				foreach (LinkedAction Action in Actions)
-				{
-					ActionToId[Action] = ActionToId.Count;
-				}
+			Dictionary<LinkedAction, int> ActionToId = new Dictionary<LinkedAction, int>();
+			foreach (LinkedAction Action in Actions)
+			{
+				ActionToId[Action] = ActionToId.Count;
+			}
 
-				Writer.WriteArrayStart("Actions");
-				foreach (LinkedAction Action in Actions)
-				{
-					Writer.WriteObjectStart();
-					Action.ExportJson(ActionToId, Writer);
-					Writer.WriteObjectEnd();
-				}
-				Writer.WriteArrayEnd();
+			Writer.WriteArrayStart("Actions");
+			foreach (LinkedAction Action in Actions)
+			{
+				Writer.WriteObjectStart();
+				Action.ExportJson(ActionToId, Writer);
 				Writer.WriteObjectEnd();
 			}
+			Writer.WriteArrayEnd();
+			Writer.WriteObjectEnd();
 		}
 	}
 }
