UnrealBuildTool: VS2019 shouldn't return VS2022 clang toolchain, but also scan VS2022 install for clang as well

#rnx

diff --git a/Engine/Source/Programs/UnrealBuildTool/Platform/Windows/UEBuildWindows.cs b/Engine/Source/Programs/UnrealBuildTool/Platform/Windows/UEBuildWindows.cs
--- a/Engine/Source/Programs/UnrealBuildTool/Platform/Windows/UEBuildWindows.cs
+++ b/Engine/Source/Programs/UnrealBuildTool/Platform/Windows/UEBuildWindows.cs
@@ -1093,7 +1093,7 @@ namespace UnrealBuildTool
 										{
 											continue;
 										}
-										else if (Compiler == WindowsCompiler.VisualStudio2019 && Version < Version2019)
+										else if (Compiler == WindowsCompiler.VisualStudio2019 && (Version < Version2019 || Version >= Version2022))
 										{
 											continue;
 										}
@@ -1168,6 +1168,12 @@ namespace UnrealBuildTool
 							AddClangToolChain(DirectoryReference.Combine(Installation.BaseDir, "VC", "Tools", "Llvm"), ToolChains);
 						}
 
+						// Check for installations bundled with Visual Studio 2022
+						foreach (VisualStudioInstallation Installation in FindVisualStudioInstallations(WindowsCompiler.VisualStudio2022))
+						{
+							AddClangToolChain(DirectoryReference.Combine(Installation.BaseDir, "VC", "Tools", "Llvm"), ToolChains);
+						}
+
 						// Check for AutoSDK paths
 						DirectoryReference AutoSdkDir;
 						if(UEBuildPlatformSDK.TryGetHostPlatformAutoSDKDir(out AutoSdkDir))
