UBT: Fix files added to the working set not invalidating the makefile.

diff --git a/Engine/Source/Programs/UnrealBuildTool/System/Unity.cs b/Engine/Source/Programs/UnrealBuildTool/System/Unity.cs
--- a/Engine/Source/Programs/UnrealBuildTool/System/Unity.cs
+++ b/Engine/Source/Programs/UnrealBuildTool/System/Unity.cs
@@ -362,6 +362,15 @@ namespace UnrealBuildTool
 			{
 				Graph.AddFileToWorkingSet(File);
 			}
+
+			HashSet<FileItem> CandidateAdaptiveFiles = new HashSet<FileItem>();
+			CandidateAdaptiveFiles.UnionWith(CPPFiles);
+			CandidateAdaptiveFiles.UnionWith(HeaderFiles);
+
+			foreach (FileItem File in CandidateAdaptiveFiles)
+			{
+				Graph.AddCandidateForWorkingSet(File);
+			}
 		}
 
 		static void AddUniqueDiagnostic(TargetMakefile Makefile, string Message)
