UnrealBuildTool: Prefer Clang 11 on Windows, to match the current version supported for Linux

#rnx

diff --git a/Engine/Source/Programs/UnrealBuildTool/Platform/Windows/UEBuildWindows.cs b/Engine/Source/Programs/UnrealBuildTool/Platform/Windows/UEBuildWindows.cs
--- a/Engine/Source/Programs/UnrealBuildTool/Platform/Windows/UEBuildWindows.cs
+++ b/Engine/Source/Programs/UnrealBuildTool/Platform/Windows/UEBuildWindows.cs
@@ -745,6 +745,7 @@ namespace UnrealBuildTool
 		/// </summary>
 		static readonly VersionNumberRange[] PreferredClangVersions =
 		{
+			VersionNumberRange.Parse("11.0.0", "11.999"),
 			VersionNumberRange.Parse("10.0.0", "10.999")
 		};
 
