UnrealBuildTool: Always print number of tasks complete when building with XGE even if bWriteMarkup isn't enabled

#rnx

diff --git a/Engine/Source/Programs/UnrealBuildTool/Executors/XGE.cs b/Engine/Source/Programs/UnrealBuildTool/Executors/XGE.cs
--- a/Engine/Source/Programs/UnrealBuildTool/Executors/XGE.cs
+++ b/Engine/Source/Programs/UnrealBuildTool/Executors/XGE.cs
@@ -351,7 +351,7 @@ namespace UnrealBuildTool
 			{
 				// Write the actions to execute to a XGE task file.
 				string XGETaskFilePath = FileReference.Combine(Unreal.EngineDirectory, "Intermediate", "Build", "XGETasks.xml").FullName;
-				WriteTaskFile(Actions, XGETaskFilePath, ProgressWriter.bWriteMarkup, false);
+				WriteTaskFile(Actions, XGETaskFilePath, true, false);
 
 				XGEResult = ExecuteTaskFileWithProgressMarkup(XGETaskFilePath, Actions.Count);
 			}
@@ -635,6 +635,7 @@ namespace UnrealBuildTool
 			using (ProgressWriter Writer = new ProgressWriter("Compiling C++ source files...", false))
 			{
 				int NumCompletedActions = 0;
+				string ProgressText = string.Empty;
 
 				// Create a wrapper delegate that will parse the output actions
 				DataReceivedEventHandler EventHandlerWrapper = (Sender, Args) =>
@@ -644,14 +645,30 @@ namespace UnrealBuildTool
 						string Text = Args.Data;
 						if (Text.StartsWith(ProgressMarkupPrefix))
 						{
+							// Flush old progress text
+							if (!string.IsNullOrEmpty(ProgressText))
+							{
+								Log.TraceInformation($"[{NumCompletedActions}/{NumActions}] Complete {ProgressText}");
+								ProgressText = string.Empty;
+							}
 							Writer.Write(++NumCompletedActions, NumActions);
 
 							// Strip out anything that is just an XGE timer. Some programs don't output anything except the progress text.
 							Text = Args.Data.Substring(ProgressMarkupPrefix.Length);
 							if(Text.StartsWith(" (") && Text.EndsWith(")"))
 							{
+								// Write the progress text with the next line of output if the current doesn't have any status.
+								ProgressText = Text.Trim();
 								return;
 							}
+							Log.TraceInformation($"[{NumCompletedActions}/{NumActions}] {Text}");
+							return;
+						}
+						if (!string.IsNullOrEmpty(ProgressText))
+						{
+							Log.TraceInformation($"[{NumCompletedActions}/{NumActions}] {Text} {ProgressText}");
+							ProgressText = string.Empty;
+							return;
 						}
 						Log.TraceInformation(Text);
 					}
