Fix issue with resizableActivity in manifest

#ue5
#android

diff --git a/Engine/Source/Programs/UnrealBuildTool/Platform/Android/UEDeployAndroid.cs b/Engine/Source/Programs/UnrealBuildTool/Platform/Android/UEDeployAndroid.cs
--- a/Engine/Source/Programs/UnrealBuildTool/Platform/Android/UEDeployAndroid.cs
+++ b/Engine/Source/Programs/UnrealBuildTool/Platform/Android/UEDeployAndroid.cs
@@ -2668,6 +2668,10 @@ namespace UnrealBuildTool
 				Text.AppendLine("\t\t          android:label=\"@string/app_name\"");
 				Text.AppendLine("\t\t          android:theme=\"@style/UnrealSplashTheme\"");
 				Text.AppendLine("\t\t          android:launchMode=\"singleTask\"");
+				if (SDKLevelInt >= 24)
+				{
+					Text.AppendLine("\t\t          android:resizeableActivity=\"false\"");
+				}
 				Text.AppendLine(string.Format("\t\t          android:screenOrientation=\"{0}\"", Orientation));
 				Text.AppendLine(string.Format("\t\t          android:debuggable=\"{0}\">", bIsForDistribution ? "false" : "true"));
 				Text.AppendLine("\t\t\t<intent-filter>");
