- Allow for a target to add extra linker commandline options in Visual Studio
#p4v-preflight-copy 17108935

diff --git a/Engine/Source/Programs/UnrealBuildTool/Platform/Windows/UEBuildWindows.cs b/Engine/Source/Programs/UnrealBuildTool/Platform/Windows/UEBuildWindows.cs
--- a/Engine/Source/Programs/UnrealBuildTool/Platform/Windows/UEBuildWindows.cs
+++ b/Engine/Source/Programs/UnrealBuildTool/Platform/Windows/UEBuildWindows.cs
@@ -251,11 +251,10 @@ namespace UnrealBuildTool
 		public int PCHMemoryAllocationFactor = 0;
 
 		/// <summary>
-		/// Override the PDB page size. Used when the PDB file grows over 4GB. Expected values are 4096 (default), 8192, 16384, 32768
+		/// Allow the target to specify extra options for linking that aren't otherwise noted here
 		/// </summary>
 		[XmlConfigFile(Category = "WindowsPlatform")]
-		[CommandLine("-PDBPageSize")]
-		public int PDBPageSize = 0;
+		public string AdditionalLinkerOptions = "";
 
 		/// <summary>
 		/// True if we allow using addresses larger than 2GB on 32 bit builds
@@ -546,9 +545,9 @@ namespace UnrealBuildTool
 			get { return Inner.PCHMemoryAllocationFactor; }
 		}
 
-		public int PDBPageSize
+		public string AdditionalLinkerOptions
 		{
-			get { return Inner.PDBPageSize; }
+			get { return Inner.AdditionalLinkerOptions; }
 		}
 
 		public bool bBuildLargeAddressAwareBinary
diff --git a/Engine/Source/Programs/UnrealBuildTool/Platform/Windows/VCToolChain.cs b/Engine/Source/Programs/UnrealBuildTool/Platform/Windows/VCToolChain.cs
--- a/Engine/Source/Programs/UnrealBuildTool/Platform/Windows/VCToolChain.cs
+++ b/Engine/Source/Programs/UnrealBuildTool/Platform/Windows/VCToolChain.cs
@@ -1056,14 +1056,10 @@ namespace UnrealBuildTool
 			}
 
 
-			// Override the PDB page size, if requested (VS2019 16.11 Preview 1 & above)
-			if (Target.WindowsPlatform.PDBPageSize > 0 && Target.WindowsPlatform.Compiler >= WindowsCompiler.VisualStudio2019)
+			// Add any extra options from the target
+			if (!string.IsNullOrEmpty(Target.WindowsPlatform.AdditionalLinkerOptions))
 			{
-				Version CompilerVersion;
-				if (Version.TryParse(Target.WindowsPlatform.CompilerVersion, out CompilerVersion) && CompilerVersion >= new Version("14.29.30130"))
-				{
-					Arguments.Add($"/PDBPAGESIZE:{Target.WindowsPlatform.PDBPageSize}");
-				}
+				Arguments.Add(Target.WindowsPlatform.AdditionalLinkerOptions);
 			}
 
 
