Add sourceFileMap to generated vs code launch.json for installed builds, so the debugger can load source files correctly when debugging the editor

diff --git a/Engine/Source/Programs/UnrealBuildTool/ProjectFiles/VisualStudioCode/VSCodeProjectFileGenerator.cs b/Engine/Source/Programs/UnrealBuildTool/ProjectFiles/VisualStudioCode/VSCodeProjectFileGenerator.cs
--- a/Engine/Source/Programs/UnrealBuildTool/ProjectFiles/VisualStudioCode/VSCodeProjectFileGenerator.cs
+++ b/Engine/Source/Programs/UnrealBuildTool/ProjectFiles/VisualStudioCode/VSCodeProjectFileGenerator.cs
@@ -1277,7 +1277,7 @@ namespace UnrealBuildTool
 					if (HostPlatform == UnrealTargetPlatform.Win64)
 					{
 						OutFile.AddField("stopAtEntry", false);
-						OutFile.AddField("externalConsole", true);
+						OutFile.AddField("console", "externalTerminal");
 
 						OutFile.AddField("type", "cppvsdbg");
 						OutFile.AddField("visualizerFile", MakeUnquotedPathString(FileReference.Combine(UE4ProjectRoot, "Engine", "Extras", "VisualStudioDebugging", "Unreal.natvis"), EPathType.Absolute));
@@ -1292,6 +1292,17 @@ namespace UnrealBuildTool
 					{
 						OutFile.AddField("type", "lldb");
 					}
+
+					if (UnrealBuildTool.OriginalCompilationRootDirectory != UE4ProjectRoot)
+					{
+						OutFile.BeginObject("sourceFileMap");
+						{
+							OutFile.AddField(
+								MakeUnquotedPathString(UnrealBuildTool.OriginalCompilationRootDirectory, EPathType.Absolute),
+								MakeUnquotedPathString(UE4ProjectRoot, EPathType.Absolute));
+						}
+						OutFile.EndObject();
+					}
 				}
 				OutFile.EndObject();
 			}
diff --git a/Engine/Source/Programs/UnrealBuildTool/UnrealBuildTool.cs b/Engine/Source/Programs/UnrealBuildTool/UnrealBuildTool.cs
--- a/Engine/Source/Programs/UnrealBuildTool/UnrealBuildTool.cs
+++ b/Engine/Source/Programs/UnrealBuildTool/UnrealBuildTool.cs
@@ -7,6 +7,7 @@ using System.IO;
 using System.Linq;
 using System.Reflection;
 using System.Text;
+using System.Text.Json;
 using System.Threading;
 using EpicGames.Core;
 using UnrealBuildBase;
@@ -51,6 +52,12 @@ namespace UnrealBuildTool
 		/// </summary>
 		static DirectoryReference? CachedWritableEngineDirectory;
 
+		/// <summary>
+		/// Cached copy of the source root directory that was used to compile the installed engine
+		/// Used to remap source code paths when debugging.
+		/// </summary>
+		static DirectoryReference? CachedOriginalCompilationRootDirectory;
+
 		/// <summary>
 		/// Writable engine directory. Uses the user's settings folder for installed builds.
 		/// </summary>
@@ -100,6 +107,55 @@ namespace UnrealBuildTool
 			}
 		}
 
+		/// <summary>
+		/// The original root directory that was used to compile the installed engine
+		/// Used to remap source code paths when debugging.
+		/// </summary>
+		public static DirectoryReference OriginalCompilationRootDirectory
+		{
+			get
+			{
+				if (CachedOriginalCompilationRootDirectory == null)
+				{
+					if (Unreal.IsEngineInstalled())
+					{
+						// Load Engine\Intermediate\Build\BuildRules\*RulesManifest.json
+						DirectoryReference BuildRules = DirectoryReference.Combine(Unreal.EngineDirectory, "Intermediate", "Build", "BuildRules");
+						FileReference? RulesManifest = DirectoryReference.EnumerateFiles(BuildRules, "*RulesManifest.json").FirstOrDefault();
+						if (RulesManifest != null)
+						{
+							JsonObject Manifest = JsonObject.Read(RulesManifest);
+							if (Manifest.TryGetStringArrayField("SourceFiles", out string[] SourceFiles))
+							{
+								FileReference? SourceFile = FileReference.FromString(SourceFiles.FirstOrDefault());
+								if (SourceFile != null && !SourceFile.IsUnderDirectory(Unreal.EngineDirectory))
+								{
+									// Walk up parent directory until Engine is found
+									DirectoryReference? Directory = SourceFile.Directory;
+									while (Directory != null && !Directory.IsRootDirectory())
+									{
+										if (Directory.GetDirectoryName() == "Engine" && Directory.ParentDirectory != null)
+										{
+											CachedOriginalCompilationRootDirectory = Directory.ParentDirectory;
+											break;
+										}
+
+										Directory = Directory.ParentDirectory;
+									}
+								}
+							}
+						}
+					}
+
+					if (CachedOriginalCompilationRootDirectory == null)
+					{
+						CachedOriginalCompilationRootDirectory = Unreal.RootDirectory;
+					}
+				}
+				return CachedOriginalCompilationRootDirectory;
+			}
+		}
+
 		/// <summary>
 		/// The Remote Ini directory.  This should always be valid when compiling using a remote server.
 		/// </summary>
