UnrealBuildTool: Don't allow remote linking for XGE when on vpn

#rnx

diff --git a/Engine/Source/Programs/UnrealBuildTool/Executors/XGE.cs b/Engine/Source/Programs/UnrealBuildTool/Executors/XGE.cs
--- a/Engine/Source/Programs/UnrealBuildTool/Executors/XGE.cs
+++ b/Engine/Source/Programs/UnrealBuildTool/Executors/XGE.cs
@@ -247,7 +247,10 @@ namespace UnrealBuildTool
 						{
 							if (Subnet.Contains(AddressBytes))
 							{
-								Log.TraceInformationOnce("XGE coordinator {0} will be not be used over VPN (adapter '{1}' with IP {2} is in subnet {3}). Set <XGE><bAllowOverVpn>true</bAllowOverVpn></XGE> in BuildConfiguration.xml to override.", HostName, Interface.Description, UnicastAddressInfo.Address, Subnet);
+								if (!bAllowOverVpn)
+								{
+									Log.TraceInformationOnce("XGE coordinator {0} will be not be used over VPN (adapter '{1}' with IP {2} is in subnet {3}). Set <XGE><bAllowOverVpn>true</bAllowOverVpn></XGE> in BuildConfiguration.xml to override.", HostName, Interface.Description, UnicastAddressInfo.Address, Subnet);
+								}
 								return true;
 							}
 						}
@@ -363,6 +366,8 @@ namespace UnrealBuildTool
 		/// </summary>
 		static void WriteTaskFile(List<LinkedAction> InActions, string TaskFilePath, bool bProgressMarkup, bool bXGEExport)
 		{
+			bool HostOnVpn = TryGetCoordinatorHost(out string CoordinatorHost) && IsHostOnVpn(CoordinatorHost);
+
 			Dictionary<string, string> ExportEnv = new Dictionary<string, string>();
 
 			List<LinkedAction> Actions = InActions;
@@ -418,14 +423,17 @@ namespace UnrealBuildTool
 			{
 				LinkedAction Action = Actions[ActionIndex];
 
+				// Don't allow remote linking if on VPN.
+				bool CanExecuteRemotely = Action.bCanExecuteRemotely && !(Action.ActionType == ActionType.Link && HostOnVpn);
+
 				// <Tool ... />
 				XmlElement ToolElement = XGETaskDocument.CreateElement("Tool");
 				ToolsElement.AppendChild(ToolElement);
 				ToolElement.SetAttribute("Name", string.Format("Tool{0}", ActionIndex));
-				ToolElement.SetAttribute("AllowRemote", Action.bCanExecuteRemotely.ToString());
+				ToolElement.SetAttribute("AllowRemote", CanExecuteRemotely.ToString());
 
 				// The XGE documentation says that 'AllowIntercept' must be set to 'true' for all tools where 'AllowRemote' is enabled
-				ToolElement.SetAttribute("AllowIntercept", Action.bCanExecuteRemotely.ToString());
+				ToolElement.SetAttribute("AllowIntercept", CanExecuteRemotely.ToString());
 
 				string OutputPrefix = "";
 				if (bProgressMarkup)
