VsCodeProjectGen: Remove references to UE4 and fix some old references that weren't updated

#rnx

diff --git a/Engine/Source/Programs/UnrealBuildTool/ProjectFiles/VisualStudioCode/VSCodeProjectFileGenerator.cs b/Engine/Source/Programs/UnrealBuildTool/ProjectFiles/VisualStudioCode/VSCodeProjectFileGenerator.cs
--- a/Engine/Source/Programs/UnrealBuildTool/ProjectFiles/VisualStudioCode/VSCodeProjectFileGenerator.cs
+++ b/Engine/Source/Programs/UnrealBuildTool/ProjectFiles/VisualStudioCode/VSCodeProjectFileGenerator.cs
@@ -39,7 +39,7 @@ namespace UnrealBuildTool
 		private DirectoryReference VSCodeDir;
 		private UnrealTargetPlatform HostPlatform = BuildHostPlatform.Current.Platform;
 		private bool bForeignProject;
-		private DirectoryReference UE4ProjectRoot;
+		private DirectoryReference ProjectRoot;
 		private string FrameworkExecutableExtension = ".exe";
 		private string FrameworkLibraryExtension = ".dll";
 
@@ -79,7 +79,7 @@ namespace UnrealBuildTool
 		{
 			if (InRelativeRoot == null)
 			{
-				InRelativeRoot = UE4ProjectRoot;
+				InRelativeRoot = ProjectRoot;
 			}
 
 			string Processed = InRef.ToString();
@@ -143,11 +143,11 @@ namespace UnrealBuildTool
 		{
 			if (bForceSkipQuotes)
 			{
-				return MakeUnquotedPathString(InRef, bInAbsolute ? EPathType.Absolute : EPathType.Relative, UE4ProjectRoot);
+				return MakeUnquotedPathString(InRef, bInAbsolute ? EPathType.Absolute : EPathType.Relative, ProjectRoot);
 			}
 			else
 			{
-				return MakeQuotedPathString(InRef, bInAbsolute ? EPathType.Absolute : EPathType.Relative, UE4ProjectRoot);
+				return MakeQuotedPathString(InRef, bInAbsolute ? EPathType.Absolute : EPathType.Relative, ProjectRoot);
 			}
 		}
 
@@ -266,8 +266,8 @@ namespace UnrealBuildTool
 			VSCodeDir = DirectoryReference.Combine(MasterProjectPath, ".vscode");
 			DirectoryReference.CreateDirectory(VSCodeDir);
 
-			UE4ProjectRoot = Unreal.RootDirectory;
-			bForeignProject = !VSCodeDir.IsUnderDirectory(UE4ProjectRoot);
+			ProjectRoot = Unreal.RootDirectory;
+			bForeignProject = !VSCodeDir.IsUnderDirectory(ProjectRoot);
 
 			List<ProjectFile> Projects;
 
@@ -300,8 +300,8 @@ namespace UnrealBuildTool
 			if (bForeignProject && bIncludeEngineSource)
 			{
 				// for installed builds we need to write the cpp properties file under the installed engine as well for intellisense to work
-				DirectoryReference Ue4CodeDirectory = DirectoryReference.Combine(Unreal.RootDirectory, ".vscode");
-				WriteCppPropertiesFile(Ue4CodeDirectory, ProjectData);
+				DirectoryReference VsCodeDirectory = DirectoryReference.Combine(Unreal.RootDirectory, ".vscode");
+				WriteCppPropertiesFile(VsCodeDirectory, ProjectData);
 			}
 
 			return true;
@@ -700,11 +700,11 @@ namespace UnrealBuildTool
 
 					if (HostPlatform == UnrealTargetPlatform.Win64)
 					{
-						OutFile.AddField("command", MakePathString(FileReference.Combine(UE4ProjectRoot, "Engine", "Build", "BatchFiles", "RunUAT.bat")));
+						OutFile.AddField("command", MakePathString(FileReference.Combine(ProjectRoot, "Engine", "Build", "BatchFiles", "RunUAT.bat")));
 					}
 					else
 					{
-						OutFile.AddField("command", MakePathString(FileReference.Combine(UE4ProjectRoot, "Engine", "Build", "BatchFiles", "RunUAT.sh")));
+						OutFile.AddField("command", MakePathString(FileReference.Combine(ProjectRoot, "Engine", "Build", "BatchFiles", "RunUAT.sh")));
 					}
 
 					OutFile.BeginArray("args");
@@ -772,7 +772,7 @@ namespace UnrealBuildTool
 
 					OutFile.BeginObject("options");
 					{
-						OutFile.AddField("cwd", MakeUnquotedPathString(UE4ProjectRoot, EPathType.Absolute));
+						OutFile.AddField("cwd", MakeUnquotedPathString(ProjectRoot, EPathType.Absolute));
 					}
 					OutFile.EndObject();
 				}
@@ -859,12 +859,12 @@ namespace UnrealBuildTool
 
 							if (HostPlatform == UnrealTargetPlatform.Win64)
 							{
-								OutFile.AddField("command", MakePathString(FileReference.Combine(UE4ProjectRoot, "Engine", "Build", "BatchFiles", Command + ".bat")));
+								OutFile.AddField("command", MakePathString(FileReference.Combine(ProjectRoot, "Engine", "Build", "BatchFiles", Command + ".bat")));
 								CleanParam = null;
 							}
 							else
 							{
-								OutFile.AddField("command", MakePathString(FileReference.Combine(UE4ProjectRoot, "Engine", "Build", "BatchFiles", HostPlatform.ToString(), "Build.sh")));
+								OutFile.AddField("command", MakePathString(FileReference.Combine(ProjectRoot, "Engine", "Build", "BatchFiles", HostPlatform.ToString(), "Build.sh")));
 
 								if (Command == "Clean")
 								{
@@ -919,7 +919,7 @@ namespace UnrealBuildTool
 
 							OutFile.BeginObject("options");
 							{
-								OutFile.AddField("cwd", MakeUnquotedPathString(UE4ProjectRoot, EPathType.Absolute));
+								OutFile.AddField("cwd", MakeUnquotedPathString(ProjectRoot, EPathType.Absolute));
 							}
 							OutFile.EndObject();
 						}
@@ -953,7 +953,7 @@ namespace UnrealBuildTool
 							OutFile.AddField("group", "build");
 							if (!RuntimePlatform.IsWindows)
 							{
-								OutFile.AddField("command", MakePathString(FileReference.Combine(UE4ProjectRoot, "Engine", "Build", "BatchFiles", HostPlatform.ToString(), "RunDotnet.sh")));
+								OutFile.AddField("command", MakePathString(FileReference.Combine(ProjectRoot, "Engine", "Build", "BatchFiles", HostPlatform.ToString(), "RunDotnet.sh")));
 							}
 							else
 							{
@@ -974,7 +974,7 @@ namespace UnrealBuildTool
 
 						OutFile.BeginObject("options");
 						{
-							OutFile.AddField("cwd", MakeUnquotedPathString(UE4ProjectRoot, EPathType.Absolute));
+							OutFile.AddField("cwd", MakeUnquotedPathString(ProjectRoot, EPathType.Absolute));
 						}
 
 						OutFile.EndObject();
@@ -1171,7 +1171,7 @@ namespace UnrealBuildTool
 						OutFile.BeginObject("application");
 						{
 							OutFile.AddField("package", PackageName);
-							OutFile.AddField("activity", "com.epicgames.ue4.GameActivity");
+							OutFile.AddField("activity", "com.epicgames.unreal.GameActivity");
 						}
 						OutFile.EndObject();
 
@@ -1183,8 +1183,8 @@ namespace UnrealBuildTool
 							OutFile.EndArray();
 
 							OutFile.BeginArray("lldbPreTargetCreateCommands");
-							FileReference UE4DataFormatters = FileReference.Combine(UE4ProjectRoot, "Engine", "Extras", "LLDBDataFormatters", "UE4DataFormatters_2ByteChars.py");
-							OutFile.AddUnnamedField("command script import \\\"" + UE4DataFormatters.FullName.Replace("\\", "/") + "\\\"");
+							FileReference DataFormatters = FileReference.Combine(ProjectRoot, "Engine", "Extras", "LLDBDataFormatters", "UEDataFormatters_2ByteChars.py");
+							OutFile.AddUnnamedField("command script import \\\"" + DataFormatters.FullName.Replace("\\", "/") + "\\\"");
 							OutFile.EndArray();
 
 							OutFile.BeginArray("lldbPostTargetCreateCommands");
@@ -1272,7 +1272,7 @@ namespace UnrealBuildTool
 										OutFile.AddField("cwd", MakePathString(CWD, true, true));
 									}
 					*/
-					OutFile.AddField("cwd", MakeUnquotedPathString(UE4ProjectRoot, EPathType.Absolute));
+					OutFile.AddField("cwd", MakeUnquotedPathString(ProjectRoot, EPathType.Absolute));
 
 					if (HostPlatform == UnrealTargetPlatform.Win64)
 					{
@@ -1280,12 +1280,12 @@ namespace UnrealBuildTool
 						OutFile.AddField("console", "externalTerminal");
 
 						OutFile.AddField("type", "cppvsdbg");
-						OutFile.AddField("visualizerFile", MakeUnquotedPathString(FileReference.Combine(UE4ProjectRoot, "Engine", "Extras", "VisualStudioDebugging", "Unreal.natvis"), EPathType.Absolute));
+						OutFile.AddField("visualizerFile", MakeUnquotedPathString(FileReference.Combine(ProjectRoot, "Engine", "Extras", "VisualStudioDebugging", "Unreal.natvis"), EPathType.Absolute));
 					}
 					else if (HostPlatform == UnrealTargetPlatform.Linux)
 					{
 						OutFile.AddField("type", "cppdbg");
-						OutFile.AddField("visualizerFile", MakeUnquotedPathString(FileReference.Combine(UE4ProjectRoot, "Engine", "Extras", "VisualStudioDebugging", "Unreal.natvis"), EPathType.Absolute));
+						OutFile.AddField("visualizerFile", MakeUnquotedPathString(FileReference.Combine(ProjectRoot, "Engine", "Extras", "VisualStudioDebugging", "Unreal.natvis"), EPathType.Absolute));
 						OutFile.AddField("showDisplayString", true);
 					}
 					else
@@ -1293,13 +1293,13 @@ namespace UnrealBuildTool
 						OutFile.AddField("type", "lldb");
 					}
 
-					if (UnrealBuildTool.OriginalCompilationRootDirectory != UE4ProjectRoot)
+					if (UnrealBuildTool.OriginalCompilationRootDirectory != ProjectRoot)
 					{
 						OutFile.BeginObject("sourceFileMap");
 						{
 							OutFile.AddField(
 								MakeUnquotedPathString(UnrealBuildTool.OriginalCompilationRootDirectory, EPathType.Absolute),
-								MakeUnquotedPathString(UE4ProjectRoot, EPathType.Absolute));
+								MakeUnquotedPathString(ProjectRoot, EPathType.Absolute));
 						}
 						OutFile.EndObject();
 					}
@@ -1339,7 +1339,7 @@ namespace UnrealBuildTool
 					OutFile.AddField("preLaunchTask", InBuildTaskName);
 				}
 				
-				DirectoryReference CWD = UE4ProjectRoot;
+				DirectoryReference CWD = ProjectRoot;
 
 				OutFile.AddField("program", MakeUnquotedPathString(InExecutable, EPathType.Absolute));
 				OutFile.BeginArray("args");
@@ -1445,7 +1445,7 @@ namespace UnrealBuildTool
 					PreLaunchTask = "UnrealBuildTool " + HostPlatform.ToString() + " Development Build";
 				}
 
-				FileReference UbtPath = FileReference.Combine(UE4ProjectRoot, "Engine", "Binaries", "DotNET", "UnrealBuildTool", "UnrealBuildTool");
+				FileReference UbtPath = FileReference.Combine(ProjectRoot, "Engine", "Binaries", "DotNET", "UnrealBuildTool", "UnrealBuildTool");
 				WriteSingleCSharpLaunchConfig(
 					OutFile,
 					"Generate Project Files",
