Only generate Makefile and VSCode now. As VSCode is the default source code for Linux.

If you want to generate for another IDE do -help to find a list to pick from

diff --git a/Engine/Source/Programs/UnrealBuildTool/System/BuildHostPlatform.cs b/Engine/Source/Programs/UnrealBuildTool/System/BuildHostPlatform.cs
--- a/Engine/Source/Programs/UnrealBuildTool/System/BuildHostPlatform.cs
+++ b/Engine/Source/Programs/UnrealBuildTool/System/BuildHostPlatform.cs
@@ -420,10 +420,6 @@ namespace UnrealBuildTool
 		{
 			yield return ProjectFileFormat.Make;
 			yield return ProjectFileFormat.VisualStudioCode;
-			yield return ProjectFileFormat.KDevelop;
-			yield return ProjectFileFormat.QMake;
-			yield return ProjectFileFormat.CMake;
-			yield return ProjectFileFormat.CodeLite;
 			#if __VPROJECT_AVAILABLE__
 				yield return ProjectFileFormat.VProject;
 			#endif
