StableTopologicalSort():

Perform dependency calculation in parallel, and overlap with array sorting. Reduces function runtime on my machine by 0.5 seconds for UnrealEditor target.

diff --git a/Engine/Source/Programs/UnrealBuildTool/Configuration/UEBuildModule.cs b/Engine/Source/Programs/UnrealBuildTool/Configuration/UEBuildModule.cs
--- a/Engine/Source/Programs/UnrealBuildTool/Configuration/UEBuildModule.cs
+++ b/Engine/Source/Programs/UnrealBuildTool/Configuration/UEBuildModule.cs
@@ -870,7 +870,7 @@ namespace UnrealBuildTool
 		/// <param name="bOnlyDirectDependencies">True to return only this module's direct dependencies</param>
 		public virtual void GetAllDependencyModules(List<UEBuildModule> ReferencedModules, HashSet<UEBuildModule> IgnoreReferencedModules, bool bIncludeDynamicallyLoaded, bool bForceCircular, bool bOnlyDirectDependencies)
 		{
-			List<UEBuildModule> AllDependencyModules = new List<UEBuildModule>();
+			List<UEBuildModule> AllDependencyModules = new List<UEBuildModule>(PrivateDependencyModules!.Count + PublicDependencyModules!.Count + (bIncludeDynamicallyLoaded ? DynamicallyLoadedModules!.Count : 0));
 			AllDependencyModules.AddRange(PrivateDependencyModules!);
 			AllDependencyModules.AddRange(PublicDependencyModules!);
 			if (bIncludeDynamicallyLoaded)
@@ -883,8 +883,7 @@ namespace UnrealBuildTool
 				if (!IgnoreReferencedModules.Contains(DependencyModule))
 				{
 					// Don't follow circular back-references!
-					bool bIsCircular = HasCircularDependencyOn(DependencyModule.Name);
-					if (bForceCircular || !bIsCircular)
+					if (bForceCircular || !HasCircularDependencyOn(DependencyModule.Name))
 					{
 						IgnoreReferencedModules.Add(DependencyModule);
 
diff --git a/Engine/Source/Programs/UnrealBuildTool/System/ExternalExecution.cs b/Engine/Source/Programs/UnrealBuildTool/System/ExternalExecution.cs
--- a/Engine/Source/Programs/UnrealBuildTool/System/ExternalExecution.cs
+++ b/Engine/Source/Programs/UnrealBuildTool/System/ExternalExecution.cs
@@ -14,6 +14,8 @@ using System.Runtime.Serialization;
 using EpicGames.Core;
 using System.Text.Json.Serialization;
 using UnrealBuildBase;
+using System.Collections.Concurrent;
+using System.Threading.Tasks;
 
 namespace UnrealBuildTool
 {
@@ -450,29 +452,75 @@ namespace UnrealBuildTool
 		/// remain in their same relative order within the original Nodes sequence.
 		/// </summary>
 		/// <param name="NodeList">The list of nodes to sort.</param>
-		static void StableTopologicalSort(List<UEBuildModuleCPP> NodeList)
+		static List<UEBuildModuleCPP> StableTopologicalSort(List<UEBuildModuleCPP> NodeList)
 		{
 			int NodeCount = NodeList.Count;
 
-			Dictionary<UEBuildModule, HashSet<UEBuildModule>> Cache = new Dictionary<UEBuildModule, HashSet<UEBuildModule>>();
+			// For each Node in NodeList, populated with the full circular dependency list from
+			// Node.GetAllDependencyModules()
+			List<Task<HashSet<UEBuildModule>>> NodeDependencies = new List<Task<HashSet<UEBuildModule>>>(NodeCount);
+			// Used to populate an element of NodeDependencies
+			HashSet<UEBuildModule> FetchDependencies(int NodeIndex)
+			{
+				HashSet<UEBuildModule> Dependencies = new HashSet<UEBuildModule>();
+				NodeList[NodeIndex].GetAllDependencyModules(new List<UEBuildModule>(), Dependencies, true, true, false);
+				return Dependencies;
+			}
+			
+			// For each Node in NodeList, populated with the nodes with a lower index in NodeList that Node depends on
+			List<Task<HashSet<UEBuildModuleCPP>>> PrecedingDependents = new List<Task<HashSet<UEBuildModuleCPP>>>(NodeCount);
+			HashSet<UEBuildModuleCPP> ComputePrecedingDependents(int NodeIndex)
+			{
+				HashSet<UEBuildModuleCPP> Results = new HashSet<UEBuildModuleCPP>();
 
+				UEBuildModuleCPP Node = NodeList[NodeIndex];
+				HashSet<UEBuildModule> Dependencies = NodeDependencies[NodeIndex].Result;
+
+				for (int I = 0; I < NodeIndex; ++I)
+				{
+					if (NodeDependencies[I].Result.Contains(Node) && !Dependencies.Contains(NodeList[I]))
+					{
+						Results.Add(NodeList[I]);
+					}
+				}
+				
+				return Results;
+			}
+			
+			for (int I = 0; I < NodeCount; ++I)
+			{
+				int LocalI = I;
+				NodeDependencies.Add(Task.Run(() => FetchDependencies(LocalI)));
+				PrecedingDependents.Add(Task.Run(() => ComputePrecedingDependents(LocalI)));
+			}
+			
+			List<UEBuildModuleCPP> Out = new List<UEBuildModuleCPP>(NodeCount);
+			// Write the ordered output			
 			for (int Index1 = 0; Index1 != NodeCount; ++Index1)
 			{
 				UEBuildModuleCPP Node1 = NodeList[Index1];
+				HashSet<UEBuildModuleCPP> NodesThatDependOnNode1 = PrecedingDependents[Index1].Result;
+				Out.Add(Node1);
+
+				if (NodesThatDependOnNode1.Count == 0)
+				{
+					continue;
+				}
 
 				for (int Index2 = 0; Index2 != Index1; ++Index2)
 				{
-					UEBuildModuleCPP Node2 = NodeList[Index2];
+					UEBuildModuleCPP Node2 = Out[Index2];
 
-					if (IsDependency(Node2, Node1, Cache) && !IsDependency(Node1, Node2, Cache))
+					if (NodesThatDependOnNode1.Contains(Node2))
 					{
 						// Rotate element at Index1 into position at Index2
-						for (int Index3 = Index1; Index3 != Index2; )
+						for (int Index3 = Index1; Index3 != Index2;)
 						{
 							--Index3;
-							NodeList[Index3 + 1] = NodeList[Index3];
+							Out[Index3 + 1] = Out[Index3];
 						}
-						NodeList[Index2] = Node1;
+
+						Out[Index2] = Node1;
 
 						// Break out of this loop, because this iteration must have covered all existing cases
 						// involving the node formerly at position Index1
@@ -480,25 +528,8 @@ namespace UnrealBuildTool
 					}
 				}
 			}
-		}
 
-		/// <summary>
-		/// Tests whether one module has a dependency on another
-		/// </summary>
-		/// <param name="FromModule">The module to test</param>
-		/// <param name="ToModule">The module to look for a dependency</param>
-		/// <param name="Cache">Cache mapping module to all its dependencies</param>
-		/// <returns>True if ToModule is a dependency of FromModule, false otherwise</returns>
-		static bool IsDependency(UEBuildModuleCPP FromModule, UEBuildModuleCPP ToModule, Dictionary<UEBuildModule, HashSet<UEBuildModule>> Cache)
-		{
-			HashSet<UEBuildModule>? Dependencies;
-			if(!Cache.TryGetValue(FromModule, out Dependencies))
-			{
-				Dependencies = new HashSet<UEBuildModule>();
-				FromModule.GetAllDependencyModules(new List<UEBuildModule>(), Dependencies, true, true, false);
-				Cache.Add(FromModule, Dependencies);
-			}
-			return Dependencies.Contains(ToModule);
+			return Out;
 		}
 
 		/// <summary>
@@ -595,7 +626,7 @@ namespace UnrealBuildTool
 
 			// Sort modules by type, then by dependency
 			List<UEBuildModuleCPP> ModulesSortedByType = ModulesToGenerateHeadersFor.OrderBy(c => ModuleToType[c]).ToList();
-			StableTopologicalSort(ModulesSortedByType);
+			ModulesSortedByType = StableTopologicalSort(ModulesSortedByType);
 
 			// Create the info for each module in parallel
 			UHTModuleInfo[] ModuleInfoArray = new UHTModuleInfo[ModulesSortedByType.Count];
