UnrealBuildTool:

Support for writing files directly as lines, without converting to a single string first.

diff --git a/Engine/Source/Programs/UnrealBuildTool/Platform/HoloLens/HoloLensToolChain.cs b/Engine/Source/Programs/UnrealBuildTool/Platform/HoloLens/HoloLensToolChain.cs
--- a/Engine/Source/Programs/UnrealBuildTool/Platform/HoloLens/HoloLensToolChain.cs
+++ b/Engine/Source/Programs/UnrealBuildTool/Platform/HoloLens/HoloLensToolChain.cs
@@ -1020,7 +1020,7 @@ namespace UnrealBuildTool
 			FileReference ResponseFileName = GetResponseFileName(LinkEnvironment, OutputFile);
 			if (!ProjectFileGenerator.bGenerateProjectFiles)
 			{
-				FileItem ResponseFile = Graph.CreateIntermediateTextFile(ResponseFileName, String.Join(Environment.NewLine, Arguments), StringComparison.InvariantCultureIgnoreCase);
+				FileItem ResponseFile = Graph.CreateIntermediateTextFile(ResponseFileName, Arguments, StringComparison.InvariantCultureIgnoreCase);
 				PrerequisiteItems.Add(ResponseFile);
 			}
 
diff --git a/Engine/Source/Programs/UnrealBuildTool/Platform/Windows/VCToolChain.cs b/Engine/Source/Programs/UnrealBuildTool/Platform/Windows/VCToolChain.cs
--- a/Engine/Source/Programs/UnrealBuildTool/Platform/Windows/VCToolChain.cs
+++ b/Engine/Source/Programs/UnrealBuildTool/Platform/Windows/VCToolChain.cs
@@ -1921,7 +1921,7 @@ namespace UnrealBuildTool
 			FileReference ResponseFileName = GetResponseFileName(LinkEnvironment, OutputFile);
 			if (!ProjectFileGenerator.bGenerateProjectFiles)
 			{
-				FileItem ResponseFile = Graph.CreateIntermediateTextFile(ResponseFileName, String.Join(Environment.NewLine, Arguments), StringComparison.InvariantCultureIgnoreCase);
+				FileItem ResponseFile = Graph.CreateIntermediateTextFile(ResponseFileName, Arguments, StringComparison.InvariantCultureIgnoreCase);
 				PrerequisiteItems.Add(ResponseFile);
 			}
 
diff --git a/Engine/Source/Programs/UnrealBuildTool/System/ActionGraphBuilder.cs b/Engine/Source/Programs/UnrealBuildTool/System/ActionGraphBuilder.cs
--- a/Engine/Source/Programs/UnrealBuildTool/System/ActionGraphBuilder.cs
+++ b/Engine/Source/Programs/UnrealBuildTool/System/ActionGraphBuilder.cs
@@ -33,6 +33,15 @@ namespace UnrealBuildTool
 		/// <returns>New file item</returns>
 		void CreateIntermediateTextFile(FileItem Location, string Contents, StringComparison ComparisonForContentChanges);
 
+		/// <summary>
+		/// Creates a response file for use in the action graph, with a newline between each string in ContentLines
+		/// </summary>
+		/// <param name="Location">Location of the response file</param>
+		/// <param name="ContentLines">Contents of the file</param>
+		/// <param name="ComparisonForContentChanges">The StringComparison mode to use to determine if a file should be re-written</param>
+		/// <returns>New file item</returns>
+		void CreateIntermediateTextFile(FileItem Location, IEnumerable<string> ContentLines, StringComparison ComparisonForContentChanges);
+
 		/// <summary>
 		/// Adds a file which is in the non-unity working set
 		/// </summary>
@@ -88,6 +97,12 @@ namespace UnrealBuildTool
 			Utils.WriteFileIfChanged(FileItem, Contents, ComparisonForContentChanges);
 		}
 
+		/// <inheritdoc/>
+		public virtual void CreateIntermediateTextFile(FileItem FileItem, IEnumerable<string> ContentLines, StringComparison ComparisonForContentChanges)
+		{
+			Utils.WriteFileIfChanged(FileItem, ContentLines, ComparisonForContentChanges);
+		}
+
 		/// <inheritdoc/>
 		public virtual void AddSourceDir(DirectoryItem SourceDir)
 		{
@@ -150,6 +165,12 @@ namespace UnrealBuildTool
 			Inner.CreateIntermediateTextFile(FileItem, Contents, ComparisonForContentChanges);
 		}
 
+		/// <inheritdoc/>
+		public virtual void CreateIntermediateTextFile(FileItem FileItem, IEnumerable<string> ContentLines, StringComparison ComparisonForContentChanges)
+		{
+			Inner.CreateIntermediateTextFile(FileItem, ContentLines, ComparisonForContentChanges);
+		}
+
 		/// <inheritdoc/>
 		public virtual void AddSourceDir(DirectoryItem SourceDir)
 		{
@@ -288,31 +309,20 @@ namespace UnrealBuildTool
 			return FileItem;
 		}
 
-		/// <summary>
-		/// Creates a text file with the given contents.  If the contents of the text file aren't changed, it won't write the new contents to
-		/// the file to avoid causing an action to be considered outdated.
-		/// </summary>
-		/// <param name="Graph">The action graph</param>
-		/// <param name="FileItem">Path to the intermediate file to create</param>
-		/// <param name="Contents">Contents of the new file</param>
-		/// <param name="ComparisonForContentChanges">The StringComparison mode to use to determine if a file should be re-written</param>
-		public static void CreateIntermediateTextFile(this IActionGraphBuilder Graph, FileItem FileItem, IEnumerable<string> Contents, StringComparison ComparisonForContentChanges)
-		{
-			Graph.CreateIntermediateTextFile(FileItem, string.Join(Environment.NewLine, Contents), ComparisonForContentChanges);
-		}
-
 		/// <summary>
 		/// Creates a text file with the given contents.  If the contents of the text file aren't changed, it won't write the new contents to
 		/// the file to avoid causing an action to be considered outdated.
 		/// </summary>
 		/// <param name="Graph">The action graph</param>
 		/// <param name="AbsolutePath">Path to the intermediate file to create</param>
-		/// <param name="Contents">Contents of the new file</param>
+		/// <param name="ContentLines">Contents of the new file</param>
 		/// <param name="ComparisonForContentChanges">The StringComparison mode to use to determine if a file should be re-written</param>
 		/// <returns>File item for the newly created file</returns>
-		public static FileItem CreateIntermediateTextFile(this IActionGraphBuilder Graph, FileReference AbsolutePath, IEnumerable<string> Contents, StringComparison ComparisonForContentChanges)
+		public static FileItem CreateIntermediateTextFile(this IActionGraphBuilder Graph, FileReference AbsolutePath, IEnumerable<string> ContentLines, StringComparison ComparisonForContentChanges)
 		{
-			return Graph.CreateIntermediateTextFile(AbsolutePath, string.Join(Environment.NewLine, Contents), ComparisonForContentChanges);
+			FileItem FileItem = UnrealBuildBase.FileItem.GetItemByFileReference(AbsolutePath);
+			Graph.CreateIntermediateTextFile(FileItem, ContentLines, ComparisonForContentChanges);
+			return FileItem;
 		}
 	}
 }
diff --git a/Engine/Source/Programs/UnrealBuildTool/System/TargetMakefile.cs b/Engine/Source/Programs/UnrealBuildTool/System/TargetMakefile.cs
--- a/Engine/Source/Programs/UnrealBuildTool/System/TargetMakefile.cs
+++ b/Engine/Source/Programs/UnrealBuildTool/System/TargetMakefile.cs
@@ -705,6 +705,17 @@ namespace UnrealBuildTool
 			FileItem.ResetCachedInfo();
 		}
 
+		/// <inheritdoc/>
+		public void CreateIntermediateTextFile(FileItem FileItem, IEnumerable<string> ContentLines, StringComparison ComparisonForContentChanges)
+		{
+			// Write the file
+			Utils.WriteFileIfChanged(FileItem, ContentLines, ComparisonForContentChanges);
+
+			// Reset the file info, in case it already knows about the old file
+			InternalDependencies.Add(FileItem);
+			FileItem.ResetCachedInfo();
+		}
+
 		/// <inheritdoc/>
 		public void AddSourceDir(DirectoryItem SourceDir)
 		{
diff --git a/Engine/Source/Programs/UnrealBuildTool/System/Utils.cs b/Engine/Source/Programs/UnrealBuildTool/System/Utils.cs
--- a/Engine/Source/Programs/UnrealBuildTool/System/Utils.cs
+++ b/Engine/Source/Programs/UnrealBuildTool/System/Utils.cs
@@ -1422,6 +1422,18 @@ namespace UnrealBuildTool
 			WriteFileIfChanged(FileItem, Contents, Comparison);
 		}
 
+		/// <summary>
+		/// Writes a file if the contents have changed
+		/// </summary>
+		/// <param name="Location">Location of the file</param>
+		/// <param name="ContentLines">New contents of the file</param>
+		/// <param name="Comparison">The type of string comparison to use</param>
+		internal static void WriteFileIfChanged(FileReference Location, IEnumerable<string> ContentLines, StringComparison Comparison)
+		{
+			FileItem FileItem = FileItem.GetItemByFileReference(Location);
+			WriteFileIfChanged(FileItem, ContentLines, Comparison);
+		}
+
 		/// <summary>
 		/// Writes a file if the contents have changed
 		/// </summary>
@@ -1461,8 +1473,50 @@ namespace UnrealBuildTool
 			}
 		}
 
+		/// <summary>
+		/// Writes a file if the contents have changed
+		/// </summary>
+		/// <param name="FileItem">Location of the file</param>
+		/// <param name="ContentLines">New contents of the file</param>
+		/// <param name="Comparison">The type of string comparison to use</param>
+		internal static void WriteFileIfChanged(FileItem FileItem, IEnumerable<string> ContentLines, StringComparison Comparison)
+		{
+			// Only write the file if its contents have changed.
+			FileReference Location = FileItem.Location;
+			
+			if (!FileItem.Exists)
+			{
+				DirectoryReference.CreateDirectory(Location.Directory);
+				FileReference.WriteAllLines(Location, ContentLines, GetEncodingForStrings(ContentLines));
+				FileItem.ResetCachedInfo();
+			}
+			else
+			{
+				string[] CurrentContents = File.ReadAllLines(FileItem.FullName);
+				if (!CurrentContents.SequenceEqual(ContentLines, StringComparer.FromComparison(Comparison)))
+				{
+					FileReference BackupFile = new FileReference($"{FileItem.FullName}.old");
+					try
+					{
+						Log.TraceLog($"Updating {FileItem}: contents have changed. Saving previous version to {BackupFile}.");
+						FileReference.Delete(BackupFile);
+						FileReference.Move(Location, BackupFile);
+					}
+					catch (Exception Ex)
+					{
+						Log.TraceWarning($"Unable to rename {FileItem} to {BackupFile}");
+						Log.TraceLog(ExceptionUtils.FormatExceptionDetails(Ex));
+					}
+					FileReference.WriteAllLines(Location, ContentLines, GetEncodingForStrings(ContentLines));
+					FileItem.ResetCachedInfo();
+				}
+			}
+		}
+
 		/// <summary>
 		/// Determines the appropriate encoding for a string: either ASCII or UTF-8.
+		/// If the string length is equivalent to the encoded length, then no non-ASCII characters were present in the string.
+		/// Don't write BOM as it messes with clang when loading response files.
 		/// </summary>
 		/// <param name="Str">The string to test.</param>
 		/// <returns>Either System.Text.Encoding.ASCII or System.Text.Encoding.UTF8, depending on whether or not the string contains non-ASCII characters.</returns>
@@ -1470,7 +1524,19 @@ namespace UnrealBuildTool
 		{
 			// If the string length is equivalent to the encoded length, then no non-ASCII characters were present in the string.
 			// Don't write BOM as it messes with clang when loading response files.
-			return (Encoding.UTF8.GetByteCount(Str) == Str.Length) ? Encoding.ASCII : new UTF8Encoding(false);
+			return (Encoding.UTF8.GetByteCount(Str) != Str.Length) ?  new UTF8Encoding(false) : Encoding.ASCII;
+		}
+		
+		/// <summary>
+		/// Determines the appropriate encoding for a list of strings: either ASCII or UTF-8.
+		/// If the string length is equivalent to the encoded length, then no non-ASCII characters were present in the string.
+		/// Don't write BOM as it messes with clang when loading response files.
+		/// </summary>
+		/// <param name="Strings">The string to test.</param>
+		/// <returns>Either System.Text.Encoding.ASCII or System.Text.Encoding.UTF8, depending on whether or not the strings contains non-ASCII characters.</returns>
+		private static Encoding GetEncodingForStrings(IEnumerable<string> Strings)
+		{
+			return Strings.Any(S => Encoding.UTF8.GetByteCount(S) != S.Length) ? new UTF8Encoding(false) : Encoding.ASCII;
 		}
 	}
 }
