UnrealBuildTool ExternalExecution.UpdateDirectoryTimestamps:

Write timestamp files in parallel.

diff --git a/Engine/Source/Programs/UnrealBuildTool/System/ExternalExecution.cs b/Engine/Source/Programs/UnrealBuildTool/System/ExternalExecution.cs
--- a/Engine/Source/Programs/UnrealBuildTool/System/ExternalExecution.cs
+++ b/Engine/Source/Programs/UnrealBuildTool/System/ExternalExecution.cs
@@ -466,7 +466,7 @@ namespace UnrealBuildTool
 				NodeList[NodeIndex].GetAllDependencyModules(new List<UEBuildModule>(), Dependencies, true, true, false);
 				return Dependencies;
 			}
-			
+
 			// For each Node in NodeList, populated with the nodes with a lower index in NodeList that Node depends on
 			List<Task<HashSet<UEBuildModuleCPP>>> PrecedingDependents = new List<Task<HashSet<UEBuildModuleCPP>>>(NodeCount);
 			HashSet<UEBuildModuleCPP> ComputePrecedingDependents(int NodeIndex)
@@ -1032,9 +1032,9 @@ namespace UnrealBuildTool
 		/// </summary>
 		private static void UpdateDirectoryTimestamps(List<UHTModuleInfo> UObjectModules)
 		{
-			foreach (UHTModuleInfo Module in UObjectModules)
+			Parallel.ForEach(UObjectModules, async Module =>
 			{
-				if(!Module.bIsReadOnly)
+				if (!Module.bIsReadOnly)
 				{
 					string GeneratedCodeDirectory = Path.GetDirectoryName(Module.GeneratedCPPFilenameBase)!;
 					DirectoryInfo GeneratedCodeDirectoryInfo = new DirectoryInfo(GeneratedCodeDirectory);
@@ -1058,7 +1058,7 @@ namespace UnrealBuildTool
 								AllUObjectFiles.AddRange(Module.PublicUObjectHeaders.ConvertAll(Item => Item.AbsolutePath));
 								AllUObjectFiles.AddRange(Module.InternalUObjectHeaders.ConvertAll(Item => Item.AbsolutePath));
 								AllUObjectFiles.AddRange(Module.PrivateUObjectHeaders.ConvertAll(Item => Item.AbsolutePath));
-								FileReference.WriteAllLines(TimestampFile, AllUObjectFiles);
+								await FileReference.WriteAllLinesAsync(TimestampFile, AllUObjectFiles);
 							}
 
 							// Because new .cpp and .h files may have been generated by UHT, invalidate the DirectoryLookupCache
@@ -1070,7 +1070,7 @@ namespace UnrealBuildTool
 						throw new BuildException(Exception, "Couldn't touch header directories: " + Exception.Message);
 					}
 				}
-			}
+			});
 		}
 
 		/// <summary>
