UBT: Check if target is null before filtering out invalid configurations when generating project files

#rnx

diff --git a/Engine/Source/Programs/UnrealBuildTool/ProjectFiles/VisualStudio/VCProject.cs b/Engine/Source/Programs/UnrealBuildTool/ProjectFiles/VisualStudio/VCProject.cs
--- a/Engine/Source/Programs/UnrealBuildTool/ProjectFiles/VisualStudio/VCProject.cs
+++ b/Engine/Source/Programs/UnrealBuildTool/ProjectFiles/VisualStudio/VCProject.cs
@@ -926,7 +926,7 @@ namespace UnrealBuildTool
 				{
 					continue;
 				}
-				if (!IsValidProjectPlatformAndConfiguration(Combination.ProjectTarget, (UnrealTargetPlatform)Combination.Platform, Combination.Configuration, PlatformProjectGenerators))
+				if (Combination.ProjectTarget != null && !IsValidProjectPlatformAndConfiguration(Combination.ProjectTarget, (UnrealTargetPlatform)Combination.Platform, Combination.Configuration, PlatformProjectGenerators))
 				{
 					continue;
 				}
