Various spelling & phrasing in comments & string

##trivial

diff --git a/Engine/Source/Programs/AutomationTool/Program.cs b/Engine/Source/Programs/AutomationTool/Program.cs
--- a/Engine/Source/Programs/AutomationTool/Program.cs
+++ b/Engine/Source/Programs/AutomationTool/Program.cs
@@ -247,7 +247,7 @@ namespace AutomationToolDriver
 					{"-Compile", "Force all script modules to be compiled"},
 					{"-IgnoreBuildRecords", "Ignore build record (.uatbuildrecord) files when determining if script modules are up to date" },
 					{"-UseLocalBuildStorage", @"Allows you to use local storage for your root build storage dir {default of P:\Builds {on PC} is changed to Engine\Saved\LocalBuilds}. Used for local testing."},
-					{"-WaitForDebugger", "Waits for a debugger to be attached, and breaks once debugger sucessfully attached."},
+					{"-WaitForDebugger", "Waits for a debugger to be attached, and breaks once debugger successfully attached."},
 					{"-BuildMachine", "" },
 					{"-WaitForUATMutex", "" }
 				}
diff --git a/Engine/Source/Programs/UnrealBuildTool/System/DynamicCompilation.cs b/Engine/Source/Programs/UnrealBuildTool/System/DynamicCompilation.cs
--- a/Engine/Source/Programs/UnrealBuildTool/System/DynamicCompilation.cs
+++ b/Engine/Source/Programs/UnrealBuildTool/System/DynamicCompilation.cs
@@ -50,7 +50,7 @@ namespace UnrealBuildTool
 				return true;
 			}
 
-			// Check the time stamp of the UnrealBuildTool.exe file.  If Unreal Build Tool was compiled more
+			// Check the time stamp of the UnrealBuildTool assembly.  If Unreal Build Tool was compiled more
 			// recently than the dynamically-compiled assembly, then we'll always recompile it.  This is
 			// because Unreal Build Tool's code may have changed in such a way that invalidate these
 			// previously-compiled assembly files.
diff --git a/Engine/Source/Programs/UnrealBuildTool/System/TargetMakefile.cs b/Engine/Source/Programs/UnrealBuildTool/System/TargetMakefile.cs
--- a/Engine/Source/Programs/UnrealBuildTool/System/TargetMakefile.cs
+++ b/Engine/Source/Programs/UnrealBuildTool/System/TargetMakefile.cs
@@ -382,13 +382,13 @@ namespace UnrealBuildTool
 					}
 				}
 
-				// Check to see if UnrealBuildTool.exe was compiled more recently than the makefile
+				// Check to see if UnrealBuildTool was compiled more recently than the makefile
 				DateTime UnrealBuildToolTimestamp = new FileInfo(Assembly.GetExecutingAssembly().Location).LastWriteTime;
 				if (MakefileInfo.LastWriteTime.CompareTo(UnrealBuildToolTimestamp) < 0)
 				{
-					// UnrealBuildTool.exe was compiled more recently than the makefile
-					Log.TraceLog("Makefile is older than UnrealBuildTool.exe, ignoring it");
-					ReasonNotLoaded = "UnrealBuildTool.exe is newer";
+					// UnrealBuildTool was compiled more recently than the makefile
+					Log.TraceLog("Makefile is older than UnrealBuildTool assembly, ignoring it");
+					ReasonNotLoaded = "UnrealBuildTool assembly is newer";
 					return null;
 				}
 
