UBT: Fix typo.

diff --git a/Engine/Source/Programs/UnrealBuildTool/System/RulesAssembly.cs b/Engine/Source/Programs/UnrealBuildTool/System/RulesAssembly.cs
--- a/Engine/Source/Programs/UnrealBuildTool/System/RulesAssembly.cs
+++ b/Engine/Source/Programs/UnrealBuildTool/System/RulesAssembly.cs
@@ -597,7 +597,7 @@ namespace UnrealBuildTool
 			// Set the default value for whether to use the shared build environment
 			if(Rules.BuildEnvironment == TargetBuildEnvironment.Unique && Unreal.IsEngineInstalled())
 			{
-				throw new BuildException("Targets with a unique build environment cannot be built an installed engine.");
+				throw new BuildException("Targets with a unique build environment cannot be built with an installed engine.");
 			}
 
 			// Automatically include CoreUObject
