Update Xcode project XCBuildConfiguration to handle unique build environments.   Also implement platform process GetModulesDirectory() on macOS which allows re-launching of the correct app to also work.

diff --git a/Engine/Source/Programs/UnrealBuildTool/ProjectFiles/Xcode/XcodeProject.cs b/Engine/Source/Programs/UnrealBuildTool/ProjectFiles/Xcode/XcodeProject.cs
--- a/Engine/Source/Programs/UnrealBuildTool/ProjectFiles/Xcode/XcodeProject.cs
+++ b/Engine/Source/Programs/UnrealBuildTool/ProjectFiles/Xcode/XcodeProject.cs
@@ -1589,7 +1589,8 @@ namespace UnrealBuildTool
 
 											// Get the output directory
 											DirectoryReference RootDirectory = Unreal.EngineDirectory;
-											if ((ProjectTarget.TargetRules.Type == TargetType.Game || ProjectTarget.TargetRules.Type == TargetType.Client || ProjectTarget.TargetRules.Type == TargetType.Server) && bShouldCompileMonolithic)
+											// Unique and Monolithic both need to use the target directory not the engine directory
+											if (ProjectTarget.TargetRules.Type != TargetType.Program && (bShouldCompileMonolithic || ProjectTarget.TargetRules.BuildEnvironment == TargetBuildEnvironment.Unique))
 											{
 												if (ProjectTarget.UnrealProjectFilePath != null)
 												{
@@ -1611,7 +1612,11 @@ namespace UnrealBuildTool
 												// Figure out what the compiled binary will be called so that we can point the IDE to the correct file
 												if (ProjectTarget.TargetRules.Type != TargetType.Game)
 												{
-													ExeName = "Unreal" + ProjectTarget.TargetRules.Type.ToString();
+													// Only if shared - unique retains the Target Name
+													if(ProjectTarget.TargetRules.BuildEnvironment == TargetBuildEnvironment.Shared)
+													{
+														ExeName = "Unreal" + ProjectTarget.TargetRules.Type.ToString();
+													}
 												}
 											}
 
diff --git a/Engine/Source/Runtime/Core/Private/Mac/MacPlatformProcess.cpp b/Engine/Source/Runtime/Core/Private/Mac/MacPlatformProcess.cpp
--- a/Engine/Source/Runtime/Core/Private/Mac/MacPlatformProcess.cpp
+++ b/Engine/Source/Runtime/Core/Private/Mac/MacPlatformProcess.cpp
@@ -1246,6 +1246,11 @@ const TCHAR* FMacPlatformProcess::GetBinariesSubdirectory()
 	return TEXT("Mac");
 }
 
+const FString FMacPlatformProcess::GetModulesDirectory()
+{
+	return FString(BaseDir());
+}
+
 bool FMacPlatformProcess::LaunchFileInDefaultExternalApplication( const TCHAR* FileName, const TCHAR* Parms /*= NULL*/, ELaunchVerb::Type Verb /*= ELaunchVerb::Open*/, bool bPromptToOpenOnFailure /*= true */)
 {
 	SCOPED_AUTORELEASE_POOL;
diff --git a/Engine/Source/Runtime/Core/Public/Mac/MacPlatformProcess.h b/Engine/Source/Runtime/Core/Public/Mac/MacPlatformProcess.h
--- a/Engine/Source/Runtime/Core/Public/Mac/MacPlatformProcess.h
+++ b/Engine/Source/Runtime/Core/Public/Mac/MacPlatformProcess.h
@@ -238,6 +238,7 @@ struct CORE_API FMacPlatformProcess : public FGenericPlatformProcess
 	static FString GenerateApplicationPath( const FString& AppName, EBuildConfiguration BuildConfiguration);
 	static const TCHAR* GetModuleExtension();
 	static const TCHAR* GetBinariesSubdirectory();
+	static const FString GetModulesDirectory();
 	static bool CanLaunchURL(const TCHAR* URL);
 	static void LaunchURL(const TCHAR* URL, const TCHAR* Parms, FString* Error);
 	static FString GetGameBundleId();
