Add compile prerequisites to PVS Gather action so that is not filtered out during single-file compile.
#rnx

diff --git a/Engine/Source/Programs/UnrealBuildTool/Platform/Windows/PVSToolChain.cs b/Engine/Source/Programs/UnrealBuildTool/Platform/Windows/PVSToolChain.cs
--- a/Engine/Source/Programs/UnrealBuildTool/Platform/Windows/PVSToolChain.cs
+++ b/Engine/Source/Programs/UnrealBuildTool/Platform/Windows/PVSToolChain.cs
@@ -621,8 +621,8 @@ namespace UnrealBuildTool
 
 			List<FileReference> InputFiles = Makefile.OutputItems.Select(x => x.Location).Where(x => x.HasExtension(".pvslog")).ToList();
 
-			// Collect the prerequisite items off of the Compile action added in CompileCPPFiles so that in SingleFileCompile mode the PVSGather step is also not filtered out
-			List<FileItem> AnalyzeActionPrerequisiteItems = Makefile.Actions.SelectMany(x => x.ProducedItems).ToList();
+			// Collect the sourcefile items off of the Compile action added in CompileCPPFiles so that in SingleFileCompile mode the PVSGather step is also not filtered out
+			List<FileItem> CompileSourceFiles = Makefile.Actions.Where(x => x is VCCompileAction).Select(x => (x as VCCompileAction).SourceFile).ToList();
 
 			FileItem InputFileListItem = Makefile.CreateIntermediateTextFile(OutputFile.ChangeExtension(".input"), InputFiles.Select(x => x.FullName), StringComparison.InvariantCultureIgnoreCase);
 
@@ -632,7 +632,7 @@ namespace UnrealBuildTool
 			AnalyzeAction.WorkingDirectory = UnrealBuildTool.EngineSourceDirectory;
 			AnalyzeAction.PrerequisiteItems.Add(InputFileListItem);
 			AnalyzeAction.PrerequisiteItems.AddRange(Makefile.OutputItems);
-			AnalyzeAction.PrerequisiteItems.AddRange(AnalyzeActionPrerequisiteItems);
+			AnalyzeAction.PrerequisiteItems.AddRange(CompileSourceFiles);
 			AnalyzeAction.ProducedItems.Add(FileItem.GetItemByFileReference(OutputFile));
 			AnalyzeAction.DeleteItems.AddRange(AnalyzeAction.ProducedItems);
 
