UnrealBuildTool:

Remove always true bIsAssembling mode variable.

#trivial

diff --git a/Engine/Source/Programs/UnrealBuildTool/Configuration/UEBuildTarget.cs b/Engine/Source/Programs/UnrealBuildTool/Configuration/UEBuildTarget.cs
--- a/Engine/Source/Programs/UnrealBuildTool/Configuration/UEBuildTarget.cs
+++ b/Engine/Source/Programs/UnrealBuildTool/Configuration/UEBuildTarget.cs
@@ -1594,7 +1594,7 @@ namespace UnrealBuildTool
 		/// <summary>
 		/// Builds the target, appending list of output files and returns building result.
 		/// </summary>
-		public TargetMakefile Build(BuildConfiguration BuildConfiguration, ISourceFileWorkingSet WorkingSet, bool bIsAssemblingBuild, List<FileReference> SpecificFilesToCompile)
+		public TargetMakefile Build(BuildConfiguration BuildConfiguration, ISourceFileWorkingSet WorkingSet, List<FileReference> SpecificFilesToCompile)
 		{
 			CppConfiguration CppConfiguration = GetCppConfiguration(Configuration);
 
@@ -1714,7 +1714,7 @@ namespace UnrealBuildTool
 			HashSet<UEBuildModuleCPP> ModulesToGenerateHeadersFor = GatherDependencyModules(OriginalBinaries.ToList());
 			using(Timeline.ScopeEvent("ExternalExecution.SetupUObjectModules()"))
 			{
-				ExternalExecution.SetupUObjectModules(ModulesToGenerateHeadersFor, Rules.Platform, ProjectDescriptor, Makefile.UObjectModules, Makefile.UObjectModuleHeaders, Rules.GeneratedCodeVersion, bIsAssemblingBuild, MetadataCache);
+				ExternalExecution.SetupUObjectModules(ModulesToGenerateHeadersFor, Rules.Platform, ProjectDescriptor, Makefile.UObjectModules, Makefile.UObjectModuleHeaders, Rules.GeneratedCodeVersion, MetadataCache);
 			}
 
 			List<string> UHTAdditionalArguments = new List<string>();
@@ -1734,7 +1734,7 @@ namespace UnrealBuildTool
 			if (Makefile.UObjectModules.Count > 0)
 			{
 				FileReference ModuleInfoFileName = FileReference.Combine(ProjectIntermediateDirectory, TargetName + ".uhtmanifest");
-				ExternalExecution.ExecuteHeaderToolIfNecessary(BuildConfiguration, ProjectFile, TargetName, TargetType, bHasProjectScriptPlugin, Makefile.UObjectModules, ModuleInfoFileName, true, bIsAssemblingBuild, WorkingSet, Makefile.UHTAdditionalArguments);
+				ExternalExecution.ExecuteHeaderToolIfNecessary(BuildConfiguration, ProjectFile, TargetName, TargetType, bHasProjectScriptPlugin, Makefile.UObjectModules, ModuleInfoFileName, true, WorkingSet, Makefile.UHTAdditionalArguments);
 			}
 
 			// Find all the shared PCHs.
diff --git a/Engine/Source/Programs/UnrealBuildTool/Modes/BuildMode.cs b/Engine/Source/Programs/UnrealBuildTool/Modes/BuildMode.cs
--- a/Engine/Source/Programs/UnrealBuildTool/Modes/BuildMode.cs
+++ b/Engine/Source/Programs/UnrealBuildTool/Modes/BuildMode.cs
@@ -715,8 +715,7 @@ namespace UnrealBuildTool
 				// Build the target
 				using(Timeline.ScopeEvent("UEBuildTarget.Build()"))
 				{
-					const bool bIsAssemblingBuild = true;
-					Makefile = Target.Build(BuildConfiguration, WorkingSet, bIsAssemblingBuild, TargetDescriptor.SpecificFilesToCompile);
+					Makefile = Target.Build(BuildConfiguration, WorkingSet, TargetDescriptor.SpecificFilesToCompile);
 				}
 
 				// Save the pre-build scripts onto the makefile
@@ -755,10 +754,9 @@ namespace UnrealBuildTool
 				if (Makefile.UObjectModules.Count > 0)
 				{
 					const bool bIsGatheringBuild = false;
-					const bool bIsAssemblingBuild = true;
 
 					FileReference ModuleInfoFileName = FileReference.Combine(Makefile.ProjectIntermediateDirectory, TargetDescriptor.Name + ".uhtmanifest");
-					ExternalExecution.ExecuteHeaderToolIfNecessary(BuildConfiguration, TargetDescriptor.ProjectFile, TargetDescriptor.Name, Makefile.TargetType, Makefile.bHasProjectScriptPlugin, UObjectModules: Makefile.UObjectModules, ModuleInfoFileName: ModuleInfoFileName, bIsGatheringBuild: bIsGatheringBuild, bIsAssemblingBuild: bIsAssemblingBuild, WorkingSet: WorkingSet, UHTAdditionalArguments: Makefile.UHTAdditionalArguments);
+					ExternalExecution.ExecuteHeaderToolIfNecessary(BuildConfiguration, TargetDescriptor.ProjectFile, TargetDescriptor.Name, Makefile.TargetType, Makefile.bHasProjectScriptPlugin, UObjectModules: Makefile.UObjectModules, ModuleInfoFileName: ModuleInfoFileName, bIsGatheringBuild: bIsGatheringBuild, WorkingSet: WorkingSet, UHTAdditionalArguments: Makefile.UHTAdditionalArguments);
 				}
 			}
 			return Makefile;
diff --git a/Engine/Source/Programs/UnrealBuildTool/Modes/JsonExportMode.cs b/Engine/Source/Programs/UnrealBuildTool/Modes/JsonExportMode.cs
--- a/Engine/Source/Programs/UnrealBuildTool/Modes/JsonExportMode.cs
+++ b/Engine/Source/Programs/UnrealBuildTool/Modes/JsonExportMode.cs
@@ -55,8 +55,7 @@ namespace UnrealBuildTool
 						Arguments.ApplyTo(BuildConfiguration);
 
 						// Create the makefile
-						const bool bIsAssemblingBuild = true;
-						TargetMakefile Makefile = Target.Build(BuildConfiguration, WorkingSet, bIsAssemblingBuild, TargetDescriptor.SpecificFilesToCompile);
+						TargetMakefile Makefile = Target.Build(BuildConfiguration, WorkingSet, TargetDescriptor.SpecificFilesToCompile);
 						List<LinkedAction> Actions = Makefile.Actions.ConvertAll(x => new LinkedAction(x));
 						ActionGraph.Link(Actions);
 
diff --git a/Engine/Source/Programs/UnrealBuildTool/System/ExternalExecution.cs b/Engine/Source/Programs/UnrealBuildTool/System/ExternalExecution.cs
--- a/Engine/Source/Programs/UnrealBuildTool/System/ExternalExecution.cs
+++ b/Engine/Source/Programs/UnrealBuildTool/System/ExternalExecution.cs
@@ -615,7 +615,7 @@ namespace UnrealBuildTool
 			}
 		}
 
-		public static void SetupUObjectModules(IEnumerable<UEBuildModuleCPP> ModulesToGenerateHeadersFor, UnrealTargetPlatform Platform, ProjectDescriptor? ProjectDescriptor, List<UHTModuleInfo> UObjectModules, List<UHTModuleHeaderInfo> UObjectModuleHeaders, EGeneratedCodeVersion GeneratedCodeVersion, bool bIsAssemblingBuild, SourceFileMetadataCache MetadataCache)
+		public static void SetupUObjectModules(IEnumerable<UEBuildModuleCPP> ModulesToGenerateHeadersFor, UnrealTargetPlatform Platform, ProjectDescriptor? ProjectDescriptor, List<UHTModuleInfo> UObjectModules, List<UHTModuleHeaderInfo> UObjectModuleHeaders, EGeneratedCodeVersion GeneratedCodeVersion, SourceFileMetadataCache MetadataCache)
 		{
 			// Find the type of each module
 			Dictionary<UEBuildModuleCPP, UHTModuleType> ModuleToType = new Dictionary<UEBuildModuleCPP, UHTModuleType>();
@@ -853,9 +853,8 @@ namespace UnrealBuildTool
 		/// <param name="UObjectModules">Modules that we generate headers for</param>
 		/// <param name="HeaderToolTimestampUtc">Timestamp for UHT</param>
 		/// <param name="bIsGatheringBuild"></param>
-		/// <param name="bIsAssemblingBuild"></param>
 		/// <returns>True if the code files are out of date</returns>
-		private static bool AreGeneratedCodeFilesOutOfDate(BuildConfiguration BuildConfiguration, List<UHTModuleInfo> UObjectModules, DateTime HeaderToolTimestampUtc, bool bIsGatheringBuild, bool bIsAssemblingBuild)
+		private static bool AreGeneratedCodeFilesOutOfDate(BuildConfiguration BuildConfiguration, List<UHTModuleInfo> UObjectModules, DateTime HeaderToolTimestampUtc, bool bIsGatheringBuild)
 		{
 			// Get CoreUObject.init.gen.cpp timestamp.  If the source files are older than the CoreUObject generated code, we'll
 			// need to regenerate code for the module
@@ -977,7 +976,7 @@ namespace UnrealBuildTool
 					// @todo ubtmake: Possibly, we should never be doing this check these days.
 					//
 					// We don't need to do this check if using hot reload makefiles, since makefile out-of-date checks already handle it.
-					if (!BuildConfiguration.bUseUBTMakefiles && (bIsGatheringBuild || !bIsAssemblingBuild))
+					if (!BuildConfiguration.bUseUBTMakefiles && bIsGatheringBuild)
 					{
 						// Also check the timestamp on the directory the source file is in.  If the directory timestamp has
 						// changed, new source files may have been added or deleted.  We don't know whether the new/deleted
@@ -1111,7 +1110,7 @@ namespace UnrealBuildTool
 		/// Builds and runs the header tool and touches the header directories.
 		/// Performs any early outs if headers need no changes, given the UObject modules, tool path, game name, and configuration
 		/// </summary>
-		public static void ExecuteHeaderToolIfNecessary(BuildConfiguration BuildConfiguration, FileReference? ProjectFile, string TargetName, TargetType TargetType, bool bHasProjectScriptPlugin, List<UHTModuleInfo> UObjectModules, FileReference ModuleInfoFileName, bool bIsGatheringBuild, bool bIsAssemblingBuild, ISourceFileWorkingSet WorkingSet, string[]? UHTAdditionalArguments)
+		public static void ExecuteHeaderToolIfNecessary(BuildConfiguration BuildConfiguration, FileReference? ProjectFile, string TargetName, TargetType TargetType, bool bHasProjectScriptPlugin, List<UHTModuleInfo> UObjectModules, FileReference ModuleInfoFileName, bool bIsGatheringBuild, ISourceFileWorkingSet WorkingSet, string[]? UHTAdditionalArguments)
 		{
 			if (ProgressWriter.bWriteMarkup)
 			{
@@ -1146,7 +1145,7 @@ namespace UnrealBuildTool
 				{
 					bUHTNeedsToRun = true;
 				}
-				else if(AreGeneratedCodeFilesOutOfDate(BuildConfiguration, UObjectModules, HeaderToolTimestampUtc, bIsGatheringBuild, bIsAssemblingBuild))
+				else if(AreGeneratedCodeFilesOutOfDate(BuildConfiguration, UObjectModules, HeaderToolTimestampUtc, bIsGatheringBuild))
 				{
 					bUHTNeedsToRun = true;
 				}
@@ -1187,7 +1186,7 @@ namespace UnrealBuildTool
 					// Always build UnrealHeaderTool if header regeneration is required, unless we're running within an installed ecosystem or hot-reloading
 					if ((!Unreal.IsEngineInstalled() || bHasProjectScriptPlugin) &&
 						!BuildConfiguration.bDoNotBuildUHT &&
-						!(bHaveHeaderTool && !bIsGatheringBuild && bIsAssemblingBuild))	// If running in "assembler only" mode, we assume UHT is already up to date for much faster iteration!
+						!(bHaveHeaderTool && !bIsGatheringBuild))
 					{
 						// If it is out of date or not there it will be built.
 						// If it is there and up to date, it will add 0.8 seconds to the build time.
