UnrealBuildTool:

Remove bIsGatheringBuild; directory timestamps do not need to be checked, and a null build of UHT is much cheaper than it used to be.

diff --git a/Engine/Source/Programs/UnrealBuildTool/Configuration/UEBuildTarget.cs b/Engine/Source/Programs/UnrealBuildTool/Configuration/UEBuildTarget.cs
--- a/Engine/Source/Programs/UnrealBuildTool/Configuration/UEBuildTarget.cs
+++ b/Engine/Source/Programs/UnrealBuildTool/Configuration/UEBuildTarget.cs
@@ -1734,7 +1734,7 @@ namespace UnrealBuildTool
 			if (Makefile.UObjectModules.Count > 0)
 			{
 				FileReference ModuleInfoFileName = FileReference.Combine(ProjectIntermediateDirectory, TargetName + ".uhtmanifest");
-				ExternalExecution.ExecuteHeaderToolIfNecessary(BuildConfiguration, ProjectFile, TargetName, TargetType, bHasProjectScriptPlugin, Makefile.UObjectModules, ModuleInfoFileName, true, WorkingSet, Makefile.UHTAdditionalArguments);
+				ExternalExecution.ExecuteHeaderToolIfNecessary(BuildConfiguration, ProjectFile, TargetName, TargetType, bHasProjectScriptPlugin, Makefile.UObjectModules, ModuleInfoFileName, WorkingSet, Makefile.UHTAdditionalArguments);
 			}
 
 			// Find all the shared PCHs.
diff --git a/Engine/Source/Programs/UnrealBuildTool/Modes/BuildMode.cs b/Engine/Source/Programs/UnrealBuildTool/Modes/BuildMode.cs
--- a/Engine/Source/Programs/UnrealBuildTool/Modes/BuildMode.cs
+++ b/Engine/Source/Programs/UnrealBuildTool/Modes/BuildMode.cs
@@ -753,10 +753,8 @@ namespace UnrealBuildTool
 				// If the target needs UHT to be run, we'll go ahead and do that now
 				if (Makefile.UObjectModules.Count > 0)
 				{
-					const bool bIsGatheringBuild = false;
-
 					FileReference ModuleInfoFileName = FileReference.Combine(Makefile.ProjectIntermediateDirectory, TargetDescriptor.Name + ".uhtmanifest");
-					ExternalExecution.ExecuteHeaderToolIfNecessary(BuildConfiguration, TargetDescriptor.ProjectFile, TargetDescriptor.Name, Makefile.TargetType, Makefile.bHasProjectScriptPlugin, UObjectModules: Makefile.UObjectModules, ModuleInfoFileName: ModuleInfoFileName, bIsGatheringBuild: bIsGatheringBuild, WorkingSet: WorkingSet, UHTAdditionalArguments: Makefile.UHTAdditionalArguments);
+					ExternalExecution.ExecuteHeaderToolIfNecessary(BuildConfiguration, TargetDescriptor.ProjectFile, TargetDescriptor.Name, Makefile.TargetType, Makefile.bHasProjectScriptPlugin, UObjectModules: Makefile.UObjectModules, ModuleInfoFileName: ModuleInfoFileName, WorkingSet: WorkingSet, UHTAdditionalArguments: Makefile.UHTAdditionalArguments);
 				}
 			}
 			return Makefile;
diff --git a/Engine/Source/Programs/UnrealBuildTool/System/ExternalExecution.cs b/Engine/Source/Programs/UnrealBuildTool/System/ExternalExecution.cs
--- a/Engine/Source/Programs/UnrealBuildTool/System/ExternalExecution.cs
+++ b/Engine/Source/Programs/UnrealBuildTool/System/ExternalExecution.cs
@@ -852,9 +852,8 @@ namespace UnrealBuildTool
 		/// <param name="BuildConfiguration">Build configuration</param>
 		/// <param name="UObjectModules">Modules that we generate headers for</param>
 		/// <param name="HeaderToolTimestampUtc">Timestamp for UHT</param>
-		/// <param name="bIsGatheringBuild"></param>
 		/// <returns>True if the code files are out of date</returns>
-		private static bool AreGeneratedCodeFilesOutOfDate(BuildConfiguration BuildConfiguration, List<UHTModuleInfo> UObjectModules, DateTime HeaderToolTimestampUtc, bool bIsGatheringBuild)
+		private static bool AreGeneratedCodeFilesOutOfDate(BuildConfiguration BuildConfiguration, List<UHTModuleInfo> UObjectModules, DateTime HeaderToolTimestampUtc)
 		{
 			// Get CoreUObject.init.gen.cpp timestamp.  If the source files are older than the CoreUObject generated code, we'll
 			// need to regenerate code for the module
@@ -968,28 +967,6 @@ namespace UnrealBuildTool
 						Log.TraceLog("UnrealHeaderTool needs to run because SavedTimestamp is older than HeaderFileTimestamp ({0}) for module {1}", HeaderFile.AbsolutePath, Module.ModuleName);
 						return true;
 					}
-
-					// When we're running in assembler mode, outdatedness cannot be inferred by checking the directory timestamp
-					// of the source headers.  We don't care if source files were added or removed in this mode, because we're only
-					// able to process the known UObject headers that are in the Makefile.  If UObject header files are added/removed,
-					// we expect the user to re-run GenerateProjectFiles which will force UBTMakefile outdatedness.
-					// @todo ubtmake: Possibly, we should never be doing this check these days.
-					//
-					// We don't need to do this check if using hot reload makefiles, since makefile out-of-date checks already handle it.
-					if (!BuildConfiguration.bUseUBTMakefiles && bIsGatheringBuild)
-					{
-						// Also check the timestamp on the directory the source file is in.  If the directory timestamp has
-						// changed, new source files may have been added or deleted.  We don't know whether the new/deleted
-						// files were actually UObject headers, but because we don't know all of the files we processed
-						// in the previous run, we need to assume our generated code is out of date if the directory timestamp
-						// is newer.
-						DateTime HeaderDirectoryTimestampUtc = new DirectoryInfo(Path.GetDirectoryName(HeaderFile.AbsolutePath)).LastWriteTimeUtc;
-						if (HeaderDirectoryTimestampUtc > SavedTimestampUtc)
-						{
-							Log.TraceLog("UnrealHeaderTool needs to run because the directory containing an existing header ({0}) has changed, and headers may have been added to or deleted from module {1}", HeaderFile.AbsolutePath, Module.ModuleName);
-							return true;
-						}
-					}
 				}
 			}
 
@@ -1110,7 +1087,7 @@ namespace UnrealBuildTool
 		/// Builds and runs the header tool and touches the header directories.
 		/// Performs any early outs if headers need no changes, given the UObject modules, tool path, game name, and configuration
 		/// </summary>
-		public static void ExecuteHeaderToolIfNecessary(BuildConfiguration BuildConfiguration, FileReference? ProjectFile, string TargetName, TargetType TargetType, bool bHasProjectScriptPlugin, List<UHTModuleInfo> UObjectModules, FileReference ModuleInfoFileName, bool bIsGatheringBuild, ISourceFileWorkingSet WorkingSet, string[]? UHTAdditionalArguments)
+		public static void ExecuteHeaderToolIfNecessary(BuildConfiguration BuildConfiguration, FileReference? ProjectFile, string TargetName, TargetType TargetType, bool bHasProjectScriptPlugin, List<UHTModuleInfo> UObjectModules, FileReference ModuleInfoFileName, ISourceFileWorkingSet WorkingSet, string[]? UHTAdditionalArguments)
 		{
 			if (ProgressWriter.bWriteMarkup)
 			{
@@ -1145,7 +1122,7 @@ namespace UnrealBuildTool
 				{
 					bUHTNeedsToRun = true;
 				}
-				else if(AreGeneratedCodeFilesOutOfDate(BuildConfiguration, UObjectModules, HeaderToolTimestampUtc, bIsGatheringBuild))
+				else if(AreGeneratedCodeFilesOutOfDate(BuildConfiguration, UObjectModules, HeaderToolTimestampUtc))
 				{
 					bUHTNeedsToRun = true;
 				}
@@ -1185,11 +1162,10 @@ namespace UnrealBuildTool
 				{
 					// Always build UnrealHeaderTool if header regeneration is required, unless we're running within an installed ecosystem or hot-reloading
 					if ((!Unreal.IsEngineInstalled() || bHasProjectScriptPlugin) &&
-						!BuildConfiguration.bDoNotBuildUHT &&
-						!(bHaveHeaderTool && !bIsGatheringBuild))
+						!BuildConfiguration.bDoNotBuildUHT)
 					{
 						// If it is out of date or not there it will be built.
-						// If it is there and up to date, it will add 0.8 seconds to the build time.
+						// If it is there and up to date, it will add 0.1 seconds or less to the build time.
 
 						// Which desktop platform do we need to compile UHT for?
 						UnrealTargetPlatform Platform = BuildHostPlatform.Current.Platform;
