UnrealBuildTool:

Update comments & method names - timestamps are recorded in the Timestamps file. The directory timestamp is no longer relevant.

#trivial

diff --git a/Engine/Source/Programs/UnrealBuildTool/System/ExternalExecution.cs b/Engine/Source/Programs/UnrealBuildTool/System/ExternalExecution.cs
--- a/Engine/Source/Programs/UnrealBuildTool/System/ExternalExecution.cs
+++ b/Engine/Source/Programs/UnrealBuildTool/System/ExternalExecution.cs
@@ -894,7 +894,7 @@ namespace UnrealBuildTool
 				}
 
 				// Grab our special "Timestamp" file that we saved after the last set of headers were generated.  This file
-				// actually contains the list of source files which contained UObjects, so that we can compare to see if any
+				// contains the list of source files which contained UObjects, so that we can compare to see if any
 				// UObject source files were deleted (or no longer contain UObjects), which means we need to run UHT even
 				// if no other source files were outdated
 				string TimestampFile = Path.Combine(GeneratedCodeDirectory, @"Timestamp");
@@ -1004,9 +1004,9 @@ namespace UnrealBuildTool
 		}
 
 		/// <summary>
-		/// Updates the intermediate include directory timestamps of all the passed in UObject modules
+		/// Updates the recorded timestamps of all the passed in UObject modules
 		/// </summary>
-		private static void UpdateDirectoryTimestamps(List<UHTModuleInfo> UObjectModules)
+		private static void UpdateTimestamps(List<UHTModuleInfo> UObjectModules)
 		{
 			Parallel.ForEach(UObjectModules, async Module =>
 			{
@@ -1019,10 +1019,6 @@ namespace UnrealBuildTool
 					{
 						if (GeneratedCodeDirectoryInfo.Exists)
 						{
-							// Touch the include directory since we have technically 'generated' the headers
-							// However, the headers might not be touched at all since that would cause the compiler to recompile everything
-							// We can't alter the directory timestamp directly, because this may throw exceptions when the directory is
-							// open in visual studio or windows explorer, so instead we create a blank file that will change the timestamp for us
 							FileReference TimestampFile = FileReference.Combine(new DirectoryReference(GeneratedCodeDirectoryInfo.FullName), "Timestamp");
 
 							// Save all of the UObject files to a timestamp file.  We'll load these on the next run to see if any new
@@ -1043,7 +1039,7 @@ namespace UnrealBuildTool
 					}
 					catch (Exception Exception)
 					{
-						throw new BuildException(Exception, "Couldn't touch header directories: " + Exception.Message);
+						throw new BuildException(Exception, "Couldn't write Timestamp file: " + Exception.Message);
 					}
 				}
 			});
@@ -1296,10 +1292,9 @@ namespace UnrealBuildTool
 
 				Progress.Write(2, 3);
 
-				// touch the directories
-				using(Timeline.ScopeEvent("ExternalExecution.UpdateDirectoryTimestamps()"))
+				using(Timeline.ScopeEvent("ExternalExecution.UpdateTimestamps()"))
 				{
-					UpdateDirectoryTimestamps(UObjectModules);
+					UpdateTimestamps(UObjectModules);
 				}
 
 				Progress.Write(3, 3);
