UnrealBuildTool:

Simplify ExternalExecution.ExecuteHeaderToolIfNecessary() interface

diff --git a/Engine/Source/Programs/UnrealBuildTool/Configuration/UEBuildTarget.cs b/Engine/Source/Programs/UnrealBuildTool/Configuration/UEBuildTarget.cs
--- a/Engine/Source/Programs/UnrealBuildTool/Configuration/UEBuildTarget.cs
+++ b/Engine/Source/Programs/UnrealBuildTool/Configuration/UEBuildTarget.cs
@@ -1733,8 +1733,7 @@ namespace UnrealBuildTool
 			// files that are injected as top level prerequisites.  If UHT only emitted included header files, we wouldn't need to run it during the Gather phase at all.
 			if (Makefile.UObjectModules.Count > 0)
 			{
-				FileReference ModuleInfoFileName = FileReference.Combine(ProjectIntermediateDirectory, TargetName + ".uhtmanifest");
-				ExternalExecution.ExecuteHeaderToolIfNecessary(BuildConfiguration, ProjectFile, TargetName, TargetType, bHasProjectScriptPlugin, Makefile.UObjectModules, ModuleInfoFileName, WorkingSet, Makefile.UHTAdditionalArguments);
+				ExternalExecution.ExecuteHeaderToolIfNecessary(BuildConfiguration, ProjectFile, Makefile, TargetName, WorkingSet);
 			}
 
 			// Find all the shared PCHs.
@@ -3426,9 +3425,6 @@ namespace UnrealBuildTool
 			// this for each file when generating their dependencies.
 			GlobalCompileEnvironment.bHackHeaderGenerator = (AppName == "UnrealHeaderTool");
 
-			// UHTLite
-			GlobalCompileEnvironment.bHackHeaderGenerator |= (AppName == "UHTLite");
-
 			bool bUseDebugCRT = GlobalCompileEnvironment.Configuration == CppConfiguration.Debug && Rules.bDebugBuildsActuallyUseDebugCRT;
 			GlobalCompileEnvironment.bUseDebugCRT = bUseDebugCRT;
 			GlobalCompileEnvironment.bEnableOSX109Support = Rules.bEnableOSX109Support;
diff --git a/Engine/Source/Programs/UnrealBuildTool/Modes/BuildMode.cs b/Engine/Source/Programs/UnrealBuildTool/Modes/BuildMode.cs
--- a/Engine/Source/Programs/UnrealBuildTool/Modes/BuildMode.cs
+++ b/Engine/Source/Programs/UnrealBuildTool/Modes/BuildMode.cs
@@ -753,8 +753,7 @@ namespace UnrealBuildTool
 				// If the target needs UHT to be run, we'll go ahead and do that now
 				if (Makefile.UObjectModules.Count > 0)
 				{
-					FileReference ModuleInfoFileName = FileReference.Combine(Makefile.ProjectIntermediateDirectory, TargetDescriptor.Name + ".uhtmanifest");
-					ExternalExecution.ExecuteHeaderToolIfNecessary(BuildConfiguration, TargetDescriptor.ProjectFile, TargetDescriptor.Name, Makefile.TargetType, Makefile.bHasProjectScriptPlugin, UObjectModules: Makefile.UObjectModules, ModuleInfoFileName: ModuleInfoFileName, WorkingSet: WorkingSet, UHTAdditionalArguments: Makefile.UHTAdditionalArguments);
+					ExternalExecution.ExecuteHeaderToolIfNecessary(BuildConfiguration, TargetDescriptor.ProjectFile, Makefile, TargetDescriptor.Name, WorkingSet);
 				}
 			}
 			return Makefile;
diff --git a/Engine/Source/Programs/UnrealBuildTool/System/ExternalExecution.cs b/Engine/Source/Programs/UnrealBuildTool/System/ExternalExecution.cs
--- a/Engine/Source/Programs/UnrealBuildTool/System/ExternalExecution.cs
+++ b/Engine/Source/Programs/UnrealBuildTool/System/ExternalExecution.cs
@@ -1083,7 +1083,7 @@ namespace UnrealBuildTool
 		/// Builds and runs the header tool and touches the header directories.
 		/// Performs any early outs if headers need no changes, given the UObject modules, tool path, game name, and configuration
 		/// </summary>
-		public static void ExecuteHeaderToolIfNecessary(BuildConfiguration BuildConfiguration, FileReference? ProjectFile, string TargetName, TargetType TargetType, bool bHasProjectScriptPlugin, List<UHTModuleInfo> UObjectModules, FileReference ModuleInfoFileName, ISourceFileWorkingSet WorkingSet, string[]? UHTAdditionalArguments)
+		public static void ExecuteHeaderToolIfNecessary(BuildConfiguration BuildConfiguration, FileReference? ProjectFile, TargetMakefile Makefile, string TargetName, ISourceFileWorkingSet WorkingSet)
 		{
 			if (ProgressWriter.bWriteMarkup)
 			{
@@ -1094,15 +1094,12 @@ namespace UnrealBuildTool
 				// We never want to try to execute the header tool when we're already trying to build it!
 				bool bIsBuildingUHT = TargetName.Equals("UnrealHeaderTool", StringComparison.InvariantCultureIgnoreCase);
 
-				// UHTLite
-				bIsBuildingUHT |= TargetName.Equals("UHTLite", StringComparison.InvariantCultureIgnoreCase);
-
 				string RootLocalPath = Unreal.RootDirectory.FullName;
 
 				UnrealTargetConfiguration UHTConfig = BuildConfiguration.bForceDebugUnrealHeaderTool ? UnrealTargetConfiguration.Debug : UnrealTargetConfiguration.Development;
 
 				// Figure out the receipt path
-				FileReference HeaderToolReceipt = GetHeaderToolReceiptFile(ProjectFile, UHTConfig, bHasProjectScriptPlugin);
+				FileReference HeaderToolReceipt = GetHeaderToolReceiptFile(ProjectFile, UHTConfig, Makefile.bHasProjectScriptPlugin);
 
 				// check if UHT is out of date
 				DateTime HeaderToolTimestampUtc = DateTime.MaxValue;
@@ -1118,12 +1115,13 @@ namespace UnrealBuildTool
 				{
 					bUHTNeedsToRun = true;
 				}
-				else if(AreGeneratedCodeFilesOutOfDate(BuildConfiguration, UObjectModules, HeaderToolTimestampUtc))
+				else if(AreGeneratedCodeFilesOutOfDate(BuildConfiguration, Makefile.UObjectModules, HeaderToolTimestampUtc))
 				{
 					bUHTNeedsToRun = true;
 				}
 
 				// Check we're not using a different version of UHT
+				FileReference ModuleInfoFileName = FileReference.Combine(Makefile.ProjectIntermediateDirectory, TargetName + ".uhtmanifest");
 				FileReference ToolInfoFile = ModuleInfoFileName.ChangeExtension(".uhtpath");
 				if(!bUHTNeedsToRun)
 				{
@@ -1148,16 +1146,16 @@ namespace UnrealBuildTool
 				// @todo ubtmake: Optimization: Ideally we could avoid having to generate this data in the case where UHT doesn't even need to run!  Can't we use the existing copy?  (see below use of Manifest)
 
 				List<UHTManifest.Module> Modules = new List<UHTManifest.Module>();
-				foreach(UHTModuleInfo UObjectModule in UObjectModules)
+				foreach(UHTModuleInfo UObjectModule in Makefile.UObjectModules)
 				{
 					Modules.Add(new UHTManifest.Module(UObjectModule));
 				}
-				UHTManifest Manifest = new UHTManifest(TargetName, TargetType, RootLocalPath, ExternalDependenciesFile.FullName, Modules);
+				UHTManifest Manifest = new UHTManifest(TargetName, Makefile.TargetType, RootLocalPath, ExternalDependenciesFile.FullName, Modules);
 
 				if (!bIsBuildingUHT && bUHTNeedsToRun)
 				{
 					// Always build UnrealHeaderTool if header regeneration is required, unless we're running within an installed ecosystem or hot-reloading
-					if ((!Unreal.IsEngineInstalled() || bHasProjectScriptPlugin) &&
+					if ((!Unreal.IsEngineInstalled() || Makefile.bHasProjectScriptPlugin) &&
 						!BuildConfiguration.bDoNotBuildUHT)
 					{
 						// If it is out of date or not there it will be built.
@@ -1182,7 +1180,7 @@ namespace UnrealBuildTool
 
 						// Add UHT plugins to UBT command line as external plugins
 						FileReference? ScriptProjectFile = null;
-						if(bHasProjectScriptPlugin && ProjectFile != null)
+						if(Makefile.bHasProjectScriptPlugin && ProjectFile != null)
 						{
 							ScriptProjectFile = ProjectFile;
 						}
@@ -1215,9 +1213,9 @@ namespace UnrealBuildTool
 					string CmdLine = (ProjectFile != null) ? "\"" + ProjectFile.FullName + "\"" : TargetName;
 					CmdLine += " \"" + ModuleInfoFileName + "\" -LogCmds=\"loginit warning, logexit warning, logdatabase error\" -Unattended -WarningsAsErrors";
 
-					if (UHTAdditionalArguments != null)
+					if (Makefile.UHTAdditionalArguments != null)
 					{
-						foreach (string UHTArgument in UHTAdditionalArguments)
+						foreach (string UHTArgument in Makefile.UHTAdditionalArguments)
 						{
 							CmdLine += " " + UHTArgument;
 						}
@@ -1282,7 +1280,7 @@ namespace UnrealBuildTool
 					// Otherwise UBT might not detect changes UHT made.
 					using(Timeline.ScopeEvent("ExternalExecution.ResetCachedHeaderInfo()"))
 					{
-						ResetCachedHeaderInfo(UObjectModules);
+						ResetCachedHeaderInfo(Makefile.UObjectModules);
 					}
 				}
 				else
@@ -1294,7 +1292,7 @@ namespace UnrealBuildTool
 
 				using(Timeline.ScopeEvent("ExternalExecution.UpdateTimestamps()"))
 				{
-					UpdateTimestamps(UObjectModules);
+					UpdateTimestamps(Makefile.UObjectModules);
 				}
 
 				Progress.Write(3, 3);
