UnrealBuildTool:

Remove unnecessary check for Makefile validity.

#trivial

diff --git a/Engine/Source/Programs/UnrealBuildTool/ProjectFiles/ProjectFileGenerator.cs b/Engine/Source/Programs/UnrealBuildTool/ProjectFiles/ProjectFileGenerator.cs
--- a/Engine/Source/Programs/UnrealBuildTool/ProjectFiles/ProjectFileGenerator.cs
+++ b/Engine/Source/Programs/UnrealBuildTool/ProjectFiles/ProjectFileGenerator.cs
@@ -347,11 +347,6 @@ namespace UnrealBuildTool
 		/// </summary>
 		public static DirectoryReference IntermediateProjectFilesPath = DirectoryReference.Combine(Unreal.EngineDirectory, "Intermediate", "ProjectFiles");
 
-		/// <summary>
-		/// Path to timestamp file, recording when was the last time projects were created.
-		/// </summary>
-		public static string ProjectTimestampFile = Path.Combine(IntermediateProjectFilesPath.FullName, "Timestamp");
-
 		/// <summary>
 		/// Global static new line string used by ProjectFileGenerator to generate project files.
 		/// </summary>
diff --git a/Engine/Source/Programs/UnrealBuildTool/System/TargetMakefile.cs b/Engine/Source/Programs/UnrealBuildTool/System/TargetMakefile.cs
--- a/Engine/Source/Programs/UnrealBuildTool/System/TargetMakefile.cs
+++ b/Engine/Source/Programs/UnrealBuildTool/System/TargetMakefile.cs
@@ -334,29 +334,9 @@ namespace UnrealBuildTool
 					return null;
 				}
 
-				// @todo ubtmake: This will only work if the directory timestamp actually changes with every single GPF.  Force delete existing files before creating new ones?  Eh... really we probably just want to delete + create a file in that folder
-				//			-> UPDATE: Seems to work OK right now though on Windows platform, maybe due to GUID changes
-				// @todo ubtmake: Some platforms may not save any files into this folder.  We should delete + generate a "touch" file to force the directory timestamp to be updated (or just check the timestamp file itself.  We could put it ANYWHERE, actually)
-
-				// Installed Build doesn't need to check engine projects for outdatedness
-				if (!Unreal.IsEngineInstalled())
-				{
-					if (DirectoryReference.Exists(ProjectFileGenerator.IntermediateProjectFilesPath))
-					{
-						DateTime EngineProjectFilesLastUpdateTime = new FileInfo(ProjectFileGenerator.ProjectTimestampFile).LastWriteTime;
-						if (MakefileInfo.LastWriteTime.CompareTo(EngineProjectFilesLastUpdateTime) < 0)
-						{
-							// Engine project files are newer than makefile
-							Log.TraceLog("Existing makefile is older than generated engine project files, ignoring it");
-							ReasonNotLoaded = "project files are newer";
-							return null;
-						}
-					}
-				}
-
-				// Check the game project directory too
 				if (ProjectFile != null)
 				{
+					// Check the game project
 					string ProjectFilename = ProjectFile.FullName;
 					FileInfo ProjectFileInfo = new FileInfo(ProjectFilename);
 					if (!ProjectFileInfo.Exists || MakefileInfo.LastWriteTime.CompareTo(ProjectFileInfo.LastWriteTime) < 0)
@@ -366,20 +346,6 @@ namespace UnrealBuildTool
 						ReasonNotLoaded = ".uproject file is newer";
 						return null;
 					}
-
-					DirectoryReference MasterProjectRelativePath = ProjectFile.Directory;
-					string GameIntermediateProjectFilesPath = Path.Combine(MasterProjectRelativePath.FullName, "Intermediate", "ProjectFiles");
-					if (Directory.Exists(GameIntermediateProjectFilesPath))
-					{
-						DateTime GameProjectFilesLastUpdateTime = new DirectoryInfo(GameIntermediateProjectFilesPath).LastWriteTime;
-						if (MakefileInfo.LastWriteTime.CompareTo(GameProjectFilesLastUpdateTime) < 0)
-						{
-							// Game project files are newer than makefile
-							Log.TraceLog("Makefile is older than generated game project files, ignoring it");
-							ReasonNotLoaded = "game project files are newer";
-							return null;
-						}
-					}
 				}
 
 				// Check to see if UnrealBuildTool was compiled more recently than the makefile
