UnrealBuildTool:

Make ConfigValueTracker immutable.
Propagate some nullability in and around TargetRules.

diff --git a/Engine/Source/Programs/Shared/EpicGames.Core/BinaryArchiveWriter.cs b/Engine/Source/Programs/Shared/EpicGames.Core/BinaryArchiveWriter.cs
--- a/Engine/Source/Programs/Shared/EpicGames.Core/BinaryArchiveWriter.cs
+++ b/Engine/Source/Programs/Shared/EpicGames.Core/BinaryArchiveWriter.cs
@@ -424,7 +424,7 @@ namespace EpicGames.Core
 		/// <param name="Dictionary">The dictionary to write</param>
 		/// <param name="WriteKey">Delegate used to read a single key</param>
 		/// <param name="WriteValue">Delegate used to read a single value</param>
-		public void WriteDictionary<K, V>(IDictionary<K, V> Dictionary, Action<K> WriteKey, Action<V> WriteValue)
+		public void WriteDictionary<K, V>(IReadOnlyDictionary<K, V> Dictionary, Action<K> WriteKey, Action<V> WriteValue)
 		{
 			if (Dictionary == null)
 			{
diff --git a/Engine/Source/Programs/UnrealBuildTool/Configuration/TargetRules.cs b/Engine/Source/Programs/UnrealBuildTool/Configuration/TargetRules.cs
--- a/Engine/Source/Programs/UnrealBuildTool/Configuration/TargetRules.cs
+++ b/Engine/Source/Programs/UnrealBuildTool/Configuration/TargetRules.cs
@@ -212,8 +212,10 @@ namespace UnrealBuildTool
 
 		/// <summary>
 		/// File containing this target
+		/// Nullable, but expected to be set before the constructor of this object is run see e.g. RulesAssembly.CreateTargetRulesInstance()
+		/// Every use of this property can be null-forgiven i.e. File!
 		/// </summary>
-		internal FileReference File { get; set; }
+		internal FileReference? File { get; set; }
 
 		/// <summary>
 		/// Platform that this target is being built for.
@@ -259,7 +261,7 @@ namespace UnrealBuildTool
 		/// <summary>
 		/// Tracks a list of config values read while constructing this target
 		/// </summary>
-		internal ConfigValueTracker ConfigValueTracker = new ConfigValueTracker();
+		internal readonly ConfigValueTracker ConfigValueTracker;
 
 		/// <summary>
 		/// Whether the target uses Steam.
@@ -1492,7 +1494,7 @@ namespace UnrealBuildTool
 				{
 					return BuildEnvironmentOverride.Value;
 				}
-				if (Type == TargetType.Program && ProjectFile != null && File.IsUnderDirectory(ProjectFile.Directory))
+				if (Type == TargetType.Program && ProjectFile != null && File!.IsUnderDirectory(ProjectFile.Directory))
 				{
 					return TargetBuildEnvironment.Unique;
 				}
@@ -1638,13 +1640,14 @@ namespace UnrealBuildTool
 		[ConfigSubObject]
 		public HoloLensTargetRules HoloLensPlatform;
 
-#nullable disable
 		/// <summary>
 		/// Constructor.
 		/// </summary>
 		/// <param name="Target">Information about the target being built</param>
 		public TargetRules(TargetInfo Target)
 		{
+			// it is expected that this.File has been set before this constructor is invoked
+
 			this.DefaultName = Target.Name;
 			this.Platform = Target.Platform;
 			this.Configuration = Target.Configuration;
@@ -1655,15 +1658,17 @@ namespace UnrealBuildTool
 			this.HoloLensPlatform = new HoloLensTargetRules(Target);
 
 			// Read settings from config files
+			Dictionary<ConfigDependencyKey, IReadOnlyList<string>?> ConfigValues = new Dictionary<ConfigDependencyKey, IReadOnlyList<string>?>();
 			foreach (object ConfigurableObject in GetConfigurableObjects())
 			{
-				ConfigCache.ReadSettings(DirectoryReference.FromFile(ProjectFile), Platform, ConfigurableObject, ConfigValueTracker);
+				ConfigCache.ReadSettings(DirectoryReference.FromFile(ProjectFile), Platform, ConfigurableObject, ConfigValues);
 				XmlConfig.ApplyTo(ConfigurableObject);
 				if (Target.Arguments != null)
 				{
 					Target.Arguments.ApplyTo(ConfigurableObject);
 				}
 			}
+			ConfigValueTracker = new ConfigValueTracker(ConfigValues);
 
 			// If we've got a changelist set, set that we're making a formal build
 			bFormalBuild = (Version.Changelist != 0 && Version.IsPromotedBuild);
@@ -1688,8 +1693,8 @@ namespace UnrealBuildTool
 			// Get the directory to use for crypto settings. We can build engine targets (eg. UHT) with 
 			// a project file, but we can't use that to determine crypto settings without triggering
 			// constant rebuilds of UHT.
-			DirectoryReference CryptoSettingsDir = DirectoryReference.FromFile(ProjectFile);
-			if (CryptoSettingsDir != null && File != null && !File.IsUnderDirectory(CryptoSettingsDir))
+			DirectoryReference? CryptoSettingsDir = DirectoryReference.FromFile(ProjectFile);
+			if (CryptoSettingsDir != null && !File!.IsUnderDirectory(CryptoSettingsDir))
 			{
 				CryptoSettingsDir = null;
 			}
@@ -1714,7 +1719,6 @@ namespace UnrealBuildTool
 				ProjectDefinitions.Add("IMPLEMENT_SIGNING_KEY_REGISTRATION()=");
 			}
 		}
-#nullable restore
 
 		/// <summary>
 		/// Formats an array of bytes as a sequence of values
@@ -1894,7 +1898,7 @@ namespace UnrealBuildTool
 						Diagnostics.Add(String.Format(FormatString, ModifiedSetting.Item1, ModifiedSetting.Item2));
 					}
 				}
-				Diagnostics.Add(String.Format("[Upgrade] Suppress this message by setting 'DefaultBuildSettings = BuildSettingsVersion.{1};' in {2}, and explicitly overriding settings that differ from the new defaults.", Version, (BuildSettingsVersion)(BuildSettingsVersion.Latest - 1), File.GetFileName()));
+				Diagnostics.Add(String.Format("[Upgrade] Suppress this message by setting 'DefaultBuildSettings = BuildSettingsVersion.{1};' in {2}, and explicitly overriding settings that differ from the new defaults.", Version, (BuildSettingsVersion)(BuildSettingsVersion.Latest - 1), File!.GetFileName()));
 				Diagnostics.Add("[Upgrade]");
 			}
 		}
@@ -1941,7 +1945,7 @@ namespace UnrealBuildTool
 
 		internal FileReference File
 		{
-			get { return Inner.File; }
+			get { return Inner.File!; }
 		}
 
 		public UnrealTargetPlatform Platform
diff --git a/Engine/Source/Programs/UnrealBuildTool/ProjectFiles/VisualStudio/VCProject.cs b/Engine/Source/Programs/UnrealBuildTool/ProjectFiles/VisualStudio/VCProject.cs
--- a/Engine/Source/Programs/UnrealBuildTool/ProjectFiles/VisualStudio/VCProject.cs
+++ b/Engine/Source/Programs/UnrealBuildTool/ProjectFiles/VisualStudio/VCProject.cs
@@ -1592,13 +1592,13 @@ namespace UnrealBuildTool
 
 					// Get the output directory
 					DirectoryReference RootOutputDirectory;
-					if (UProjectDirectory != null && (bShouldCompileMonolithic || TargetRulesObject.BuildEnvironment == TargetBuildEnvironment.Unique) && TargetRulesObject.File.IsUnderDirectory(UProjectDirectory))
+					if (UProjectDirectory != null && (bShouldCompileMonolithic || TargetRulesObject.BuildEnvironment == TargetBuildEnvironment.Unique) && TargetRulesObject.File!.IsUnderDirectory(UProjectDirectory))
 					{
-						RootOutputDirectory = UEBuildTarget.GetOutputDirectoryForExecutable(UProjectDirectory, TargetRulesObject.File);
+						RootOutputDirectory = UEBuildTarget.GetOutputDirectoryForExecutable(UProjectDirectory, TargetRulesObject.File!);
 					}
 					else
 					{
-						RootOutputDirectory = UEBuildTarget.GetOutputDirectoryForExecutable(Unreal.EngineDirectory, TargetRulesObject.File);
+						RootOutputDirectory = UEBuildTarget.GetOutputDirectoryForExecutable(Unreal.EngineDirectory, TargetRulesObject.File!);
 					}
 
 					// Get the output directory
diff --git a/Engine/Source/Programs/UnrealBuildTool/System/ConfigCache.cs b/Engine/Source/Programs/UnrealBuildTool/System/ConfigCache.cs
--- a/Engine/Source/Programs/UnrealBuildTool/System/ConfigCache.cs
+++ b/Engine/Source/Programs/UnrealBuildTool/System/ConfigCache.cs
@@ -294,8 +294,8 @@ namespace UnrealBuildTool
 		/// <param name="ProjectDir">Path to the project directory</param>
 		/// <param name="Platform">The platform being built</param>
 		/// <param name="TargetObject">Object to receive the settings</param>
-		/// <param name="Tracker">Tracks the set of config values that were retrieved. May be null.</param>
-		internal static void ReadSettings(DirectoryReference ProjectDir, UnrealTargetPlatform Platform, object TargetObject, ConfigValueTracker? Tracker)
+		/// <param name="ConfigValues">Will be populated with config values that were retrieved. May be null.</param>
+		internal static void ReadSettings(DirectoryReference? ProjectDir, UnrealTargetPlatform Platform, object TargetObject, Dictionary<ConfigDependencyKey, IReadOnlyList<string>?>? ConfigValues)
 		{
 			List<ConfigField> Fields = FindConfigFieldsForType(TargetObject.GetType());
 			foreach(ConfigField Field in Fields)
@@ -338,9 +338,10 @@ namespace UnrealBuildTool
 				}
 
 				// Save the dependency
-				if (Tracker != null)
+				if (ConfigValues != null)
 				{
-					Tracker.Add(Field.Attribute.ConfigType, ProjectDir, Platform, Field.Attribute.SectionName, KeyName, Values);
+					ConfigDependencyKey Key = new ConfigDependencyKey(Field.Attribute.ConfigType, ProjectDir, Platform, Field.Attribute.SectionName, KeyName);
+					ConfigValues[Key] = Values;
 				}
 			}
 		}
diff --git a/Engine/Source/Programs/UnrealBuildTool/System/ConfigValueTracker.cs b/Engine/Source/Programs/UnrealBuildTool/System/ConfigValueTracker.cs
--- a/Engine/Source/Programs/UnrealBuildTool/System/ConfigValueTracker.cs
+++ b/Engine/Source/Programs/UnrealBuildTool/System/ConfigValueTracker.cs
@@ -128,14 +128,14 @@ namespace UnrealBuildTool
 		/// <summary>
 		/// The dependencies list
 		/// </summary>
-		Dictionary<ConfigDependencyKey, IReadOnlyList<string>?> Dependencies;
+		readonly IReadOnlyDictionary<ConfigDependencyKey, IReadOnlyList<string>?> Dependencies;
 
 		/// <summary>
 		/// Constructor
 		/// </summary>
-		public ConfigValueTracker()
+		public ConfigValueTracker(IReadOnlyDictionary<ConfigDependencyKey, IReadOnlyList<string>?> ConfigValues)
 		{
-			Dependencies = new Dictionary<ConfigDependencyKey, IReadOnlyList<string>?>();
+			Dependencies = new Dictionary<ConfigDependencyKey, IReadOnlyList<string>?>(ConfigValues);
 		}
 
 		/// <summary>
@@ -156,21 +156,6 @@ namespace UnrealBuildTool
 			Writer.WriteDictionary(Dependencies, Key => Key.Write(Writer), Value => Writer.WriteList(Value, x => Writer.WriteString(x)));
 		}
 
-		/// <summary>
-		/// Adds a new configuration value
-		/// </summary>
-		/// <param name="Type">The config hierarchy type</param>
-		/// <param name="ProjectDir">The project directory</param>
-		/// <param name="Platform">The platform being built</param>
-		/// <param name="SectionName">Name of the config file section</param>
-		/// <param name="KeyName">Name of the config file key</param>
-		/// <param name="Values">Current values for this key</param>
-		public void Add(ConfigHierarchyType Type, DirectoryReference ProjectDir, UnrealTargetPlatform Platform, string SectionName, string KeyName, IReadOnlyList<string>? Values)
-		{
-			ConfigDependencyKey Key = new ConfigDependencyKey(Type, ProjectDir, Platform, SectionName, KeyName);
-			Dependencies[Key] = Values;
-		}
-
 		/// <summary>
 		/// Checks whether the list of dependencies is still valid
 		/// </summary>
