UnrealBuildTool:

TargetRules.File may be null in the TargetRules() constructor when run from AutomationTool.ProjectUtils.CompileAndLoadTargetsAssembly()

diff --git a/Engine/Source/Programs/UnrealBuildTool/Configuration/TargetRules.cs b/Engine/Source/Programs/UnrealBuildTool/Configuration/TargetRules.cs
--- a/Engine/Source/Programs/UnrealBuildTool/Configuration/TargetRules.cs
+++ b/Engine/Source/Programs/UnrealBuildTool/Configuration/TargetRules.cs
@@ -212,8 +212,6 @@ namespace UnrealBuildTool
 
 		/// <summary>
 		/// File containing this target
-		/// Nullable, but expected to be set before the constructor of this object is run see e.g. RulesAssembly.CreateTargetRulesInstance()
-		/// Every use of this property can be null-forgiven i.e. File!
 		/// </summary>
 		internal FileReference? File { get; set; }
 
@@ -1646,8 +1644,6 @@ namespace UnrealBuildTool
 		/// <param name="Target">Information about the target being built</param>
 		public TargetRules(TargetInfo Target)
 		{
-			// it is expected that this.File has been set before this constructor is invoked
-
 			this.DefaultName = Target.Name;
 			this.Platform = Target.Platform;
 			this.Configuration = Target.Configuration;
@@ -1694,7 +1690,7 @@ namespace UnrealBuildTool
 			// a project file, but we can't use that to determine crypto settings without triggering
 			// constant rebuilds of UHT.
 			DirectoryReference? CryptoSettingsDir = DirectoryReference.FromFile(ProjectFile);
-			if (CryptoSettingsDir != null && !File!.IsUnderDirectory(CryptoSettingsDir))
+			if (CryptoSettingsDir != null && File != null && !File.IsUnderDirectory(CryptoSettingsDir))
 			{
 				CryptoSettingsDir = null;
 			}
