UBT: Make sure deprecation warnings are treated as errors when -WarningsAsErrors is specified on the command line. Fixes a number of intermittent errors through Horde.

diff --git a/Engine/Source/Programs/UnrealBuildTool/Configuration/TargetRules.cs b/Engine/Source/Programs/UnrealBuildTool/Configuration/TargetRules.cs
--- a/Engine/Source/Programs/UnrealBuildTool/Configuration/TargetRules.cs
+++ b/Engine/Source/Programs/UnrealBuildTool/Configuration/TargetRules.cs
@@ -746,7 +746,7 @@ namespace UnrealBuildTool
 		/// Whether to include PerfCounters support.
 		/// </summary>
 		[RequiresUniqueBuildEnvironment]
-        public bool bWithPerfCounters
+		public bool bWithPerfCounters
 		{
 			get { return bWithPerfCountersOverride ?? (Type == TargetType.Editor || Type == TargetType.Server); }
 			set { bWithPerfCountersOverride = value; }
@@ -1009,14 +1009,26 @@ namespace UnrealBuildTool
 		/// <summary>
 		/// Default treatment of uncategorized warnings
 		/// </summary>
-		[XmlConfigFile(Category = "BuildConfiguration")]
-		public WarningLevel DefaultWarningLevel = WarningLevel.Warning;
+		public WarningLevel DefaultWarningLevel
+		{
+			get => DefaultWarningLevelPrivate ?? (bWarningsAsErrors ? WarningLevel.Error : WarningLevel.Error);
+			set => DefaultWarningLevelPrivate = value;
+		}
+
+		[XmlConfigFile(Category = "BuildConfiguration", Name = nameof(DefaultWarningLevel))]
+		private WarningLevel? DefaultWarningLevelPrivate;
 
 		/// <summary>
-		/// Whether to treat all warnings as errors. UE generally treats most warnings as errors, with the exception of deprecation warnings,
+		/// Level to report deprecation warnings as errors
 		/// </summary>
-		[XmlConfigFile(Category = "BuildConfiguration")]
-		public WarningLevel DeprecationWarningLevel = WarningLevel.Warning;
+		public WarningLevel DeprecationWarningLevel
+		{
+			get => DeprecationWarningLevelPrivate ?? DefaultWarningLevel;
+			set => DeprecationWarningLevelPrivate = value;
+		}
+
+		[XmlConfigFile(Category = "BuildConfiguration", Name = nameof(DeprecationWarningLevel))]
+		private WarningLevel? DeprecationWarningLevelPrivate;
 
 		/// <summary>
 		/// Forces shadow variable warnings to be treated as errors on platforms that support it.
diff --git a/Engine/Source/Programs/UnrealBuildTool/System/XmlConfig.cs b/Engine/Source/Programs/UnrealBuildTool/System/XmlConfig.cs
--- a/Engine/Source/Programs/UnrealBuildTool/System/XmlConfig.cs
+++ b/Engine/Source/Programs/UnrealBuildTool/System/XmlConfig.cs
@@ -449,11 +449,16 @@ namespace UnrealBuildTool
 			Element.MinOccurs = 0;
 			Element.MaxOccurs = 1;
 
-			if(Type == typeof(string))
+			if (TryGetNullableStructType(Type, out Type? InnerType))
+			{
+				Type = InnerType;
+			}
+
+			if (Type == typeof(string))
 			{
 				Element.SchemaTypeName = XmlSchemaType.GetBuiltInSimpleType(XmlTypeCode.String).QualifiedName;
 			}
-			else if(Type == typeof(bool) || Type == typeof(bool?))
+			else if(Type == typeof(bool))
 			{
 				Element.SchemaTypeName = XmlSchemaType.GetBuiltInSimpleType(XmlTypeCode.Boolean).QualifiedName;
 			}
@@ -545,7 +550,24 @@ namespace UnrealBuildTool
 				File.WriteAllText(Location.FullName, OutputText);
 			}
 		}
-		
+
+		/// <summary>
+		/// Tests whether a type is a nullable struct, and extracts the inner type if it is
+		/// </summary>
+		static bool TryGetNullableStructType(Type Type, [NotNullWhen(true)] out Type? InnerType)
+		{
+			if (Type.IsGenericType && Type.GetGenericTypeDefinition() == typeof(Nullable<>))
+			{
+				InnerType = Type.GetGenericArguments()[0];
+				return true;
+			}
+			else
+			{
+				InnerType = null;
+				return false;
+			}
+		}
+
 		/// <summary>
 		/// Reads an XML config file and merges it to the given cache
 		/// </summary>
@@ -576,10 +598,14 @@ namespace UnrealBuildTool
 						{
 							// Parse the corresponding value
 							object Value;
-							if(Field.FieldType == typeof(string[]))
+							if (Field.FieldType == typeof(string[]))
 							{
 								Value = KeyElement.ChildNodes.OfType<XmlElement>().Where(x => x.Name == "Item").Select(x => x.InnerText).ToArray();
 							}
+							else if (TryGetNullableStructType(Field.FieldType, out Type? StructType))
+							{
+								Value = ParseValue(StructType, KeyElement.InnerText);
+							}
 							else
 							{
 								Value = ParseValue(Field.FieldType, KeyElement.InnerText);
